[mutable] Move Mutable plugin from the restricted folder to the experimental plugins folder.
Remove plugin folder remap from ini file.
Corrections for inclusive wording from the latest coding standard.

diff --git a/Engine/Plugins/Experimental/Mutable/Config/DefaultMutable.ini b/Engine/Plugins/Experimental/Mutable/Config/DefaultMutable.ini
new file mode 100644
--- /dev/null
+++ b/Engine/Plugins/Experimental/Mutable/Config/DefaultMutable.ini
@@ -0,0 +1,46 @@
+; Settings for the Mutable plugin
+
+[CoreRedirects]
+; None yet
+
+
+[Features]
+
+; If this is enabled, CustomizableObjects will be allowed to use 8 bone influences per vertex. 
+; Otherwise only the default 4 will be used. This is used at object compile time.
+bExtraBoneInfluencesEnabled=False
+
+;--------------------------------------------------------------------------------------------------
+; These are scalability settings that can be defined in any xxxScalability.ini file.
+[ScalabilitySettings]
+
+; If different than 0, limit the amount of memory (in KB) to use to cache streaming data when 
+; building characters. 0 means no limit, -1 means use default (10Mb in PC and 3Mb in consoles).
++CVars=b.MutableStreamingMemory=10000
+
+; Override in platform scalability .ini files
+; b.MutableStreamingMemory=3000
+
+[ViewDistanceQuality@0]
+; If different than 0, limit the number of mutable instances with full LODs to have at a given time.
++CVars=b.NumGeneratedInstancesLimit=7
+
+; If different than 0, limit the number of mutable instances with LOD 1 to have at a given time.
++CVars=b.NumGeneratedInstancesLimitLOD1=0
+
+; If different than 0, limit the number of mutable instances with LOD 2 to have at a given time.
++CVars=b.NumGeneratedInstancesLimitLOD2=0
+
+; If NumGeneratedInstancesLimit is different than 0, sets the distance at which the system will fix 
+; the LOD of an instance to the lowest res one (LOD2) to prevent unnecessary LOD changes and memory 
+; consumption
++CVars=b.DistanceForFixedLOD2=0
+
+[ViewDistanceQuality@1]
++CVars=b.NumGeneratedInstancesLimit=0
+[ViewDistanceQuality@2]
++CVars=b.NumGeneratedInstancesLimit=0
+[ViewDistanceQuality@3]
++CVars=b.NumGeneratedInstancesLimit=0
+[ViewDistanceQuality@Cine]
++CVars=b.NumGeneratedInstancesLimit=0
diff --git a/Engine/Plugins/Experimental/Mutable/Source/CustomizableObject/Classes/CustomizableInstanceLODManagement.h b/Engine/Plugins/Experimental/Mutable/Source/CustomizableObject/Classes/CustomizableInstanceLODManagement.h
new file mode 100644
--- /dev/null
+++ b/Engine/Plugins/Experimental/Mutable/Source/CustomizableObject/Classes/CustomizableInstanceLODManagement.h
@@ -0,0 +1,68 @@
+// Copyright Epic Games, Inc. All Rights Reserved.
+
+#pragma once
+#include "HAL/IConsoleManager.h"
+#include "GameFramework/Actor.h"
+
+#include "CustomizableInstanceLODManagement.generated.h"
+
+
+// This is an abstract base class, override it to create a new Instance LOD management system and register with UCustomizableObjectSystem::SetInstanceLODManagement
+UCLASS(Blueprintable, BlueprintType)
+class CUSTOMIZABLEOBJECT_API UCustomizableInstanceLODManagementBase : public UObject
+{
+	GENERATED_BODY()
+public:
+	UCustomizableInstanceLODManagementBase() : UObject() {};
+	virtual ~UCustomizableInstanceLODManagementBase() {};
+
+	virtual void UpdateInstanceDistsAndLODs() { check(0 && "You must override this"); };
+
+	virtual int32 GetNumGeneratedInstancesLimitFullLODs() { check(0 && "You must override this"); return 0; };
+	virtual int32 GetNumGeneratedInstancesLimitLOD1() { check(0 && "You must override this"); return 0; };
+	virtual int32 GetNumGeneratedInstancesLimitLOD2() { check(0 && "You must override this"); return 0; };
+
+	virtual bool IsOnlyUpdateCloseCustomizableObjectsEnabled() const { check(0 && "You must override this"); return false; };
+	virtual float GetOnlyUpdateCloseCustomizableObjectsDist() const { check(0 && "You must override this"); return 0.f; };
+};
+
+
+UCLASS(Blueprintable, BlueprintType)
+class CUSTOMIZABLEOBJECT_API UCustomizableInstanceLODManagement : public UCustomizableInstanceLODManagementBase
+{
+	GENERATED_BODY()
+public:
+	UCustomizableInstanceLODManagement();
+	virtual ~UCustomizableInstanceLODManagement();
+
+	virtual void UpdateInstanceDistsAndLODs() override;
+
+	virtual int32 GetNumGeneratedInstancesLimitFullLODs() override;
+	virtual int32 GetNumGeneratedInstancesLimitLOD1() override;
+	virtual int32 GetNumGeneratedInstancesLimitLOD2() override;
+
+	// Used to define the actors that will be considered the centers of LOD distance calculations. If left empty it will revert to the default, the first player controller actor
+	void AddViewCenter(const AActor* const InCentralActor) { ViewCenters.Add(TWeakObjectPtr<const AActor>(InCentralActor)); };
+	void RemoveViewCenter(const AActor* const InCentralActor) { ViewCenters.Remove(TWeakObjectPtr<const AActor>(InCentralActor)); };
+	void ClearViewCenters() { ViewCenters.Empty(); }
+
+	// Sets how many of the nearest instances to the player will have updates with priority over LOD changes, only works if LODs are enabled with b.NumGeneratedInstancesLimitLODs
+	void SetNumberOfPriorityUpdateInstances(int32 NumPriorityUpdateInstances);
+	int32 GetNumberOfPriorityUpdateInstances() const;
+
+	virtual bool IsOnlyUpdateCloseCustomizableObjectsEnabled() const override;
+	virtual float GetOnlyUpdateCloseCustomizableObjectsDist() const override;
+	void EnableOnlyUpdateCloseCustomizableObjects(float CloseDist);
+	void DisableOnlyUpdateCloseCustomizableObjects();
+
+private:
+	// The list of actors that define a view radius (CloseDist) around them used for LOD management
+	TSet<TWeakObjectPtr<const AActor>> ViewCenters;
+	
+	// Sets how many of the nearest instances to the player will have updates with priority over LOD changes
+	int32 NumPriorityUpdateInstances = 1;
+
+	bool bOnlyUpdateCloseCustomizableObjects;
+	float CloseCustomizableObjectsDist;	
+};
+
diff --git a/Engine/Plugins/Experimental/Mutable/Source/CustomizableObject/Classes/CustomizableObject.h b/Engine/Plugins/Experimental/Mutable/Source/CustomizableObject/Classes/CustomizableObject.h
new file mode 100644
--- /dev/null
+++ b/Engine/Plugins/Experimental/Mutable/Source/CustomizableObject/Classes/CustomizableObject.h
@@ -0,0 +1,1323 @@
+// Copyright Epic Games, Inc. All Rights Reserved.
+ 
+#pragma once
+
+#include "Engine/Texture.h"
+#include "EdGraph/EdGraphPin.h"
+#include "EdGraph/EdGraph.h"
+#include "EdGraph/EdGraphNode.h"
+#include "CustomizableObjectIdentifier.h"
+#include "CustomizableObjectParameterTypeDefinitions.h"
+#include "CustomizableObjectClothingTypes.h"
+#include "CustomizableObjectUIData.h"
+#include "Animation/AnimInstance.h"
+
+#if WITH_EDITORONLY_DATA
+#include "PerPlatformProperties.h"
+#include "PerQualityLevelProperties.h"
+#endif
+
+
+#include "CustomizableObject.generated.h"
+
+class UCustomizableObject;
+
+DECLARE_MULTICAST_DELEGATE(FPostCompileDelegate)
+ 
+// Forward declaration of Mutable classes necessary for the interface
+namespace mu
+{
+	class Model;
+	class Parameters;
+}
+
+
+CUSTOMIZABLEOBJECT_API DECLARE_LOG_CATEGORY_EXTERN(LogMutable, Warning, All);
+
+USTRUCT()
+struct FFParameterOptionsTags
+{
+	GENERATED_USTRUCT_BODY()
+
+	/** List of tags of a Parameter Options */
+	UPROPERTY(Category = "CustomizablePopulation", EditAnywhere)
+	TArray<FString> Tags;
+};
+
+
+USTRUCT()
+struct FParameterTags
+{
+public:
+
+	GENERATED_USTRUCT_BODY()
+
+	/** List of tags of a parameter */
+	UPROPERTY(Category = "CustomizablePopulation", EditAnywhere)
+	TArray<FString> Tags;
+
+	/** Map of options available for a parameter can have and their tags */
+	UPROPERTY(Category = "CustomizablePopulation", EditAnywhere)
+	TMap<FString, FFParameterOptionsTags> ParameterOptions;
+};
+
+
+UENUM()
+enum class ECustomizableObjectRelevancy : uint8
+{
+	// 
+	All = 0 UMETA(DisplayName = "Relevant for client and server"),
+	// 
+	ClientOnly = 1 UMETA(DisplayName = "Only necessary on clients")
+};
+
+
+// This is used to hide Mutable SDK members in the public headers.
+class FCustomizableObjectPrivateData;
+
+USTRUCT()
+struct FProfileParameterDat
+{
+	GENERATED_USTRUCT_BODY()
+
+	UPROPERTY()
+	FString ProfileName;
+
+	UPROPERTY()
+	TArray<FCustomizableObjectBoolParameterValue> BoolParameters;
+
+	////
+	UPROPERTY()
+	TArray<FCustomizableObjectIntParameterValue> IntParameters;
+
+	////
+	UPROPERTY()
+	TArray<FCustomizableObjectFloatParameterValue> FloatParameters;
+
+	////
+	UPROPERTY()
+	TArray<FCustomizableObjectTextureParameterValue> TextureParameters;
+
+	////
+	UPROPERTY()
+	TArray<FCustomizableObjectVectorParameterValue> VectorParameters;
+
+	////
+	UPROPERTY()
+	TArray<FCustomizableObjectProjectorParameterValue> ProjectorParameters;
+};
+
+
+USTRUCT()
+struct FCompilationOptions
+{
+	GENERATED_USTRUCT_BODY()
+
+	// Flag to know if texture compression should be enabled
+	UPROPERTY()
+	bool bTextureCompression = true;
+
+	// From 0 to 3
+	UPROPERTY()
+	int32 OptimizationLevel = 1;
+
+	// If enable, multi-threaded compilation will be used when compiling this models.
+	// For some memory-heavy models, this may need to be disabled, since it may multiply the
+	// use of memory during compilation for each thread.
+	UPROPERTY()
+	bool bUseParallelCompilation = true;
+
+	// Use the disk to store intermediate compilation data. This slows down the object compilation
+	// but it may be necessary for huge objects.
+	UPROPERTY()
+	bool bUseDiskCompilation = false;
+
+	// Did we have the extra bones enabled when we compiled?
+	bool bExtraBoneInfluencesEnabled = false;
+
+	// Compiling for cook
+	bool bIsCooking = false;
+
+	// Save compilation data for cook on disk 
+	bool bSaveCookedDataToDisk = false;
+
+	// This can be set for additional settings
+	const class ITargetPlatform* TargetPlatform = nullptr;
+
+	// Used to force a check to compare the guids stored in the base object compared to the children to see if any of them has changes and thus mark the object as modified.
+	bool bCheckChildrenGuids = false;
+
+	// Used to prevent the current model from being stored in the editor's streamed data and cache when the compilation finishes
+	bool bDontUpdateStreamedDataAndCache = false;
+
+	// Used to enable the use of real time morph targets.
+	bool bRealTimeMorphTargetsEnabled = false;
+
+	// Used to enable the use of clothing.
+	bool bClothingEnabled = false;
+
+	// Used to reduce the number of notifications when compiling objects
+	bool bSilentCompilation = true;
+
+};
+
+
+UENUM()
+enum class ECustomizableObjectCompilationState : uint8
+{
+	//
+	None,
+	// 
+	InProgress,
+	//
+	Completed,
+	//
+	Failed
+};
+
+
+class FCustomizableObjectCompilerBase
+{
+public:
+
+	FCustomizableObjectCompilerBase() {};
+
+	// Ensure virtual destruction
+	virtual ~FCustomizableObjectCompilerBase() {};
+	
+	virtual void Compile(UCustomizableObject& Object, const FCompilationOptions& Options, bool bAsync) {};
+
+	virtual bool Tick() { return false; }
+	virtual void ForceFinishCompilation() {};
+
+	// Return true if this object doesn't reference a parent object.
+	virtual bool IsRootObject(const class UCustomizableObject* Object) const { return true; }
+
+	/** Returns the Customizable Object that does start the CO tree */
+	virtual UCustomizableObject* GetRootObject( class UCustomizableObject* Object) = 0;
+
+	/** Provides the caller with the warning and error messages produced during compilation */
+	virtual void GetCompilationMessages(TArray<FText>& OutWarningMessages, TArray<FText>& OutErrorMessages) const = 0;
+	
+	virtual ECustomizableObjectCompilationState GetCompilationState() const { return ECustomizableObjectCompilationState::None; }
+};
+
+
+// Warning! MutableCompiledDataHeader must be the first data serialized in a stream
+struct MutableCompiledDataStreamHeader
+{
+	int32 InternalVersion;
+	FGuid VersionId;
+
+	MutableCompiledDataStreamHeader() { }
+	MutableCompiledDataStreamHeader(int32 InInternalVersion, FGuid InVersionId) : InternalVersion(InInternalVersion), VersionId(InVersionId) { }
+
+	friend FArchive& operator<<(FArchive& Ar, MutableCompiledDataStreamHeader& Header)
+	{
+		Ar << Header.InternalVersion;
+		Ar << Header.VersionId;
+
+		return Ar;
+	}
+};
+
+
+USTRUCT()
+struct FMutableModelImageProperties
+{
+	GENERATED_USTRUCT_BODY()
+
+	FMutableModelImageProperties()
+		: Filter(TF_Default)
+		, SRGB(0)
+		, FlipGreenChannel(0)
+		, LODBias(0)
+		, LODGroup(TEXTUREGROUP_World)
+		, AddressX(TA_Clamp)
+		, AddressY(TA_Clamp)
+	{}
+
+	FMutableModelImageProperties(const FString& InTextureParameterName, TextureFilter InFilter, uint32 InSRGB, 
+		uint32 InFlipGreenChannel, int32 InLODBias, TEnumAsByte<enum TextureGroup> InLODGroup,
+		TEnumAsByte<enum TextureAddress> InAddressX, TEnumAsByte<enum TextureAddress> InAddressY)
+		: TextureParameterName(InTextureParameterName)
+		, Filter(InFilter)
+		, SRGB(InSRGB)
+		, FlipGreenChannel(InFlipGreenChannel)
+		, LODBias(InLODBias)
+		, LODGroup(InLODGroup)
+		, AddressX(InAddressX)
+		, AddressY(InAddressY)
+	{}
+
+	// Name in the material.
+	UPROPERTY()
+	FString TextureParameterName;
+
+	UPROPERTY()
+	TEnumAsByte<enum TextureFilter> Filter;
+
+	UPROPERTY()
+	uint32 SRGB : 1;
+
+	UPROPERTY()
+	uint32 FlipGreenChannel : 1;
+
+	UPROPERTY()
+	int32 LODBias;
+
+	UPROPERTY()
+	TEnumAsByte<enum TextureGroup> LODGroup;
+
+	UPROPERTY()
+	TEnumAsByte<enum TextureAddress> AddressX;
+
+	UPROPERTY()
+	TEnumAsByte<enum TextureAddress> AddressY;
+
+	bool operator!=(const FMutableModelImageProperties& rhs)
+	{
+		return
+			TextureParameterName != rhs.TextureParameterName ||
+			Filter != rhs.Filter ||
+			SRGB != rhs.SRGB ||
+			FlipGreenChannel != rhs.FlipGreenChannel ||
+			LODBias != rhs.LODBias ||
+			LODGroup != rhs.LODGroup ||
+			AddressX != rhs.AddressX ||
+			AddressY != rhs.AddressY;
+	}
+
+	friend FArchive& operator<<(FArchive& Ar, FMutableModelImageProperties& ImageProps)
+	{
+		Ar << ImageProps.TextureParameterName;
+		Ar << ImageProps.Filter;
+
+		// Bitfields don't serialize automatically with FArchive
+		if (Ar.IsLoading())
+		{
+			int32 Aux = 0;
+			Ar << Aux;
+			ImageProps.SRGB = Aux;
+
+			Aux = 0;
+			Ar << Aux;
+			ImageProps.FlipGreenChannel = Aux;
+		}
+		else
+		{
+			int32 Aux = ImageProps.SRGB;
+			Ar << Aux;
+
+			Aux = ImageProps.FlipGreenChannel;
+			Ar << Aux;
+		}
+
+		Ar << ImageProps.LODBias;
+		Ar << ImageProps.LODGroup;
+
+		Ar << ImageProps.AddressX;
+		Ar << ImageProps.AddressY;
+
+		return Ar;
+	}
+};
+
+
+USTRUCT()
+struct FMutableModelParameterValue
+{
+	GENERATED_USTRUCT_BODY()
+
+	FMutableModelParameterValue() = default;
+
+	UPROPERTY()
+	FString Name;
+
+	UPROPERTY()
+	int Value = 0;
+};
+
+
+USTRUCT()
+struct FMutableModelParameterProperties
+{
+	GENERATED_USTRUCT_BODY()
+
+	FMutableModelParameterProperties() = default;
+	FString Name;
+
+	UPROPERTY()
+	EMutableParameterType Type = EMutableParameterType::None;
+
+	UPROPERTY()
+	int ImageDescriptionCount = 0;
+
+	UPROPERTY()
+	TArray<FMutableModelParameterValue> PossibleValues;
+};
+
+
+// Represents a texture used for masking-out areas of an object from projectors
+USTRUCT()
+struct FMaskOutTexture
+{
+	GENERATED_USTRUCT_BODY()
+
+	bool operator ==(const FMaskOutTexture& Other) const
+	{ 
+		return SizeX == Other.SizeX && SizeY == Other.SizeY && Data == Other.Data;
+	}
+
+	void SetTextureSize(int32 InSizeX, int32 InSizeY)
+	{
+		SizeX = InSizeX;
+		SizeY = InSizeY;
+		Data.SetNumUninitialized(FMath::DivideAndRoundUp(InSizeX * InSizeY, NumBitsPerDWORD)); 
+	}
+
+	int32 GetSizeX() const { return SizeX; }
+	int32 GetSizeY() const { return SizeY; }
+
+	FBitReference GetTexelReference(int32 Index)
+	{
+		check(Index >= 0 && Index < SizeX * SizeY);
+		return FBitReference(Data[Index / NumBitsPerDWORD], 1 << (Index & (NumBitsPerDWORD - 1)));
+	}
+
+private:
+	UPROPERTY()
+	int32 SizeX = 0;
+
+	UPROPERTY()
+	int32 SizeY = 0;
+
+	UPROPERTY()
+	TArray<uint32> Data; // Only alpha channel of PF_R8G8B8A8 is stored as binary bits
+};
+
+
+USTRUCT()
+struct FMorphTargetInfo
+{
+	GENERATED_USTRUCT_BODY()
+
+	UPROPERTY()
+	FName Name;
+
+	UPROPERTY()
+	int32 LodNum = 0;
+
+	friend FArchive& operator<<(FArchive& Ar, FMorphTargetInfo& Info)
+	{
+		Ar << Info.Name;
+		Ar << Info.LodNum;
+
+		return Ar;
+	}
+};
+
+
+USTRUCT()
+struct FMorphTargetVertexData
+{
+	GENERATED_USTRUCT_BODY()
+
+	UPROPERTY()
+	FVector3f PositionDelta = FVector3f::ZeroVector;
+
+	UPROPERTY()
+	FVector3f TangentZDelta = FVector3f::ZeroVector;
+	
+	UPROPERTY()
+	int32 MorphIndex = 0;
+
+	// Unused padding so memory footprint is the same the aligned.
+	int32 Padding = 0;
+
+	friend FArchive& operator<<(FArchive& Ar, FMorphTargetVertexData& Data)
+	{
+		Ar << Data.PositionDelta;
+		Ar << Data.TangentZDelta;
+		Ar << Data.MorphIndex;
+
+		return Ar;
+	}
+};
+template<> struct TCanBulkSerialize<FMorphTargetVertexData> { enum { Value = true }; };
+
+static_assert(sizeof(FMorphTargetVertexData) == 32, "");
+
+
+// A USTRUCT version of FMeshToMeshVertData in SkeletalMeshTypes.h
+// We are taking advantage of the padding data to store from which asset this data comes from
+// maintaining the same memory footprint than the original.
+USTRUCT()
+struct FCustomizableObjectMeshToMeshVertData
+{
+	GENERATED_USTRUCT_BODY()
+
+	FCustomizableObjectMeshToMeshVertData() = default;
+
+	explicit FCustomizableObjectMeshToMeshVertData(const FMeshToMeshVertData& Original)
+		: Weight(Original.Weight)
+	{
+		for (int i = 0; i < 4; ++i)
+		{
+			PositionBaryCoordsAndDist[i] = Original.PositionBaryCoordsAndDist[i];
+			NormalBaryCoordsAndDist[i] = Original.NormalBaryCoordsAndDist[i];
+			TangentBaryCoordsAndDist[i] = Original.TangentBaryCoordsAndDist[i];
+			SourceMeshVertIndices[i] = Original.SourceMeshVertIndices[i];
+		}
+	}
+		
+
+	explicit operator FMeshToMeshVertData() const
+	{
+		FMeshToMeshVertData ReturnValue;
+
+		for (int i = 0; i < 4; ++i)
+		{
+			ReturnValue.PositionBaryCoordsAndDist[i] = PositionBaryCoordsAndDist[i];
+			ReturnValue.NormalBaryCoordsAndDist[i] = NormalBaryCoordsAndDist[i];
+			ReturnValue.TangentBaryCoordsAndDist[i] = TangentBaryCoordsAndDist[i];
+			ReturnValue.SourceMeshVertIndices[i] = SourceMeshVertIndices[i];
+		}
+		ReturnValue.Weight = Weight;
+
+		return ReturnValue;
+	}
+
+	
+	// Barycentric coords and distance along normal for the position of the final vert
+	UPROPERTY()
+	float PositionBaryCoordsAndDist[4] = {0.f};
+
+	// Barycentric coords and distance along normal for the location of the unit normal endpoint
+	// Actual normal = ResolvedNormalPosition - ResolvedPosition
+	UPROPERTY()
+	float NormalBaryCoordsAndDist[4] = {0.f};
+
+	// Barycentric coords and distance along normal for the location of the unit Tangent endpoint
+	// Actual normal = ResolvedNormalPosition - ResolvedPosition
+	UPROPERTY()
+	float TangentBaryCoordsAndDist[4] = {0.f};
+
+	// Contains the 3 indices for verts in the source mesh forming a triangle, the last element
+	// is a flag to decide how the skinning works, 0xffff uses no simulation, and just normal
+	// skinning, anything else uses the source mesh and the above skin data to get the final position
+	UPROPERTY()
+	uint16	 SourceMeshVertIndices[4] = {0u, 0u, 0u, 0u};
+
+	UPROPERTY()
+	float Weight = 0.0f;
+
+	// Dummy for alignment (8 bytes). Originally not used.
+	UPROPERTY()
+	int16 SourceAssetIndex = 0;
+	
+	//
+	UPROPERTY()
+	int16 SourceAssetLodIndex = 0;
+
+	/**
+	 * Serializer
+	 *
+	 * @param Ar - archive to serialize with
+	 * @param V - vertex to serialize
+	 * @return archive that was used
+	 */
+	friend FArchive& operator<<(FArchive& Ar, FCustomizableObjectMeshToMeshVertData& V)
+	{
+		Ar << V.PositionBaryCoordsAndDist[0];
+		Ar << V.PositionBaryCoordsAndDist[1];
+		Ar << V.PositionBaryCoordsAndDist[2];
+		Ar << V.PositionBaryCoordsAndDist[3];
+		Ar << V.NormalBaryCoordsAndDist[0];
+		Ar << V.NormalBaryCoordsAndDist[1];
+		Ar << V.NormalBaryCoordsAndDist[2];
+		Ar << V.NormalBaryCoordsAndDist[3];
+		Ar << V.TangentBaryCoordsAndDist[0];
+		Ar << V.TangentBaryCoordsAndDist[1];
+		Ar << V.TangentBaryCoordsAndDist[2];
+		Ar << V.TangentBaryCoordsAndDist[3];
+
+		Ar << V.SourceMeshVertIndices[0];
+		Ar << V.SourceMeshVertIndices[1];
+		Ar << V.SourceMeshVertIndices[2];
+		Ar << V.SourceMeshVertIndices[3];
+
+		Ar << V.Weight;
+
+		Ar << V.SourceAssetIndex;
+		Ar << V.SourceAssetLodIndex;
+		return Ar;
+	}
+};
+template<> struct TCanBulkSerialize<FCustomizableObjectMeshToMeshVertData> { enum { Value = true }; };
+
+UCLASS()
+class CUSTOMIZABLEOBJECT_API UMutableMaskOutCache : public UObject
+{
+public:
+	GENERATED_BODY()
+
+	UPROPERTY()
+	TMap<FString, FString> Materials; // Maps a UMaterial's asset path to a UTexture's asset path
+
+	UPROPERTY()
+	TMap<FString, FMaskOutTexture> Textures; // Maps a UTexture's asset path to the cached mask-out texture data
+};
+
+
+USTRUCT()
+struct CUSTOMIZABLEOBJECT_API FMutableStreamableBlock
+{
+	GENERATED_USTRUCT_BODY()
+
+	UPROPERTY()
+	uint16 FileIndex = 0;
+
+	UPROPERTY()
+	uint64 Offset = 0;
+
+	UPROPERTY()
+	uint32 Size = 0;
+
+
+	friend FArchive& operator<<(FArchive& Ar, FMutableStreamableBlock& Data)
+	{
+		Ar << Data.FileIndex;
+		Ar << Data.Offset;
+		Ar << Data.Size;
+
+		return Ar;
+	}
+};
+template<> struct TCanBulkSerialize<FMutableStreamableBlock> { enum { Value = true }; };
+
+
+struct CUSTOMIZABLEOBJECT_API FMutableCachedPlatformData
+{
+	/** */
+	TArray64<uint8> ModelData;
+
+	/** */
+	TArray64<uint8> StreamableData;
+};
+
+
+struct FMutableRefLODInfo
+{
+	float ScreenSize = 0.f;
+	float LODHysteresis = 0.f;
+	bool bSupportUniformlyDistributedSampling = false;
+	bool bAllowCPUAccess = false;
+
+	friend FArchive& operator<<(FArchive& Ar, FMutableRefLODInfo& Data)
+	{
+		Ar << Data.ScreenSize;
+		Ar << Data.LODHysteresis;
+		Ar << Data.bSupportUniformlyDistributedSampling;
+		Ar << Data.bAllowCPUAccess;
+
+		return Ar;
+	}
+};
+
+
+struct FMutableRefLODRenderData
+{
+	bool bIsLODOptional = false;
+	bool bStreamedDataInlined = false;
+	
+	friend FArchive& operator<<(FArchive& Ar, FMutableRefLODRenderData& Data)
+	{
+		Ar << Data.bIsLODOptional;
+		Ar << Data.bStreamedDataInlined;
+
+		return Ar;
+	}
+};
+
+
+struct FMutableRefLODData
+{
+	FMutableRefLODInfo LODInfo;
+
+	FMutableRefLODRenderData RenderData;
+	
+	friend FArchive& operator<<(FArchive& Ar, FMutableRefLODData& Data)
+	{
+		Ar << Data.LODInfo;
+		Ar << Data.RenderData;
+		
+		return Ar;
+	}
+};
+
+
+struct FMutableRefSocket
+{
+	FName SocketName;
+	FName BoneName;
+
+	FVector RelativeLocation;
+	FRotator RelativeRotation;
+	FVector RelativeScale;
+
+	bool bForceAlwaysAnimated = false;
+	
+	friend FArchive& operator<<(FArchive& Ar, FMutableRefSocket& Data)
+	{
+		Ar << Data.SocketName;
+		Ar << Data.BoneName;
+		Ar << Data.RelativeLocation;
+		Ar << Data.RelativeRotation;
+		Ar << Data.RelativeScale;
+
+		return Ar;
+	}
+};
+
+
+struct FMutableRefSkeletalMeshSettings
+{
+	bool bEnablePerPolyCollision = false;
+
+	friend FArchive& operator<<(FArchive& Ar, FMutableRefSkeletalMeshSettings& Data)
+	{
+		Ar << Data.bEnablePerPolyCollision;
+
+		return Ar;
+	}
+};
+
+USTRUCT()
+struct FMutableRefSkeletalMeshData
+{
+	GENERATED_BODY()
+	
+	// LOD info
+	TArray<FMutableRefLODData> LODData;
+	
+	// Sockets
+	TArray<FMutableRefSocket> Sockets;
+
+	// Bounding Box
+	FBoxSphereBounds Bounds;
+
+	// Settings
+	FMutableRefSkeletalMeshSettings Settings;
+
+	// Skeleton, must be stored in the ReferencedSkeletons too
+	UPROPERTY(Transient)
+	TSoftObjectPtr<USkeleton> Skeleton;
+	
+	// PhysicsAsset, must be stored in the PhysicsAssetMap too
+	UPROPERTY(Transient)
+	TSoftObjectPtr<UPhysicsAsset> PhysicsAsset;
+	
+	// Post Processing AnimBP, must be stored in the AnimBPAssetsMap too
+	UPROPERTY(Transient) 
+	TSoftClassPtr<UAnimInstance> PostProcessAnimInst;
+	
+	// Shadow PhysicsAsset, must be stored in the PhysicsAssetMap too
+	UPROPERTY(Transient)
+	TSoftObjectPtr<UPhysicsAsset> ShadowPhysicsAsset;
+
+
+	friend FArchive& operator<<(FArchive& Ar, FMutableRefSkeletalMeshData& Data)
+	{
+		Ar << Data.LODData;
+		Ar << Data.Sockets;
+		Ar << Data.Bounds;
+		Ar << Data.Settings;
+
+
+		if(Ar.IsSaving())
+		{
+			FString AssetPath = Data.Skeleton.ToSoftObjectPath().ToString();
+			Ar << AssetPath;
+
+			AssetPath = Data.PhysicsAsset.ToSoftObjectPath().ToString();
+			Ar << AssetPath;
+
+			AssetPath = Data.PostProcessAnimInst.ToSoftObjectPath().ToString();
+			Ar << AssetPath;
+
+			AssetPath = Data.ShadowPhysicsAsset.ToSoftObjectPath().ToString();
+			Ar << AssetPath;
+		
+		}
+		else
+		{
+			FString SkeletonAssetPath;
+			Ar << SkeletonAssetPath;
+			Data.Skeleton = TSoftObjectPtr<USkeleton>(FSoftObjectPath(SkeletonAssetPath));
+
+			FString PhysicsAssetPath;
+			Ar << PhysicsAssetPath;
+			Data.PhysicsAsset = TSoftObjectPtr<UPhysicsAsset>(FSoftObjectPath(PhysicsAssetPath));
+			
+			FString PostProcessAnimInstAssetPath;
+			Ar << PostProcessAnimInstAssetPath;
+			Data.PostProcessAnimInst = TSoftClassPtr<UAnimInstance>(FSoftObjectPath(PostProcessAnimInstAssetPath));
+			
+			FString ShadowPhysicsAssetPath;
+			Ar << ShadowPhysicsAssetPath;
+			Data.ShadowPhysicsAsset = TSoftObjectPtr<UPhysicsAsset>(FSoftObjectPath(ShadowPhysicsAssetPath));
+		}
+
+		return Ar;
+	}
+
+};
+
+USTRUCT()
+struct FMutableLODSettings
+{
+	GENERATED_BODY()
+
+#if WITH_EDITORONLY_DATA
+
+	/** Minimum LOD to render per Platform. */
+	UPROPERTY(EditAnywhere, Category = LODSettings, meta = (DisplayName = "Minimum LOD"))
+	FPerPlatformInt MinLOD;
+
+	/** Minimum LOD to render per Quality level.*/
+	UPROPERTY(EditAnywhere, Category = LODSettings, meta = (DisplayName = "Quality Level Minimum LOD"))
+	FPerQualityLevelInt MinQualityLevelLOD;
+
+	/** Override the LOD Streaming settings from the reference skeletal meshes.*/
+	UPROPERTY(EditAnywhere, Category = LODSettings, meta = (DisplayName = "Override LOD Streaming Settings"))
+	bool bOverrideLODStreamingSettings = true;
+
+	/** Enabled: streaming LODs will trigger automatic updates to generate and discard LODs. Streaming may decrease the amount of memory used, but will stress the CPU and Streaming of resources.
+	  *	Keep in mind that, even though updates may be faster depending on the amount of LODs to generate, there may be more updates to process.
+	  * 
+	  * Disabled: all LODs will be generated at once. It may increase the amount of memory used by the meshes and the generation may take longer, but less updates will be required. */
+	UPROPERTY(EditAnywhere, Category = LODSettings, meta = (EditCondition = "bOverrideLODStreamingSettings", DisplayName = "Enable LOD Streaming"))
+	FPerPlatformBool bEnableLODStreaming = true;
+
+	/** Limit the number of LODs to stream. A value of 0 is the same as disabling streaming of LODs.*/
+	UPROPERTY(EditAnywhere, Category = LODSettings, meta = (EditCondition = "bOverrideLODStreamingSettings"))
+	FPerPlatformInt NumMaxStreamedLODs = MAX_MESH_LOD_COUNT;
+
+#endif
+
+	/** Number of LODs that the object root has, which can be higher than the reference mesh LOD count. It's set at the end of the compilation process */
+	UPROPERTY()
+	int32 NumLODsInRoot = 0;
+
+	/** Frist LOD available, some platforms may remove lower LODs when cooking, this MinLOD represents the first LOD we can generate */
+	UPROPERTY()
+	int32 FirstLODAvailable = 0;
+
+	/** Whether we should stream LODs for on the running platform */
+	UPROPERTY()
+	bool bLODStreamingEnabled = false;
+
+	/** If bEnableLODStreaming is true, maximum number of LODs to stream */
+	UPROPERTY()
+	uint32 NumLODsToStream = 0;
+
+};
+
+
+UCLASS( config=Engine )
+class CUSTOMIZABLEOBJECT_API UCustomizableObjectBulk : public UObject
+{
+public:
+	GENERATED_BODY()
+
+	virtual void PostLoad() override;
+
+	/** Creates and returns an array with IAsyncReadFileHandles for each BulkData file.
+	 * Used by Mutable to stream in resources when generating instances. Must be deleted by the caller. */
+	TArray<IAsyncReadFileHandle*> GetAsyncReadFileHandles() const;
+
+#if WITH_EDITOR
+
+	//~ Begin UObject Interface
+	virtual void CookAdditionalFilesOverride(const TCHAR* PackageFilename, const ITargetPlatform* TargetPlatform,
+		TFunctionRef<void(const TCHAR* Filename, void* Data, int64 Size)> WriteAdditionalFile) override;
+	//~ End UObject Interface
+
+	/** Compute the number of files and sizes the BulkData will be split into and fix up 
+	 * HashToStreamableBlock's FileIndices and Offsets. */	
+	void PrepareBulkData(UCustomizableObject* InCustomizableObject, const ITargetPlatform* TargetPlatform);
+
+#endif
+	
+private:
+
+#if WITH_EDITOR
+
+	/** Helper to store the size of each BulkData partition. Only valid while cooking */
+	TArray<uint64> BulkDataFilesSize;
+
+	/** Helper to retrieve the BulkData from within the CookAdditionalFilesOverride */
+	TObjectPtr<UCustomizableObject> CustomizableObject;
+
+#endif
+
+	/** BulkData file paths */
+	UPROPERTY()
+	TArray<FString> BulkDataFileNames;
+};
+
+
+UCLASS( config=Engine )
+class CUSTOMIZABLEOBJECT_API UCustomizableObject : public UObject
+{
+public:
+	GENERATED_BODY()
+
+	UCustomizableObject();
+
+	/** All the SkeletalMeshes generated for this CustomizableObject instances will use the Reference Skeletal Mesh 
+	* properties for everything that Mutable doesn't create or modify. This includes data like LOD distances, Physics
+	* properties, Bounding Volumes, Skeleton, etc.
+	*
+	* While a CustomizableObject instance is being created for the first time, and in some situation with lots of 
+	* objects this may require some seconds, the Reference Skeletal Mesh is used for the actor. This works as a better
+	* solution than the alternative of not showing anything, although this can be disabled with the function
+	* "SetReplaceDiscardedWithReferenceMeshEnabled" (See the c++ section).
+	* 
+	* For more information on this topic read the Basic Concepts at work.anticto.com.
+	*/
+	UPROPERTY()
+	TObjectPtr<class USkeletalMesh> ReferenceSkeletalMesh_DEPRECATED;
+
+	/** All the SkeletalMeshes generated for this CustomizableObject instances will use the Reference Skeletal Mesh
+	* properties for everything that Mutable doesn't create or modify. This includes data like LOD distances, Physics
+	* properties, Bounding Volumes, Skeleton, etc.
+	*
+	* While a CustomizableObject instance is being created for the first time, and in some situation with lots of
+	* objects this may require some seconds, the Reference Skeletal Mesh is used for the actor. This works as a better
+	* solution than the alternative of not showing anything, although this can be disabled with the function
+	* "SetReplaceDiscardedWithReferenceMeshEnabled" (See the c++ section).
+	*
+	* For more information on this topic read the Basic Concepts at work.anticto.com.
+	*/
+	UPROPERTY(EditAnywhere, Category=CustomizableObject)
+	TArray< TObjectPtr<class USkeletalMesh> > ReferenceSkeletalMeshes;
+
+	/** All the SkeletalMeshes generated for this CustomizableObject instances will use the Reference Skeletal Mesh
+	 * properties for everything that Mutable doesn't create or modify. This struct stores the information used from
+	 * the Reference Skeletal Meshes to avoid having them loaded at all times. This includes data like LOD distances,
+	 * LOD render data settings, Mesh sockets, Bounding volumes, etc.
+	 */
+	UPROPERTY(Transient)
+	TArray<FMutableRefSkeletalMeshData> ReferenceSkeletalMeshesData;
+
+	// Hide for now, since it is not supported yet
+	//UPROPERTY(EditAnywhere, Category = CustomizableObject)
+	UPROPERTY()
+	TObjectPtr<class UStaticMesh> ReferenceStaticMesh;
+
+	/** List of Materials referenced by this or any child customizable object. */
+	UPROPERTY()
+	TArray<TSoftObjectPtr<UMaterialInterface>> ReferencedMaterials;
+
+	/** List of Material slot names for the materials referenced by this or any child customizable object. */
+	UPROPERTY()
+	TArray<FName> ReferencedMaterialSlotNames;
+
+	/** List of skeletons referenced by any of the parts of this customizable object. 
+	 * The position in this array is used as skeleton ID passed to Mutable bones.
+	 */
+	UPROPERTY()
+	TArray<TSoftObjectPtr<USkeleton>> ReferencedSkeletons;
+
+	UPROPERTY(EditAnywhere, Category = CustomizableObject, meta = (DisplayName = "LOD Settings"))
+	FMutableLODSettings LODSettings;
+
+	UPROPERTY(VisibleAnywhere, Category = CustomizableObject)
+	TArray<FMutableModelImageProperties> ImageProperties;
+
+	UPROPERTY(Transient)
+	TArray<FMorphTargetInfo> ContributingMorphTargetsInfo;
+	
+	UPROPERTY(Transient)
+	TArray<FMorphTargetVertexData> MorphTargetReconstructionData;
+
+	UPROPERTY(Transient)
+	TArray<FCustomizableObjectClothConfigData> ClothSharedConfigsData;	
+
+	UPROPERTY(Transient)
+	TArray<FCustomizableObjectClothingAssetData> ContributingClothingAssetsData;
+	
+	UPROPERTY(Transient)
+	TArray<FCustomizableObjectMeshToMeshVertData> ClothMeshToMeshVertData;
+
+#if WITH_EDITORONLY_DATA
+
+	// Hide this property because it is not used yet.
+	//UPROPERTY(EditAnywhere, Category = CustomizableObject)
+	UPROPERTY()
+	ECustomizableObjectRelevancy Relevancy;
+
+	// Compilation options to use in editor and for packaging for this object.
+	UPROPERTY()
+	FCompilationOptions CompileOptions;
+
+	// 
+	UPROPERTY(EditAnywhere, Category = CustomizableObject)
+	bool bDisableTextureLayoutManagement = false;
+
+	//
+	UPROPERTY(EditAnywhere, Category = CustomizableObject)
+	bool bEnableRealTimeMorphTargets = false;
+
+	//
+	UPROPERTY(EditAnywhere, Category = CustomizableObject)
+	bool bEnableClothing = false;
+
+	// Options when compiling this customizable object (see EMutableCompileMeshType declaration for info)
+	UPROPERTY(EditAnywhere, Category = CustomizableObject)
+	EMutableCompileMeshType MeshCompileType = EMutableCompileMeshType::LocalAndChildren;
+
+	// Array of elements to use with compile option CompileType = WorkingSet
+	UPROPERTY(EditAnywhere, Category = CustomizableObject)
+	TArray<TSoftObjectPtr<UCustomizableObject>> WorkingSet;
+
+	// Editor graph
+	UPROPERTY()
+	TObjectPtr<UEdGraph> Source;
+
+	// Used to verify the derived data matches this version of the Customizable Object.
+	UPROPERTY()
+	FGuid VersionId;
+
+	// Set of all the guids of all the CustomizableObjects in the compilation
+	UPROPERTY()
+	TSet<FGuid> CustomizableObjectGuidsInCompilation;
+
+	/* Map to identify what CustomizableObject owns a parameter. Used to display a tooltip when hovering a parameter
+	   in the Prev. instance panel */
+	UPROPERTY(Transient)
+	TMap<FString, FString> CustomizableObjectPathMap;
+
+	UPROPERTY(Transient)
+	TMap<FString, FCustomizableObjectIdPair> GroupNodeMap; 
+
+	// 
+	UPROPERTY()
+	TArray<FProfileParameterDat>  InstancePropertiesProfiles;
+
+#endif // WITH_EDITORONLY_DATA
+
+	// Object parameters interface. This is used to query static data about the parameters available
+	// in instances of this object.
+
+	// Get the number of parameters available in any instance.
+	UFUNCTION(BlueprintCallable, Category = CustomizableObject)
+	int GetParameterCount() const;
+
+	// Get the index of a parameter
+	UFUNCTION(BlueprintCallable, Category = CustomizableObject)
+	int FindParameter(const FString& Name) const;
+
+	// Get the type of a parameter
+	UFUNCTION(BlueprintCallable, Category = CustomizableObject)
+	EMutableParameterType GetParameterType(int32 ParamIndex) const;
+	UFUNCTION(BlueprintCallable, Category = CustomizableObject)
+	EMutableParameterType GetParameterTypeByName(const FString& Name) const;
+
+	// Get the name of a parameter
+	UFUNCTION(BlueprintCallable, Category = CustomizableObject)
+	const FString& GetParameterName(int32 ParamIndex) const;
+
+	// Get the number of description images available for a parameter
+	int32 GetParameterDescriptionCount(int32 ParamIndex) const;
+
+	// Get the number of description images available for a parameter
+	UFUNCTION(BlueprintCallable, Category = CustomizableObject)
+	int32 GetParameterDescriptionCount(const FString& ParamName) const;
+
+	// Returns how many possible options an int parameter has
+	UFUNCTION(BlueprintCallable, Category = CustomizableObject)
+	int32 GetIntParameterNumOptions(int32 ParamIndex) const;
+
+	// Gets the Name of the option at position K in the list of available options for the int parameter.
+	// Useful to enumerate the int parameter's possible options (Ex: "Hat1", "Hat2", "Cap", "Nothing")
+	UFUNCTION(BlueprintCallable, Category = CustomizableObject)
+	const FString& GetIntParameterAvailableOption(int32 ParamIndex, int32 K) const;
+
+	//UFUNCTION(BlueprintCallable, Category = CustomizableObject)
+	int32 FindIntParameterValue( int32 ParamIndex, const FString& Value ) const;
+	//UFUNCTION(BlueprintCallable, Category = CustomizableObject)
+	FString FindIntParameterValueName(int32 ParamIndex, int32 ParamValue) const;
+
+	//
+	USkeletalMesh* GetRefSkeletalMesh(int32 ComponentIndex = 0);
+	
+	//
+	FMutableRefSkeletalMeshData* GetRefSkeletalMeshData(int32 ComponentIndex = 0);
+	
+	TSoftObjectPtr<UMaterialInterface> GetReferencedMaterialAssetPtr(uint32 Index);
+
+	// Call before using Mutable's Projector testing with mask out features. It should only be loaded when needed because it can spend quite a lot of memory
+	// Can cause a loading hitch
+	UFUNCTION(BlueprintCallable, Category = CustomizableObject)
+	void LoadMaskOutCache();
+
+	// Call after having used Mutable's Projector testing with mask out features. It should be unloaded because it can spend quite a lot of memory
+	UFUNCTION(BlueprintCallable, Category = CustomizableObject)
+	void UnloadMaskOutCache();
+
+private: 
+
+	// This is information about the parameters in the model that is generated at model compile time.
+	UPROPERTY(Transient)
+	TArray<FMutableModelParameterProperties> ParameterProperties;
+
+	// Map of name to index of ParameterProperties.
+	// use this to lookup fast by Name
+	TMap<FString, int32> ParameterPropertiesLookupTable;
+
+	// Rebuild ParameterProperties from the current compiled model.
+	void UpdateParameterPropertiesFromModel();
+
+	// This is a manual version number for the binary blobs in this asset.
+	// Increasing it invalidates all the previously compiled models.
+	// Warning: If while merging code both versions have changed, take the highest+1.
+	static const int32 CurrentSupportedVersion = 347;
+
+public:
+
+#if WITH_EDITOR
+	
+	// Compile the object if Automatic Compilation is enabled and the object can be compiled.
+	// Automatic compilation can be enabled/disabled in the Mutable's Plugin Settings.
+	bool ConditionalAutoCompile();
+	
+	// Add a profile that stores the values of the parameters used by the CustomInstance.
+	FReply AddNewParameterProfile(FString Name, class UCustomizableObjectInstance& CustomInstance);
+
+	// Create new GUID for this CO
+	void UpdateVersionId();
+	FGuid GetVersionId() const { return VersionId; }
+
+	int GetCurrentSupportedVersion() const { return CurrentSupportedVersion; };
+
+	// Compose folder name where the data is stored
+	FString GetCompiledDataFolderPath(bool bIsEditorData) const;
+
+	// Compose file name 
+	FString GetCompiledDataFileName(bool bIsModel, const ITargetPlatform* InTargetPlatform = nullptr) const;
+
+	/** Used to set the flag IsRoot */
+	void GetAssetRegistryTags(TArray<FAssetRegistryTag>& OutTags) const override;
+
+	/** Returns a one line description of an object for viewing in the thumbnail view of the generic browser */
+	FString GetDesc() override;
+
+	bool IsEditorOnly() const override;
+
+	// Initialize if not set.
+	void InitializeIdentifier();
+
+	// Begin UObject interface.
+	void PreSave(FObjectPreSaveContext ObjectSaveContext) override;
+	void PostRename(UObject* OldOuter, const FName OldName) override;
+	void PostDuplicate(EDuplicateMode::Type DuplicateMode) override;
+
+	void BeginCacheForCookedPlatformData(const ITargetPlatform* TargetPlatform) override;
+	bool IsCachedCookedPlatformDataLoaded(const ITargetPlatform* TargetPlatform) override;
+	// End UObject interface.
+
+	// To avoid missing resources in packages we load all the resources. Loading the references added in the BeginCacheForCookedPlatformData will help us ensure they are included in the final package
+	void LoadReferencedObjects();
+	
+	// Compile the object for a specific platform - Compile for Cook Customizable Object
+	void CompileForTargetPlatform(const ITargetPlatform* TargetPlatform, bool bIsOnCookStart = false);
+
+	// Unless we are packaging there is no need for keeping all the data generated during compilation, this information is stored in the derived data.
+	void ClearCompiledData();
+
+	// Rebuild HashToStreamableBlocks and ParameterProperties from the current compiled model.
+	void UpdateCompiledDataFromModel();
+	
+	/** Generic Save/Load methods to write/read compiled data */
+	void SaveCompiledData(FArchive& Ar, bool bSkipEditorOnlyData = false);
+	void LoadCompiledData(FArchive& Ar, bool bSkipEditorOnlyData = false);
+
+	/** Load compiled data from disk, this is used to load Editor Compilations and Cook Compilations (when using OnCookStart) */
+	void LoadCompiledDataFromDisk(bool bIsEditorData = true, const ITargetPlatform* InTargetPlatform = nullptr);
+
+	/** Cache platform data for cook */
+	void CachePlatformData(const ITargetPlatform* InTargetPlatform, const TArray64<uint8>& InObjectBytes, const TArray64<uint8>& InBulkBytes);
+
+#if WITH_EDITORONLY_DATA
+
+	/** Map of PlatformName to CachedPlatformData. Only valid while cooking. */
+	TMap<FString, FMutableCachedPlatformData> CachedPlatformsData;
+	
+#endif
+	
+	void SaveEmbeddedData(FArchive& Ar);
+
+#endif
+
+	// Data that may be stored in the asset itself, only in packaged builds.
+	void LoadEmbeddedData(FArchive& Ar);
+
+	void PostLoad() override;
+
+	void Serialize(FArchive& Ar) override;
+
+	// 
+	void SerializeClothingDerivedData(FMemoryWriter64& Ar);
+	void DeserializeClothingDerivedData(FMemoryReaderView& Ar);
+
+
+	int32 FindState( const FString& Name ) const;
+
+	UFUNCTION(BlueprintCallable, Category = CustomizableObject)
+	int32 GetStateCount() const;
+
+	UFUNCTION(BlueprintCallable, Category = CustomizableObject)
+	FString GetStateName(int32 StateIndex) const;
+
+	int32 GetStateParameterCount(int32 StateIndex) const;
+	int32 GetStateParameterIndex(int32 StateIndex, int32 ParameterIndex) const;
+
+	UFUNCTION(BlueprintCallable, Category = CustomizableObject)
+	int32 GetStateParameterCount(const FString& StateName) const;
+
+	UFUNCTION(BlueprintCallable, Category = CustomizableObject)
+	FString GetStateParameterName(const FString& StateName, int32 ParameterIndex) const;
+	FString GetStateParameterName(int32 StateIndex, int32 ParameterIndex) const;
+
+	UFUNCTION(BlueprintCallable, Category = CustomizableObject)
+	FParameterUIData GetStateUIMetadata(const FString& StateName) const;
+
+	UFUNCTION(BlueprintCallable, Category = CustomizableObject)
+	FParameterUIData GetStateUIMetadataFromIndex(int32 StateIndex) const;
+
+	UFUNCTION(BlueprintCallable, Category = CustomizableObject)
+	FParameterUIData GetParameterUIMetadata(const FString& ParamName) const;
+
+	UFUNCTION(BlueprintCallable, Category = CustomizableObject)
+	FParameterUIData GetParameterUIMetadataFromIndex(int32 ParamIndex) const;
+
+	TSoftObjectPtr<USkeleton> GetReferencedSkeletonAssetPtr( uint32 Index );
+
+	/** Stores all the parameter UI metadata information for all the dependencies of this Customizable Object */
+	UPROPERTY()
+	TMap<FString, FParameterUIData> ParameterUIDataMap;
+
+	/** Stores all the state UI metadata information for all the dependencies of this Customizable Object */
+	UPROPERTY()
+	TMap<FString, FParameterUIData> StateUIDataMap;
+
+	/** Stores the physics assets gathered from the SkeletalMesh nodes during compilation, to be used in mesh generation in-game */
+	UPROPERTY()
+	TMap<FString, TSoftObjectPtr<class UPhysicsAsset>> PhysicsAssetsMap;
+
+	/** Stores the UAnimBlueprint assets gathered from the SkeletalMesh nodes during compilation, to be used in mesh generation in-game */
+	UPROPERTY()
+	TMap<FString, TSoftClassPtr<UAnimInstance>> AnimBPAssetsMap;
+
+	/** Stores the textures that will be used to mask-out areas in projectors. The cache isn't used for rendering, but for coverage testing */
+	UPROPERTY()
+	TSoftObjectPtr<UMutableMaskOutCache> MaskOutCache;
+
+	/** Map of Hash to Streaming blocks, used to stream a block of data representing a resource from the BulkData */
+	UPROPERTY()
+	TMap<uint64, FMutableStreamableBlock> HashToStreamableBlock;
+
+	// Customizable Object Population data start ------------------------------------------------------
+	/** Array to store the selected Population Class tags for this Customizable Object */
+	UPROPERTY()
+	TArray<FString> CustomizableObjectClassTags;
+	
+	/** Array to strore all the Population Class tags */
+	UPROPERTY()
+	TArray<FString> PopulationClassTags;
+
+	/** Map of parameters available for the Customizable Object and their tags */
+	UPROPERTY()
+	TMap<FString, FParameterTags> CustomizableObjectParametersTags;
+	// Customizable Object Population data end --------------------------------------------------------
+
+#if WITH_EDITORONLY_DATA
+
+	/** True if this object references a parent object. This is used basically to exclude this object
+	  * from cooking. This is actually derived from the source graph object node pointing to another
+	  * object or not, but it needs to be cached here because the source graph is not always available.
+	  * For old objects this may be false even if they are child objects until they are resaved, but 
+	  * that is the conservative case and shouldn't cause a problem. */
+	UPROPERTY()
+	bool bIsChildObject = false;
+
+	FPostCompileDelegate PostCompileDelegate;
+#endif
+
+	UFUNCTION(BlueprintCallable, Category = CustomizableObject)
+	class UCustomizableObjectInstance* CreateInstance();
+
+	// Return a pointer to the BulkData subobject, only valid in packaged builds
+	const UCustomizableObjectBulk* GetStreamableBulkData() const { return BulkData; }
+
+	FCustomizableObjectPrivateData* GetPrivate() const { return PrivateData.Get(); }
+
+	// This will always return true in a packaged game
+	UFUNCTION(BlueprintCallable, Category = CustomizableObject)
+	bool IsCompiled() const;
+
+	// See UCustomizableObjectSystem::LockObject()
+	bool IsLocked() const;
+
+	mu::Model* GetModel() const;
+
+#if WITH_EDITOR
+	void SetModel(mu::Model* Model);
+#endif
+
+	int32 GetNumLODs() const;
+
+	/** Modify the provided mutable parameters so that the forced values for the given customizable object state are applied. */
+	void ApplyStateForcedValuesToParameters(int32 State, mu::Parameters* Parameters);
+
+#if WITH_EDITOR
+
+	// UObject Interface -> Data validation
+	virtual EDataValidationResult IsDataValid(TArray<FText>& ValidationErrors) override;
+	virtual EDataValidationResult
+	IsDataValid(TArray<FText>& ValidationErrors, TArray<FText>& ValidationWarnings) override;
+	// End of UObject Interface
+
+	// UObject Interface -> Asset saving
+	virtual void PreSaveRoot(FObjectPreSaveRootContext ObjectSaveContext) override;
+	// End of UObject Interface
+
+private:
+	/** Flag that tells us if the validation of the asset has been triggered by the saving of it.
+	 * It gets consumed by IsDataValid and it is that same function the one that resets the value of this flag.
+	 */
+	bool bIsValidationTriggeredBySave = false;
+
+#endif
+
+private:
+
+	/** Used to prevent GC of MaskOutCache and keep it in memory while it's needed */
+	UPROPERTY(Transient)
+	TObjectPtr<UMutableMaskOutCache> MaskOutCache_HardRef;
+
+	/** Unique Identifier - used to locate Model and Streamable data on disk. Should not be modified. */
+	UPROPERTY()
+	FGuid Identifier;
+
+	/** BulkData that stores all in-game resources used by Mutable when generating instances.
+	  * Only valid in packaged builds */
+	UPROPERTY()
+	TObjectPtr<UCustomizableObjectBulk> BulkData;
+
+	TSharedPtr<FCustomizableObjectPrivateData> PrivateData;
+
+};
diff --git a/Engine/Plugins/Experimental/Mutable/Source/CustomizableObject/Classes/CustomizableObjectInstance.h b/Engine/Plugins/Experimental/Mutable/Source/CustomizableObject/Classes/CustomizableObjectInstance.h
new file mode 100644
--- /dev/null
+++ b/Engine/Plugins/Experimental/Mutable/Source/CustomizableObject/Classes/CustomizableObjectInstance.h
@@ -0,0 +1,644 @@
+// Copyright Epic Games, Inc. All Rights Reserved.
+
+#pragma once
+
+#include "CustomizableSkeletalComponent.h"
+#include "CustomizableObject.h"
+#include "CustomizableObjectInstanceDescriptor.h"
+#include "MultilayerProjector.h"
+#include "Engine/Texture2D.h"
+
+#include "CustomizableObjectParameterTypeDefinitions.h"
+#include "GameplayTagContainer.h"
+
+#include "CustomizableObjectInstance.generated.h"
+
+class UCustomizableInstancePrivateData; // This is used to hide Mutable SDK members in the public headers.
+
+//! Order of the unreal vertex buffers when in mutable data
+#define MUTABLE_VERTEXBUFFER_POSITION	0
+#define MUTABLE_VERTEXBUFFER_TANGENT	1
+#define MUTABLE_VERTEXBUFFER_TEXCOORDS	2
+
+
+USTRUCT()
+struct FCustomizedMaterialTexture2D
+{
+	GENERATED_USTRUCT_BODY()
+	
+	UPROPERTY(Category=CustomizedMaterialTexture2D, VisibleAnywhere)
+	FName Name;
+	
+	UPROPERTY(Category=CustomizedMaterialTexture2D, EditAnywhere) // Replaced EditInline with EditAnywhere for 4.10
+	TObjectPtr<UTexture2D> Texture;
+
+
+	FCustomizedMaterialTexture2D()
+		: Texture(nullptr)
+	{
+	}
+
+	inline friend FArchive& operator<<( FArchive& Ar, FCustomizedMaterialTexture2D& Data )
+	{
+		Ar << Data.Name;
+		Ar << Data.Texture;
+		return Ar;
+	}
+	
+};
+
+
+// Current state of a projector associated with a parameter.
+namespace EProjectorState
+{
+	typedef uint8 Type;
+	
+	const Type Hidden = 0;
+	const Type Translate = 1;
+	const Type Rotate = 2;
+	const Type Scale = 3;	
+	const Type Selected = 4;
+	const Type TypeChanged = 5;
+};
+
+
+namespace mu
+{
+	class Instance;
+};
+
+
+/** FString with the possible errors from skeletal mesh update */
+namespace ESkeletalMeshState
+{
+	const FString Correct = "Correct";
+	const FString UpdateError = "Update error";
+	const FString PostUpdateError = "Post Update Error";
+	const FString AsyncUpdatePending = "Async update pending";
+};
+
+
+class USkeletalMesh;
+
+DECLARE_DELEGATE(FObjectInstanceUpdateBeginDelegate);
+DECLARE_DYNAMIC_MULTICAST_DELEGATE_OneParam(FObjectInstanceUpdatedDelegate, UCustomizableObjectInstance*, Instance);
+DECLARE_MULTICAST_DELEGATE_OneParam(FObjectInstanceUpdatedNativeDelegate, UCustomizableObjectInstance*);
+DECLARE_DELEGATE_OneParam( FProjectorStateChangedDelegate, FString );
+
+DECLARE_DYNAMIC_DELEGATE_TwoParams(FEachComponentAnimInstanceClassDelegate, int32, SlotIndex, TSubclassOf<UAnimInstance>, AnimInstClass);
+DECLARE_DELEGATE_TwoParams(FEachComponentAnimInstanceClassNativeDelegate, int32 /*SlotIndex*/, TSubclassOf<UAnimInstance> /*AnimInstClass*/);
+
+UCLASS( Blueprintable, BlueprintType, HideCategories=(CustomizableObjectInstance) )
+class CUSTOMIZABLEOBJECT_API UCustomizableObjectInstance : public UObject
+{
+public:
+	 GENERATED_BODY()
+
+	UCustomizableObjectInstance();
+
+	FCustomizableObjectInstanceDescriptor& GetDescriptor();
+
+	const FCustomizableObjectInstanceDescriptor& GetDescriptor() const;
+	
+	void SetDescriptor(const FCustomizableObjectInstanceDescriptor& InDescriptor);
+	
+	/** The generated skeletal meshes for this Instance, one for each component */
+	UPROPERTY(Transient, VisibleAnywhere, Category = CustomizableSkeletalMesh)
+	TArray< TObjectPtr<USkeletalMesh> > SkeletalMeshes;
+
+	// Flag to control the build of the parameter description images required for customization UI
+	// These descriptions get generated with every instance update, so it should be disabled when not needed.
+	UPROPERTY(Transient, BlueprintReadwrite, Category=CustomizableObjectInstance, EditAnywhere)
+	bool bBuildParameterDecorations = false;
+
+	/** Broadcasts when the Customizable Object Instance is updated. */
+	UPROPERTY(Transient, BlueprintAssignable, Category = CustomizableObjectInstance)
+	FObjectInstanceUpdatedDelegate UpdatedDelegate;
+
+	/** Non-dynamic version. */
+	FObjectInstanceUpdatedNativeDelegate UpdatedNativeDelegate;
+
+	// Delegate to notify the beginning of an update of this instance
+	// NOTE: this delegate can be called outside the main thread, make sure operations done in the
+	//       callbacks are safe, avoid using it outside editor functionality
+	FObjectInstanceUpdateBeginDelegate UpdateBeginDelegate;
+
+	TMap<FString, bool> ParamNameToExpandedMap; // Used to check whether a mutable param is expanded in the editor to show its child params
+
+	// Will store status description of current skeletal mesh generation (for instance, "EmptyLOD0" or "EmptyMesh"
+	UPROPERTY()
+	FString SkeletalMeshStatus;
+
+#if WITH_EDITOR
+	// Will store the previous status description to avoid losing notifications with partial updates.
+	FString PreUpdateSkeletalMeshStatus;
+
+	/** During editor, always remember the duration of the last update in the mutable runtime, for profiling. */
+	int32 LastUpdateMutableRuntimeCycles = 0;
+#endif
+
+public:
+
+	// UObject interface.
+#if WITH_EDITOR
+	virtual void PostEditChangeProperty(FPropertyChangedEvent& PropertyChangedEvent) override;
+	virtual bool CanEditChange( const FProperty* InProperty ) const override;
+	bool InstanceUpdated; // Flag for the editor, to know when the instance's skeletal mesh has been updated
+#endif //WITH_EDITOR
+
+	virtual void Serialize(FArchive& Ar) override;
+	virtual void PostLoad() override;
+	virtual void BeginDestroy() override;
+	virtual bool IsReadyForFinishDestroy() override;
+	virtual FString GetDesc() override;
+	virtual bool IsEditorOnly() const override;
+
+	void SetObject(UCustomizableObject* InObject);
+
+	//Get the current CustomizableObject 
+	UCustomizableObject* GetCustomizableObject() const;
+
+	int32 GetState() const;
+	void SetState(int32 InState);
+
+	UFUNCTION(BlueprintCallable, Category = CustomizableObjectInstance)
+	FString GetCurrentState() const;
+
+	UFUNCTION(BlueprintCallable, Category = CustomizableObjectInstance)
+	void SetCurrentState(const FString& StateName);
+
+	UFUNCTION(BlueprintCallable, Category = CustomizableObjectInstance)
+	USkeletalMesh* GetSkeletalMesh(int32 ComponentIndex = 0) const;
+	
+	UFUNCTION(BlueprintCallable, Category = CustomizableObjectInstance)
+	bool HasAnySkeletalMesh() const;
+
+	UFUNCTION(BlueprintCallable, Category = CustomizableObjectInstance)
+	TArray<FCustomizableObjectBoolParameterValue>& GetBoolParameters();
+
+	const TArray<FCustomizableObjectBoolParameterValue>& GetBoolParameters() const;
+
+	UFUNCTION(BlueprintCallable, Category = CustomizableObjectInstance)
+	TArray<FCustomizableObjectIntParameterValue>& GetIntParameters();
+
+	const TArray<FCustomizableObjectIntParameterValue>& GetIntParameters() const;
+
+	UFUNCTION(BlueprintCallable, Category = CustomizableObjectInstance)
+	TArray<FCustomizableObjectFloatParameterValue>& GetFloatParameters();
+
+	const TArray<FCustomizableObjectFloatParameterValue>& GetFloatParameters() const;
+	
+	UFUNCTION(BlueprintCallable, Category = CustomizableObjectInstance)
+	TArray<FCustomizableObjectTextureParameterValue>& GetTextureParameters();
+
+	const TArray<FCustomizableObjectTextureParameterValue>& GetTextureParameters() const;
+
+	UFUNCTION(BlueprintCallable, Category = CustomizableObjectInstance)
+	TArray<FCustomizableObjectVectorParameterValue>& GetVectorParameters();
+
+	const TArray<FCustomizableObjectVectorParameterValue>& GetVectorParameters() const;
+
+	UFUNCTION(BlueprintCallable, Category = CustomizableObjectInstance)
+	TArray<FCustomizableObjectProjectorParameterValue>& GetProjectorParameters();
+	
+	const TArray<FCustomizableObjectProjectorParameterValue>& GetProjectorParameters() const;
+
+	/** See FCustomizableObjectInstanceDescriptor::HasAnyParameters. */
+	UFUNCTION(BlueprintCallable, Category = CustomizableObjectInstance)
+	bool HasAnyParameters() const;
+	
+	//! Set random values to the parameters. Useful for testing only.
+	UFUNCTION(BlueprintCallable, Category = CustomizableObjectInstance)
+	void SetRandomValues();
+
+	// Utilities to manage saving and loading parameters from profiles.
+	bool LoadParametersFromProfile(int32 ProfileIndex);
+	bool SaveParametersToProfile(int32 ProfileIndex);
+	bool MigrateProfileParametersToCurrentInstance(int32 ProfileIndex);
+	void SetSelectedParameterProfileDirty();
+	bool IsSelectedParameterProfileDirty() const;
+
+	//
+	UFUNCTION(BlueprintCallable, Category = CustomizableObjectInstance)
+	void UpdateSkeletalMeshAsync(bool bIgnoreCloseDist = false, bool bForceHighPriority = false);
+
+	// Clones the instance creating a new identical transient instance.
+	UFUNCTION(BlueprintCallable, Category = CustomizableObjectInstance)
+	UCustomizableObjectInstance* Clone();
+
+	// Clones the instance creating a new identical static instance.
+	UFUNCTION(BlueprintCallable, Category = CustomizableObjectInstance)
+	UCustomizableObjectInstance* CloneStatic(UObject* Outer);
+
+	// Copy parameters from input instance
+	void CopyParametersFromInstance(UCustomizableObjectInstance* Instance);
+
+	// Releases all the mutable resources this instance holds, should only be called when it is not going to be used any more.
+	void ReleaseMutableResources(bool bCalledFromBeginDestroy);
+
+	// Get a description image for a parameter.
+	// The param index is in the range of o to CustomizableObject->GetParameterCount.
+	// This will only be valid if bBuildParameterDecorations was set to true before the last update.
+	UTexture2D* GetParameterDescription(int32 ParamIndex, int32 DescIndex);
+
+	// Returns de description texture (ex: color bar) for this parameter and DescIndex
+	// This will only be valid if bBuildParameterDecorations was set to true before the last update.
+	UFUNCTION(BlueprintCallable, Category = CustomizableObjectInstance)
+	UTexture2D* GetParameterDescription(const FString& ParamName, int32 DescIndex);
+
+	//! This is only valid if bBuildParameterDecorations has been set before the last update.
+	bool IsParameterRelevant( int32 ParameterIndex ) const;
+
+	//! This is only valid if bBuildParameterDecorations has been set before the last update.
+	UFUNCTION(BlueprintCallable, Category = CustomizableObjectInstance)
+	bool IsParameterRelevant(const FString& ParamName) const;
+
+	UFUNCTION(BlueprintCallable, Category = CustomizableObjectInstance)
+	bool IsParamMultidimensional(const FString& ParamName) const;
+
+	UFUNCTION(BlueprintCallable, Category = CustomizableObjectInstance)
+	int32 CurrentParamRange(const FString& ParamName) const;
+	
+	bool IsParamMultidimensional(int32 ParamIndex) const;
+	FCustomizableObjectProjector GetProjectorDefaultValue(int32 ParamIndex) const;
+
+	bool bShowOnlyRuntimeParameters = true;
+	bool bShowOnlyRelevantParameters = true;
+
+	/** Control the display of support widgets to edit projectors of this instance. */
+	void SetProjectorState( const FString& ParamName, int32 RangeIndex, EProjectorState::Type state );
+	void ResetProjectorStates();
+	EProjectorState::Type GetProjectorState(const FString& ParamName, int32 RangeIndex) const;
+	FProjectorStateChangedDelegate ProjectorStateChangedDelegate;
+
+	// DEPRECATED: Use the method in the CustomizableObject instead which takes an index among all parameters
+	// Returns how many possible options an int parameter has
+	//int32 GetIntParameterNumOptions(int32 IntParamIndex);
+
+	// DEPRECATED: Use the method in the CustomizableObject instead which takes an index among all parameters
+	// Gets the Name of the option at position K in the list of available options for the int parameter. Useful to enumerate the int parameter's possible options (Ex: "Hat1", "Hat2", "Cap", "Nothing")
+	//const FString& GetIntParameterAvailableOption(int32 IntParamIndex, int32 K);
+
+	UFUNCTION(BlueprintCallable, Category = CustomizableObjectInstance)
+	const FString& GetIntParameterSelectedOption(const FString& ParamName, int32 RangeIndex = -1) const;
+
+	// Sets the selected option of an int parameter by the option's name
+	void SetIntParameterSelectedOption(int32 IntParamIndex, const FString& SelectedOption, int32 RangeIndex = -1);
+
+	// Sets the selected option of an int parameter, by the option's name
+	UFUNCTION(BlueprintCallable, Category = CustomizableObjectInstance)
+	void SetIntParameterSelectedOption(const FString& ParamName, const FString& SelectedOptionName, int32 RangeIndex = -1);
+
+	// Gets the value of a float parameter with name "FloatParamName"
+	UFUNCTION(BlueprintCallable, Category = CustomizableObjectInstance)
+	float GetFloatParameterSelectedOption(const FString& FloatParamName, int32 RangeIndex = -1) const;
+
+	// Sets the float value "FloatValue" of a float parameter with index "FloatParamIndex"
+	UFUNCTION(BlueprintCallable, Category = CustomizableObjectInstance)
+	void SetFloatParameterSelectedOption(const FString& FloatParamName, float FloatValue, int32 RangeIndex = -1);
+
+	// Gets the value of a color parameter with name "ColorParamName"
+	UFUNCTION(BlueprintCallable, Category = CustomizableObjectInstance)
+	FLinearColor GetColorParameterSelectedOption(const FString& ColorParamName) const;
+
+	// Sets the color value "ColorValue" of a color parameter with index "ColorParamIndex"
+	UFUNCTION(BlueprintCallable, Category = CustomizableObjectInstance)
+	void SetColorParameterSelectedOption(const FString& ColorParamName, const FLinearColor& ColorValue);
+
+	// Sets the bool value "BoolValue" of a bool parameter with name "BoolParamName"
+	UFUNCTION(BlueprintCallable, Category = CustomizableObjectInstance)
+	bool GetBoolParameterSelectedOption(const FString& BoolParamName) const;
+
+	// Sets the bool value "BoolValue" of a bool parameter with name "BoolParamName"
+	UFUNCTION(BlueprintCallable, Category = CustomizableObjectInstance)
+	void SetBoolParameterSelectedOption(const FString& BoolParamName, bool BoolValue);
+
+	// Sets the vector value "VectorValue" of a bool parameter with index "VectorParamIndex"
+	UFUNCTION(BlueprintCallable, Category = CustomizableObjectInstance)
+	void SetVectorParameterSelectedOption(const FString& VectorParamName, const FLinearColor& VectorValue);
+
+	// Sets the projector values of a projector parameter with index "ProjectorParamIndex"
+	UFUNCTION(BlueprintCallable, Category = CustomizableObjectInstance)
+	void SetProjectorValue(const FString& ProjectorParamName,
+		const FVector& OutPos, const FVector& OutDirection, const FVector& OutUp, const FVector& OutScale,
+		float OutAngle,
+		int32 RangeIndex = -1);
+
+	/** Set only the projector position. */
+	void SetProjectorPosition(const FString& ProjectorParamName, const FVector3f& Pos, int32 RangeIndex = -1);
+
+	// Get the projector values of a projector parameter with index "ProjectorParamIndex"
+	UFUNCTION(BlueprintCallable, Category = CustomizableObjectInstance)
+	void GetProjectorValue(const FString& ProjectorParamName,
+		FVector& OutPos, FVector& OutDirection, FVector& OutUp, FVector& OutScale,
+		float& OutAngle, ECustomizableObjectProjectorType& OutType,
+		int32 RangeIndex = -1) const;
+
+	void GetProjectorValueF(const FString& ProjectorParamName,
+		FVector3f& Pos, FVector3f& Direction, FVector3f& Up, FVector3f& Scale,
+		float& Angle, ECustomizableObjectProjectorType& Type,
+		int32 RangeIndex = -1) const;
+
+	// Get the current projector position for the parameter with the given name
+	UFUNCTION(BlueprintCallable, Category = CustomizableObjectInstance)
+	FVector GetProjectorPosition(const FString & ParamName, int32 RangeIndex = -1) const;
+
+	// Get the current projector direction vector for the parameter with the given name
+	UFUNCTION(BlueprintCallable, Category = CustomizableObjectInstance)
+	FVector GetProjectorDirection(const FString & ParamName, int32 RangeIndex = -1) const;
+
+	// Get the current projector up vector for the parameter with the given name
+	UFUNCTION(BlueprintCallable, Category = CustomizableObjectInstance)
+	FVector GetProjectorUp(const FString & ParamName, int32 RangeIndex = -1) const;
+
+	// Get the current projector scale for the parameter with the given name
+	UFUNCTION(BlueprintCallable, Category = CustomizableObjectInstance)
+	FVector GetProjectorScale(const FString & ParamName, int32 RangeIndex = -1) const;
+
+	// Get the current cylindrical projector angle for the parameter with the given name
+	UFUNCTION(BlueprintCallable, Category = CustomizableObjectInstance)
+	float GetProjectorAngle(const FString& ParamName, int32 RangeIndex = -1) const;
+
+	// Get the current projector type for the parameter with the given name
+	UFUNCTION(BlueprintCallable, Category = CustomizableObjectInstance)
+	ECustomizableObjectProjectorType GetProjectorParameterType(const FString& ParamName, int32 RangeIndex = -1) const;
+
+	/** See FCustomizableObjectInstanceDescriptor::GetProjector. */
+	FCustomizableObjectProjector GetProjector(const FString& ParamName, int32 RangeIndex) const;
+	
+	// Finds in IntParameters a parameter with name ParamName, returns the index if found, -1 otherwise
+	UFUNCTION(BlueprintCallable, Category = CustomizableObjectInstance)
+	int32 FindIntParameterNameIndex(const FString& ParamName) const;
+
+	// Finds in FloatParameters a parameter with name ParamName, returns the index if found, -1 otherwise
+	UFUNCTION(BlueprintCallable, Category = CustomizableObjectInstance)
+	int32 FindFloatParameterNameIndex(const FString& ParamName) const;
+
+	// Finds in BoolParameters a parameter with name ParamName, returns the index if found, -1 otherwise
+	UFUNCTION(BlueprintCallable, Category = CustomizableObjectInstance)
+	int32 FindBoolParameterNameIndex(const FString& ParamName) const;
+
+	// Finds in VectorParameters a parameter with name ParamName, returns the index if found, -1 otherwise
+	UFUNCTION(BlueprintCallable, Category = CustomizableObjectInstance)
+	int32 FindVectorParameterNameIndex(const FString& ParamName) const;
+
+	// Finds in ProjectorParameters a parameter with name ParamName, returns the index if found, -1 otherwise
+	UFUNCTION(BlueprintCallable, Category = CustomizableObjectInstance)
+	int32 FindProjectorParameterNameIndex(const FString& ParamName) const;
+
+	// Increases the range of values of the integer with ParamName, returns the index of the new integer value, -1 otherwise.
+	// The added value is initialized with the first integer option and is the last one of the range.
+	UFUNCTION(BlueprintCallable, Category = CustomizableObjectInstance)
+	int32 AddValueToIntRange(const FString& ParamName);
+
+	// Increases the range of values of the float with ParamName, returns the index of the new float value, -1 otherwise.
+	// The added value is initialized with 0.5f and is the last one of the range.
+	UFUNCTION(BlueprintCallable, Category = CustomizableObjectInstance)
+	int32 AddValueToFloatRange(const FString& ParamName);
+
+	// Increases the range of values of the projector with ParamName, returns the index of the new projector value, -1 otherwise.
+	// The added value is initialized with the default projector as set up in the editor and is the last one of the range.
+	UFUNCTION(BlueprintCallable, Category = CustomizableObjectInstance)
+	int32 AddValueToProjectorRange(const FString& ParamName);
+
+	// Remove the last of the integer range of values from the parameter ParamName, returns the index of the last valid integer, -1 if no values left.
+	UFUNCTION(BlueprintCallable, Category = CustomizableObjectInstance)
+	int32 RemoveValueFromIntRange(const FString& ParamName);
+
+	// Remove the RangeIndex element of the integer range of values from the parameter ParamName, returns the index of the last valid integer, -1 if no values left.
+	int32 RemoveValueFromIntRange(const FString& ParamName, int32 RangeIndex);
+
+	// Remove the last of the float range of values from the parameter ParamName, returns the index of the last valid float, -1 if no values left.
+	UFUNCTION(BlueprintCallable, Category = CustomizableObjectInstance)
+	int32 RemoveValueFromFloatRange(const FString& ParamName);
+
+	// Remove the RangeIndex element of the float range of values from the parameter ParamName, returns the index of the last valid float, -1 if no values left.
+	int32 RemoveValueFromFloatRange(const FString& ParamName, int32 RangeIndex);
+
+	// Remove the last of the projector range of values from the parameter ParamName, returns the index of the last valid projector, -1 if no values left.
+	UFUNCTION(BlueprintCallable, Category = CustomizableObjectInstance)
+	int32 RemoveValueFromProjectorRange(const FString& ParamName);
+
+	// Remove the RangeIndex element of the projector range of values from the parameter ParamName, returns the index of the last valid projector, -1 if no values left.
+	int32 RemoveValueFromProjectorRange(const FString& ParamName, int32 RangeIndex);
+
+	// ------------------------------------------------------------
+	// Multilayer Projectors
+	// ------------------------------------------------------------
+	
+	/** Given Multilayer Projector name, create a new Multilayer Projector Helper (if non-existent). See FMultilayerProjector.
+	 *
+	 * @return ture if successfully created (or was already created).
+	 */
+	UFUNCTION(BlueprintCallable, Category = CustomizableObjectInstance)
+	bool CreateMultiLayerProjector(const FName& ProjectorParamName);
+	
+	/** Given Multilayer Projector name, remove a Multilayer Projector Helper. See FMultilayerProjector. */
+	UFUNCTION(BlueprintCallable, Category = CustomizableObjectInstance)
+	void RemoveMultilayerProjector(const FName& ProjectorParamName);
+	
+	// Layers
+
+	/** See FMultilayerProjector::NumLayers. */
+	UFUNCTION(BlueprintCallable, Category = CustomizableObjectInstance)
+	int32 MultilayerProjectorNumLayers(const FName& ProjectorParamName) const;
+
+	/** See FMultilayerProjector::CreateLayer. */
+	UFUNCTION(BlueprintCallable, Category = CustomizableObjectInstance)
+	void MultilayerProjectorCreateLayer(const FName& ProjectorParamName, int32 Index);
+
+	/** See FMultilayerProjector::RemoveLayerAt. */
+	UFUNCTION(BlueprintCallable, Category = CustomizableObjectInstance)
+	void MultilayerProjectorRemoveLayerAt(const FName& ProjectorParamName, int32 Index);
+
+	/** See FMultilayerProjector::GetLayer. */
+	UFUNCTION(BlueprintCallable, Category = CustomizableObjectInstance)
+	FMultilayerProjectorLayer MultilayerProjectorGetLayer(const FName& ProjectorParamName, int32 Index) const;
+
+	/** See FMultilayerProjector::UpdateLayer. */
+	UFUNCTION(BlueprintCallable, Category = CustomizableObjectInstance)
+	void MultilayerProjectorUpdateLayer(const FName& ProjectorParamName, int32 Index, const FMultilayerProjectorLayer& Layer);
+
+	// Virtual layers
+
+	/** See FMultilayerProjector::GetVirtualLayers. */
+	UFUNCTION(BlueprintCallable, Category = CustomizableObjectInstance)
+	TArray<FName> MultilayerProjectorGetVirtualLayers(const FName& ProjectorParamName) const;
+	
+	/** See FMultilayerProjector::VirtualLayer. */
+	UFUNCTION(BlueprintCallable, Category = CustomizableObjectInstance)
+	void MultilayerProjectorCreateVirtualLayer(const FName& ProjectorParamName, const FName& Id);
+
+	/** See FMultilayerProjector::FindOrCreateVirtualLayer. */
+	UFUNCTION(BlueprintCallable, Category = CustomizableObjectInstance)
+	FMultilayerProjectorVirtualLayer MultilayerProjectorFindOrCreateVirtualLayer(const FName& ProjectorParamName, const FName& Id);
+
+	/** See FMultilayerProjector::RemoveVirtualLayer. */
+	UFUNCTION(BlueprintCallable, Category = CustomizableObjectInstance)
+	void MultilayerProjectorRemoveVirtualLayer(const FName& ProjectorParamName, const FName& Id);
+
+	/** See FMultilayerProjector::GetVirtualLayer. */
+	UFUNCTION(BlueprintCallable, Category = CustomizableObjectInstance)
+	FMultilayerProjectorVirtualLayer MultilayerProjectorGetVirtualLayer(const FName& ProjectorParamName, const FName& Id) const;
+
+	/** See FMultilayerProjector::UpdateVirtualLayer. */
+	UFUNCTION(BlueprintCallable, Category = CustomizableObjectInstance)
+	void MultilayerProjectorUpdateVirtualLayer(const FName& ProjectorParamName, const FName& Id, const FMultilayerProjectorVirtualLayer& Layer);
+	
+	// ------------------------------------------------------------
+	
+	// Returns the animation BP for the parameter component and slot, gathered from all the meshes that compose this instance
+	UFUNCTION(BlueprintCallable, Category = CustomizableObjectInstance)
+	TSubclassOf<UAnimInstance> GetAnimBP(int32 ComponentIndex, int32 SlotIndex) const;
+
+	UFUNCTION(BlueprintCallable, Category = CustomizableObjectInstance)
+	const FGameplayTagContainer& GetAnimationGameplayTags() const;
+	
+	UFUNCTION(BlueprintCallable, Category = CustomizableObjectInstance)
+	void ForEachAnimInstance(int32 ComponentIndex, FEachComponentAnimInstanceClassDelegate Delegate) const;
+
+	void ForEachAnimInstance(int32 ComponentIndex, FEachComponentAnimInstanceClassNativeDelegate Delegate) const;
+
+	/** Serializes/Deserializes all the customization parameters to/from a descriptor, ready to be sent/read, works like a typical UE4 Two-Way Save System Function  */
+	void SaveDescriptor(FArchive &CustomizableObjectDescriptor);
+	void LoadDescriptor(FArchive &CustomizableObjectDescriptor);
+
+	// Enable physics asset replacement so that generated skeletal meshes have the merged physics assets of their skeletal mesh parts and reference mesh
+	UFUNCTION(BlueprintCallable, Category = CustomizableObjectInstance)
+	void SetReplacePhysicsAssets(bool bReplaceEnabled);
+
+	// Enables the reuse of all possible textures when the instance is updated without any changes in geometry or state (the first update after creation doesn't reuse any)
+	// It will only work if the textures aren't compressed, so set the instance to a Mutable state with texture compression disabled
+	// WARNING! If texture reuse is enabled, do NOT keep external references to the textures of the instance. The instance owns the textures.
+	void SetReuseInstanceTextures(bool bTextureReuseEnabled);
+
+	// Adds/removes a texture channel coverage query to the instance update process. Every material updated by mutable with a channel named TextureName will be checked
+	// It only works properly when the Instance is in a state with texture compression disabled.
+	void AddQueryTextureCoverage(const FString& TextureName, const FString* MaskOutChannelName = nullptr);
+	void RemoveQueryTextureCoverage(const FString& TextureName);
+
+	// Returns the result of a texture coverage query previously registered with AddQueryTextureCoverage. The query is run during an update of the instance.
+	// It only works properly when the Instance is in a state with texture compression disabled.
+	float GetQueryResultTextureCoverage(const FString& TextureName);
+	float GetQueryResultTextureCoverageMasked(const FString& TextureName); // Same as the previous query but this time intersected with the mask texture
+
+	void AdditionalAssetsAsyncLoaded( FGraphEventRef CompletionEvent );
+
+	// The following methods should only be used in an LOD management class
+	void SetIsBeingUsedByComponentInPlay(bool bIsUsedByComponent );
+	bool GetIsBeingUsedByComponentInPlay() const;
+	void SetIsDiscardedBecauseOfTooManyInstances(bool bIsDiscarded);
+	bool GetIsDiscardedBecauseOfTooManyInstances() const;
+	void SetIsPlayerOrNearIt(bool NewValue);
+	float GetMinSquareDistToPlayer() const;
+	void SetMinSquareDistToPlayer(float NewValue);
+	void SetMinMaxLODToLoad(int32 NewMinLOD, int32 NewMaxLOD, bool LimitLODUpgrades = true);
+	int32 GetNumLODsAvailable() const;
+
+	/* If enabled, CurrentMinLOD will be the first LOD of the generated SkeletalMesh. Otherwise the number of LODs will remain constant and LODs [0 .. CurrentMinLOD] will share the same RenderData. 
+	 * Enabled by default */
+	void SetUseCurrentMinLODAsBaseLOD(bool bIsBaseLOD);
+	bool GetUseCurrentMinLODAsBaseLOD() const;
+
+	/** Instance updated. */
+	void Updated();
+	
+	/** Hash representing the actual state of the instance. */
+	uint32 GetLastUpdateDescriptorHash() const;
+
+	// --------------------------------------------------------------------
+
+	/** Flag to know if a property of this instance changed in the editor */
+	bool bEditorPropertyChanged = false;
+
+	UCustomizableInstancePrivateData* GetPrivate() const { 
+		check(nullptr != PrivateData); // Currently this is initialized in the constructor so we expect it always to exist.
+		return PrivateData; 
+	}
+
+	bool ProjectorUpdatedInViewport = false;
+
+	// TEMP VARIABLE to ease updating the gizmo in the editor after pasting a new projector's transform information
+	bool TempUpdateGizmoInViewport = false;
+
+	// TEMP VARIABLE to verify the projector parameter with transform modified by Paste Transform is set as selected
+	FString TempProjectorParameterName;
+
+	// TEMP VARIABLE to verify the projector parameter with transform modified by Paste Transform is set as selected
+	int32 TempProjectorParameterRangeIndex;
+
+	// TEMP VARIABLE to avoid the projector selection being reset after pasting transform
+	bool AvoidResetProjectorVisibilityForNonNode = false;
+
+	// TEMP VARIABLE to check the Min desired LODs for this instance
+	TWeakObjectPtr<UCustomizableSkeletalComponent> NearestToActor;
+	TWeakObjectPtr<const AActor> NearestToViewCenter;
+
+#if WITH_EDITOR
+	/** If there's a projector parameter pending to be hidden, name of that projector parameter */
+	FString LastSelectedProjectorParameter;
+
+	/** For the case of projector parameter in several layers (the name of the projector is the same and only the index will change */
+	FString LastSelectedProjectorParameterWithIndex;
+
+	/** Flag to ease detect a projector layer change event in FCustomizableObjectEditor::OnObjectModified */
+	bool ProjectorLayerChange = false;
+
+	/** Flag to avoid resetting projector state when editing the alpha value of a layer */
+	bool ProjectorAlphaChange = false;
+
+	/** Flag to unselect the projector */
+	bool UnselectProjector = false;
+
+	/** Profile index the instance parameters are in and if the profile needs to be refreshed */
+	int32 SelectedProfileIndex = INDEX_NONE;
+	bool bSelectedProfileDirty = false;
+
+	/** Tag required to avoid updating the wrong projector parameter range index once removed.
+	On a Group Projector Parameter, if a projector gizmo is removed while is selected, the projector is unselected and removed. The unselection will causes a late update
+	to that projector index, which is performed after the projector being removed. During the update, since the projector index has been removed, the update index is no
+	longer valid. This tag allows to check if the last removed index matches with the currently update index. In the case they match, the update does not occur.*/
+	FString RemovedProjectorParameterNameWithIndex;
+#endif 
+
+private:
+	UPROPERTY()
+	FCustomizableObjectInstanceDescriptor Descriptor;
+
+	UPROPERTY( Transient )
+	TObjectPtr<UCustomizableInstancePrivateData> PrivateData;
+	
+	/** Multilayer Projector helpers. See FMultilayerProjector.*/
+	UPROPERTY()
+	TMap<FName, FMultilayerProjector> MultilayerProjectors;
+
+	/** Hash of the UCustomizableObjectInstance::Descriptor on the last update. */
+	uint32 LastUpdateDescriptorHash = 0;
+
+	// Deprecated properties
+	
+	UPROPERTY()
+	TObjectPtr<UCustomizableObject> CustomizableObject_DEPRECATED;
+	
+	UPROPERTY()
+	TArray<FCustomizableObjectBoolParameterValue> BoolParameters_DEPRECATED;
+
+	UPROPERTY()
+	TArray<FCustomizableObjectIntParameterValue> IntParameters_DEPRECATED;
+
+	UPROPERTY()
+	TArray<FCustomizableObjectFloatParameterValue> FloatParameters_DEPRECATED;
+
+	UPROPERTY()
+	TArray<FCustomizableObjectTextureParameterValue> TextureParameters_DEPRECATED;
+
+	UPROPERTY()
+	TArray<FCustomizableObjectVectorParameterValue> VectorParameters_DEPRECATED;
+
+	UPROPERTY()
+	TArray<FCustomizableObjectProjectorParameterValue> ProjectorParameters_DEPRECATED;
+};
+
+#if WITH_EDITOR
+CUSTOMIZABLEOBJECT_API void CopyTextureProperties(UTexture2D* Texture, const UTexture2D* SourceTexture);
+#endif	
\ No newline at end of file
diff --git a/Engine/Plugins/Experimental/Mutable/Source/CustomizableObject/Classes/CustomizableObjectSystem.h b/Engine/Plugins/Experimental/Mutable/Source/CustomizableObject/Classes/CustomizableObjectSystem.h
new file mode 100644
--- /dev/null
+++ b/Engine/Plugins/Experimental/Mutable/Source/CustomizableObject/Classes/CustomizableObjectSystem.h
@@ -0,0 +1,455 @@
+// Copyright Epic Games, Inc. All Rights Reserved.
+
+#pragma once
+
+#include "Engine/StreamableManager.h"
+#include "Templates/SharedPointer.h"
+
+#include "Containers/Ticker.h"
+#include "CustomizableObjectInstance.h"
+
+#include "Runtime/Launch/Resources/Version.h"
+
+#if WITH_EDITOR
+#include "Framework/Notifications/NotificationManager.h"
+#include "Widgets/Notifications/SNotificationList.h"
+#endif
+
+#include "CustomizableObjectSystem.generated.h"
+
+// This sets the amount of memory in bytes used to keep streaming data after using it, to reduce the streaming load.
+// High values use more memory, but save object construction time.
+// Setting it to 0 was the original behaviour, and keeps all the data loaded for a mutable operation until the
+// next operation.
+// This can be overwritten with the console variable b.MutableStreamingMemory
+#if !PLATFORM_DESKTOP
+	#define MUTABLE_STREAMING_CACHE				(3 * 1024 * 1024)
+#else
+	#define MUTABLE_STREAMING_CACHE				(12 * 1024 * 1024)
+#endif
+
+// Split StreamedBulkData into chunks smaller than MUTABLE_STREAMED_DATA_MAXCHUNKSIZE
+#define MUTABLE_STREAMED_DATA_MAXCHUNKSIZE		(512 * 1024 * 1024)
+
+// In case of async file operations, what priority to use
+#define MUTABLE_SYSTEM_ASYNC_STREAMING_PRIORITY			EAsyncIOPriorityAndFlags::AIOP_Normal
+
+// This is used to hide Mutable SDK members in the public headers.
+class FCustomizableObjectSystemPrivate;
+
+
+// Mutable stats
+DECLARE_STATS_GROUP(TEXT("Mutable"), STATGROUP_Mutable, STATCAT_Advanced);
+DECLARE_DWORD_ACCUMULATOR_STAT_EXTERN(TEXT("Num Created Mutable Skeletal Meshes"), STAT_MutableNumSkeletalMeshes, STATGROUP_Mutable, );
+DECLARE_DWORD_ACCUMULATOR_STAT_EXTERN(TEXT("Num Cached Mutable Skeletal Meshes"), STAT_MutableNumCachedSkeletalMeshes, STATGROUP_Mutable, );
+DECLARE_DWORD_ACCUMULATOR_STAT_EXTERN(TEXT("Num Allocated Mutable Skeletal Meshes"), STAT_MutableNumAllocatedSkeletalMeshes, STATGROUP_Mutable, );
+DECLARE_DWORD_ACCUMULATOR_STAT_EXTERN(TEXT("Num Instances at LOD 0"), STAT_MutableNumInstancesLOD0, STATGROUP_Mutable, );
+DECLARE_DWORD_ACCUMULATOR_STAT_EXTERN(TEXT("Num Instances at LOD 1"), STAT_MutableNumInstancesLOD1, STATGROUP_Mutable, );
+DECLARE_DWORD_ACCUMULATOR_STAT_EXTERN(TEXT("Num Instances at LOD 2 or more"), STAT_MutableNumInstancesLOD2, STATGROUP_Mutable, );
+DECLARE_DWORD_ACCUMULATOR_STAT_EXTERN(TEXT("Skeletal Mesh Resource Memory"), STAT_MutableSkeletalMeshResourceMemory, STATGROUP_Mutable, );
+DECLARE_DWORD_ACCUMULATOR_STAT_EXTERN(TEXT("Num Created Mutable Textures"), STAT_MutableNumTextures, STATGROUP_Mutable, );
+DECLARE_DWORD_ACCUMULATOR_STAT_EXTERN(TEXT("Num Cached Mutable Textures"), STAT_MutableNumCachedTextures, STATGROUP_Mutable, );
+DECLARE_DWORD_ACCUMULATOR_STAT_EXTERN(TEXT("Num Allocated Mutable Textures"), STAT_MutableNumAllocatedTextures, STATGROUP_Mutable, );
+DECLARE_DWORD_ACCUMULATOR_STAT_EXTERN(TEXT("Texture Resource Memory"), STAT_MutableTextureResourceMemory, STATGROUP_Mutable, );
+DECLARE_DWORD_ACCUMULATOR_STAT_EXTERN(TEXT("Texture Generated Memory"), STAT_MutableTextureGeneratedMemory, STATGROUP_Mutable, );
+DECLARE_DWORD_ACCUMULATOR_STAT_EXTERN(TEXT("Texture Locked Memory"), STAT_MutableTextureCacheMemory, STATGROUP_Mutable, );
+DECLARE_DWORD_ACCUMULATOR_STAT_EXTERN(TEXT("Texture Parameter Decoration Memory"), STAT_MutableTextureParameterDecorationMemory, STATGROUP_Mutable, );
+DECLARE_DWORD_ACCUMULATOR_STAT_EXTERN(TEXT("Pending Instance Updates"), STAT_MutablePendingInstanceUpdates, STATGROUP_Mutable, );
+DECLARE_DWORD_ACCUMULATOR_STAT_EXTERN(TEXT("Abandoned Instance Updates"), STAT_MutableAbandonedInstanceUpdates, STATGROUP_Mutable, );
+DECLARE_DWORD_ACCUMULATOR_STAT_EXTERN(TEXT("Last Instance Build Time"), STAT_MutableInstanceBuildTime, STATGROUP_Mutable, );
+DECLARE_DWORD_ACCUMULATOR_STAT_EXTERN(TEXT("Avrg Instance Build Time"), STAT_MutableInstanceBuildTimeAvrg, STATGROUP_Mutable, );
+DECLARE_DWORD_ACCUMULATOR_STAT_EXTERN(TEXT("Streaming Ops"), STAT_MutableStreamingOps, STATGROUP_Mutable, );
+DECLARE_DWORD_ACCUMULATOR_STAT_EXTERN(TEXT("Streaming Cache"), STAT_MutableStreamingCache, STATGROUP_Mutable, );
+
+// These stats are provided by the mutable runtime
+DECLARE_DWORD_ACCUMULATOR_STAT_EXTERN(TEXT("Profile-LiveInstanceCount"), STAT_MutableProfile_LiveInstanceCount, STATGROUP_Mutable, );
+DECLARE_DWORD_ACCUMULATOR_STAT_EXTERN(TEXT("Profile-StreamingCacheBytes"), STAT_MutableProfile_StreamingCacheBytes, STATGROUP_Mutable, );
+DECLARE_DWORD_ACCUMULATOR_STAT_EXTERN(TEXT("Profile-InstanceUpdateCount"), STAT_MutableProfile_InstanceUpdateCount, STATGROUP_Mutable, );
+
+
+#if WITH_EDITOR
+
+// Struct used to keep a copy of the EditorSettings needed to compile Customizable Objects.
+struct FEditorCompileSettings 
+{
+	// General case
+	bool bDisableCompilation;
+
+	// Auto Compile 
+	bool bEnableAutomaticCompilation = true;
+	bool bCompileObjectsSynchronously = true;
+	bool bCompileRootObjectsOnStartPIE = false;
+};
+
+#endif
+
+//
+namespace EMutableProfileMetric
+{
+	typedef uint8 Type;
+
+	const Type BuiltInstances = 1;
+	const Type UpdateOperations = 2;
+	const Type Count = 4;
+
+};
+
+USTRUCT()
+struct FPendingReleaseSkeletalMeshInfo
+{
+public:
+	GENERATED_USTRUCT_BODY()
+
+public:
+	UPROPERTY()
+	TObjectPtr<USkeletalMesh> SkeletalMesh = nullptr;
+
+	UPROPERTY()
+	double TimeStamp = 0.0f;
+};
+
+
+
+/** Key to identify an image inside a generated mutable runtime instance. */
+struct FMutableImageReference
+{
+	/** Original image ID. Once generated it will be unique. However, future updates of the image may return a different ID for the
+	* same image, if many other resources have been built in the middle. For this reason the rest of the data in the struct is what
+	* must be used to request the additional mips.
+	*/
+	uint32 ImageID = 0;
+
+	/** */
+	uint8 LOD = 0;
+	uint8 Component = 0;
+	uint8 Surface = 0;
+	uint8 Image = 0;
+};
+
+
+
+// Opaque representation of a possible registered value for texture parameters.
+struct FCustomizableObjectExternalTexture
+{
+	FCustomizableObjectExternalTexture() = default;
+
+	FString Name;
+	int64_t Value = 0;
+};
+
+
+//
+UCLASS(abstract)
+class CUSTOMIZABLEOBJECT_API UCustomizableSystemImageProvider : public UObject
+{
+public:
+	GENERATED_BODY()
+
+	enum class ValueType : uint8
+	{
+		// This texture is not provided by this provider.
+		None,
+
+		// Data will be provided with size and pointer
+		Raw,
+
+		// Data will be provided from an unreal texture
+		Unreal
+	};
+
+	// Query that Mutable will run to find out if a texture will be provided as an Unreal UTexture2D,
+	// or as a raw data blob.
+	virtual ValueType HasTextureParameterValue(int64 ID) { return ValueType::None; }
+
+	// In case IsTextureParameterValueUnreal returns false, this will be used to query the texture size data.
+	virtual FIntVector GetTextureParameterValueSize(int64 ID) { return FIntVector(0, 0, 0); }
+
+	// In case IsTextureParameterValueUnreal returns false, this will be used to query the texture data that must
+	// be copied in the preallocated buffer. The pixel format is assumed to be 4-channel RGBA, uint8_t per channel.
+	virtual void GetTextureParameterValueData(int64 ID, uint8* OutData) {}
+
+	// In case IsTextureParameterValueUnreal returns true, this will be used to query the texture.
+	virtual UTexture2D* GetTextureParameterValue(int64 ID) { return nullptr; }
+
+	// Used in the editor to show the list of available options.
+	// Only necessary if the images are required in editor previews.
+	virtual void GetTextureParameterValues(TArray<FCustomizableObjectExternalTexture>& OutValues) {};
+};
+
+
+// Example implementation of a ICustomizableSystemImageProvider that just uses a predefined array of textures.
+// It is also used by editors to set some preview images.
+UCLASS()
+class CUSTOMIZABLEOBJECT_API UCustomizableObjectImageProviderArray : public UCustomizableSystemImageProvider
+{
+public:
+	GENERATED_BODY()
+
+	UPROPERTY(EditAnywhere, Category = Preview)
+	TArray< TObjectPtr<UTexture2D> > Textures;
+
+	// UObject interface
+#if WITH_EDITOR
+	void PostEditChangeProperty(struct FPropertyChangedEvent& PropertyChangedEvent) override;
+#endif
+
+	// UCustomizableSystemImageProvider interface
+	ValueType HasTextureParameterValue(int64 ID) override;
+	UTexture2D* GetTextureParameterValue(int64 ID) override;
+	void GetTextureParameterValues(TArray<FCustomizableObjectExternalTexture>& OutValues) override;
+
+	// Own interface
+	void InvalidateIds();
+
+	DECLARE_MULTICAST_DELEGATE(FOnTexturesChanged);
+	FOnTexturesChanged TexturesChangeDelegate;
+
+private:
+
+	// The preview texture values will start at this ID, to avoid clashing with other game-specific custom providers.
+	int FirstId = 100000;
+};
+
+
+struct FMutableMipUpdateLevel
+{
+	FMutableMipUpdateLevel(int32 InMipLevel, void* InDest, int32 InSizeX, int32 InSizeY, int32 InDataSize, EPixelFormat InFormat) :
+		MipLevel(InMipLevel), Dest(InDest), SizeX(InSizeX), SizeY(InSizeY), DataSize(InDataSize), Format(InFormat) {}
+
+	int32 MipLevel;
+	void* Dest; // Only access from the FMutableTextureMipDataProvider, owned by the FTextureMipInfoArray so don't delete
+	int32 SizeX;
+	int32 SizeY;
+	int32 DataSize;
+	EPixelFormat Format;
+};
+
+
+class UCustomizableInstanceLODManagementBase;
+
+UCLASS(Blueprintable, BlueprintType)
+class CUSTOMIZABLEOBJECT_API UCustomizableObjectSystem : public UObject
+{
+public:
+	GENERATED_BODY()
+
+	UCustomizableObjectSystem() = default;
+	void InitSystem();
+
+	// Get the singleton object. It will be created if it doesn't exist yet.
+	UFUNCTION(BlueprintCallable, BlueprintPure, Category = Status)
+	static UCustomizableObjectSystem* GetInstance();
+
+	// Return true if the singleton has been created. It is different than GetInstance in that GetInstance will create it if it doesn't exist.
+	static bool IsCreated();
+
+	// Begin UObject interface.
+	virtual void BeginDestroy() override;
+	virtual FString GetDesc() override;
+	// End UObject interface.
+
+	// Creates a new Customizable Object Compiler (Only does real work in editor builds). The caller is responsible for freeing the new compiler
+	class FCustomizableObjectCompilerBase* GetNewCompiler();
+	void SetNewCompilerFunc(FCustomizableObjectCompilerBase* (*NewCompilerFunc)());
+
+	bool IsReplaceDiscardedWithReferenceMeshEnabled() const;
+	void SetReplaceDiscardedWithReferenceMeshEnabled(bool bIsEnabled);
+
+	// Unprotect the resources in the instances of this object of being garbage collector
+	// while an instance is being built or updated, so that they can be reused.
+	void ClearResourceCacheProtected();
+
+#if WITH_EDITOR
+	// Lock a CustomizableObjects, preventing the generation or update of any of its instances
+	// Will return true if successful, false if it fails to lock because an update is already underway
+	// This is usually only used in the editor
+	bool LockObject(const class UCustomizableObject*);
+	void UnlockObject(const class UCustomizableObject*);
+	
+	// Called whenever the Mutable Editor Settings change, copying the new value of the current needed settings to the Customizable Object System
+	void EditorSettingsChanged(const FEditorCompileSettings& InEditorSettings);
+
+	// If compilation is disabled, Customizable Objects won't be compiled in the editor
+	bool IsCompilationDisabled() const;
+
+	// If true, uncompiled Customizable Objects will be compiled whenever an instance update is required
+	bool IsAutoCompileEnabled() const;
+
+	// If true, uncompiled Customizable Objects will be compiled synchronously
+	bool IsAutoCompilationSync() const;
+
+	// Copy of the Mutable Editor Settings tied to CO compilation. They are updated whenever changed
+	FEditorCompileSettings EditorSettings;
+
+#endif
+
+	bool IsCompactSerializationEnabled() const;
+
+	bool IsSupport16BitBoneIndexEnabled() const;
+
+	bool IsProgressiveMipStreamingEnabled() const;
+
+public:
+
+	void AddPendingReleaseSkeletalMesh( USkeletalMesh* SkeletalMesh );
+
+	void PurgePendingReleaseSkeletalMesh();
+
+private:
+
+	UPROPERTY()
+	TArray<FPendingReleaseSkeletalMeshInfo> PendingReleaseSkeletalMesh;
+
+public:
+
+	/** [Texture Parameters] Get a list of all the possible values for external texture parameters according to the various providers registered with RegisterImageProvider. */
+	TArray<FCustomizableObjectExternalTexture> GetTextureParameterValues();
+
+	/** [Texture Parameters] Add a new image provider to the CustomizableObject System. This will be queried for when an external image ID is provided to mutable in an Texture Parameter node. */
+	void RegisterImageProvider(UCustomizableSystemImageProvider* Provider);
+
+	/** [Texture Parameters] Remove a previously regiostered provider. */
+	void UnregisterImageProvider(UCustomizableSystemImageProvider* Provider);
+
+#if WITH_EDITOR
+	/** [Texture Parameters] Get the editor-side preview external texture provider. 
+	 * \TODO: Move to the editor? 
+	 */
+	UCustomizableObjectImageProviderArray* GetEditorExternalImageProvider();
+#endif
+
+private:
+
+	/** [Texture Parameters] If in editor, this will hold a reference to the image provider used to give examples of external images to preview in the customizable objects that use this functionality.
+	 * \TODO : Move to the editor ?
+	 */
+	UPROPERTY()
+	TObjectPtr<UCustomizableObjectImageProviderArray> PreviewExternalImageProvider = nullptr;
+
+public:
+    
+	// Show a warning when a CustomizableObject is being used and it's not compiled 
+	void AddUncompiledCOWarning(UCustomizableObject* InObject);
+
+	// Enables the collection of internal mutabe performance data. It has a performance cost.
+	void EnableBenchmark();
+	// Writes the benchmark results
+	void EndBenchmark();
+
+	// Show data about all UCustomizableObjectInstance existing elements
+	void LogShowData(bool bFullInfo, bool ShowMaterialInfo) const;
+
+
+	// Give access to the internal object data.
+	FCustomizableObjectSystemPrivate* GetPrivate() { return Private.Get(); }
+
+	FStreamableManager& GetStreamableManager() { return StreamableManager; }
+
+	UCustomizableInstanceLODManagementBase* GetInstanceLODManagement() { return CurrentInstanceLODManagement.Get(); }
+
+	// Pass a null ptr to reset to the default InstanceLODManagement
+	void SetInstanceLODManagement(UCustomizableInstanceLODManagementBase* NewInstanceLODManagement) 
+		{ CurrentInstanceLODManagement = NewInstanceLODManagement ? NewInstanceLODManagement : DefaultInstanceLODManagement.Get(); }
+
+	// Find out the version of the plugin
+	UFUNCTION(BlueprintCallable, Category = Status)
+	FString GetPluginVersion() const;
+
+	// Get the number of instances built and alive.
+	UFUNCTION(BlueprintCallable, Category = Status)
+	int32 GetNumInstances() const;
+
+	// Get the number of instances waiting to be updated.
+	UFUNCTION(BlueprintCallable, Category = Status)
+	int32 GetNumPendingInstances() const;
+
+	// Get the total number of instances includingbuilt and not built.
+	UFUNCTION(BlueprintCallable, Category = Status)
+	int32 GetTotalInstances() const;
+
+	// Get the amount of memory in use for textures generated by mutable.
+	UFUNCTION(BlueprintCallable, Category = Status)
+	int32 GetTextureMemoryUsed() const;
+
+	// Return the average build/update time of an instance in ms.
+	UFUNCTION(BlueprintCallable, Category = Status)
+	int32 GetAverageBuildTime() const;
+
+	// If set to true, Mutable will release Mutable-generated textures immediately when they are no longer used without waiting for GC
+	// IMPORTANT!!! Do NOT keep references to any Mutable generated textures or skeletal meshes if this is enabled,
+	// they are owned by Mutable and will be destroyed without notice
+	UFUNCTION(BlueprintCallable, Category = Status)
+	void SetReleaseMutableTexturesImmediately(bool bReleaseTextures);
+
+#if WITH_EDITOR
+	void RecompileCustomizableObjectAsync(const FAssetData& InAssetData, const UCustomizableObject* InObject);
+	
+	void RecompileCustomizableObjects(const TArray<FAssetData>& InObjects);
+
+	// Get the maximum size a chunk can have on a specific platform. If unspecified return MUTABLE_STREAMED_DATA_MAXCHUNKSIZE
+	uint64 GetMaxChunkSizeForPlatform(const ITargetPlatform* InTargetPlatform);
+#endif
+
+	void ClearCurrentMutableOperation();
+
+	UPROPERTY(Transient)
+	TObjectPtr<UCustomizableInstanceLODManagementBase> DefaultInstanceLODManagement = nullptr;
+
+	UPROPERTY(Transient)
+	TObjectPtr<UCustomizableInstanceLODManagementBase> CurrentInstanceLODManagement = nullptr;
+
+	// Array where textures are added temporarily while the mutable thread may want to
+	// reused them for some instance under construction.
+	UPROPERTY(Transient)
+	TArray< TObjectPtr<UTexture2D> > ProtectedCachedTextures;
+
+private:
+
+	TSharedPtr<FCustomizableObjectSystemPrivate> Private = nullptr;
+
+	// For async material loading
+	FStreamableManager StreamableManager;
+	bool bMaterialsAlreadyLoaded = false;
+
+	// Most of the work in this plugin happens here.
+	bool Tick(float DeltaTime);
+
+	// If there is an on-going operation, advance it.
+	void AdvanceCurrentOperation();
+
+	FTimerHandle TimerHandle_ApplyRandomValues;
+
+	bool ApplyRandomValuesToAllPawns = false;
+
+	bool RandomValuesCommandTimerWorking = false;
+
+	// TODO: Can we move this to the editor module?
+#if WITH_EDITOR
+
+	// Used to ask the user if they want to recompile uncompiled PIE COs
+	void OnPreBeginPIE(const bool bIsSimulatingInEditor);
+
+	void StartNextRecompile();
+	void TickRecompileCustomizableObjects();
+
+	FCustomizableObjectCompilerBase* RecompileCustomizableObjectsCompiler = nullptr;
+	
+	TArray<FAssetData> ObjectsToRecompile;
+	uint32 TotalNumObjectsToRecompile = 0;
+	uint32 NumObjectsCompiled = 0;
+
+	/** Recompile progress bar handle */
+	FProgressNotificationHandle RecompileNotificationHandle;
+
+	// Array to keep track of cached objects
+	TArray<FGuid> UncompiledCustomizableObjectIds;
+
+	/** Weak pointer to the Uncompiled Customizable Objects notification */
+	TWeakPtr<SNotificationItem> UncompiledCustomizableObjectsNotificationPtr;
+
+	/** Map used to cache per platform MaxChunkSize. If MaxChunkSize > 0, streamed data will be split in multiple files */
+	TMap<FString, int64> PlatformMaxChunkSize;
+	
+#endif
+};
+
