[mutable] Move Mutable plugin from the restricted folder to the experimental plugins folder.

diff --git a/Engine/Plugins/Experimental/Mutable/Source/CustomizableObject/Classes/CustomizableSkeletalComponent.h b/Engine/Plugins/Experimental/Mutable/Source/CustomizableObject/Classes/CustomizableSkeletalComponent.h
new file mode 100644
--- /dev/null
+++ b/Engine/Plugins/Experimental/Mutable/Source/CustomizableObject/Classes/CustomizableSkeletalComponent.h
@@ -0,0 +1,71 @@
+// Copyright Epic Games, Inc. All Rights Reserved.
+
+#pragma once
+
+#include "Components/SceneComponent.h"
+#include "CustomizableSkeletalComponent.generated.h"
+
+namespace mu
+{
+	class Instance;
+};
+
+class USkeletalMesh;
+
+
+DECLARE_DELEGATE_TwoParams(FCustomizableSkeletalComponentPreUpdateDelegate, UCustomizableSkeletalComponent* Component, USkeletalMesh* NextMesh);
+DECLARE_DELEGATE(FCustomizableSkeletalComponentUpdatedDelegate);
+
+UCLASS(Blueprintable, BlueprintType, ClassGroup = (CustomizableObject), meta = (BlueprintSpawnableComponent))
+class CUSTOMIZABLEOBJECT_API UCustomizableSkeletalComponent : public USceneComponent
+{
+public:
+	GENERATED_BODY()
+
+	UCustomizableSkeletalComponent();
+
+	// Used to replace the SkeletalMesh of the parent component by the ReferenceSkeletalMesh or the generated SkeletalMesh 
+	bool bPendingSetSkeletalMesh = false;
+
+	// Used to avoid replacing the SkeletalMesh of the parent component by the ReferenceSkeletalMesh if bPendingSetSkeletalMesh is true
+	bool bSkipSetReferenceSkeletalMesh = false;
+
+	UPROPERTY(Transient)
+	float SkippedLastRenderTime;
+
+	UPROPERTY(BlueprintReadWrite, EditAnywhere, Category = CustomizableSkeletalMesh)
+	TObjectPtr<class UCustomizableObjectInstance> CustomizableObjectInstance;
+
+	UPROPERTY(BlueprintReadWrite, EditAnywhere, Category = CustomizableSkeletalMesh)
+	int32 ComponentIndex;
+
+	FCustomizableSkeletalComponentPreUpdateDelegate PreUpdateDelegate;
+	FCustomizableSkeletalComponentUpdatedDelegate UpdatedDelegate;
+
+	USkeletalMesh* GetSkeletalMesh() const;
+	USkeletalMesh* GetAttachedSkeletalMesh() const;
+
+	UFUNCTION(BlueprintCallable, Category = CustomizableObject)
+	void UpdateSkeletalMeshAsync(bool bNeverSkipUpdate = false);
+
+	void SetSkeletalMesh(USkeletalMesh* SkeletalMesh, bool bReinitPose=true, bool bForceClothReset=true);
+	void SetPhysicsAsset(class UPhysicsAsset* PhysicsAsset);
+	void UpdateDistFromComponentToPlayer(const AActor* const Pawn, bool bForceEvenIfNotBegunPlay = false);
+
+	void SetVisibilityOfSkeletalMeshSectionWithMaterialName(bool bVisible, const FString& MaterialName, int32 LOD);
+
+	virtual void BeginDestroy() override;
+	virtual void PreSave(FObjectPreSaveContext ObjectSaveContext) override;
+
+#if WITH_EDITOR
+	// Used to generate instances outside the CustomizableObject editor and PIE
+	void UpdateDistFromComponentToLevelEditorCamera(const FVector& CameraPosition);
+	void EditorUpdateComponent();
+#endif
+
+private:
+	void TickComponent(float DeltaTime, enum ELevelTick TickType, FActorComponentTickFunction* ThisTickFunction) override;
+	void OnAttachmentChanged() override;
+
+};
+
diff --git a/Engine/Plugins/Experimental/Mutable/Source/CustomizableObject/Classes/CustomizableSkeletalMeshActor.h b/Engine/Plugins/Experimental/Mutable/Source/CustomizableObject/Classes/CustomizableSkeletalMeshActor.h
new file mode 100644
--- /dev/null
+++ b/Engine/Plugins/Experimental/Mutable/Source/CustomizableObject/Classes/CustomizableSkeletalMeshActor.h
@@ -0,0 +1,36 @@
+// Copyright Epic Games, Inc. All Rights Reserved.
+
+#pragma once
+
+#include "CoreMinimal.h"
+#include "UObject/ObjectMacros.h"
+#include "Animation/SkeletalMeshActor.h"
+#include "CustomizableSkeletalComponent.h"
+#include "CustomizableSkeletalMeshActor.generated.h"
+
+UCLASS(ClassGroup = CustomizableObject, Blueprintable, ComponentWrapperClass, ConversionRoot, meta = (ChildCanTick))
+class CUSTOMIZABLEOBJECT_API ACustomizableSkeletalMeshActor : public ASkeletalMeshActor
+{
+	GENERATED_UCLASS_BODY()
+public:
+	
+	class UCustomizableSkeletalComponent* GetCustomizableSkeletalComponent(int32 Index = 0) const { return CustomizableSkeletalComponents[Index]; }
+	
+	class USkeletalMeshComponent* GetSkeletalMeshComponentAt(int32 Index = 0) const { return SkeletalMeshComponents[Index]; }
+
+	void AttachNewComponent();
+
+	int32 GetNumComponents() { return CustomizableSkeletalComponents.Num(); }
+
+private:
+	
+	UPROPERTY(Category = CustomizableSkeletalMesh, VisibleAnywhere, BlueprintReadOnly, meta = (AllowPrivateAccess = "true"))
+	TArray< TObjectPtr<class UCustomizableSkeletalComponent> > CustomizableSkeletalComponents;
+
+	UPROPERTY(Category = CustomizableSkeletalMesh, VisibleAnywhere, BlueprintReadOnly, meta = (AllowPrivateAccess = "true"))
+	TArray< TObjectPtr<class USkeletalMeshComponent>> SkeletalMeshComponents;
+
+	// TODO: This is a temporary fix to not break the demos, we should update the demos to support the arrays of components
+	UPROPERTY(Category = CustomizableSkeletalMesh, VisibleAnywhere, BlueprintReadOnly, meta = (AllowPrivateAccess = "true"))
+	TObjectPtr<class UCustomizableSkeletalComponent> CustomizableSkeletalComponent;
+};
\ No newline at end of file
diff --git a/Engine/Plugins/Experimental/Mutable/Source/CustomizableObject/Private/CustomizableInstanceLODManagement.cpp b/Engine/Plugins/Experimental/Mutable/Source/CustomizableObject/Private/CustomizableInstanceLODManagement.cpp
new file mode 100644
--- /dev/null
+++ b/Engine/Plugins/Experimental/Mutable/Source/CustomizableObject/Private/CustomizableInstanceLODManagement.cpp
@@ -0,0 +1,442 @@
+// Copyright Epic Games, Inc. All Rights Reserved.
+
+#include "CustomizableInstanceLODManagement.h"
+
+//#include "CustomizableObjectSystemPrivate.h"
+#include "GameFramework/PlayerController.h"
+#include "CustomizableObjectSystem.h"
+#include "CustomizableSkeletalComponent.h"
+#include "CustomizableObject.h"
+#include "Camera/CameraComponent.h"
+#include "Camera/CameraActor.h"
+#include "Kismet/GameplayStatics.h"
+
+#if WITH_EDITOR
+#include "LevelEditorViewport.h"
+#endif
+
+static TAutoConsoleVariable<int32> CVarNumGeneratedInstancesLimit(
+	TEXT("b.NumGeneratedInstancesLimit"),
+	0,
+	TEXT("If different than 0, limit the number of mutable instances with full LODs to have at a given time."),
+	ECVF_Scalability);
+
+static TAutoConsoleVariable<int32> CVarNumGeneratedInstancesLimitLOD1(
+	TEXT("b.NumGeneratedInstancesLimitLOD1"),
+	0,
+	TEXT("If different than 0, limit the number of mutable instances with LOD 1 to have at a given time."),
+	ECVF_Scalability);
+
+static TAutoConsoleVariable<int32> CVarNumGeneratedInstancesLimitLOD2(
+	TEXT("b.NumGeneratedInstancesLimitLOD2"),
+	0,
+	TEXT("If different than 0, limit the number of mutable instances with LOD 2 to have at a given time."),
+	ECVF_Scalability);
+
+static TAutoConsoleVariable<int32> CVarDistanceForFixedLOD2(
+	TEXT("b.DistanceForFixedLOD2"),
+	50000,
+	TEXT("If NumGeneratedInstancesLimit is different than 0, sets the distance at which the system will fix the LOD of an instance to the lowest res one (LOD2) to prevent unnecessary LOD changes and memory consumption"),
+	ECVF_Scalability);
+
+
+UCustomizableInstanceLODManagement::UCustomizableInstanceLODManagement() : UCustomizableInstanceLODManagementBase()
+{
+	bOnlyUpdateCloseCustomizableObjects = false;
+	CloseCustomizableObjectsDist = 2000.f;
+}
+
+
+UCustomizableInstanceLODManagement::~UCustomizableInstanceLODManagement()
+{
+
+}
+
+// Used to manually update distances used in "OnlyUpdateCloseCustomizableObjects" system. If OnlyForInstance is null, all instances have their distance updated
+// ViewCenter is the origin where the distances will be measured from.
+void UpdatePawnToInstancesDistances(const class UCustomizableObjectInstance* OnlyForInstance, const TWeakObjectPtr<const AActor> ViewCenter)
+{
+	for (TObjectIterator<UCustomizableSkeletalComponent> CustomizableSkeletalComponent; CustomizableSkeletalComponent; ++CustomizableSkeletalComponent)
+	{
+		if (CustomizableSkeletalComponent->IsValidLowLevel() && (OnlyForInstance == nullptr || CustomizableSkeletalComponent->CustomizableObjectInstance == OnlyForInstance))
+		{
+			CustomizableSkeletalComponent->UpdateDistFromComponentToPlayer(ViewCenter.IsValid() ? ViewCenter.Get() : nullptr, OnlyForInstance != nullptr);
+		}
+	}
+}
+
+#if WITH_EDITOR
+// Used to manually update instances in the level editor
+void UpdateCameraToInstancesDistance(const FVector CameraPosition)
+{
+	for (TObjectIterator<UCustomizableSkeletalComponent> CustomizableSkeletalComponent; CustomizableSkeletalComponent; ++CustomizableSkeletalComponent)
+	{
+		if (CustomizableSkeletalComponent->IsValidLowLevel() && !CustomizableSkeletalComponent->IsTemplate())
+		{
+			CustomizableSkeletalComponent->UpdateDistFromComponentToLevelEditorCamera(CameraPosition);
+		}
+	}
+}
+#endif
+
+
+void UCustomizableInstanceLODManagement::UpdateInstanceDistsAndLODs()
+{
+	UCustomizableObjectSystem* CustomizableObjectSystem = UCustomizableObjectSystem::GetInstance();
+
+	if (ViewCenters.Num() == 0) // Just use the first pawn
+	{
+		UCustomizableSkeletalComponent* FirstCustomizableSkeletalComponent = nullptr;
+
+#if WITH_EDITOR
+		bool bLevelEditorInstancesUpdated = false;
+#endif
+
+		for (TObjectIterator<UCustomizableSkeletalComponent> CustomizableSkeletalComponent; CustomizableSkeletalComponent; ++CustomizableSkeletalComponent)
+		{
+			if (CustomizableSkeletalComponent && !CustomizableSkeletalComponent->IsTemplate())
+			{
+				UWorld* LocalWorld = CustomizableSkeletalComponent->GetWorld();
+				APlayerController* Controller = LocalWorld ? LocalWorld->GetFirstPlayerController() : nullptr;
+				TWeakObjectPtr<const AActor> ViewCenter = Controller ? TWeakObjectPtr<const AActor>(Controller->GetPawn()) : nullptr;
+
+				if (ViewCenter.IsValid())
+				{
+					UpdatePawnToInstancesDistances(nullptr, ViewCenter);
+					break;
+				}
+
+#if WITH_EDITOR
+				else if (CustomizableSkeletalComponent->GetWorld())
+				{
+					EWorldType::Type worldType = CustomizableSkeletalComponent->GetWorld()->WorldType;
+					
+					// Level Editor Instances (non PIE)
+					if (!bLevelEditorInstancesUpdated && worldType == EWorldType::Editor)
+					{
+						for (FLevelEditorViewportClient* LevelVC : GEditor->GetLevelViewportClients())
+						{
+							if (LevelVC && LevelVC->IsPerspective())
+							{
+								UpdateCameraToInstancesDistance(LevelVC->GetViewLocation());
+								bLevelEditorInstancesUpdated = true;
+								break;
+							}
+						}
+					}
+
+					// Blueprint instances
+					else if (worldType == EWorldType::EditorPreview)
+					{
+						CustomizableSkeletalComponent->EditorUpdateComponent();
+					}
+				}
+#endif // WITH_EDITOR
+
+			}
+		}
+	}
+	else
+	{
+		for (const TWeakObjectPtr<const AActor> ViewCenter : ViewCenters)
+		{
+			if (ViewCenter.IsValid())
+			{
+				UpdatePawnToInstancesDistances(nullptr, ViewCenter);
+			}
+		}
+	}
+
+	int32 NumGeneratedInstancesLimitLODs = GetNumGeneratedInstancesLimitFullLODs();
+
+	if (NumGeneratedInstancesLimitLODs > 0)
+	{
+		int32 NumGeneratedInstancesLimitLOD1 = GetNumGeneratedInstancesLimitLOD1();
+		int32 NumGeneratedInstancesLimitLOD2 = GetNumGeneratedInstancesLimitLOD2();
+
+		TArray<UCustomizableObjectInstance*> SortedInstances;
+
+		for (TObjectIterator<UCustomizableObjectInstance> CustomizableObjectInstance; CustomizableObjectInstance; ++CustomizableObjectInstance)
+		{
+			if (IsValidChecked(*CustomizableObjectInstance) &&
+				CustomizableObjectInstance->GetPrivate() &&
+				CustomizableObjectInstance->GetIsBeingUsedByComponentInPlay())
+			{
+				if (const UCustomizableObject* CustomizableObject = CustomizableObjectInstance->GetCustomizableObject();
+					CustomizableObject &&
+					!CustomizableObject->IsLocked())
+				{
+					UCustomizableObjectInstance* Ptr = *CustomizableObjectInstance;
+					Ptr->SetIsDiscardedBecauseOfTooManyInstances(false);
+					SortedInstances.Add(Ptr);
+				}
+			}
+		}
+
+		for (int32 i = 0; i < GetNumberOfPriorityUpdateInstances() && i < SortedInstances.Num(); ++i)
+		{
+			SortedInstances[i]->SetIsPlayerOrNearIt(true);
+		}
+
+		if (SortedInstances.Num() > NumGeneratedInstancesLimitLODs)
+		{
+			SortedInstances.Sort([](UCustomizableObjectInstance& A, UCustomizableObjectInstance& B)
+			{
+				return A.GetMinSquareDistToPlayer() < B.GetMinSquareDistToPlayer();
+			});
+
+			bool bAlreadyReachedFixedLOD = false;
+			const float DistanceForFixedLOD = float(CVarDistanceForFixedLOD2.GetValueOnGameThread());
+			const float DistanceForFixedLODSquared = FMath::Square(DistanceForFixedLOD);
+
+			for (int32 i = 0; i < NumGeneratedInstancesLimitLODs && i < SortedInstances.Num(); ++i)
+			{
+				if (!bAlreadyReachedFixedLOD && SortedInstances[i]->GetMinSquareDistToPlayer() < DistanceForFixedLODSquared)
+				{
+					SortedInstances[i]->SetMinMaxLODToLoad(0, INT32_MAX);
+				}
+				else
+				{
+					SortedInstances[i]->SetMinMaxLODToLoad(2, 2);
+					bAlreadyReachedFixedLOD = true;
+				}
+			}
+
+			for (int32 i = NumGeneratedInstancesLimitLODs; i < NumGeneratedInstancesLimitLODs + NumGeneratedInstancesLimitLOD1 && i < SortedInstances.Num(); ++i)
+			{
+				if (!bAlreadyReachedFixedLOD && SortedInstances[i]->GetMinSquareDistToPlayer() < DistanceForFixedLODSquared)
+				{
+					SortedInstances[i]->SetMinMaxLODToLoad(1, 2);
+				}
+				else
+				{
+					SortedInstances[i]->SetMinMaxLODToLoad(2, 2);
+					bAlreadyReachedFixedLOD = true;
+				}
+			}
+
+			for (int32 i = NumGeneratedInstancesLimitLODs + NumGeneratedInstancesLimitLOD1; i < NumGeneratedInstancesLimitLODs + NumGeneratedInstancesLimitLOD1 + NumGeneratedInstancesLimitLOD2 && i < SortedInstances.Num(); ++i)
+			{
+				SortedInstances[i]->SetMinMaxLODToLoad(2, 2);
+			}
+
+			for (int32 i = NumGeneratedInstancesLimitLODs + NumGeneratedInstancesLimitLOD1 + NumGeneratedInstancesLimitLOD2; i < SortedInstances.Num(); ++i)
+			{
+				SortedInstances[i]->SetIsDiscardedBecauseOfTooManyInstances(true);
+			}
+		}
+		else
+		{
+			// No limit surpassed, set all instances to have all LODs, there will be an UpdateSkeletalMesh only if there's a change in LOD state
+			for (int32 i = 0; i < SortedInstances.Num(); ++i)
+			{
+				SortedInstances[i]->SetMinMaxLODToLoad(0, INT32_MAX);
+			}
+		}
+	}
+	else if (bOnlyUpdateCloseCustomizableObjects)
+	{
+		TArray<FVector> Locations;
+		TArray<FMatrix> ProjectionMatrices;
+		TMap <const AActor*, int32> ViewCentersToIndexMap;
+
+		UCameraComponent* PlayerCamera = nullptr;
+
+		if (ViewCenters.Num() == 0)
+		{
+			for (TObjectIterator<UCustomizableSkeletalComponent> CustomizableSkeletalComponent; CustomizableSkeletalComponent; ++CustomizableSkeletalComponent)
+			{
+				UWorld* LocalWorld = CustomizableSkeletalComponent->GetWorld();
+				APlayerController* Controller = LocalWorld ? LocalWorld->GetFirstPlayerController() : nullptr;
+
+				if (Controller && Controller->GetViewTarget())
+				{
+					PlayerCamera = Cast<UCameraComponent>(Controller->GetViewTarget()->GetComponentByClass(UCameraComponent::StaticClass()));
+					ViewCentersToIndexMap.Add(Controller, 0);
+					break;
+				}
+			}
+
+			if (PlayerCamera && PlayerCamera->IsActive())
+			{
+				FMinimalViewInfo MinimalViewInfo;
+				PlayerCamera->GetCameraView(16.6f, MinimalViewInfo); // The Delta Time parameter is irrelevant
+
+				FMatrix ViewMatrix;
+				FMatrix ProjectionMatrix;
+				FMatrix ViewProjectionMatrix;
+				UGameplayStatics::GetViewProjectionMatrix(MinimalViewInfo, ViewMatrix, ProjectionMatrix, ViewProjectionMatrix);
+
+				Locations.Add(MinimalViewInfo.Location);
+				ProjectionMatrices.Add(ProjectionMatrix);
+			}
+		}
+		else
+		{
+			int32 index = 0;
+			for (const TWeakObjectPtr<const AActor> ViewCenter : ViewCenters)
+			{
+				if (!ViewCenter.IsValid())
+				{
+					continue;
+				}
+
+				UCameraComponent* Camera = nullptr;
+				if (const ACameraActor* CameraActor = Cast<ACameraActor>(ViewCenter.Get()))
+				{
+					Camera = CameraActor->GetCameraComponent();
+				}
+				else
+				{
+					Camera = Cast<UCameraComponent>(ViewCenter->GetComponentByClass(UCameraComponent::StaticClass()));
+				}
+				
+				if (Camera && Camera->IsActive())
+				{
+					FMinimalViewInfo MinimalViewInfo;
+					Camera->GetCameraView(16.6f, MinimalViewInfo); // The Delta Time parameter is irrelevant
+
+					FMatrix ViewMatrix;
+					FMatrix ProjectionMatrix;
+					FMatrix ViewProjectionMatrix;
+					UGameplayStatics::GetViewProjectionMatrix(MinimalViewInfo, ViewMatrix, ProjectionMatrix, ViewProjectionMatrix);
+
+					Locations.Add(MinimalViewInfo.Location);
+					ProjectionMatrices.Add(ProjectionMatrix);
+					ViewCentersToIndexMap.Add(ViewCenter.Get(), index);
+					index++;
+				}
+			}
+		}
+
+		if (Locations.Num())
+		{
+			float SquaredDist = FMath::Square(CloseCustomizableObjectsDist);
+			float ScreenSizeLOD1 = 0.04f;
+			float ScreenSizeLOD2 = 0.02f;
+
+			for (TObjectIterator<UCustomizableObjectInstance> CustomizableObjectInstance; CustomizableObjectInstance; ++CustomizableObjectInstance)
+			{
+				if ( IsValidChecked(*CustomizableObjectInstance) && CustomizableObjectInstance->GetPrivate() )
+				{
+					UCustomizableObject* CustomizableObject = CustomizableObjectInstance->GetCustomizableObject(); 
+					if (!CustomizableObject)
+					{
+						continue;
+					}
+														
+					if (CustomizableObjectInstance->GetMinSquareDistToPlayer() > SquaredDist || !CustomizableObjectInstance->NearestToActor.IsValid())
+					{
+						continue;
+					}
+
+					USkeletalMesh* RefSkeletalMesh = CustomizableObject->GetRefSkeletalMesh();
+					if (!RefSkeletalMesh)
+					{
+						continue;
+					}
+					
+					if (RefSkeletalMesh->GetLODInfoArray().Num() > 1)
+					{
+						ScreenSizeLOD1 = FMath::Square(RefSkeletalMesh->GetLODInfo(1)->ScreenSize.Default);
+					}
+					
+					if (RefSkeletalMesh->GetLODInfoArray().Num() > 2)
+					{
+						ScreenSizeLOD2 = FMath::Square(RefSkeletalMesh->GetLODInfo(2)->ScreenSize.Default);
+					}
+					
+					float BoundSize = RefSkeletalMesh->GetBounds().SphereRadius;
+					FVector BoundOrigin;
+
+					AActor* ParentActor = CustomizableObjectInstance->NearestToActor->GetAttachmentRootActor();
+					if (ParentActor)
+					{
+						BoundSize *= ParentActor->GetActorScale().Size() * 1.55f;
+						BoundOrigin = ParentActor->GetActorLocation();
+					}
+					else
+					{
+						BoundOrigin = CustomizableObjectInstance->NearestToActor->GetComponentLocation();
+					}
+
+
+					int32* ViewIndex = ViewCentersToIndexMap.Find(CustomizableObjectInstance->NearestToViewCenter.Get());
+					float InstanceScreenSize = ComputeBoundsScreenRadiusSquared(BoundOrigin, BoundSize, Locations[ViewIndex ? *ViewIndex : 0 ], ProjectionMatrices[ ViewIndex ? *ViewIndex : 0]);
+
+					int32 InstanceMaxLOD = FMath::Min(RefSkeletalMesh->GetLODNum(), CustomizableObjectInstance->GetNumLODsAvailable()) - 1;
+
+					bool bIsPawn = CustomizableObjectInstance->GetMinSquareDistToPlayer() <= 0;
+
+					if (InstanceMaxLOD == 0 || InstanceScreenSize > ScreenSizeLOD1 || bIsPawn)
+					{
+						CustomizableObjectInstance->SetMinMaxLODToLoad(0, InstanceMaxLOD, false);
+					}
+					else if (InstanceMaxLOD == 1 || InstanceScreenSize > ScreenSizeLOD2 )
+					{
+						CustomizableObjectInstance->SetMinMaxLODToLoad(1, InstanceMaxLOD, false);
+					}
+					else
+					{
+						CustomizableObjectInstance->SetMinMaxLODToLoad(2, 2);
+					}
+				}
+			}
+		}
+	}
+}
+
+
+int32 UCustomizableInstanceLODManagement::GetNumGeneratedInstancesLimitFullLODs()
+{
+	return CVarNumGeneratedInstancesLimit.GetValueOnGameThread();
+}
+
+
+int32 UCustomizableInstanceLODManagement::GetNumGeneratedInstancesLimitLOD1()
+{
+	return CVarNumGeneratedInstancesLimitLOD1.GetValueOnGameThread();
+}
+
+
+int32 UCustomizableInstanceLODManagement::GetNumGeneratedInstancesLimitLOD2()
+{
+	return CVarNumGeneratedInstancesLimitLOD2.GetValueOnGameThread();
+}
+
+
+void UCustomizableInstanceLODManagement::SetNumberOfPriorityUpdateInstances(int32 InNumPriorityUpdateInstances)
+{
+	NumPriorityUpdateInstances = InNumPriorityUpdateInstances;
+}
+
+
+int32 UCustomizableInstanceLODManagement::GetNumberOfPriorityUpdateInstances() const
+{
+	return NumPriorityUpdateInstances;
+}
+
+
+bool UCustomizableInstanceLODManagement::IsOnlyUpdateCloseCustomizableObjectsEnabled() const
+{
+	return bOnlyUpdateCloseCustomizableObjects;
+}
+
+
+void UCustomizableInstanceLODManagement::EnableOnlyUpdateCloseCustomizableObjects(float CloseDist)
+{
+	bOnlyUpdateCloseCustomizableObjects = true;
+	CloseCustomizableObjectsDist = CloseDist;
+}
+
+
+void UCustomizableInstanceLODManagement::DisableOnlyUpdateCloseCustomizableObjects()
+{
+	bOnlyUpdateCloseCustomizableObjects = false;
+}
+
+
+float UCustomizableInstanceLODManagement::GetOnlyUpdateCloseCustomizableObjectsDist() const
+{
+	return CloseCustomizableObjectsDist;
+}
diff --git a/Engine/Plugins/Experimental/Mutable/Source/CustomizableObject/Private/CustomizableInstancePrivateData.h b/Engine/Plugins/Experimental/Mutable/Source/CustomizableObject/Private/CustomizableInstancePrivateData.h
new file mode 100644
--- /dev/null
+++ b/Engine/Plugins/Experimental/Mutable/Source/CustomizableObject/Private/CustomizableInstancePrivateData.h
@@ -0,0 +1,379 @@
+// Copyright Epic Games, Inc. All Rights Reserved.
+
+#pragma once
+
+#include "MutableRuntime/Public/Model.h"
+#include "MutableRuntime/Public/System.h"
+#include "MutableRuntime/Public/Mesh.h"
+#include "MutableRuntime/Public/Image.h"
+#include "Engine/StaticMesh.h"
+#include "Engine/StreamableManager.h"
+
+#include "CustomizableObjectInstance.h"
+#include "CustomizableObjectSystemPrivate.h"
+#include "Engine/Texture2D.h"
+
+#include "ClothingAsset.h"
+#include "GameplayTagContainer.h"
+
+#include "CustomizableInstancePrivateData.generated.h"
+
+/** CustomizableObject Instance flags for internal use  */
+enum ECOInstanceFlags
+{
+	None							= 0,
+
+	// Update process
+	Updating						= 1 << 0,	//
+	CreatingSkeletalMesh			= 1 << 1,	//
+	Generated						= 1 << 2,	//
+	AssetsLoaded					= 1 << 3,	// Assets required to generate a mesh. Mainly Materials, PhysicsAssets and AnimBps 
+	ReduceLODs						= 1 << 4,	// CurrentMinLOD will be mapped to the LOD 0 of the generated SkeletalMesh. 
+	ReuseTextures					= 1 << 5, 	// 
+	ReplacePhysicsAssets			= 1 << 6,	// Merge active PhysicsAssets and replace the base physics asset
+	PendingMeshUpdate				= 1 << 7,	// Internaly used to know if we should regenerate the SkeletalMeshes when updating the instance
+
+	// Update priorities
+	UsedByComponent					= 1 << 8,	// If any components are using this instance, they will set flag every frame
+	UsedByComponentInPlay			= 1 << 9,	// If any components are using this instance in play, they will set flag every frame
+	UsedByPlayerOrNearIt			= 1 << 10,	// The instance is used by the player or is near the player, used to give more priority to its updates
+	DiscardedByNumInstancesLimit	= 1 << 11,	// The instance is descarded because we exceeded the limit of instances generated 
+
+	// Types of updates
+	PendingLODsUpdate				= 1 << 12,	// Used to queue an update due to a change in LODs required by the instance
+	PendingLODsUpdateSecondStage	= 1 << 13,	// Second stage of the previous flag, Internaly used to queue an update due to a change in LODs required by the instance
+	PendingLODsDowngrade			= 1 << 14,	// Used to queue a downgrade update to reduce the number of LODs. LOD update goes from a high res level to a low res one, ex: 0 to 1 or 1 to 2
+
+	// Streaming
+	LODsStreamingEnabled			= 1 << 15,	// Stream LODs instead of generating all LODs at once. Enables LODs update(upgrade)/downgrade.
+};
+
+
+USTRUCT()
+struct FGeneratedTexture
+{
+	GENERATED_USTRUCT_BODY();
+
+	UPROPERTY(Category = CustomizableObjectInstance, EditAnywhere)
+	int32 Id = 0;
+
+	UPROPERTY(Category = CustomizableObjectInstance, EditAnywhere) 
+	FString Name;
+
+	UPROPERTY(Category = CustomizableObjectInstance, EditAnywhere) 
+	TObjectPtr<UTexture2D> Texture = nullptr;
+};
+
+
+USTRUCT()
+struct FGeneratedMaterial
+{
+	GENERATED_USTRUCT_BODY();
+
+	UPROPERTY(Category = CustomizedMaterial, VisibleAnywhere)
+	TArray< FGeneratedTexture > Textures;
+};
+
+
+// Unreal-side data for the parameter decorations.
+USTRUCT()
+struct FParameterDecorations
+{
+	GENERATED_USTRUCT_BODY();
+
+	UPROPERTY(Transient)
+	TArray< TObjectPtr<UTexture2D> > Images;
+};
+
+
+USTRUCT()
+struct FReferencedPhysicsAssets
+{
+	GENERATED_USTRUCT_BODY();
+	
+	TArray<FString> PhysicsAssetToLoad;
+	
+	UPROPERTY(Transient)
+	TArray< TObjectPtr<UPhysicsAsset> > PhysicsAssetsToMerge;
+};
+
+
+USTRUCT()
+struct FReferencedSkeletons
+{
+	GENERATED_USTRUCT_BODY();
+
+	UPROPERTY()
+	TArray<int32> SkeletonsToLoad;
+
+	UPROPERTY()
+	TArray< TObjectPtr<USkeleton> > SkeletonsToMerge;
+};
+
+
+USTRUCT()
+struct FCustomizableInstanceComponentData
+{
+	GENERATED_USTRUCT_BODY();
+
+	uint16 ComponentIndex = 0;
+
+	// AnimBP data gathered for a component from its constituent meshes
+	UPROPERTY(Transient, Category = CustomizableObjectInstance, editfixedsize, VisibleAnywhere)
+	TMap<int32, TSoftClassPtr<UAnimInstance>> AnimSlotToBP;
+
+	/** Skeletons required by the current generated instance. Skeletons to be loaded and merged.*/
+	UPROPERTY(Transient)
+	FReferencedSkeletons Skeletons;
+	
+	/** PhysicsAssets required by the current generated instance. PhysicsAssets to be loaded and merged.*/
+	UPROPERTY(Transient)
+	FReferencedPhysicsAssets PhysicsAssets;
+
+	/** Clothing PhysicsAssets required by the current generated instance. PhysicsAssets to be loaded and merged.*/
+	TArray<TPair<int32, FString>> ClothingPhysicsAssetsToStream;
+
+	/** Array of generated MeshIds per each LOD, used to decide if the mesh should be updated or not.
+	 *  Size == NumLODsAvailable
+	 *  LODs without mesh will be set to -1 */
+	TArray<int32> LastMeshIdPerLOD;
+
+	bool operator==(const FCustomizableInstanceComponentData& Other) const { return ComponentIndex == Other.ComponentIndex; }
+};
+
+
+UCLASS()
+class UCustomizableInstancePrivateData : public UObject
+{
+public:
+	GENERATED_BODY()
+
+	UCustomizableInstancePrivateData();
+
+	UPROPERTY( Transient )
+	TArray<FGeneratedMaterial> GeneratedMaterials;
+
+	UPROPERTY( Transient )
+	TArray<FGeneratedTexture> GeneratedTextures;
+
+	// Projector state for each parameter name
+	TMap<TPair<FString, int32>, EProjectorState::Type> ProjectorStates;
+
+	// Parameter decoration unreal-side data if generated.
+	UPROPERTY(Transient)
+	TArray<FParameterDecorations> ParameterDecorations;
+
+	// Indices of the parameters that are relevant for the given parameter values.
+	// This only gets updated if parameter decorations are generated.
+	TArray<int> RelevantParameters;
+
+	// If Texture reuse is enabled, stores which texture is being used in a particular <LODIndex, ComponentIndex, MeshSurfaceIndex, image>
+	// \TODO: Create a key based on a struct instead of generating strings dynamically.
+	UPROPERTY(Transient)
+	TMap<FString, TWeakObjectPtr<UTexture2D>> TextureReuseCache;
+
+	// Handle used to store a streaming request operation if one is ongoing.
+	TSharedPtr<FStreamableHandle> StreamingHandle;
+
+	FCustomizableInstanceComponentData* GetComponentData(int32 ComponentIndex);
+	const FCustomizableInstanceComponentData* GetComponentData(int32 ComponentIndex) const;
+
+	ECOInstanceFlags GetCOInstanceFlags() const { return InstanceFlagsPrivate; }
+	void SetCOInstanceFlags(ECOInstanceFlags FlagsToSet) { InstanceFlagsPrivate = (ECOInstanceFlags)(InstanceFlagsPrivate | FlagsToSet); }
+	void ClearCOInstanceFlags(ECOInstanceFlags FlagsToClear) { InstanceFlagsPrivate = (ECOInstanceFlags)(InstanceFlagsPrivate & ~FlagsToClear); }
+	bool HasCOInstanceFlags(ECOInstanceFlags FlagsToCheck) const { return (InstanceFlagsPrivate & FlagsToCheck) != 0; }
+
+	void BuildMaterials(const TSharedPtr<FMutableOperationData>& OperationData, UCustomizableObjectInstance* Public);
+
+	void ReuseTexture(UTexture2D* Texture);
+
+	// Return an event that will be fired when the assets  have been loaded. It returns null if no asset needs loading.
+	FGraphEventRef LoadAdditionalAssetsAsync(const TSharedPtr<FMutableOperationData>& OperationData, UCustomizableObjectInstance* Public, struct FStreamableManager &StreamableManager);
+	void AdditionalAssetsAsyncLoaded(UCustomizableObjectInstance* Public);
+	
+	mu::ParametersPtr ReloadParametersFromObject(UCustomizableObjectInstance* Public, bool ClearLastMeshIds = false);
+
+	void TickUpdateCloseCustomizableObjects(UCustomizableObjectInstance* Public);
+	void UpdateInstanceIfNotGenerated(UCustomizableObjectInstance* Public, bool bAsync);
+
+	// Returns true if success (?)
+	bool UpdateSkeletalMesh_PostBeginUpdate0(UCustomizableObjectInstance* Public, const TSharedPtr<FMutableOperationData>& OperationData);
+
+	static void ReleaseMutableTexture(int32 MutableTextureId, UTexture2D* Texture, struct FMutableResourceCache& Cache);
+
+	// Parameter decoration generation: game thread step
+	void UpdateParameterDecorationsEngineResources(const TSharedPtr<FMutableOperationData>& OperationData);
+
+	// Prepare the data for the unreal textures, but don't create them because
+	// it runs in the mutable thread
+	// \TODO: CustomizableObject shouldn't be here
+	static void ProcessTextureCoverageQueries(const TSharedPtr<FMutableOperationData>& OperationData, UCustomizableObject* CustomizableObject, const FString& ImageKeyName, FTexturePlatformData *PlatformData, UMaterialInterface* Material);
+
+	// Prepare the data for the unreal reference skeleton
+	// it runs in the mutable thread
+	static void PrepareSkeletonData(const TSharedPtr<FMutableOperationData>& OperationData);
+
+	// Copy data generated in the mutable thread over to the instance and initializes additional data required during the update
+	void PrepareForUpdate(const TSharedPtr<FMutableOperationData>& OperationData);
+
+	void SetMinMaxLODToLoad(int32 NewMinLOD, int32 NewMaxLOD, bool bLimitLODUpgrades = true);
+	int32 GetMinLODToLoad() { return CurrentMinLOD; }
+	int32 GetMaxLODToLoad() { return CurrentMaxLOD; }
+	
+	int32 GetNumLODsAvailable() const { return NumLODsAvailable; }
+	
+	// The following method is basically copied from PostEditChangeProperty and/or SkeletalMesh.cpp to be able to replicate PostEditChangeProperty without the editor
+	void PostEditChangePropertyWithoutEditor(USkeletalMesh* SkeletalMesh);
+	
+	void DoUpdateSkeletalMeshAsync(UCustomizableObjectInstance* Instance, FMutableQueueElem::EQueuePriorityType Priority = FMutableQueueElem::EQueuePriorityType::Low);
+	void DoUpdateSkeletalMesh(UCustomizableObjectInstance* Instance, bool bAsync, bool bIsCloseDistTick=false, bool bOnlyUpdateIfNotGenerated=false, bool bIgnoreCloseDist=false, bool bForceHighPriority = false);
+
+	void DiscardResourcesAndSetReferenceSkeletalMesh(UCustomizableObjectInstance* Public);
+
+	/** 
+	* \param OnlyLOD: If not 0, extract and convert only one single LOD from the source image.
+	* \param ExtractChannel: If different than -1, extract a single-channel image with the specified source channel data.
+	*/
+	CUSTOMIZABLEOBJECT_API static void ConvertImage(class UTexture2D* Texture, mu::ImagePtrConst MutableImage, const FMutableModelImageProperties& Props, int32 OnlyLOD=-1, int32 ExtractChannel=-1);
+
+	/** Set OnlyLOD to -1 to generate all mips */
+	CUSTOMIZABLEOBJECT_API static FTexturePlatformData* MutableCreateImagePlatformData(const mu::Image* MutableImage, int32 OnlyLOD, uint16 FullSizeX, uint16 FullSizeY);
+
+private:
+
+	bool BuildSkeletalMeshSkeletonData(const TSharedPtr<FMutableOperationData>& OperationData, USkeletalMesh* SkeletalMesh, const FMutableRefSkeletalMeshData* RefSkeletalMeshData, UCustomizableObjectInstance* CustomizableObjectIntance, int32 ComponentIndex);
+	void BuildMorphTargetsData(const TSharedPtr<FMutableOperationData>& OperationData, USkeletalMesh* SkeletalMesh, UCustomizableObjectInstance* CustomizableObjectInstance, int32 ComponentIndex);
+	void BuildClothingData(const TSharedPtr<FMutableOperationData>& OperationData, USkeletalMesh* SkeletalMesh, UCustomizableObjectInstance* CustomizableObjectInstance, int32 ComponentIndex);
+	void BuildSkeletalMeshElementData(const TSharedPtr<FMutableOperationData>& OperationData, USkeletalMesh* SkeletalMesh, UCustomizableObjectInstance* CustomizableObjectInstance, int32 ComponentIndex);
+	bool BuildSkeletalMeshRenderData(const TSharedPtr<FMutableOperationData>& OperationData, USkeletalMesh* SkeletalMesh, UCustomizableObjectInstance* CustomizableObjectInstance, int32 ComponentIndex);
+
+	//
+	USkeleton* MergeSkeletons(UCustomizableObjectInstance* Public, const FMutableRefSkeletalMeshData* RefSkeletalMeshData, int32 ComponentIndex);
+
+	//
+	UPhysicsAsset* MergePhysicalAssets(class UPhysicsAsset* BaseAsset, int32 ComponentIndex);
+	void RefitPhysicsBodies(UPhysicsAsset* InOutPhysicsAsset, const mu::PhysicsBody* PhysicsBody);
+	
+	
+	int32 GetLastMeshId(int32 ComponentIndex, int32 LODIndex) const;
+	void SetLastMeshId(int32 ComponentIndex, int32 LODIndex, int32 MeshId);
+	void ClearAllLastMeshIds();
+
+	bool MeshNeedsUpdate(UCustomizableObjectInstance* Public, const TSharedPtr<FMutableOperationData>& OperationData, bool& bOutEmptyMesh);
+
+	UTexture2D* CreateTexture();
+	
+public:
+
+	// If any components are using this instance, they will store the min of their distances to the player here every frame for LOD purposes
+	float MinSquareDistFromComponentToPlayer;
+	float LastMinSquareDistFromComponentToPlayer; // The same as the previous dist for last frame
+												
+	double LastUpdateTime;
+
+	// These are the LODs that MUST be used in an update through GetMinLODToLoad and GetMaxLODToLoad, they are thread safe
+	int32 CurrentMinLOD;
+	int32 CurrentMaxLOD;
+
+	// These are the LODs we want to have, they MUSN'T be used in an update, they aren't thread safe to use from the mutable thread
+	int32 MinLODToLoad;
+	int32 MaxLODToLoad;
+	
+	// This is the LODs that the Customizable Object has
+	int32 NumLODsAvailable;
+
+	// First SkeletalMesh LOD we can generate on the running platform
+	uint8 FirstLODAvailable;
+
+	// Maximum number of SkeletalMesh LODs to stream
+	uint8 NumMaxLODsToStream;
+	
+	
+	TMap<FString, FTextureCoverageQueryData> TextureCoverageQueries;
+
+	UPROPERTY(Transient)
+	TArray<FCustomizableInstanceComponentData> ComponentsData;
+
+	UPROPERTY(Transient)
+	TArray< TObjectPtr<UMaterialInterface> > ReferencedMaterials;
+
+	// Converts a ReferencedMaterials index from the CustomizableObject to an index in the ReferencedMaterials in the Instance
+	TMap<uint32, uint32> ObjectToInstanceIndexMap;
+
+	TMap<uint32, FGeneratedTexture> TexturesToRelease;
+
+	UPROPERTY(Transient)
+	TArray< TObjectPtr<UPhysicsAsset> > ClothingPhysicsAssets;
+
+	// To keep loaded AnimBPs referenced and prevent GC
+	UPROPERTY(Transient, Category = CustomizableObjectInstance, editfixedsize, VisibleAnywhere)
+	TArray<TSubclassOf<UAnimInstance>> GatheredAnimBPs;
+
+	UPROPERTY(Transient, Category = CustomizableObjectInstance, editfixedsize, VisibleAnywhere)
+	FGameplayTagContainer AnimBPGameplayTags;
+
+private:
+	
+	ECOInstanceFlags InstanceFlagsPrivate = ECOInstanceFlags::None;
+
+	// Struct used by BuildMaterials() to identify common materials between LODs
+	struct FMutableMaterialPlaceholder
+	{
+		enum class EPlaceHolderParamType { Vector, Scalar, Texture };
+
+		struct FMutableMaterialPlaceHolderParam
+		{
+			FName ParamName;
+			int32 LayerIndex; // Set to -1 for non-multilayer params
+			FLinearColor Vector;
+			float Scalar;
+			FGeneratedTexture Texture;
+			EPlaceHolderParamType Type;
+
+			FMutableMaterialPlaceHolderParam(const FName& InParamName, const int32 InLayerIndex, const FLinearColor& InVector)
+				: ParamName(InParamName), LayerIndex(InLayerIndex), Vector(InVector), Type(EPlaceHolderParamType::Vector) {}
+
+			FMutableMaterialPlaceHolderParam(const FName& InParamName, const int32 InLayerIndex, const float InScalar)
+				: ParamName(InParamName), LayerIndex(InLayerIndex), Scalar(InScalar), Type(EPlaceHolderParamType::Scalar) {}
+
+			FMutableMaterialPlaceHolderParam(const FName& InParamName, const int32 InLayerIndex, const FGeneratedTexture& InTexture)
+				: ParamName(InParamName), LayerIndex(InLayerIndex), Texture(InTexture), Type(EPlaceHolderParamType::Texture) {}
+		};
+
+		UMaterialInterface* ParentMaterial;
+		TArray<FMutableMaterialPlaceHolderParam> Params;
+		int32 MatIndex = -1;
+
+		void AddParam(const FMutableMaterialPlaceHolderParam& NewParam) { Params.Add(NewParam); }
+
+		FString GetSerialization() const
+		{
+			FString Serialization = ParentMaterial ? FString::FromInt(ParentMaterial->GetUniqueID()) : FString("null");
+
+			for (const FMutableMaterialPlaceHolderParam& Param : Params)
+			{
+				Serialization += FString("-");
+				Serialization += FString::FromInt((int32)Param.Type) + FString("_") + Param.ParamName.ToString() + FString("_")
+					+ FString::FromInt(Param.LayerIndex) + FString("_");
+
+				switch (Param.Type)
+				{
+				case EPlaceHolderParamType::Vector:
+					Serialization += Param.Vector.ToString();
+					break;
+
+				case EPlaceHolderParamType::Scalar:
+					Serialization += FString::Printf(TEXT("%f"), Param.Scalar);
+					break;
+
+				case EPlaceHolderParamType::Texture:
+					Serialization += FString::FromInt(Param.Texture.Texture->GetUniqueID());
+					break;
+				}
+			}
+
+			return Serialization;
+		}
+	};
+};
+
diff --git a/Engine/Plugins/Experimental/Mutable/Source/CustomizableObject/Private/CustomizableObject.cpp b/Engine/Plugins/Experimental/Mutable/Source/CustomizableObject/Private/CustomizableObject.cpp
new file mode 100644
--- /dev/null
+++ b/Engine/Plugins/Experimental/Mutable/Source/CustomizableObject/Private/CustomizableObject.cpp
@@ -0,0 +1,1682 @@
+// Copyright Epic Games, Inc. All Rights Reserved.
+
+#include "CustomizableObject.h"
+
+#include "AssetRegistry/AssetRegistryModule.h"
+#include "Serialization/MemoryReader.h"
+#include "Serialization/MemoryWriter.h"
+#include "Misc/Paths.h"
+#include "HAL/FileManager.h"
+#include "HAL/PlatformFileManager.h"
+#include "Interfaces/ITargetPlatform.h"
+#include "Runtime/Launch/Resources/Version.h"
+#include "Algo/Count.h"
+#include "Algo/Copy.h"
+#include "PhysicsEngine/PhysicsAsset.h"
+#include "ClothConfig.h"
+#include "UObject/ObjectSaveContext.h"
+
+#include "MutableRuntime/Public/Model.h"
+#include "MutableRuntime/Public/Mesh.h"
+#include "MutableRuntime/Public/Image.h"
+#include "MutableRuntime/Public/Layout.h"
+#include "MutableRuntime/Public/System.h"
+#include "MutableRuntime/Public/MutableTrace.h"
+#include "MutableRuntime/Private/ModelPrivate.h"
+
+#include "ICustomizableObjectModule.h"
+#include "CustomizableObjectPrivate.h"
+#include "CustomizableObjectSystem.h"
+#include "UnrealMutableModelDiskStreamer.h"
+
+#include "Materials/MaterialInstance.h"
+#include "Misc/SecureHash.h"
+
+#include "ClothingAsset.h"
+
+#define LOCTEXT_NAMESPACE "CustomizableObject"
+
+DEFINE_LOG_CATEGORY(LogMutable);
+
+#if WITH_EDITOR
+static bool bUsesOnCookStart = false;
+#endif
+//-------------------------------------------------------------------------------------------------
+
+UCustomizableObject::UCustomizableObject()
+	: UObject()
+{
+	PrivateData = TSharedPtr<FCustomizableObjectPrivateData>( new FCustomizableObjectPrivateData() );
+
+#if WITH_EDITORONLY_DATA
+	const FString CVarName = TEXT("r.SkeletalMesh.MinLodQualityLevel");
+	const FString ScalabilitySectionName = TEXT("ViewDistanceQuality");
+	LODSettings.MinQualityLevelLOD.Init(*CVarName, *ScalabilitySectionName);
+#endif
+}
+
+
+#if WITH_EDITOR
+bool UCustomizableObject::IsEditorOnly() const
+{
+	return bIsChildObject;
+}
+
+
+void UCustomizableObject::UpdateVersionId()
+{
+	VersionId = FGuid::NewGuid();
+}
+
+
+void UCustomizableObject::GetAssetRegistryTags(TArray<FAssetRegistryTag>& OutTags) const
+{
+	int32 isRoot = 0;
+#if WITH_EDITOR
+	FCustomizableObjectCompilerBase* Compiler = UCustomizableObjectSystem::GetInstance()->GetNewCompiler();
+	if (Compiler)
+	{
+		isRoot = Compiler->IsRootObject(this) ? 1 : 0;
+		delete Compiler;
+	}
+#endif
+
+	OutTags.Add(FAssetRegistryTag("IsRoot", FString::FromInt(isRoot), FAssetRegistryTag::TT_Numerical));
+	Super::GetAssetRegistryTags(OutTags);
+}
+
+
+void UCustomizableObject::PreSave(FObjectPreSaveContext ObjectSaveContext)
+{
+	Super::PreSave(ObjectSaveContext);
+
+	// Update the derived child object flag
+	FCustomizableObjectCompilerBase* Compiler = UCustomizableObjectSystem::GetInstance()->GetNewCompiler();
+	if (Compiler)
+	{
+		bIsChildObject = !Compiler->IsRootObject(this);
+		delete Compiler;
+
+		if (bIsChildObject)
+		{
+			GetPackage()->SetPackageFlags(PKG_EditorOnly);
+		}
+		else
+		{
+			GetPackage()->ClearPackageFlags(PKG_EditorOnly);
+		}
+	}
+
+	if (!Identifier.IsValid())
+	{
+		Identifier = FGuid::NewGuid();
+	}
+
+#if WITH_EDITORONLY_DATA
+	if (ObjectSaveContext.IsCooking() && !bIsChildObject)
+	{
+		if (const FMutableCachedPlatformData* PlatformData = CachedPlatformsData.Find(ObjectSaveContext.GetTargetPlatform()->PlatformName()))
+		{
+			// Load cached data before saving
+			FMemoryReaderView MemoryReader(PlatformData->ModelData);
+			LoadCompiledData(MemoryReader, true);
+
+			// Create an export object to manage the streamable data
+			BulkData = NewObject<UCustomizableObjectBulk>(this);
+			BulkData->Mark(OBJECTMARK_TagExp);
+
+			// Split streamable data into smaller chunks and fix up the CO HashToStreamableBlock's FileIndex and Offset
+			BulkData->PrepareBulkData(this, ObjectSaveContext.GetTargetPlatform());
+		}
+		else
+		{
+			UE_LOG(LogMutable, Warning, TEXT("Cook: Customizable Object [%s] is missing [%s] platform data."), *GetName(),
+				*ObjectSaveContext.GetTargetPlatform()->PlatformName());
+			
+			ClearCompiledData();
+			SetModel(nullptr);
+		}
+	}
+#endif
+}
+
+
+#endif // End WITH_EDITOR
+
+
+void UCustomizableObject::PostLoad()
+{
+	Super::PostLoad();
+
+	// Make sure mutable has been initialised.
+	UCustomizableObjectSystem::GetInstance();
+
+#if WITH_EDITOR
+	if (ReferenceSkeletalMesh_DEPRECATED)
+	{
+		ReferenceSkeletalMeshes.Add(ReferenceSkeletalMesh_DEPRECATED);
+		ReferenceSkeletalMesh_DEPRECATED = nullptr;
+	}
+#endif
+}
+
+
+bool UCustomizableObject::IsLocked() const
+{
+	if (PrivateData.IsValid())
+	{
+		return PrivateData->bLocked;
+	}
+
+	return false;
+}
+
+
+void UCustomizableObject::Serialize(FArchive& Ar_Asset)
+{
+	MUTABLE_CPUPROFILER_SCOPE(UCustomizableObject::Serialize)
+	
+	Super::Serialize(Ar_Asset);
+
+#if WITH_EDITOR
+	if (Ar_Asset.IsCooking())
+	{
+		if (Ar_Asset.IsSaving())
+		{
+			UE_LOG(LogMutable, Verbose, TEXT("Serializing cooked data for Customizable Object [%s]."), *GetName());
+			SaveEmbeddedData(Ar_Asset);
+		}
+	}
+	else
+	{
+		// Can't remove this or saved customizable objects will fail to load
+		int64 InternalVersion = CurrentSupportedVersion;
+		Ar_Asset << InternalVersion;
+		
+		if (Ar_Asset.IsLoading())
+		{
+			LoadCompiledDataFromDisk();
+		}
+	}
+#else
+	if (Ar_Asset.IsLoading())
+	{
+		LoadEmbeddedData(Ar_Asset);
+	}
+#endif
+}
+
+
+#if WITH_EDITOR
+void UCustomizableObject::PostRename(UObject * OldOuter, const FName OldName)
+{
+	Super::PostRename(OldOuter, OldName);
+
+	if (Source)
+	{
+		Source->PostRename(OldOuter, OldName);
+	}
+}
+
+
+void UCustomizableObject::PostDuplicate(EDuplicateMode::Type DuplicateMode)
+{
+	Super::PostDuplicate(DuplicateMode);
+
+	if (DuplicateMode == EDuplicateMode::Normal)
+	{
+		// The only place where we can change the identifier
+		Identifier = FGuid::NewGuid();
+
+		// Create a new Private Data or it will use the same as the original Customizable Object
+		PrivateData = TSharedPtr<FCustomizableObjectPrivateData>(new FCustomizableObjectPrivateData());
+	}
+}
+
+
+void UCustomizableObject::BeginCacheForCookedPlatformData(const ITargetPlatform* TargetPlatform)
+{
+	bool bIsRelevantForThisTarget =
+		(Relevancy == ECustomizableObjectRelevancy::All)
+		||
+		(TargetPlatform && Relevancy == ECustomizableObjectRelevancy::ClientOnly && !TargetPlatform->IsServerOnly());
+	
+	if (TargetPlatform && bIsRelevantForThisTarget)
+	{
+		if (PrivateData->CachedPlatformNames.Find(TargetPlatform->PlatformName()) == INDEX_NONE)
+		{
+			if (!bUsesOnCookStart)
+			{
+				// Compile and save in the CachedPlatformsData map
+				CompileForTargetPlatform(TargetPlatform);
+			}
+			else
+			{
+				// Load from Disk
+				LoadCompiledDataFromDisk(false, TargetPlatform);
+
+				if (!PrivateData->bModelCompiledForCook)
+				{
+					LoadReferencedObjects();
+					PrivateData->bModelCompiledForCook = true;
+				}
+			}
+
+			PrivateData->CachedPlatformNames.Add(TargetPlatform->PlatformName());
+		}
+	}
+	else
+	{
+		ClearCompiledData();
+		PrivateData->SetModel(nullptr); // Discard compilation
+		if (TargetPlatform)
+		{
+			PrivateData->CachedPlatformNames.Add(TargetPlatform->PlatformName());
+		}
+	}
+}
+
+
+bool UCustomizableObject::IsCachedCookedPlatformDataLoaded( const ITargetPlatform* TargetPlatform ) 
+{ 
+	return PrivateData->CachedPlatformNames.Find(TargetPlatform->PlatformName()) != INDEX_NONE;
+}
+
+// TODO COOK: Remove Hack to add new references to the package
+void UCustomizableObject::LoadReferencedObjects()
+{
+	for (TSoftObjectPtr<UMaterialInterface> r : ReferencedMaterials)
+	{
+		r.LoadSynchronous();
+	}
+
+	for(TPair<FString, TSoftObjectPtr<UPhysicsAsset>>& p : PhysicsAssetsMap)
+	{
+		p.Value.LoadSynchronous();
+	}
+	
+	for (TPair<FString, FParameterUIData>& i : ParameterUIDataMap)
+	{
+		i.Value.LoadResources();
+	}
+
+	for (TPair<FString, FParameterUIData>& s : StateUIDataMap)
+	{
+		s.Value.LoadResources();
+	}
+}
+
+
+void UCustomizableObject::ClearCompiledData()
+{
+	ReferenceSkeletalMeshesData.Empty();
+	ReferencedMaterials.Empty();
+	ReferencedMaterialSlotNames.Empty();
+	ReferencedSkeletons.Empty();
+	ImageProperties.Empty();
+	ParameterUIDataMap.Empty();
+	StateUIDataMap.Empty();
+	PhysicsAssetsMap.Empty();
+	ContributingMorphTargetsInfo.Empty();
+	MorphTargetReconstructionData.Empty();
+	ClothMeshToMeshVertData.Empty();
+	ContributingClothingAssetsData.Empty();
+	ClothSharedConfigsData.Empty();
+
+#if WITH_EDITORONLY_DATA
+	CustomizableObjectPathMap.Empty();
+	GroupNodeMap.Empty();
+#endif
+
+	HashToStreamableBlock.Empty();
+	BulkData = nullptr;
+}
+
+
+void UCustomizableObject::UpdateCompiledDataFromModel()
+{
+	mu::ModelPtr Model = PrivateData->GetModel();
+
+	// Generate a map that using the resource id tells the offset and size of the resource inside the bulk data
+	if(Model)
+	{
+		uint64_t Offset = 0;
+		uint32_t ResourceId = 0;
+		uint32_t ResourceSize = 0;
+
+		int32 NumStreamingFiles = Model->GetPrivate()->m_program.m_roms.Num();
+		HashToStreamableBlock.Empty(NumStreamingFiles);
+
+		for (size_t FileIndex = 0; FileIndex < NumStreamingFiles; ++FileIndex)
+		{
+			ResourceId = Model->GetPrivate()->m_program.m_roms[FileIndex].Id;
+			ResourceSize = Model->GetPrivate()->m_program.m_roms[FileIndex].Size;
+
+			HashToStreamableBlock.Add(ResourceId, FMutableStreamableBlock({0, Offset, ResourceSize }));
+			Offset += ResourceSize;
+		}
+	}
+
+	// Generate ParameterProperties and IntParameterLookUpTable
+	UpdateParameterPropertiesFromModel();
+}
+
+
+void UCustomizableObject::SaveCompiledData(FArchive& MemoryWriter, bool bSkipEditorOnlyData)
+{
+	int32 InternalVersion = CurrentSupportedVersion;
+	MutableCompiledDataStreamHeader Header(InternalVersion, VersionId);
+	MemoryWriter << Header;
+
+	MemoryWriter << ReferenceSkeletalMeshesData;
+	
+	int32 NumReferencedMaterials = ReferencedMaterials.Num();
+	MemoryWriter << NumReferencedMaterials;
+
+	for (const TSoftObjectPtr<UMaterialInterface>& Material : ReferencedMaterials)
+	{
+		FString StringRef = Material.ToSoftObjectPath().ToString();
+		MemoryWriter << StringRef;
+	}
+
+	int32 NumReferencedMaterialSlotNames = ReferencedMaterialSlotNames.Num();
+	MemoryWriter << NumReferencedMaterialSlotNames;
+
+	for (const FName& MaterialSlotName : ReferencedMaterialSlotNames)
+	{
+		FString StringRef = MaterialSlotName.ToString();
+		MemoryWriter << StringRef;
+	}
+
+	int32 NumReferencedSkeletons = ReferencedSkeletons.Num();
+	MemoryWriter << NumReferencedSkeletons;
+
+	for (const TSoftObjectPtr<USkeleton>& Skeleton : ReferencedSkeletons)
+	{
+		FString StringRef = Skeleton.ToSoftObjectPath().ToString();
+		MemoryWriter << StringRef;
+	}
+
+	MemoryWriter << ImageProperties;
+	MemoryWriter << ParameterUIDataMap;
+	MemoryWriter << StateUIDataMap;
+
+	int32 NumPhysicsAssets = PhysicsAssetsMap.Num();
+	MemoryWriter << NumPhysicsAssets;
+
+	for (const TPair<FString, TSoftObjectPtr<class UPhysicsAsset>>& PhysicsAsset : PhysicsAssetsMap)
+	{
+		FString StringRef = PhysicsAsset.Value.ToSoftObjectPath().ToString();
+		MemoryWriter << StringRef;
+	}
+
+	MemoryWriter << ContributingMorphTargetsInfo;
+	MemoryWriter << MorphTargetReconstructionData;
+	
+	MemoryWriter << ClothMeshToMeshVertData;
+	MemoryWriter << ContributingClothingAssetsData;
+	MemoryWriter << ClothSharedConfigsData; 
+	
+	MemoryWriter << HashToStreamableBlock;
+
+	// All Editor Only data must be serialized here
+	if (!bSkipEditorOnlyData)
+	{
+		MemoryWriter << CustomizableObjectPathMap;
+		MemoryWriter << GroupNodeMap;
+	}
+
+	MemoryWriter << LODSettings.NumLODsInRoot;
+
+	MemoryWriter << LODSettings.FirstLODAvailable;
+
+	MemoryWriter << LODSettings.NumLODsToStream;
+	MemoryWriter << LODSettings.bLODStreamingEnabled;
+}
+
+void UCustomizableObject::LoadCompiledData(FArchive& MemoryReader, bool bSkipEditorOnlyData)
+{
+	PrivateData->SetModel(nullptr);
+	ClearCompiledData();
+
+	MutableCompiledDataStreamHeader Header;
+	MemoryReader << Header;
+
+	if (CurrentSupportedVersion == Header.InternalVersion)
+	{
+		// Make sure mutable has been initialised.
+		UCustomizableObjectSystem::GetInstance();
+
+		MemoryReader << ReferenceSkeletalMeshesData;
+		
+		// We can load
+		int32 NumReferencedMaterials = 0;
+		MemoryReader << NumReferencedMaterials;
+
+		for (int32 i = 0; i < NumReferencedMaterials; ++i)
+		{
+			FString StringRef;
+			MemoryReader << StringRef;
+
+			ReferencedMaterials.Add(TSoftObjectPtr<UMaterialInterface>(FSoftObjectPath(StringRef)));
+		}
+
+		int32 NumReferencedMaterialSlotNames = 0;
+		MemoryReader << NumReferencedMaterialSlotNames;
+
+		for (int32 i = 0; i < NumReferencedMaterialSlotNames; ++i)
+		{
+			FString StringRef;
+			MemoryReader << StringRef;
+
+			ReferencedMaterialSlotNames.Add(FName(*StringRef));
+		}
+
+		int32 NumReferencedSkeletons = 0;
+		MemoryReader << NumReferencedSkeletons;
+
+		for (int32 SkeletonIndex = 0; SkeletonIndex < NumReferencedSkeletons; ++SkeletonIndex)
+		{
+			FString StringRef;
+			MemoryReader << StringRef;
+
+			ReferencedSkeletons.Add(TSoftObjectPtr<USkeleton>(FSoftObjectPath(StringRef)));
+		}
+
+		MemoryReader << ImageProperties;
+		MemoryReader << ParameterUIDataMap;
+		MemoryReader << StateUIDataMap;
+
+		int32 NumPhysicsAssets = 0;
+		MemoryReader << NumPhysicsAssets;
+
+		for (int32 i = 0; i < NumPhysicsAssets; ++i)
+		{
+			FString StringRef;
+			MemoryReader << StringRef;
+
+			PhysicsAssetsMap.Add(StringRef, TSoftObjectPtr<UPhysicsAsset>(FSoftObjectPath(StringRef)));
+		}
+
+		MemoryReader << ContributingMorphTargetsInfo;
+		MemoryReader << MorphTargetReconstructionData;
+
+		MemoryReader << ClothMeshToMeshVertData;
+		MemoryReader << ContributingClothingAssetsData;
+		MemoryReader << ClothSharedConfigsData;
+
+		MemoryReader << HashToStreamableBlock;
+
+		// All Editor Only data must be loaded here
+		if (!bSkipEditorOnlyData)
+		{
+			MemoryReader << CustomizableObjectPathMap;
+			MemoryReader << GroupNodeMap;
+		}
+
+		MemoryReader << LODSettings.NumLODsInRoot;
+
+		MemoryReader << LODSettings.FirstLODAvailable;
+
+		MemoryReader << LODSettings.NumLODsToStream;
+		MemoryReader << LODSettings.bLODStreamingEnabled;
+
+		bool bModelSerialized = false;
+		MemoryReader << bModelSerialized;
+
+		if (bModelSerialized)
+		{
+			UnrealMutableInputStream stream(MemoryReader);
+			mu::InputArchive arch(&stream);
+			mu::ModelPtr Model = mu::Model::StaticUnserialise( arch );
+
+			PrivateData->SetModel(Model.get());
+		}
+	}
+
+	UpdateParameterPropertiesFromModel();
+}
+
+void UCustomizableObject::LoadCompiledDataFromDisk(bool bIsEditorData, const ITargetPlatform* InTargetPlatform)
+{
+	FString PlatformName = InTargetPlatform ? InTargetPlatform->PlatformName() : FPlatformProperties::PlatformName();
+
+	if (bIsEditorData)
+	{
+		// Customizable Object outdated
+		if (!Identifier.IsValid())
+		{
+			return;
+		}
+	}
+	else // Loading cooked data
+	{
+		// If we don't use OnCookStart there's nothing to be loaded. Same case for child objects.
+		if (!bUsesOnCookStart || bIsChildObject) return;
+
+		// Platform Data already cached
+		if (CachedPlatformsData.Contains(PlatformName)) return;
+	}
+
+	// Compose Folder Name
+	const FString FolderPath = GetCompiledDataFolderPath(bIsEditorData);
+
+	// Compose File Names
+	FString ModelFileName = FolderPath + GetCompiledDataFileName(true, InTargetPlatform);
+	FString StreamableFileName = FolderPath + GetCompiledDataFileName(false, InTargetPlatform);
+
+	IFileManager& FileManager = IFileManager::Get();
+	if (FileManager.FileExists(*ModelFileName) && FileManager.FileExists(*StreamableFileName))
+	{
+		// Check CompiledData
+		IFileHandle* CompiledDataFileHandle = FPlatformFileManager::Get().GetPlatformFile().OpenRead(*ModelFileName);
+		IFileHandle* StreamableDataFileHandle = FPlatformFileManager::Get().GetPlatformFile().OpenRead(*StreamableFileName);
+
+		MutableCompiledDataStreamHeader CompiledDataHeader;
+		MutableCompiledDataStreamHeader StreamableDataHeader;
+
+		int32 HeaderSize = sizeof(MutableCompiledDataStreamHeader);
+		TArray<uint8> HeaderBytes;
+		HeaderBytes.SetNum(HeaderSize);
+
+		{
+			CompiledDataFileHandle->Read(HeaderBytes.GetData(), HeaderSize);
+			FMemoryReader AuxMemoryReader(HeaderBytes);
+			AuxMemoryReader << CompiledDataHeader;
+		}
+		{
+			StreamableDataFileHandle->Read(HeaderBytes.GetData(), HeaderSize);
+			FMemoryReader AuxMemoryReader(HeaderBytes);
+			AuxMemoryReader << StreamableDataHeader;
+		}
+
+		if (CompiledDataHeader.InternalVersion == CurrentSupportedVersion
+			&&
+			CompiledDataHeader.InternalVersion == StreamableDataHeader.InternalVersion 
+			&&
+			CompiledDataHeader.VersionId == StreamableDataHeader.VersionId)
+		{
+			if ( bIsEditorData && (IsRunningGame() || CompiledDataHeader.VersionId == VersionId) )
+			{ 
+				int64 CompiledDataSize = CompiledDataFileHandle->Size() - HeaderSize;
+				TArray64<uint8> CompiledDataBytes;
+				CompiledDataBytes.SetNumUninitialized(CompiledDataSize);
+
+				CompiledDataFileHandle->Seek(HeaderSize);
+				CompiledDataFileHandle->Read(CompiledDataBytes.GetData(), CompiledDataSize);
+
+				FMemoryReaderView MemoryReader(CompiledDataBytes);
+				LoadCompiledData(MemoryReader);
+			}
+			else if (!bIsEditorData)// Caching Cooked Data
+			{
+				FMutableCachedPlatformData& CachedData = CachedPlatformsData.Add(PlatformName);
+				int64 CompiledDataSize = CompiledDataFileHandle->Size() - HeaderSize;
+				int64 StreamableDataSize = StreamableDataFileHandle->Size() - HeaderSize;
+
+				CachedData.ModelData.SetNumUninitialized(CompiledDataSize);
+				CachedData.StreamableData.SetNumUninitialized(StreamableDataSize);
+
+				// Change the current read position to exclude the header
+				CompiledDataFileHandle->Seek(HeaderSize);
+				StreamableDataFileHandle->Seek(HeaderSize);
+
+				// Read Data
+				CompiledDataFileHandle->Read(CachedData.ModelData.GetData(), CompiledDataSize);
+				StreamableDataFileHandle->Read(CachedData.StreamableData.GetData(), StreamableDataSize);
+			}
+		}
+
+		delete CompiledDataFileHandle;
+		delete StreamableDataFileHandle;
+	}
+}
+
+
+void UCustomizableObject::CachePlatformData(const ITargetPlatform* InTargetPlatform, const TArray64<uint8>& InModelBytes, const TArray64< uint8>& InBulkBytes)
+{
+	FString PlatformName = InTargetPlatform ? InTargetPlatform->PlatformName() : FPlatformProperties::PlatformName();
+
+	check(!CachedPlatformsData.Find(PlatformName));
+
+	FMutableCachedPlatformData& Data = CachedPlatformsData.Add(PlatformName);
+
+	// Cache CO data and mu::Model
+	FMemoryWriter64 MemoryWriter(Data.ModelData);
+	SaveCompiledData(MemoryWriter, true);
+	Data.ModelData.Append(InModelBytes);
+
+	// Cache streamable bulk data
+	Data.StreamableData.Reset(InBulkBytes.Num());
+	Algo::Copy(InBulkBytes, Data.StreamableData);
+}
+
+
+void UCustomizableObject::CompileForTargetPlatform(const ITargetPlatform* TargetPlatform, bool bIsOnCookStart)
+{
+	// TEMP COOK: Remove together with OnCookStart
+	bUsesOnCookStart = bIsOnCookStart;
+
+	FCustomizableObjectCompilerBase* Compiler = UCustomizableObjectSystem::GetInstance()->GetNewCompiler();
+
+	bool bIsRootObject = Compiler->IsRootObject(this);
+	
+	bool bIsRelevantForThisTarget =
+		(Relevancy == ECustomizableObjectRelevancy::All)
+		||
+		(TargetPlatform && Relevancy == ECustomizableObjectRelevancy::ClientOnly && !TargetPlatform->IsServerOnly());
+
+	// Discard any older compilation
+	PrivateData->SetModel(nullptr);
+
+	if (bIsRootObject && bIsRelevantForThisTarget)
+	{
+
+		FCompilationOptions Options;
+		Options.OptimizationLevel = 3;	// max optimization when packaging.
+		Options.bTextureCompression = true;
+		Options.bIsCooking = true;
+		Options.bSaveCookedDataToDisk = bUsesOnCookStart;
+		Options.TargetPlatform = TargetPlatform;
+		Options.bExtraBoneInfluencesEnabled = ICustomizableObjectModule::Get().AreExtraBoneInfluencesEnabled();
+
+		// If this is enabled, there are determinism problems. Disable it for packaging.
+		Options.bUseParallelCompilation = false;
+
+		Compiler->Compile(*this, Options, false);
+	}
+
+	delete Compiler;
+}
+
+
+bool UCustomizableObject::ConditionalAutoCompile()
+{
+	check(IsInGameThread());
+
+	// Don't compile compiled objects
+	if (IsCompiled())
+	{
+		return true;
+	}
+
+	// Don't compile objects being compiled
+	if (IsLocked())
+	{
+		return false;
+	}
+
+	UCustomizableObjectSystem* System = UCustomizableObjectSystem::GetInstance();
+	if (!System || !System->IsValidLowLevel() || System->HasAnyFlags(RF_BeginDestroyed))
+	{
+		return false;
+	}
+
+	// Don't compile if we're running game, Compile and/or if AutoCompile is disabled
+	if (IsRunningGame() || System->IsCompilationDisabled() || !System->IsAutoCompileEnabled())
+	{
+		System->AddUncompiledCOWarning(this);
+		return false;
+	}
+
+	// Discard any older compilation
+	PrivateData->SetModel(nullptr);
+
+	// Sync/Async compilation
+	if (System->IsAutoCompilationSync())
+	{
+		FCustomizableObjectCompilerBase* Compiler = UCustomizableObjectSystem::GetInstance()->GetNewCompiler();
+		Compiler->Compile(*this, this->CompileOptions, false);
+		delete Compiler;
+	}
+	else
+	{
+		const FAssetRegistryModule& AssetRegistryModule = FModuleManager::LoadModuleChecked<FAssetRegistryModule>("AssetRegistry");
+		const FAssetData AssetData = AssetRegistryModule.Get().GetAssetByObjectPath(*GetPathName());
+		System->RecompileCustomizableObjectAsync(AssetData, this);
+	}
+
+	return IsCompiled();
+}
+
+
+FReply UCustomizableObject::AddNewParameterProfile(FString Name, UCustomizableObjectInstance& CustomInstance)
+{
+
+	if (Name.IsEmpty())
+	{
+		Name = "Unnamed_Profile";
+	}
+
+	FString ProfileName = Name;
+	int32 Suffix = 0;
+
+	bool bUniqueNameFound = false;
+	while (!bUniqueNameFound)
+	{
+		FProfileParameterDat* Found = InstancePropertiesProfiles.FindByPredicate(
+			[&ProfileName](const FProfileParameterDat& Profile) { return Profile.ProfileName == ProfileName; });
+
+		bUniqueNameFound = static_cast<bool>(!Found);
+		if (Found)
+		{
+			ProfileName = Name + FString::FromInt(Suffix);
+			++Suffix;
+		}
+	}
+
+	int32 ProfileIndex = InstancePropertiesProfiles.Emplace();
+
+	InstancePropertiesProfiles[ProfileIndex].ProfileName = ProfileName;
+	CustomInstance.SaveParametersToProfile(ProfileIndex);
+
+	Modify();
+
+	return FReply::Handled();
+}
+
+
+void UCustomizableObject::InitializeIdentifier()
+{
+	if (!Identifier.IsValid())
+	{
+		Identifier = FGuid::NewGuid();
+		MarkPackageDirty();
+	}
+}
+
+
+FString UCustomizableObject::GetCompiledDataFolderPath(bool bIsEditorData) const
+{	
+	const FString FolderName = bIsEditorData ? TEXT("MutableStreamedDataEditor/") : TEXT("MutableStreamedData/");
+	return FPaths::ConvertRelativePathToFull(FPaths::ProjectSavedDir() + FolderName);
+}
+
+
+FString UCustomizableObject::GetCompiledDataFileName(bool bIsModel, const ITargetPlatform* InTargetPlatform) const
+{
+	const FString PlatformName = InTargetPlatform ? InTargetPlatform->PlatformName() : FPlatformProperties::PlatformName();
+	const FString FileIdentifier = Identifier.IsValid() ? Identifier.ToString() : VersionId.ToString();
+	const FString Extension = bIsModel ? TEXT("_M.mut") : TEXT("_S.mut");
+	return PlatformName + FileIdentifier + Extension;
+}
+
+
+FString UCustomizableObject::GetDesc()
+{
+	int states = GetStateCount();
+	int params = GetParameterCount();
+	return FString::Printf(TEXT("%d States, %d Parameters"), states, params);
+}
+
+
+void UCustomizableObject::SaveEmbeddedData(FArchive& Ar)
+{
+	UE_LOG(LogMutable, Log, TEXT("Saving embedded data for Customizable Object [%s] now at position %d."), *GetName(), int(Ar.Tell()));
+
+	int32 InternalVersion = PrivateData->GetModel() ? CurrentSupportedVersion : -1;
+	Ar << InternalVersion;
+
+	if (PrivateData->GetModel())
+	{
+		// Serialize RefSkeletalMeshesData
+		{
+			Ar << ReferenceSkeletalMeshesData;
+		}
+		
+		// Serialize morph data
+		{
+			Ar << ContributingMorphTargetsInfo;
+			Ar << MorphTargetReconstructionData;
+		}
+		
+		{
+			Ar << ClothMeshToMeshVertData;
+			Ar << ContributingClothingAssetsData;
+			Ar << ClothSharedConfigsData;
+		}
+
+		// Serialise the entire model, but unload the streamable data first.
+		{
+			PrivateData->GetModel()->UnloadExternalData();
+
+			UnrealMutableOutputStream stream(Ar);
+			mu::OutputArchive arch(&stream);
+			mu::Model::Serialise(PrivateData->GetModel(), arch);
+		}
+
+		UE_LOG(LogMutable, Log, TEXT("Saved embedded data for Customizable Object [%s] now at position %d."), *GetName(), int(Ar.Tell()));
+	}
+}
+
+#endif // End WITH_EDITOR 
+
+void UCustomizableObject::LoadEmbeddedData(FArchive& Ar)
+{
+	MUTABLE_CPUPROFILER_SCOPE(UCustomizableObject::LoadEmbeddedData)
+
+	int32 InternalVersion;
+	Ar << InternalVersion;
+
+	// If this fails, something went wrong with the packaging: we have data that belongs
+	// to a different version than the code.
+	check(CurrentSupportedVersion==InternalVersion);
+
+	if(CurrentSupportedVersion == InternalVersion)
+	{
+		// Load RefSkeletalMeshesData
+		{
+			Ar << ReferenceSkeletalMeshesData;
+		}
+	
+		// Load morph data
+		{
+			Ar << ContributingMorphTargetsInfo;
+			Ar << MorphTargetReconstructionData;
+		}
+		
+		{
+			Ar << ClothMeshToMeshVertData;
+			Ar << ContributingClothingAssetsData;
+			Ar << ClothSharedConfigsData;
+		}
+		
+		// Load model
+		UnrealMutableInputStream stream(Ar);
+		mu::InputArchive arch(&stream);
+		mu::ModelPtr Model = mu::Model::StaticUnserialise( arch );
+		PrivateData->SetModel( Model.get() );
+
+		// Create parameter properties
+		UpdateParameterPropertiesFromModel();
+	}
+}
+
+bool UCustomizableObject::IsCompiled() const
+{
+	bool IsCompiled = PrivateData->GetModel() != nullptr;
+
+	return IsCompiled;
+}
+
+
+USkeletalMesh* UCustomizableObject::GetRefSkeletalMesh(int32 ComponentIndex)
+{
+	if (ReferenceSkeletalMeshes.IsValidIndex(ComponentIndex))
+	{
+		return ReferenceSkeletalMeshes[ComponentIndex];
+	}
+
+	return nullptr;
+}
+
+
+FMutableRefSkeletalMeshData* UCustomizableObject::GetRefSkeletalMeshData(int32 ComponentIndex)
+{
+	if(ReferenceSkeletalMeshesData.IsValidIndex(ComponentIndex))
+	{
+		return &ReferenceSkeletalMeshesData[ComponentIndex];
+	}
+
+	UE_LOG(LogMutable, Warning, TEXT("UCustomizableObject::GetRefSkeletalMeshData with an invalid Index. "
+			"Reference SkeletalMesh data and compiled model may be out of sync. "
+			"Try recompiling the CustomizableObject [%s]."), *GetName() );
+	
+	return nullptr;
+}
+
+
+TSoftObjectPtr<UMaterialInterface> UCustomizableObject::GetReferencedMaterialAssetPtr( uint32 Index )
+{
+	if (ReferencedMaterials.IsValidIndex(Index))
+	{
+		return ReferencedMaterials[Index];
+	}
+
+	UE_LOG(LogMutable, Warning, TEXT("UCustomizableObject::GetReferencedMaterial with an invalid Index. "
+		"Source material data and CustomizableObject data may be out of sync. "
+		"Try recompiling and saving the CustomizableObject asset [%s]."), *GetName() );
+	return nullptr;
+}
+
+
+TSoftObjectPtr<USkeleton> UCustomizableObject::GetReferencedSkeletonAssetPtr( uint32 Index )
+{
+	if (ReferencedSkeletons.IsValidIndex(Index))
+	{
+		return ReferencedSkeletons[Index];
+	}
+
+	UE_LOG(LogMutable, Warning, TEXT("UCustomizableObject::GetReferencedSkeleton with an invalid Index. "
+		"Skeleton data and CustomizableObject data may be out of sync. "
+		"Try recompiling and saving the CustomizableObject asset [%s]."), *GetName());
+	return nullptr;
+}
+
+
+int32 UCustomizableObject::FindState( const FString& Name ) const
+{
+	int32 Result = -1;
+	if (PrivateData->GetModel())
+	{
+		Result = PrivateData->GetModel()->FindState( TCHAR_TO_ANSI(*Name) );
+	}
+
+	return Result;
+}
+
+
+int32 UCustomizableObject::GetStateCount() const
+{
+	int32 Result = 0;
+
+	if (PrivateData->GetModel())
+	{
+		Result = PrivateData->GetModel()->GetStateCount();
+	}
+
+	return Result;
+}
+
+
+FString UCustomizableObject::GetStateName(int32 StateIndex) const
+{
+	FString Result;
+
+	if (PrivateData->GetModel())
+	{
+		Result = ANSI_TO_TCHAR( PrivateData->GetModel()->GetStateName(StateIndex) );
+	}
+
+	return Result;
+}
+
+
+int32 UCustomizableObject::GetStateParameterCount( int32 StateIndex ) const
+{
+	int32 Result = 0;
+
+	if (PrivateData->GetModel())
+	{
+		Result = PrivateData->GetModel()->GetStateParameterCount(StateIndex);
+	}
+
+	return Result;
+}
+
+int32 UCustomizableObject::GetStateParameterIndex(int32 StateIndex, int32 ParameterIndex) const
+{
+	int32 Result = 0;
+
+	if (PrivateData->GetModel())
+	{
+		Result = PrivateData->GetModel()->GetStateParameterIndex(StateIndex, ParameterIndex);
+	}
+
+	return Result;
+}
+
+
+int32 UCustomizableObject::GetStateParameterCount(const FString& StateName) const
+{
+	int32 StateIndex = FindState(StateName);
+	
+	return GetStateParameterCount(StateIndex);
+}
+
+
+FString UCustomizableObject::GetStateParameterName(const FString& StateName, int32 ParameterIndex) const
+{
+	int32 StateIndex = FindState(StateName);
+	
+	return GetStateParameterName(StateIndex, ParameterIndex);
+}
+
+FString UCustomizableObject::GetStateParameterName(int32 StateIndex, int32 ParameterIndex) const
+{
+	return GetParameterName(GetStateParameterIndex(StateIndex, ParameterIndex));
+}
+
+
+UCustomizableObjectInstance* UCustomizableObject::CreateInstance()
+{
+	MUTABLE_CPUPROFILER_SCOPE(UCustomizableObject::CreateInstance)
+
+	UCustomizableObjectInstance* PreviewInstance = NewObject<UCustomizableObjectInstance>(GetTransientPackage(), NAME_None, RF_Transient);
+	PreviewInstance->SetObject(this);
+	PreviewInstance->bShowOnlyRuntimeParameters = false;
+
+	UE_LOG(LogMutable, Log, TEXT("Created Customizable Object Instance."));
+
+	return PreviewInstance;
+}
+
+
+mu::Model* UCustomizableObject::GetModel() const
+{
+	return PrivateData->GetModel();
+}
+
+#if WITH_EDITOR
+void UCustomizableObject::SetModel(mu::Model* Model)
+{
+	PrivateData->SetModel(Model);
+	
+	UpdateCompiledDataFromModel();
+}
+
+#endif // End WITH_EDITOR
+
+
+int32 UCustomizableObject::GetNumLODs() const
+{
+	return LODSettings.NumLODsInRoot;
+}
+
+
+int UCustomizableObject::GetParameterCount() const
+{
+	return ParameterProperties.Num();
+}
+
+
+EMutableParameterType UCustomizableObject::GetParameterType(int32 ParamIndex) const
+{
+	if (ParamIndex >= 0 && ParamIndex < ParameterProperties.Num())
+	{
+		return ParameterProperties[ParamIndex].Type;
+	}
+	else
+	{
+		UE_LOG(LogMutable, Error, TEXT("Index [%d] out of ParameterProperties bounds at GetParameterType."), ParamIndex);
+	}
+
+	return EMutableParameterType::None;
+}
+
+
+EMutableParameterType UCustomizableObject::GetParameterTypeByName(const FString& Name) const
+{
+	const int* IndexPtr = ParameterPropertiesLookupTable.Find(Name);
+	int Index = IndexPtr ? *IndexPtr : INDEX_NONE;
+
+	if (ParameterProperties.IsValidIndex(Index))
+	{
+		return ParameterProperties[Index].Type;
+	}
+
+	UE_LOG(LogMutable, Warning, TEXT("Name '%s' does not exist in ParameterProperties lookup table at GetParameterTypeByName."), *Name);
+
+	for (int32 ParamIndex = 0; ParamIndex < ParameterProperties.Num(); ++ParamIndex)
+	{
+		if (ParameterProperties[ParamIndex].Name == Name)
+		{
+			return ParameterProperties[ParamIndex].Type;
+		}
+	}
+
+	UE_LOG(LogMutable, Warning, TEXT("Name '%s' does not exist in ParameterProperties at GetParameterTypeByName."), *Name);
+
+	return EMutableParameterType::None;
+}
+
+
+static const FString s_EmptyString;
+
+const FString & UCustomizableObject::GetParameterName(int32 ParamIndex) const
+{
+	if (ParamIndex >= 0 && ParamIndex < ParameterProperties.Num())
+	{
+		return ParameterProperties[ParamIndex].Name;
+	}
+	else
+	{
+		UE_LOG(LogMutable, Warning, TEXT("Index [%d] out of ParameterProperties bounds at GetParameterName."), ParamIndex);
+	}
+
+	return s_EmptyString;
+}
+
+
+void UCustomizableObject::UpdateParameterPropertiesFromModel()
+{
+	if (PrivateData->GetModel())
+	{
+		mu::ParametersPtr MutableParameters = PrivateData->GetModel()->NewParameters();
+		int paramCount = MutableParameters->GetCount();
+
+		ParameterProperties.Reset(paramCount);
+		ParameterPropertiesLookupTable.Empty(paramCount);
+		for (int paramIndex = 0; paramIndex<paramCount; ++paramIndex)
+		{
+			FMutableModelParameterProperties Data;
+
+			Data.Name = MutableParameters->GetName(paramIndex);
+			Data.Type = EMutableParameterType::None;
+			Data.ImageDescriptionCount = MutableParameters->GetAdditionalImageCount(paramIndex);
+
+			mu::PARAMETER_TYPE mutableType = MutableParameters->GetType(paramIndex);
+			switch (mutableType)
+			{
+			case mu::PARAMETER_TYPE::T_BOOL:
+			{
+				Data.Type = EMutableParameterType::Bool;
+				break;
+			}
+
+			case mu::PARAMETER_TYPE::T_INT:
+			{
+				Data.Type = EMutableParameterType::Int;
+
+				int ValueCount = MutableParameters->GetIntPossibleValueCount(paramIndex);
+				for (int i = 0; i<ValueCount; ++i)
+				{
+					FMutableModelParameterValue ValueData;
+					ValueData.Name = MutableParameters->GetIntPossibleValueName(paramIndex,i);
+					ValueData.Value = MutableParameters->GetIntPossibleValue(paramIndex,i);
+					Data.PossibleValues.Add(ValueData);
+				}
+				break;
+			}
+
+			case mu::PARAMETER_TYPE::T_FLOAT:
+			{
+				Data.Type = EMutableParameterType::Float;
+				break;
+			}
+
+			case mu::PARAMETER_TYPE::T_COLOUR:
+			{
+				Data.Type = EMutableParameterType::Color;
+				break;
+			}
+
+			case mu::PARAMETER_TYPE::T_PROJECTOR:
+			{
+				Data.Type = EMutableParameterType::Projector;
+				break;
+			}
+
+			case mu::PARAMETER_TYPE::T_IMAGE:
+			{
+				Data.Type = EMutableParameterType::Texture;
+				break;
+			}
+
+			default:
+				// Unhandled type?
+				check(false);
+				break;
+			}
+
+			ParameterProperties.Add(Data);
+			ParameterPropertiesLookupTable.Add(Data.Name, paramIndex);
+		}
+	}
+	else
+	{
+		ParameterProperties.Reset();
+		ParameterPropertiesLookupTable.Reset();
+	}
+}
+
+
+int UCustomizableObject::GetParameterDescriptionCount(int32 ParamIndex) const
+{
+	if (ParamIndex >= 0 && ParamIndex < ParameterProperties.Num())
+	{
+		return ParameterProperties[ParamIndex].ImageDescriptionCount;
+	}
+	else
+	{
+		UE_LOG(LogMutable, Warning, TEXT("Index [%d] out of ParameterProperties bounds at GetParameterDescriptionCount."), ParamIndex);
+	}
+
+	return 0;
+}
+
+
+int32 UCustomizableObject::GetParameterDescriptionCount(const FString& ParamName) const
+{
+	return GetParameterDescriptionCount(FindParameter(ParamName));
+}
+
+
+int32 UCustomizableObject::GetIntParameterNumOptions(int32 ParamIndex) const
+{
+	if (ParamIndex >= 0 && ParamIndex < ParameterProperties.Num())
+	{
+		return ParameterProperties[ParamIndex].PossibleValues.Num();
+	}
+	else
+	{
+		UE_LOG(LogMutable, Warning, TEXT("Index [%d] out of ParameterProperties bounds at GetIntParameterNumOptions."), ParamIndex);
+	}
+
+	return 0;
+}
+
+
+const FString& UCustomizableObject::GetIntParameterAvailableOption(int32 ParamIndex, int32 K) const
+{
+	if (ParamIndex >= 0 && ParamIndex < ParameterProperties.Num())
+	{
+		if (K >= 0 && K < GetIntParameterNumOptions(ParamIndex))
+		{
+			return ParameterProperties[ParamIndex].PossibleValues[K].Name;
+		}
+		else
+		{
+			UE_LOG(LogMutable, Warning, TEXT("Index [%d] out of IntParameterNumOptions bounds at GetIntParameterAvailableOption."), K);
+		}
+	}
+	else
+	{
+		UE_LOG(LogMutable, Warning, TEXT("Index [%d] out of ParameterProperties bounds at GetIntParameterAvailableOption."), ParamIndex);
+	}
+
+	return s_EmptyString;
+}
+
+
+int UCustomizableObject::FindParameter(const FString& Name) const
+{
+	const int32 * Found = ParameterPropertiesLookupTable.Find(Name);
+	if (Found == nullptr)
+	{
+		return INDEX_NONE;
+	}
+	return *Found;
+}
+
+
+int32 UCustomizableObject::FindIntParameterValue(int32 ParamIndex, const FString& Value) const
+{
+	int32 MinValueIndex = 0;
+	if (ParamIndex >= 0 && ParamIndex < ParameterProperties.Num())
+	{
+		const TArray<FMutableModelParameterValue>& PossibleValues = ParameterProperties[ParamIndex].PossibleValues;
+		if (PossibleValues.Num())
+		{
+			MinValueIndex = PossibleValues[0].Value;
+
+			for (int32 OrderValue = 0; OrderValue < PossibleValues.Num(); ++OrderValue)
+			{
+				const FString& Name = PossibleValues[OrderValue].Name;
+
+				if (Name == Value)
+				{
+					int32 CorrectedValue = OrderValue + MinValueIndex;
+					check(PossibleValues[OrderValue].Value == CorrectedValue);
+					return CorrectedValue;
+				}
+			}
+		}
+		else
+		{
+			UE_LOG(LogMutable, Warning, TEXT("No possible values for parameter with index [%d] at FindIntParameterValue."), ParamIndex);
+		}
+	}
+	else
+	{
+		UE_LOG(LogMutable, Warning, TEXT("Index [%d] out of ParameterProperties bounds at FindIntParameterValue."), ParamIndex);
+	}
+	return MinValueIndex;
+}
+
+
+FString UCustomizableObject::FindIntParameterValueName(int32 ParamIndex, int32 ParamValue) const
+{
+	if (ParamIndex >= 0 && ParamIndex < ParameterProperties.Num())
+	{
+		const TArray<FMutableModelParameterValue> & PossibleValues = ParameterProperties[ParamIndex].PossibleValues;
+
+		if (ParamValue >= 0 && ParamValue < PossibleValues.Num())
+		{
+			int32 MinValueIndex = PossibleValues[0].Value;
+			int32 ActualIndex = ParamValue - MinValueIndex;
+
+			check(ActualIndex >= 0 && ActualIndex < PossibleValues.Num());
+
+			return PossibleValues[ActualIndex].Name;
+		}
+	}
+	else
+	{
+		UE_LOG(LogMutable, Warning, TEXT("Index [%d] out of ParameterProperties bounds at FindIntParameterValueName."), ParamIndex);
+	}
+
+	return FString();
+}
+
+
+FParameterUIData UCustomizableObject::GetParameterUIMetadata(const FString& ParamName) const
+{
+	const FParameterUIData* ParameterUIData = ParameterUIDataMap.Find(ParamName);
+
+	return ParameterUIData ? *ParameterUIData : FParameterUIData();
+}
+
+
+FParameterUIData UCustomizableObject::GetParameterUIMetadataFromIndex(int32 ParamIndex) const
+{
+	return GetParameterUIMetadata(GetParameterName(ParamIndex));
+}
+
+
+FParameterUIData UCustomizableObject::GetStateUIMetadata(const FString& StateName) const
+{
+	const FParameterUIData* StateUIData = StateUIDataMap.Find(StateName);
+
+	return StateUIData ? *StateUIData : FParameterUIData();
+}
+
+
+FParameterUIData UCustomizableObject::GetStateUIMetadataFromIndex(int32 StateIndex) const
+{
+	return GetStateUIMetadata(GetStateName(StateIndex));
+}
+
+
+void UCustomizableObject::LoadMaskOutCache()
+{
+	MaskOutCache_HardRef = MaskOutCache.LoadSynchronous();
+}
+
+
+void UCustomizableObject::UnloadMaskOutCache()
+{
+	MaskOutCache_HardRef = nullptr;
+}
+
+
+void UCustomizableObject::ApplyStateForcedValuesToParameters( int32 State, mu::Parameters* Parameters)
+{
+	FParameterUIData StateMetaData = GetStateUIMetadataFromIndex(State);
+	for (const TPair<FString, FString>& ForcedParameter : StateMetaData.ForcedParameterValues)
+	{
+		int32 ForcedParameterIndex = FindParameter(ForcedParameter.Key);
+		if (ForcedParameterIndex == INDEX_NONE)
+		{
+			continue;
+		}
+
+		bool bIsMultidimensional = Parameters->NewRangeIndex(ForcedParameterIndex).get() != nullptr;
+		if (!bIsMultidimensional)
+		{
+			switch (GetParameterType(ForcedParameterIndex))
+			{
+			case EMutableParameterType::Int:
+			{
+				FString StringValue = ForcedParameter.Value;
+				if (StringValue.IsNumeric())
+				{
+					Parameters->SetIntValue(ForcedParameterIndex, FCString::Atoi(*StringValue));
+				}
+				else
+				{
+					int32 IntParameterIndex = FindIntParameterValue(ForcedParameterIndex, StringValue);
+					Parameters->SetIntValue(ForcedParameterIndex, IntParameterIndex);
+				}
+				break;
+			}
+			case EMutableParameterType::Bool:
+			{
+				Parameters->SetBoolValue(ForcedParameterIndex, ForcedParameter.Value.ToBool());
+				break;
+			}
+			default:
+			{
+				UE_LOG(LogMutable, Log, TEXT("Forced parameter type not supported."));
+				break;
+			}
+			}
+		}
+	}
+
+}
+
+#if WITH_EDITOR
+void UCustomizableObject::PreSaveRoot(FObjectPreSaveRootContext ObjectSaveContext)
+{
+	UObject::PreSaveRoot(ObjectSaveContext);
+
+	// Tell the validation system on this object that the validation that is going to be next invoked is due to
+	// this asset being saved.
+	// This value will be set to false by the validation method on this object so subsequent validation attempts
+	// get treated as expected.
+	bIsValidationTriggeredBySave = true;
+}
+
+EDataValidationResult UCustomizableObject::IsDataValid(TArray<FText>& ValidationErrors)
+{
+	TArray<FText> ValidationWarnings;
+	return IsDataValid(ValidationErrors,ValidationWarnings);
+}
+
+EDataValidationResult UCustomizableObject::IsDataValid(TArray<FText>& ValidationErrors,
+	TArray<FText>& ValidationWarnings)
+{
+	// This method seems to be designed to check data errors (like variables with unexpected values).
+	// Currently it does not check if the root that we are compiling has already been compiled during another validation.
+	
+	EDataValidationResult Result = EDataValidationResult::NotValidated;
+
+	// If validation is invoked by the saving of the asset just skip it. It is too expensive.
+	if (bIsValidationTriggeredBySave)
+    {
+		// Reenable the validation for this object after running the saving process and skipping this validation run
+		bIsValidationTriggeredBySave = false;
+    	return Result;
+    }
+	
+	// Request a compiler to be able to locate the root and to compile it
+	FCustomizableObjectCompilerBase* Compiler = UCustomizableObjectSystem::GetInstance()->GetNewCompiler();
+	{
+		// Find out witch is the root for this CO (it may be itself but that is OK)
+		UCustomizableObject* RootObject = Compiler->GetRootObject(this);
+		check (RootObject);
+		
+		// Run Sync compilation -> Warning : Potentially long operation -------------
+		Compiler->Compile(*RootObject, this->CompileOptions, false);
+		// --------------------------------------------------------------------------
+
+		// Request the compilation warnings and errors and provide them to the caller
+		Compiler->GetCompilationMessages(ValidationWarnings,ValidationErrors);
+		
+		// Check if the compilation was able to be performed (not if there were errors but if a critical issue appeared
+		// before starting the compilation as we know.
+		switch (Compiler->GetCompilationState())
+		{
+		case ECustomizableObjectCompilationState::Completed:
+			{
+				// If a warning or error was found then this object failed the validation process
+				Result = (ValidationWarnings.IsEmpty() && ValidationErrors.IsEmpty()) ? EDataValidationResult::Valid : EDataValidationResult::Invalid;
+			}
+			break;
+			
+		case ECustomizableObjectCompilationState::Failed:
+			{
+				// Early CO compilation error (before starting mutable compilation) -> Output is invalid
+				Result = EDataValidationResult::Invalid;
+				UE_LOG(LogMutable, Error,
+					   TEXT("Compilation of %s failed : Check previous log messages to get more information."),
+					   *this->GetName())
+			}
+			break;
+		
+			// Invalid options for this context:
+		case ECustomizableObjectCompilationState::None:
+			{
+				// This would mean we are working with a locked asset, with should not be possible with a sync compilation
+			}
+		case ECustomizableObjectCompilationState::InProgress:
+		default:
+			checkNoEntry();
+		}
+	}
+	delete Compiler;
+
+	
+	return Result;
+}
+#endif
+
+//-------------------------------------------------------------------------------------------------
+//-------------------------------------------------------------------------------------------------
+//-------------------------------------------------------------------------------------------------
+
+void FCustomizableObjectPrivateData::SetModel(mu::Model* Model)
+{
+	MutableModel = Model;
+}
+
+
+mu::Model* FCustomizableObjectPrivateData::GetModel()
+{
+	return MutableModel.get();
+}
+
+
+const mu::Model* FCustomizableObjectPrivateData::GetModel() const
+{
+	return MutableModel.get();
+}
+
+//-------------------------------------------------------------------------------------------------
+//-------------------------------------------------------------------------------------------------
+//-------------------------------------------------------------------------------------------------
+
+void UCustomizableObjectBulk::PostLoad()
+{
+	UObject::PostLoad();
+
+	const FString PackageFilename = FPackageName::LongPackageNameToFilename(GetOutermost()->GetName(), TEXT(".uasset"));
+	FString OuterPathName = FPaths::GetPath(PackageFilename);
+
+	for(FString& FilePath : BulkDataFileNames)
+	{
+		FilePath = OuterPathName / FilePath;
+	}
+
+}
+
+
+TArray<IAsyncReadFileHandle*> UCustomizableObjectBulk::GetAsyncReadFileHandles() const
+{
+	TArray<IAsyncReadFileHandle*> ReadFileHandles;
+	ReadFileHandles.Reserve(BulkDataFileNames.Num());
+
+	for (const FString& FilePath : BulkDataFileNames)
+	{
+		IAsyncReadFileHandle* ReadFileHandle = FPlatformFileManager::Get().GetPlatformFile().OpenAsyncRead(*FilePath);
+		
+		if (!ReadFileHandle)
+		{
+			UE_LOG(LogMutable, Error, TEXT("Failed to create AsyncReadFileHandle. File Path [%s]."), *FilePath);
+			break;
+		}
+
+		ReadFileHandles.Add(ReadFileHandle);
+	}
+
+	if (ReadFileHandles.Num() != BulkDataFileNames.Num())
+	{
+		for (IAsyncReadFileHandle* ReadFileHandle : ReadFileHandles)
+		{
+			delete ReadFileHandle;
+		}
+
+		ReadFileHandles.Empty();
+	}
+
+	return ReadFileHandles;
+}
+
+
+#if WITH_EDITOR
+
+void UCustomizableObjectBulk::CookAdditionalFilesOverride(const TCHAR* PackageFilename,
+	const ITargetPlatform* TargetPlatform,
+	TFunctionRef<void(const TCHAR* Filename, void* Data, int64 Size)> WriteAdditionalFile)
+{
+	check(CustomizableObject);
+	
+	FMutableCachedPlatformData* PlatformData = CustomizableObject->CachedPlatformsData.Find(TargetPlatform->PlatformName());
+	check(PlatformData);
+
+	// Data to serialize in separate files
+	uint8* Data = PlatformData->StreamableData.GetData();
+
+	// Path to the asset
+	const FString CookedFilePath = FPaths::GetPath(PackageFilename);
+
+	const uint16 NumBulkDataFiles = BulkDataFilesSize.Num();
+	for(uint16 FileIndex = 0; FileIndex < NumBulkDataFiles; ++FileIndex)
+	{
+		const FString CookedBulkFileName = CookedFilePath / BulkDataFileNames[FileIndex];
+
+		const int64 NumBytes = BulkDataFilesSize[FileIndex];
+		WriteAdditionalFile(*CookedBulkFileName, (void*)Data, NumBytes);
+		Data += NumBytes;
+	}
+}
+
+
+void UCustomizableObjectBulk::PrepareBulkData(UCustomizableObject* InOuter, const ITargetPlatform* TargetPlatform)
+{
+	CustomizableObject = InOuter;
+	check(CustomizableObject);
+
+	BulkDataFilesSize.Empty();
+	BulkDataFileNames.Empty();
+	
+	// Split the Streamable data into several separate files and fix up FileIndex and Offset of each StreamableBlock
+	if(mu::ModelPtr Model = CustomizableObject->GetModel())
+	{
+		const uint64 MaxChunkSize = UCustomizableObjectSystem::GetInstance()->GetMaxChunkSizeForPlatform(TargetPlatform);
+
+		const FString BulkFileName = CustomizableObject->GetName() + FString(TEXT("_Bulk"));
+
+		uint16 CurrentFileIndex = 0;
+		uint64 CurrentChunkSize = 0;
+
+		const int32 NumStreamingFiles = Model->GetPrivate()->m_program.m_roms.Num();
+		for (size_t FileIndex = 0; FileIndex < NumStreamingFiles; ++FileIndex)
+		{
+			const uint32 ResourceId = Model->GetPrivate()->m_program.m_roms[FileIndex].Id;
+
+			FMutableStreamableBlock& StreamableBlock = CustomizableObject->HashToStreamableBlock[ResourceId];
+
+			const uint32 BlockSize = StreamableBlock.Size;
+			if(CurrentChunkSize + BlockSize > MaxChunkSize)
+			{
+				BulkDataFileNames.Add(BulkFileName + FString::Printf(TEXT("%d.mut"), CurrentFileIndex));
+
+				if(CurrentChunkSize == 0)
+				{
+					BulkDataFilesSize.Add(BlockSize);
+					StreamableBlock.FileIndex = CurrentFileIndex++;
+					StreamableBlock.Offset = 0;
+					continue;
+				}
+
+				BulkDataFilesSize.Add(CurrentChunkSize);
+			
+				CurrentFileIndex++;
+				CurrentChunkSize = 0;
+			}
+
+			StreamableBlock.FileIndex = CurrentFileIndex;
+			StreamableBlock.Offset = CurrentChunkSize;
+			CurrentChunkSize += BlockSize;			
+		}
+
+		BulkDataFilesSize.Add(CurrentChunkSize);
+		BulkDataFileNames.Add(BulkFileName + FString::Printf(TEXT("%d.mut"), CurrentFileIndex));
+	}
+}
+
+#endif
+
+#undef LOCTEXT_NAMESPACE
diff --git a/Engine/Plugins/Experimental/Mutable/Source/CustomizableObject/Private/CustomizableObjectDGGUI.cpp b/Engine/Plugins/Experimental/Mutable/Source/CustomizableObject/Private/CustomizableObjectDGGUI.cpp
new file mode 100644
--- /dev/null
+++ b/Engine/Plugins/Experimental/Mutable/Source/CustomizableObject/Private/CustomizableObjectDGGUI.cpp
@@ -0,0 +1,142 @@
+// Copyright Epic Games, Inc. All Rights Reserved.
+
+#include "CustomizableObjectDGGUI.h"
+#include "CustomizableSkeletalComponent.h"
+#include "Engine/SkeletalMesh.h"
+#include "HAL/IConsoleManager.h"
+#include "Kismet/GameplayStatics.h"
+#include "Animation/AnimInstance.h"
+#include "CustomizableObjectInstance.h"
+
+void UDGGUI::OpenDGGUI(const int32 SlotID, UCustomizableSkeletalComponent* SelectedCustomizableSkeletalComponent, const UWorld* CurrentWorld, const int32 PlayerIndex)
+{
+#if WITH_EDITOR
+	if (APlayerController* Player = UGameplayStatics::GetPlayerController(CurrentWorld, PlayerIndex))
+	{
+		FSoftClassPath DGUIPath(TEXT("/Mutable/UI/DynamicallyGeneratedGUI_DGGUI/DynamicallyGeneratedGUI_DGGUI.DynamicallyGeneratedGUI_DGGUI_C"));
+		if (UClass* DGUI = DGUIPath.TryLoadClass<UDGGUI>())
+		{
+			UDGGUI* WDGUI = CreateWidget<UDGGUI>(Player, DGUI);
+			if (WDGUI)
+			{
+				WDGUI->SetCustomizableSkeletalComponent(SelectedCustomizableSkeletalComponent);
+				SelectedCustomizableSkeletalComponent->PreUpdateDelegate.BindUObject(WDGUI, &UDGGUI::CustomizableSkeletalMeshPreUpdate);
+				WDGUI->AddToViewport();
+				Player->SetShowMouseCursor(true);
+			}
+		}
+	}
+#endif // WITH_EDITOR
+}
+
+bool UDGGUI::CloseExistingDGGUI(const UWorld* CurrentWorld)
+{
+#if WITH_EDITOR
+	bool bClosing = false;
+	for (TObjectIterator<UDGGUI> PreviousGUI; PreviousGUI; ++PreviousGUI)
+	{
+		if (PreviousGUI->IsValidLowLevel())
+		{
+			if (UCustomizableSkeletalComponent* CustomizableComponent = PreviousGUI->GetCustomizableSkeletalComponent())
+			{
+				CustomizableComponent->PreUpdateDelegate.Unbind();
+				CustomizableComponent->UpdatedDelegate.Unbind();
+				PreviousGUI->SetCustomizableSkeletalComponent(nullptr);
+				bClosing = true;
+			}
+			PreviousGUI->RemoveFromParent();
+		}
+	}
+	if (bClosing)
+	{
+		if (APlayerController* Player = UGameplayStatics::GetPlayerController(CurrentWorld, 0))
+		{
+			Player->SetShowMouseCursor(false);
+		}
+		return true;
+	}
+#endif // WITH_EDITOR
+	return false;
+}
+
+
+// This is all based on UMeshCosmeticsVariance_ApplyParameters::PostCustomizationFixup( , but I think this should be set from within mutable, like how Physics assets are set up without user intervention
+void UDGGUI::CustomizableSkeletalMeshPreUpdate(UCustomizableSkeletalComponent* Component, USkeletalMesh* NextMesh)
+{
+#if WITH_EDITOR
+	if (Component && NextMesh)
+	{
+		USkeletalMeshComponent* Parent = Cast<USkeletalMeshComponent>(Component->GetAttachParent());
+		USkeleton* NextSkeleton = NextMesh->GetSkeleton();
+		if (Parent && NextSkeleton)
+		{
+			if (USkeletalMesh* ParentSkeletalMesh = Parent->SkeletalMesh)
+			{
+				if (USkeleton* ParentSkeleton = ParentSkeletalMesh->GetSkeleton())
+				{
+					UAnimInstance* AnimInstance = Parent->GetAnimInstance();
+					UClass* AnimClass = Parent->GetAnimClass();
+					IAnimClassInterface* Animation = Cast<IAnimClassInterface>(AnimClass);
+					if (AnimInstance && Animation)
+					{
+						if (USkeleton* AnimationSkeleton = Animation->GetTargetSkeleton())
+						{
+							if (AnimationSkeleton != NextSkeleton && NextSkeleton->GetCompatibleSkeletons().Find(AnimationSkeleton) == INDEX_NONE)
+							{
+								NextSkeleton->AddCompatibleSkeleton(AnimationSkeleton);
+								LastAnimationClass = AnimClass;
+								Component->UpdatedDelegate.BindUObject(this, &UDGGUI::CustomizableSkeletalMeshUpdated);
+							}
+						}
+					}
+				}
+			}
+		}
+	}
+#endif // WITH_EDITOR
+}
+
+
+void UDGGUI::CustomizableSkeletalMeshUpdated()
+{
+#if WITH_EDITOR
+	if (UCustomizableSkeletalComponent* CustomizableComponent = GetCustomizableSkeletalComponent())
+	{
+		if (USkeletalMeshComponent* CustomizedComponent = Cast<USkeletalMeshComponent>(CustomizableComponent->GetAttachParent()))
+		{
+			CustomizedComponent->SetAnimInstanceClass(LastAnimationClass);
+
+			if (UAnimInstance* TargetAnimInstance = CustomizedComponent->GetAnimInstance())
+			{
+				const UCustomizableObjectInstance* Instance = CustomizableComponent->CustomizableObjectInstance;
+
+				const int32 ComponentIndex = CustomizableComponent->ComponentIndex;
+
+				IAnimClassInterface* AnimInstanceClassInterface = IAnimClassInterface::GetFromClass(LastAnimationClass);
+
+				USkeleton* ProportionalSkel = AnimInstanceClassInterface->GetTargetSkeleton();
+				check(ProportionalSkel);
+				const auto LinkAnimClassBySlot = [&TargetAnimInstance, ProportionalSkel](int32 SlotIndex, TSubclassOf<UAnimInstance> SlottedInstanceClass)
+				{
+					if (SlottedInstanceClass)
+					{
+						IAnimClassInterface* SlottedClassInterface = IAnimClassInterface::GetFromClass(SlottedInstanceClass);
+						USkeleton* SlottedSkeleton = SlottedClassInterface->GetTargetSkeleton();
+
+						if (SlottedSkeleton != ProportionalSkel && SlottedSkeleton->GetCompatibleSkeletons().Find(ProportionalSkel) == INDEX_NONE)
+						{
+							SlottedSkeleton->AddCompatibleSkeleton(ProportionalSkel);
+						}
+
+						TargetAnimInstance->LinkAnimGraphByTag(*FString::Printf(TEXT("Part%d"), SlotIndex), SlottedInstanceClass);
+					}
+				};
+
+				using ForEachDelType = FEachComponentAnimInstanceClassNativeDelegate;
+				Instance->ForEachAnimInstance(ComponentIndex, ForEachDelType::CreateLambda(LinkAnimClassBySlot));
+			}
+		}
+		//CustomizableComponent->UpdatedDelegate.Unbind();
+	}
+#endif // WITH_EDITOR
+}
\ No newline at end of file
diff --git a/Engine/Plugins/Experimental/Mutable/Source/CustomizableObject/Private/CustomizableObjectIdentifier.cpp b/Engine/Plugins/Experimental/Mutable/Source/CustomizableObject/Private/CustomizableObjectIdentifier.cpp
new file mode 100644
--- /dev/null
+++ b/Engine/Plugins/Experimental/Mutable/Source/CustomizableObject/Private/CustomizableObjectIdentifier.cpp
@@ -0,0 +1,21 @@
+// Copyright Epic Games, Inc. All Rights Reserved.
+
+#include "CustomizableObjectIdentifier.h"
+
+
+FCustomizableObjectIdPair::FCustomizableObjectIdPair(FString ObjectGroupName, FString ObjectName)
+	: CustomizableObjectGroupName(ObjectGroupName)
+	, CustomizableObjectName(ObjectName)
+{
+}
+
+FCustomizableObjectIdentifier::FCustomizableObjectIdentifier()
+{
+}
+
+FCustomizableObjectIdentifier::FCustomizableObjectIdentifier(FString ObjectGroupName, FString ObjectName)
+	: CustomizableObjectGroupName(ObjectGroupName)
+	, CustomizableObjectName(ObjectName)
+{
+}
+
