[mutable] Move Mutable plugin from the restricted folder to the experimental plugins folder.

diff --git a/Engine/Plugins/Experimental/Mutable/Source/CustomizableObject/Private/CustomizableObjectInstance.cpp b/Engine/Plugins/Experimental/Mutable/Source/CustomizableObject/Private/CustomizableObjectInstance.cpp
new file mode 100644
--- /dev/null
+++ b/Engine/Plugins/Experimental/Mutable/Source/CustomizableObject/Private/CustomizableObjectInstance.cpp
@@ -0,0 +1,5749 @@
+// Copyright Epic Games, Inc. All Rights Reserved.
+
+#include "CustomizableObjectInstance.h"
+
+#include "StaticMeshResources.h"
+#include "SkeletalMeshTypes.h"
+#include "Engine/SkeletalMesh.h"
+#include "Animation/Skeleton.h"
+#include "Materials/MaterialInstanceDynamic.h"
+#include "Components/SkinnedMeshComponent.h"
+#include "Engine/SkeletalMeshSocket.h"
+#include "Animation/MorphTarget.h"
+#include "ClothingAsset.h"
+#include "ClothConfig.h"
+#include "SkeletalMergingLibrary.h"
+
+#include "CustomizableObject.h"
+#include "CustomizableObjectPrivate.h"
+#include "CustomizableSkeletalComponent.h"
+#include "CustomizableInstancePrivateData.h"
+#include "ICustomizableObjectModule.h"
+#include "CustomizableObjectSystem.h"
+#include "CustomizableObjectSystemPrivate.h"
+#include "CustomizableInstanceLODManagement.h"
+#include "CustomizableObjectParameterTypeDefinitions.h"
+#include "CustomizableObjectCustomVersion.h"
+#include "GameplayTagContainer.h"
+#include "UnrealConversionUtils.h"
+#include "MultilayerProjector.h"
+
+#include "Rendering/Texture2DResource.h"
+#include "Engine/Texture2D.h"
+#include "UObject/Package.h"
+#include "Serialization/BulkData.h"
+#include "RHI.h"
+
+#include "Algo/MaxElement.h"
+#include "Algo/MinElement.h"
+#include "Math/NumericLimits.h"
+#include "Templates/IsTrivial.h"
+
+#include "Runtime/Launch/Resources/Version.h"
+
+#include "CustomizableObjectMipDataProvider.h"
+#include "PhysicsEngine/PhysicsAsset.h"
+#include "Animation/AnimBlueprint.h"
+#include "CustomizableObjectMipDataProvider.h"
+
+#if WITH_EDITOR
+#include "MessageLogModule.h"
+#include "Logging/MessageLog.h"
+#include "Framework/Notifications/NotificationManager.h"
+#endif
+
+#include "MutableRuntime/Public/Model.h"
+#include "MutableRuntime/Public/System.h"
+#include "MutableRuntime/Public/Mesh.h"
+#include "MutableRuntime/Public/Image.h"
+#include "MutableRuntime/Public/Parameters.h"
+#include "MutableRuntime/Public/MutableTrace.h"
+#include "MutableRuntime/Private/ImagePrivate.h"
+#include "MutableRuntime/Private/OpImageSwizzle.h"
+#include "MutableRuntime/Private/OpImagePixelFormat.h"
+#include "MutableRuntime/Private/MeshPrivate.h"
+#include "MutableRuntime/Private/ModelPrivate.h"
+
+// Required for engine branch preprocessor defines.
+#include "CustomizableObject/Public/UnrealPortabilityHelpers.h"
+#include "Animation/AnimBlueprintGeneratedClass.h"
+#include "PhysicsEngine/BodyInstance.h"
+
+
+UTexture2D* UCustomizableInstancePrivateData::CreateTexture()
+{
+	UTexture2D* NewTexture = NewObject<UTexture2D>(
+		GetTransientPackage(),
+		NAME_None,
+		RF_Transient
+		);
+	UCustomizableObjectSystem::GetInstance()->GetPrivate()->CreatedTexture(NewTexture);
+	NewTexture->SetPlatformData( nullptr );
+
+	return NewTexture;
+}
+
+
+int32 UCustomizableInstancePrivateData::GetLastMeshId(int32 ComponentIndex, int32 LODIndex) const
+{
+	const FCustomizableInstanceComponentData* ComponentData = GetComponentData(ComponentIndex);
+	check(ComponentData);
+	check(ComponentData->LastMeshIdPerLOD.IsValidIndex(LODIndex));
+		
+	return ComponentsData[ComponentIndex].LastMeshIdPerLOD[LODIndex];
+}
+
+
+void UCustomizableInstancePrivateData::SetLastMeshId(int32 ComponentIndex, int32 LODIndex, int32 MeshId)
+{
+	FCustomizableInstanceComponentData* ComponentData = GetComponentData(ComponentIndex);
+	check(ComponentData);
+	check(ComponentData->LastMeshIdPerLOD.IsValidIndex(LODIndex));
+
+	ComponentData->LastMeshIdPerLOD[LODIndex] = MeshId;
+}
+
+
+void UCustomizableInstancePrivateData::ClearAllLastMeshIds()
+{
+	if (NumLODsAvailable != INT32_MAX)
+	{
+		for (FCustomizableInstanceComponentData& ComponentData : ComponentsData)
+		{
+			ComponentData.LastMeshIdPerLOD.Reset(NumLODsAvailable);
+			ComponentData.LastMeshIdPerLOD.Init(INDEX_NONE, NumLODsAvailable);
+		}
+	}
+}
+
+
+FCustomizableInstanceComponentData* UCustomizableInstancePrivateData::GetComponentData(int32 ComponentIndex)
+{
+	return ComponentsData.FindByPredicate([&ComponentIndex](FCustomizableInstanceComponentData& C) { return C.ComponentIndex == ComponentIndex; });
+}
+
+
+const FCustomizableInstanceComponentData* UCustomizableInstancePrivateData::GetComponentData(int32 ComponentIndex) const
+{
+	return ComponentsData.FindByPredicate([&ComponentIndex](FCustomizableInstanceComponentData& C) { return C.ComponentIndex == ComponentIndex; });
+}
+
+
+UCustomizableObjectInstance::UCustomizableObjectInstance()
+{
+	PrivateData = CreateDefaultSubobject<UCustomizableInstancePrivateData>(FName("PrivateData"), true);
+}
+
+
+FCustomizableObjectInstanceDescriptor& UCustomizableObjectInstance::GetDescriptor()
+{
+	return Descriptor;	
+}
+
+
+const FCustomizableObjectInstanceDescriptor& UCustomizableObjectInstance::GetDescriptor() const
+{
+	return const_cast<UCustomizableObjectInstance*>(this)->GetDescriptor();
+}
+
+
+void UCustomizableObjectInstance::SetDescriptor(const FCustomizableObjectInstanceDescriptor& InDescriptor)
+{
+	Descriptor = InDescriptor;
+	PrivateData->ReloadParametersFromObject(this, true);
+}
+
+
+void UCustomizableInstancePrivateData::SetMinMaxLODToLoad(int32 NewMinLOD, int32 NewMaxLOD, bool bLimitLODUpgrades)
+{
+	if (!HasCOInstanceFlags(LODsStreamingEnabled))
+	{
+		return;
+	}
+
+	NewMinLOD = FMath::Min(FMath::Max(NewMinLOD, (int32)FirstLODAvailable), FirstLODAvailable + NumMaxLODsToStream);
+
+	const bool bIsDowngradeLODUpdate = CurrentMinLOD >= 0 && NewMinLOD > CurrentMinLOD;
+	SetCOInstanceFlags(bIsDowngradeLODUpdate ? PendingLODsDowngrade : None);
+
+	if (!bIsDowngradeLODUpdate && bLimitLODUpgrades)
+	{
+		if (FPlatformTime::Seconds() - LastUpdateTime < 3.f) // This is a sort of hysteresis/cooldown to prevent too many instances being discarded/updated per second
+		{
+			return;
+		}
+	}
+
+	if(NumLODsAvailable != INT32_MAX)
+	{
+		if(MaxLODToLoad == INT32_MAX)
+		{
+			MaxLODToLoad = NumLODsAvailable - 1;
+		}
+
+		NewMaxLOD = FMath::Min(NewMaxLOD, NumLODsAvailable - 1);
+	}
+	
+	SetCOInstanceFlags(MinLODToLoad != NewMinLOD || MaxLODToLoad != NewMaxLOD ? PendingLODsUpdate : None);
+
+	MinLODToLoad = NewMinLOD;
+	MaxLODToLoad = NewMaxLOD;
+}
+
+
+void UCustomizableInstancePrivateData::PrepareForUpdate(const TSharedPtr<FMutableOperationData>& OperationData)
+{
+	CurrentMinLOD = OperationData->CurrentMinLOD;
+	CurrentMaxLOD = OperationData->CurrentMaxLOD;
+
+	const bool bNumLODsUpdated = NumLODsAvailable != OperationData->NumLODsAvailable;
+	
+	SetCOInstanceFlags(bNumLODsUpdated ? PendingMeshUpdate : None);
+
+	NumLODsAvailable = OperationData->NumLODsAvailable;
+
+	TArray<FInstanceUpdateData::FLOD>& LODs = OperationData->InstanceUpdateData.LODs;
+	TArray<FInstanceUpdateData::FComponent>& Components = OperationData->InstanceUpdateData.Components;
+
+	// Clear, reinit or create ComponentData for each component 
+	TSet<uint16> ComponentIds;
+	for (const FInstanceUpdateData::FLOD& LOD : LODs)
+	{
+		check(LOD.FirstComponent + LOD.ComponentCount <= Components.Num());
+
+		for (uint16 ComponentIndex = 0; ComponentIndex < LOD.ComponentCount; ++ComponentIndex)
+		{
+			const FInstanceUpdateData::FComponent& Component = Components[ComponentIndex];
+			ComponentIds.Add(Component.Id);
+
+			if (FCustomizableInstanceComponentData* ComponentData = GetComponentData(Component.Id))
+			{
+				ComponentData->AnimSlotToBP.Empty();
+				ComponentData->Skeletons.SkeletonsToLoad.Empty();
+				ComponentData->Skeletons.SkeletonsToMerge.Empty();
+				ComponentData->PhysicsAssets.PhysicsAssetToLoad.Empty();
+				ComponentData->PhysicsAssets.PhysicsAssetsToMerge.Empty();
+				ComponentData->ClothingPhysicsAssetsToStream.Empty();
+
+				if (bNumLODsUpdated)
+				{
+					ComponentData->LastMeshIdPerLOD.Init(-1, NumLODsAvailable);
+				}
+				continue;
+			}
+
+			FCustomizableInstanceComponentData& NewComponentData = ComponentsData.AddDefaulted_GetRef();
+			NewComponentData.ComponentIndex = Component.Id;
+			NewComponentData.LastMeshIdPerLOD.Init(-1, NumLODsAvailable);
+		}
+	}
+
+	// Check if a component have been removed.
+	if (ComponentsData.Num() != ComponentIds.Num())
+	{
+		for (uint16 ComponentIndex = 0; ComponentIndex < ComponentsData.Num();)
+		{
+			FCustomizableInstanceComponentData& ComponentData = ComponentsData[ComponentIndex];
+
+			if (!ComponentIds.Find(ComponentData.ComponentIndex))
+			{
+				ComponentsData.RemoveSingleSwap(ComponentData);
+				continue;
+			}
+
+			ComponentIndex++;
+		}
+	}
+}
+
+
+#if WITH_EDITOR
+
+void UCustomizableObjectInstance::PostEditChangeProperty(FPropertyChangedEvent& PropertyChangedEvent)
+{
+	bEditorPropertyChanged = true;
+
+	Super::PostEditChangeProperty(PropertyChangedEvent);
+}
+
+
+bool UCustomizableObjectInstance::CanEditChange(const FProperty* InProperty) const
+{
+	bool bIsMutable = Super::CanEditChange(InProperty);
+	if (bIsMutable && InProperty != NULL)
+	{
+		if (InProperty->GetFName() == TEXT("CustomizationObject"))
+		{
+			bIsMutable = false;
+		}
+
+		if (InProperty->GetFName() == TEXT("ParameterName"))
+		{
+			bIsMutable = false;
+		}
+	}
+
+	return bIsMutable;
+}
+
+#endif // WITH_EDITOR
+
+
+bool UCustomizableObjectInstance::IsEditorOnly() const
+{
+	if (UCustomizableObject* CustomizableObject = GetCustomizableObject())
+	{
+		return CustomizableObject->IsEditorOnly();
+	}
+	return false;
+}
+
+
+void UCustomizableObjectInstance::BeginDestroy()
+{
+	if (PrivateData && PrivateData->StreamingHandle)
+	{
+		PrivateData->StreamingHandle->CancelHandle();
+		PrivateData->StreamingHandle = nullptr;
+	}
+	
+	Super::BeginDestroy();
+
+	ReleaseMutableResources(true);
+}
+
+
+void UCustomizableObjectInstance::ReleaseMutableResources(bool bCalledFromBeginDestroy)
+{
+	PrivateData->GeneratedMaterials.Empty();
+
+	if (UCustomizableObjectSystem::IsCreated()) // Need to check this because the object might be destroyed after the CustomizableObjectSystem at shutdown
+	{
+		FCustomizableObjectSystemPrivate* CustomizableObjectSystem = UCustomizableObjectSystem::GetInstance()->GetPrivate();
+		// Get the cache of resources of all live instances of this object
+		FMutableResourceCache& Cache = CustomizableObjectSystem->GetObjectCache(GetCustomizableObject());
+
+		for (FGeneratedTexture& Texture : PrivateData->GeneratedTextures)
+		{
+			if (CustomizableObjectSystem->RemoveTextureReference(Texture.Id))
+			{
+				// Do not release textures when called from BeginDestroy, it would produce a texture artifact in the
+				// instance's remaining sk meshes and GC is being performed anyway so it will free the textures if needed
+				if (!bCalledFromBeginDestroy && CustomizableObjectSystem->bReleaseTexturesImmediately)
+				{
+					UCustomizableInstancePrivateData::ReleaseMutableTexture(Texture.Id, Texture.Texture, Cache);
+				}
+			}
+		}
+	}
+
+	PrivateData->GeneratedTextures.Empty();
+}
+
+
+bool UCustomizableObjectInstance::IsReadyForFinishDestroy()
+{
+	//return ReleaseResourcesFence.IsFenceComplete();
+	return true;
+}
+
+
+void UCustomizableObjectInstance::Serialize(FArchive& Ar)
+{
+	Super::Serialize(Ar);
+
+	Ar.UsingCustomVersion(FCustomizableObjectCustomVersion::GUID);
+
+	const int32 CustomizableObjectCustomVersion = GetLinkerCustomVersion(FCustomizableObjectCustomVersion::GUID);
+
+	if (CustomizableObjectCustomVersion < FCustomizableObjectCustomVersion::GroupProjectorIntToScalarIndex)
+	{
+		TArray<int32> IntParametersToMove;
+
+		// Find the num layer parameters that were int enums
+		for (int32 i = 0; i < IntParameters_DEPRECATED.Num(); ++i)
+		{
+			if (IntParameters_DEPRECATED[i].ParameterName.EndsWith(FMultilayerProjector::NUM_LAYERS_PARAMETER_POSTFIX, ESearchCase::CaseSensitive))
+			{
+				FString ParameterNamePrefix, Aux;
+				const bool bSplit = IntParameters_DEPRECATED[i].ParameterName.Split(FMultilayerProjector::NUM_LAYERS_PARAMETER_POSTFIX, &ParameterNamePrefix, &Aux);
+				check(bSplit);
+
+				// Confirm this is actually a multilayer param by finding the corresponding pose param
+				for (int32 j = 0; j < IntParameters_DEPRECATED.Num(); ++j)
+				{
+					if (i != j)
+					{
+						if (IntParameters_DEPRECATED[j].ParameterName.StartsWith(ParameterNamePrefix, ESearchCase::CaseSensitive) &&
+							IntParameters_DEPRECATED[j].ParameterName.EndsWith(FMultilayerProjector::POSE_PARAMETER_POSTFIX, ESearchCase::CaseSensitive))
+						{
+							IntParametersToMove.Add(i);
+							break;
+						}
+					}
+				}
+			}
+		}
+
+		// Convert them to float params
+		for (int32 i = 0; i < IntParametersToMove.Num(); ++i)
+		{
+			FloatParameters_DEPRECATED.AddDefaulted();
+			FloatParameters_DEPRECATED.Last().ParameterName = IntParameters_DEPRECATED[IntParametersToMove[i]].ParameterName;
+			FloatParameters_DEPRECATED.Last().ParameterValue = FCString::Atoi(*IntParameters_DEPRECATED[IntParametersToMove[i]].ParameterValueName);
+			FloatParameters_DEPRECATED.Last().Uid = IntParameters_DEPRECATED[IntParametersToMove[i]].Uid;
+		}
+
+		// Remove them from the int params in reverse order
+		for (int32 i = IntParametersToMove.Num() - 1; i >= 0; --i)
+		{
+			IntParameters_DEPRECATED.RemoveAt(IntParametersToMove[i]);
+		}
+	}
+	
+	if (CustomizableObjectCustomVersion < FCustomizableObjectCustomVersion::CustomizableObjectInstanceDescriptor)
+	{
+		Descriptor.CustomizableObject = CustomizableObject_DEPRECATED;
+
+		Descriptor.BoolParameters = BoolParameters_DEPRECATED;
+		Descriptor.IntParameters = IntParameters_DEPRECATED;
+		Descriptor.FloatParameters = FloatParameters_DEPRECATED;
+		Descriptor.TextureParameters = TextureParameters_DEPRECATED;
+		Descriptor.VectorParameters = VectorParameters_DEPRECATED;
+		Descriptor.ProjectorParameters = ProjectorParameters_DEPRECATED;
+	}
+}
+
+
+void UCustomizableObjectInstance::PostLoad()
+{
+	// Make sure mutable has been initialised.
+	UCustomizableObjectSystem::GetInstance(); 
+
+	Super::PostLoad();
+
+	// This property is transient and should be false unless explicitely set to true. In case old
+	// data still has it, reset it here
+	bBuildParameterDecorations = false;
+
+	Descriptor.CreateParametersLookupTable();
+}
+
+
+FString UCustomizableObjectInstance::GetDesc()
+{
+	FString ObjectName = "Missing Object";
+	if (UCustomizableObject* CustomizableObject = GetCustomizableObject())
+	{
+		ObjectName = CustomizableObject->GetName();
+	}
+
+	return FString::Printf(TEXT("Instance of [%s]"), *ObjectName);
+}
+
+
+bool UCustomizableObjectInstance::IsParamMultidimensional(const FString& ParamName) const
+{
+	return Descriptor.IsParamMultidimensional(ParamName);
+}
+
+
+int32 UCustomizableObjectInstance::CurrentParamRange(const FString& ParamName) const
+{
+	return Descriptor.CurrentParamRange(ParamName);
+}
+
+
+bool UCustomizableObjectInstance::IsParamMultidimensional(int32 ParamIndex) const
+{
+	return Descriptor.IsParamMultidimensional(ParamIndex);
+}
+
+
+// Only safe to call if the Mutable texture ref count system returns 0 and absolutely sure nobody holds a reference to the texture
+void UCustomizableInstancePrivateData::ReleaseMutableTexture(int32 MutableTextureId, UTexture2D* Texture, FMutableResourceCache& Cache)
+{
+	if (Texture && Texture->IsValidLowLevel())
+	{
+		Texture->ConditionalBeginDestroy();
+
+		for (FTexture2DMipMap& Mip : Texture->GetPlatformData()->Mips)
+		{
+			Mip.BulkData.RemoveBulkData();
+		}
+	}
+
+	// Must remove texture from cache since it has been released
+	for (TPair<FMutableImageCacheKey, TWeakObjectPtr<UTexture2D> >& CachedTexture : Cache.Images)
+	{
+		if (CachedTexture.Key.Resource == MutableTextureId)
+		{
+			Cache.Images.Remove(CachedTexture.Key);
+			break;
+		}
+	}
+}
+
+
+FCustomizableObjectProjector UCustomizableObjectInstance::GetProjectorDefaultValue(const int32 ParamIndex) const
+{
+	return Descriptor.GetProjectorDefaultValue(ParamIndex);
+}
+
+
+mu::ParametersPtr UCustomizableInstancePrivateData::ReloadParametersFromObject(UCustomizableObjectInstance* Public, bool ClearLastMeshIds)
+{
+	UCustomizableObject* CustomizableObject = Public->GetCustomizableObject();
+	if (!CustomizableObject)
+	{
+		// May happen when deleting assets in editor, while they are still open.
+		return nullptr;
+	}
+
+	if (!CustomizableObject->IsCompiled())
+	{	
+		return nullptr;
+	}
+
+	FirstLODAvailable = CustomizableObject->LODSettings.FirstLODAvailable;
+
+	if (CustomizableObject->LODSettings.bLODStreamingEnabled || !UCustomizableObjectSystem::GetInstance()->IsProgressiveMipStreamingEnabled())
+	{
+		SetCOInstanceFlags(LODsStreamingEnabled);
+		NumMaxLODsToStream = CustomizableObject->LODSettings.NumLODsToStream;
+	}
+	else
+	{
+		ClearCOInstanceFlags(LODsStreamingEnabled);
+		NumMaxLODsToStream = 0;
+	}
+	
+	int32 OldState = Public->GetState();
+	TArray<FCustomizableObjectBoolParameterValue> OldBoolParameters = Public->GetBoolParameters();
+	TArray<FCustomizableObjectIntParameterValue> OldIntParameters = Public->GetIntParameters();
+	TArray<FCustomizableObjectFloatParameterValue> OldFloatParameters = Public->GetFloatParameters();
+	TArray<FCustomizableObjectTextureParameterValue> OldTextureParameters = Public->GetTextureParameters();
+	TArray<FCustomizableObjectVectorParameterValue> OldVectorParameters = Public->GetVectorParameters();
+	TArray<FCustomizableObjectProjectorParameterValue> OldProjectorParameters = Public->GetProjectorParameters();
+	TMap<TPair<FString, int32>, EProjectorState::Type> OldProjectorStates = ProjectorStates;
+
+	if (ClearLastMeshIds)
+	{
+		ClearAllLastMeshIds();
+	}
+
+	mu::ParametersPtr MutableParameters;
+	Public->SetState(FMath::Clamp(OldState, 0, CustomizableObject->GetStateCount() - 1));
+	Public->GetBoolParameters().Reset();
+	Public->GetIntParameters().Reset();
+	Public->GetFloatParameters().Reset();
+	Public->GetTextureParameters().Reset();
+	Public->GetVectorParameters().Reset();
+	Public->GetProjectorParameters().Reset();
+
+#if WITH_EDITOR
+	if(!Public->ProjectorAlphaChange)
+#endif
+	{
+		ProjectorStates.Reset();
+	}
+
+	if (CustomizableObject && CustomizableObject->GetPrivate()->GetModel())
+	{
+		TArray<bool> OldIntParametersUsed;
+		OldIntParametersUsed.SetNumUninitialized(OldIntParameters.Num());
+		FMemory::Memzero(OldIntParametersUsed.GetData(), OldIntParametersUsed.GetAllocatedSize());
+
+		MutableParameters = CustomizableObject->GetPrivate()->GetModel()->NewParameters();
+		int32 paramCount = MutableParameters->GetCount();
+
+		for (int32 paramIndex = 0; paramIndex<paramCount; ++paramIndex)
+		{
+			FString Name = MutableParameters->GetName(paramIndex);
+			FString Uid = MutableParameters->GetUid(paramIndex);
+			mu::PARAMETER_TYPE mutableType = MutableParameters->GetType(paramIndex);
+
+			switch (mutableType)
+			{
+			case mu::PARAMETER_TYPE::T_BOOL:
+			{
+				FCustomizableObjectBoolParameterValue param;
+				param.ParameterName = Name;
+				param.Uid = Uid;
+
+				bool found = false;
+				for (int32 i = 0; i < OldBoolParameters.Num(); ++i)
+				{
+					if (OldBoolParameters[i].ParameterName == Name)
+					{
+						found = true;
+						param.ParameterValue = OldBoolParameters[i].ParameterValue;
+						MutableParameters->SetBoolValue(paramIndex, param.ParameterValue);
+
+						if (OldBoolParameters[i].Uid.IsEmpty())
+						{
+							Public->Modify();
+						}
+
+						break;
+					}
+				}
+
+				if (!found)
+				{
+					for (int32 i = 0; i < OldBoolParameters.Num(); ++i)
+					{
+						if (!Uid.IsEmpty() && OldBoolParameters[i].Uid == Uid)
+						{
+							found = true;
+							param.ParameterValue = OldBoolParameters[i].ParameterValue;
+							MutableParameters->SetBoolValue(paramIndex, param.ParameterValue);
+
+							if (OldBoolParameters[i].Uid.IsEmpty())
+							{
+								Public->Modify();
+							}
+
+							break;
+						}
+					}
+				}
+
+				if (!found)
+				{
+					param.ParameterValue = MutableParameters->GetBoolValue(paramIndex);
+				}
+
+				Public->GetBoolParameters().Add(param);
+				break;
+			}
+
+			case mu::PARAMETER_TYPE::T_INT:
+			{
+				FString ParameterValueName;
+				TArray<FString> ParameterRangeValueNames;
+
+				bool found = false;
+				for (int32 i = 0; i < OldIntParameters.Num(); ++i)
+				{
+					if (OldIntParametersUsed[i] == false && OldIntParameters[i].ParameterName.Equals(Name, ESearchCase::CaseSensitive))
+					{
+						found = true;
+						OldIntParametersUsed[i] = true;
+
+						mu::RangeIndexPtr RangeIdxPtr = MutableParameters->NewRangeIndex(paramIndex);
+
+						if (!RangeIdxPtr)
+						{
+						ParameterValueName = OldIntParameters[i].ParameterValueName;
+						int32 Value = CustomizableObject->FindIntParameterValue(paramIndex, ParameterValueName);
+						MutableParameters->SetIntValue(paramIndex, Value);
+						}
+						else
+						{
+							for (int RangeIndex = 0; RangeIndex < OldIntParameters[i].ParameterRangeValueNames.Num(); ++RangeIndex)
+							{
+								FString AuxParameterValueName = OldIntParameters[i].ParameterRangeValueNames[RangeIndex];
+								ParameterRangeValueNames.Add(AuxParameterValueName);
+								int32 Value = CustomizableObject->FindIntParameterValue(paramIndex, AuxParameterValueName);
+								RangeIdxPtr->SetPosition(0, RangeIndex);
+								MutableParameters->SetIntValue(paramIndex, Value, RangeIdxPtr);
+							}
+						}
+
+						if (OldIntParameters[i].Uid.IsEmpty())
+						{
+							Public->Modify();
+						}
+
+						break;
+					}
+				}
+
+				if (!found)
+				{
+					for (int32 i = 0; i < OldIntParameters.Num(); ++i)
+					{
+						if (!Uid.IsEmpty() && OldIntParametersUsed[i] == false && OldIntParameters[i].Uid == Uid)
+						{
+							found = true;
+							OldIntParametersUsed[i] = true;
+
+							mu::RangeIndexPtr RangeIdxPtr = MutableParameters->NewRangeIndex(paramIndex);
+
+							if (!RangeIdxPtr)
+							{
+							ParameterValueName = OldIntParameters[i].ParameterValueName;
+							int32 Value = CustomizableObject->FindIntParameterValue(paramIndex, ParameterValueName);
+							MutableParameters->SetIntValue(paramIndex, Value);
+							}
+							else
+							{
+								for (int RangeIndex = 0; RangeIndex < OldIntParameters[i].ParameterRangeValueNames.Num(); ++RangeIndex)
+								{
+									FString AuxParameterValueName = OldIntParameters[i].ParameterRangeValueNames[RangeIndex];
+									ParameterRangeValueNames.Add(AuxParameterValueName);
+									int32 Value = CustomizableObject->FindIntParameterValue(paramIndex, AuxParameterValueName);
+									RangeIdxPtr->SetPosition(0, RangeIndex);
+									MutableParameters->SetIntValue(paramIndex, Value, RangeIdxPtr);
+								}
+							}
+
+							if (OldIntParameters[i].Uid.IsEmpty())
+							{
+								Public->Modify();
+							}
+
+							break;
+						}
+					}
+				}
+
+				if (!found)
+				{
+					mu::RangeIndexPtr RangeIdxPtr = MutableParameters->NewRangeIndex(paramIndex);
+
+					if (!RangeIdxPtr.get())
+					{
+					int32 Value = MutableParameters->GetIntValue(paramIndex);
+					ParameterValueName = CustomizableObject->FindIntParameterValueName(paramIndex, Value);
+					}
+					else
+					{
+						int32 ValueCount = MutableParameters->GetValueCount(paramIndex);
+
+						for (int32 ValueIterator = 0; ValueIterator < ValueCount; ++ValueIterator)
+						{
+							mu::RangeIndexPtr RangeValueIdxPtr = MutableParameters->GetValueIndex(paramIndex, ValueIterator);
+							int32 ValueIndex = RangeValueIdxPtr->GetPosition(0);
+
+							//if (!param.RangeValues.IsValidIndex(ValueIndex))
+							//{
+							//	param.RangeValues.AddDefaulted(ValueIndex + 1 - param.RangeValues.Num());
+							//}
+
+							int32 Value = MutableParameters->GetIntValue(paramIndex, RangeValueIdxPtr);
+							FString AuxParameterValueName = CustomizableObject->FindIntParameterValueName(paramIndex, Value);
+							//param.RangeValues[ValueIndex] = AuxParameterValueName;
+
+							if (!ParameterRangeValueNames.IsValidIndex(ValueIndex))
+							{
+								ParameterRangeValueNames.AddDefaulted(ValueIndex + 1 - ParameterRangeValueNames.Num());
+							}
+
+							ParameterRangeValueNames[ValueIndex] = AuxParameterValueName;
+						}
+					}
+
+					Public->Modify();
+				}
+
+				//Public->IntParameters.Add(param);
+				Public->GetIntParameters().Emplace(Name, ParameterValueName, Uid, ParameterRangeValueNames);
+
+				break;
+			}
+
+			case mu::PARAMETER_TYPE::T_FLOAT:
+			{
+				FCustomizableObjectFloatParameterValue param;
+				param.ParameterName = Name;
+				param.Uid = Uid;
+
+				bool found = false;
+				for (int32 i = 0; i < OldFloatParameters.Num(); ++i)
+				{
+					if (OldFloatParameters[i].ParameterName == Name || (!Uid.IsEmpty() && OldFloatParameters[i].Uid == Uid))
+					{
+						found = true;
+						param.ParameterValue = OldFloatParameters[i].ParameterValue;
+						param.ParameterRangeValues = OldFloatParameters[i].ParameterRangeValues;
+
+						mu::RangeIndexPtr RangeIdxPtr = MutableParameters->NewRangeIndex(paramIndex);
+
+						if (!RangeIdxPtr)
+						{
+							MutableParameters->SetFloatValue(paramIndex, param.ParameterValue);
+						}
+						else
+						{
+							for (int RangeIndex = 0; RangeIndex < OldFloatParameters[i].ParameterRangeValues.Num(); ++RangeIndex)
+							{
+								RangeIdxPtr->SetPosition(0, RangeIndex);
+								MutableParameters->SetFloatValue(paramIndex, param.ParameterRangeValues[RangeIndex], RangeIdxPtr);
+							}
+						}
+
+						if (OldFloatParameters[i].Uid.IsEmpty())
+						{
+							Public->Modify();
+						}
+					}
+				}
+
+				if (!found)
+				{
+					mu::RangeIndexPtr RangeIdxPtr = MutableParameters->NewRangeIndex(paramIndex);
+
+					if (!RangeIdxPtr.get())
+					{
+					param.ParameterValue = MutableParameters->GetFloatValue(paramIndex);
+				}
+					else
+					{
+					param.ParameterValue = MutableParameters->GetFloatValue(paramIndex);
+
+						int32 ValueCount = MutableParameters->GetValueCount(paramIndex);
+
+						for (int32 ValueIterator = 0; ValueIterator < ValueCount; ++ValueIterator)
+						{
+							mu::RangeIndexPtr RangeValueIdxPtr = MutableParameters->GetValueIndex(paramIndex, ValueIterator);
+							int32 ValueIndex = RangeValueIdxPtr->GetPosition(0);
+
+							float Value = MutableParameters->GetFloatValue(paramIndex, RangeValueIdxPtr);
+
+							if (!param.ParameterRangeValues.IsValidIndex(ValueIndex))
+							{
+								param.ParameterRangeValues.AddDefaulted(ValueIndex + 1 - param.ParameterRangeValues.Num());
+							}
+
+							param.ParameterRangeValues[ValueIndex] = Value;
+						}
+					}
+				}
+
+				Public->GetFloatParameters().Add(param);
+				break;
+			}
+
+			case mu::PARAMETER_TYPE::T_COLOUR:
+			{
+				FCustomizableObjectVectorParameterValue param;
+				param.ParameterName = Name;
+				param.Uid = Uid;
+
+				bool found = false;
+				for (int32 i = 0; i < OldVectorParameters.Num(); ++i)
+				{
+					if (OldVectorParameters[i].ParameterName == Name || (!Uid.IsEmpty() && OldVectorParameters[i].Uid == Uid))
+					{
+						found = true;
+						param.ParameterValue = OldVectorParameters[i].ParameterValue;
+						MutableParameters->SetColourValue(paramIndex, param.ParameterValue.R, param.ParameterValue.G, param.ParameterValue.B);
+
+						if (OldVectorParameters[i].Uid.IsEmpty())
+						{
+							Public->Modify();
+						}
+					}
+				}
+
+				if (!found)
+				{
+					MutableParameters->GetColourValue(paramIndex, &param.ParameterValue.R, &param.ParameterValue.G, &param.ParameterValue.B);
+					param.ParameterValue.A = 1.0f;
+				}
+
+				Public->GetVectorParameters().Add(param);
+				break;
+			}
+
+			case mu::PARAMETER_TYPE::T_PROJECTOR:
+			{
+				FCustomizableObjectProjectorParameterValue param;
+				param.ParameterName = Name;
+				param.Uid = Uid;
+
+				bool found = false;
+				for (int32 i = 0; i < OldProjectorParameters.Num(); ++i)
+				{
+					if (OldProjectorParameters[i].ParameterName == Name || (!Uid.IsEmpty() && OldProjectorParameters[i].Uid == Uid))
+					{
+						found = true;
+
+						param.Value = OldProjectorParameters[i].Value;
+						param.RangeValues = OldProjectorParameters[i].RangeValues;
+
+						mu::RangeIndexPtr RangeIdxPtr = MutableParameters->NewRangeIndex(paramIndex);
+
+						TArray<FCustomizableObjectProjector> AuxValuesArray;
+						AuxValuesArray.Add(param.Value);
+						TArray<FCustomizableObjectProjector>& Values = RangeIdxPtr.get() ? param.RangeValues : AuxValuesArray;
+						int32 ArrayIndex = 0;
+
+						for (FCustomizableObjectProjector& Value : Values)
+						{
+							if (RangeIdxPtr.get())
+							{
+								check(RangeIdxPtr->GetRangeCount() == 1);
+								RangeIdxPtr->SetPosition(0, ArrayIndex);
+							}
+
+							switch (Value.ProjectionType)
+							{
+							case ECustomizableObjectProjectorType::Planar:
+							case ECustomizableObjectProjectorType::Wrapping:
+							{
+								MutableParameters->SetProjectorValue(paramIndex,
+									Value.Position[0], Value.Position[1], Value.Position[2],
+									Value.Direction[0], Value.Direction[1], Value.Direction[2],
+									Value.Up[0], Value.Up[1], Value.Up[2],
+									Value.Scale[0], Value.Scale[1], Value.Scale[2],
+									Value.Angle,
+									RangeIdxPtr);
+								break;
+							}
+
+							case ECustomizableObjectProjectorType::Cylindrical:
+							{
+								// Apply strange swizzle for scales
+								// TODO: try to avoid this
+									float Radius = FMath::Abs(Value.Scale[0] / 2.0f);
+									float Height = Value.Scale[2];
+								// TODO: try to avoid this
+								MutableParameters->SetProjectorValue(paramIndex,
+										Value.Position[0], Value.Position[1], Value.Position[2],
+										-Value.Direction[0], -Value.Direction[1], -Value.Direction[2],
+										-Value.Up[0], -Value.Up[1], -Value.Up[2],
+										-Height, Radius, Radius,
+										Value.Angle,
+										RangeIdxPtr);
+								break;
+							}
+
+							default:
+								// Not implemented.
+								check(false);
+							}
+							
+							ArrayIndex++;
+						}
+					}
+
+					if (OldProjectorParameters[i].Uid.IsEmpty())
+					{
+						Public->Modify();
+					}
+				}
+
+				if (!found)
+				{
+					mu::RangeIndexPtr RangeIdxPtr = MutableParameters->NewRangeIndex(paramIndex);
+
+					if (!RangeIdxPtr.get())
+					{
+					mu::PROJECTOR_TYPE type;
+					MutableParameters->GetProjectorValue(paramIndex,
+						&type,
+						&param.Value.Position[0], &param.Value.Position[1], &param.Value.Position[2],
+						&param.Value.Direction[0], &param.Value.Direction[1], &param.Value.Direction[2],
+						&param.Value.Up[0], &param.Value.Up[1], &param.Value.Up[2],
+						&param.Value.Scale[0], &param.Value.Scale[1], &param.Value.Scale[2],
+						&param.Value.Angle );
+					switch (type)
+					{
+					case mu::PROJECTOR_TYPE::PLANAR:
+						param.Value.ProjectionType = ECustomizableObjectProjectorType::Planar;
+						break;
+
+					case mu::PROJECTOR_TYPE::CYLINDRICAL:
+						// Unapply strange swizzle for scales.
+						// TODO: try to avoid this
+						param.Value.ProjectionType = ECustomizableObjectProjectorType::Cylindrical;
+						param.Value.Direction = -param.Value.Direction;
+						param.Value.Up = -param.Value.Up;
+						param.Value.Scale[2] = -param.Value.Scale[0];
+						param.Value.Scale[0] = param.Value.Scale[1] = param.Value.Scale[1] * 2.0f;
+						break;
+
+					case mu::PROJECTOR_TYPE::WRAPPING:
+						param.Value.ProjectionType = ECustomizableObjectProjectorType::Wrapping;
+						break;
+					default:
+						// not implemented
+						check(false);
+					}
+				}
+					else
+					{
+						int32 ValueCount = MutableParameters->GetValueCount(paramIndex);
+
+						for (int32 ValueIterator = 0; ValueIterator < ValueCount; ++ValueIterator)
+						{
+							mu::RangeIndexPtr RangeValueIdxPtr = MutableParameters->GetValueIndex(paramIndex, ValueIterator);
+							int32 ValueIndex = RangeValueIdxPtr->GetPosition(0);
+
+							if (!param.RangeValues.IsValidIndex(ValueIndex))
+							{
+								param.RangeValues.AddDefaulted(ValueIndex + 1 - param.RangeValues.Num());
+							}
+
+							mu::PROJECTOR_TYPE type;
+							MutableParameters->GetProjectorValue(paramIndex,
+							&type,
+							&param.RangeValues[ValueIndex].Position[0], &param.RangeValues[ValueIndex].Position[1], &param.RangeValues[ValueIndex].Position[2],
+							&param.RangeValues[ValueIndex].Direction[0], &param.RangeValues[ValueIndex].Direction[1], &param.RangeValues[ValueIndex].Direction[2],
+							&param.RangeValues[ValueIndex].Up[0], &param.RangeValues[ValueIndex].Up[1], &param.RangeValues[ValueIndex].Up[2],
+							&param.RangeValues[ValueIndex].Scale[0], &param.RangeValues[ValueIndex].Scale[1], &param.RangeValues[ValueIndex].Scale[2],
+							&param.RangeValues[ValueIndex].Angle,
+							RangeValueIdxPtr);
+
+							switch (type)
+							{
+							case mu::PROJECTOR_TYPE::PLANAR:
+								param.RangeValues[ValueIndex].ProjectionType = ECustomizableObjectProjectorType::Planar;
+								break;
+
+							case mu::PROJECTOR_TYPE::CYLINDRICAL:
+								// Unapply strange swizzle for scales.
+								// TODO: try to avoid this
+								param.RangeValues[ValueIndex].ProjectionType = ECustomizableObjectProjectorType::Cylindrical;
+								param.RangeValues[ValueIndex].Direction = -param.RangeValues[ValueIndex].Direction;
+								param.RangeValues[ValueIndex].Up = -param.RangeValues[ValueIndex].Up;
+								param.RangeValues[ValueIndex].Scale[2] = -param.RangeValues[ValueIndex].Scale[0];
+								param.RangeValues[ValueIndex].Scale[0] = param.RangeValues[ValueIndex].Scale[1] = param.RangeValues[ValueIndex].Scale[1] * 2.0f;
+								break;
+
+							case mu::PROJECTOR_TYPE::WRAPPING:
+								param.RangeValues[ValueIndex].ProjectionType = ECustomizableObjectProjectorType::Wrapping;
+								break;
+							default:
+								// not implemented
+								check(false);
+							}
+						}
+					}
+				}
+
+				Public->GetProjectorParameters().Add(param);
+				break;
+			}
+
+			case mu::PARAMETER_TYPE::T_IMAGE:
+			{
+				FCustomizableObjectTextureParameterValue param;
+				param.ParameterName = Name;
+				param.Uid = Uid;
+
+				bool found = false;
+				for (int32 i = 0; i < OldTextureParameters.Num(); ++i)
+				{
+					if (OldTextureParameters[i].ParameterName == Name || (!Uid.IsEmpty() && OldTextureParameters[i].Uid == Uid))
+					{
+						found = true;
+						param.ParameterValue = OldTextureParameters[i].ParameterValue;
+						MutableParameters->SetImageValue(paramIndex, param.ParameterValue);
+					}
+
+					if (OldTextureParameters[i].Uid.IsEmpty())
+					{
+						Public->Modify();
+					}
+				}
+
+				if (!found)
+				{
+					param.ParameterValue = MutableParameters->GetImageValue(paramIndex);
+				}
+
+				Public->GetTextureParameters().Add(param);
+				break;
+			}
+
+			default:
+				// TODO
+				break;
+
+			}
+		}
+	}
+	else
+	{
+		if (!CustomizableObject)
+		{
+			UE_LOG(LogMutable, Warning, TEXT("[ReloadParametersFromObject] No CustomizableObject, generated empty parameters for [%s] "), *GetName());
+		}
+		else if (!CustomizableObject->GetPrivate()->GetModel())
+		{
+			UE_LOG(LogMutable, Warning, TEXT("[ReloadParametersFromObject] No model in object [%s], generated empty parameters for [%s] "), *CustomizableObject->GetName(), *GetName());
+		}
+	}
+	
+	Public->GetDescriptor().CreateParametersLookupTable();
+
+	return MutableParameters;
+}
+
+
+void UCustomizableObjectInstance::SetObject(UCustomizableObject* InObject)
+{
+	Descriptor.CustomizableObject = InObject;
+	PrivateData->ReloadParametersFromObject(this, true);
+}
+
+
+UCustomizableObject* UCustomizableObjectInstance::GetCustomizableObject() const
+{
+	return Descriptor.CustomizableObject;
+}
+
+
+int32 UCustomizableObjectInstance::GetState() const
+{
+	return Descriptor.GetState();
+}
+
+
+void UCustomizableObjectInstance::SetState(const int32 InState)
+{
+	const int32 OldState = GetState();
+	
+	Descriptor.SetState(InState);
+
+	if (OldState != InState)
+	{
+		PrivateData->ProjectorStates.Reset();
+
+		// State may change texture properties, so invalidate the texture reuse cache
+		PrivateData->TextureReuseCache.Empty();
+	}
+}
+
+
+FString UCustomizableObjectInstance::GetCurrentState() const
+{
+	return GetCustomizableObject()->GetStateName(GetState());
+}
+
+
+void UCustomizableObjectInstance::SetCurrentState(const FString& StateName)
+{
+	SetState(GetCustomizableObject()->FindState(StateName));
+}
+
+
+void UCustomizableObjectInstance::SetProjectorState(const FString& ParamName, int32 RangeIndex, EProjectorState::Type state)
+{
+	PrivateData->ProjectorStates.Add(TPair<FString, int32>(ParamName, RangeIndex), state);
+	ProjectorStateChangedDelegate.ExecuteIfBound(ParamName);
+}
+
+
+void UCustomizableObjectInstance::ResetProjectorStates()
+{
+	PrivateData->ProjectorStates.Empty();
+}
+
+
+EProjectorState::Type UCustomizableObjectInstance::GetProjectorState(const FString& ParamName, int32 RangeIndex) const
+{
+	if (const EProjectorState::Type* state = PrivateData->ProjectorStates.Find(TPair<FString, int32>(ParamName, RangeIndex)))
+	{
+		return *state;
+	}
+	else
+	{
+		return EProjectorState::Hidden;
+	}
+}
+
+
+UTexture2D* UCustomizableObjectInstance::GetParameterDescription(int32 ParamIndex, int32 DescIndex)
+{
+	const UCustomizableObject* CustomizableObject = GetCustomizableObject();
+    if (!CustomizableObject)
+    {
+    	return nullptr;
+    }
+	
+	if (CustomizableObject->GetPrivate()->GetModel())
+	{
+		return nullptr;
+	}
+
+	// TODO, make sure the lines below aren't actually needed. They were commented because they made SDetailsViewBase::UpdatePropertyMaps() crash
+	// Make sure this instance is updated
+	//UCustomizableObjectSystem* CustomizableObjectSystem = UCustomizableObjectSystem::GetInstance();
+	//CustomizableObjectSystem->WaitForInstance(this);
+
+	if (GetPrivate()->ParameterDecorations.IsValidIndex(ParamIndex)
+		&&
+		GetPrivate()->ParameterDecorations[ParamIndex].Images.IsValidIndex(DescIndex))
+	{
+		return GetPrivate()->ParameterDecorations[ParamIndex].Images[DescIndex];
+	}
+
+	//UE_LOG(LogMutable, Warning, TEXT("Expected parameter decoration not found in instance [%s]."), *GetName());
+	return nullptr;
+}
+
+
+UTexture2D* UCustomizableObjectInstance::GetParameterDescription(const FString& ParamName, int32 DescIndex)
+{
+	return GetParameterDescription(GetCustomizableObject()->FindParameter(ParamName), DescIndex);
+}
+
+
+bool UCustomizableObjectInstance::IsParameterRelevant(int32 ParameterIndex) const
+{
+	// This should have been precalculated in the last update if the appropiate flag in the instance was set.
+	return GetPrivate()->RelevantParameters.Contains(ParameterIndex);
+}
+
+
+bool UCustomizableObjectInstance::IsParameterRelevant(const FString& ParamName) const
+{
+	// This should have been precalculated in the last update if the appropiate flag in the instance was set.
+	int32 ParameterIndexInObject = GetCustomizableObject()->FindParameter(ParamName);
+	return GetPrivate()->RelevantParameters.Contains(ParameterIndexInObject);
+}
+
+
+void UCustomizableInstancePrivateData::PostEditChangePropertyWithoutEditor(USkeletalMesh* InSkeletalMesh)
+{
+	MUTABLE_CPUPROFILER_SCOPE(UCustomizableInstancePrivateData::PostEditChangePropertyWithoutEditor);
+
+	if (!InSkeletalMesh) return;
+
+	FSkeletalMeshRenderData* Resource = Helper_GetResourceForRendering(InSkeletalMesh);
+	if (!Resource) return;
+
+	{
+		MUTABLE_CPUPROFILER_SCOPE(InitResources);
+		// reinitialize resource
+		InSkeletalMesh->InitResources();
+	}
+
+	{
+		MUTABLE_CPUPROFILER_SCOPE(UpdateSkeletalMesh_LocalUVDensities);
+
+		for (int32 MaterialIndex = 0; MaterialIndex < InSkeletalMesh->GetMaterials().Num(); ++MaterialIndex)
+		{
+			FMeshUVChannelInfo& UVChannelData = InSkeletalMesh->GetMaterials()[MaterialIndex].UVChannelData;
+
+			UVChannelData.bInitialized = true;
+			UVChannelData.bOverrideDensities = false;
+
+			for (int32 i = 0; i < TEXSTREAM_MAX_NUM_UVCHANNELS; ++i)
+			{
+				// TODO
+				UVChannelData.LocalUVDensities[i] = 200.f;
+			}
+		}
+	}
+}
+
+
+void UCustomizableObjectInstance::UpdateSkeletalMeshAsync(bool bIgnoreCloseDist, bool bForceHighPriority)
+{
+	PrivateData->DoUpdateSkeletalMesh(this, true, false, false, bIgnoreCloseDist, bForceHighPriority);
+}
+
+
+void UCustomizableInstancePrivateData::DoUpdateSkeletalMeshAsync(UCustomizableObjectInstance* Instance, FMutableQueueElem::EQueuePriorityType Priority)
+{
+	check(IsInGameThread());
+
+	UE_LOG(LogMutable, Verbose, TEXT("Started UpdateSkeletalMesh Async. of Instance %d with priority %d at dist %f bIsPlayerOrNearIt=%d, frame=%d"), Instance->GetUniqueID(), static_cast<int32>(Priority), FMath::Sqrt(Instance->GetPrivate()->LastMinSquareDistFromComponentToPlayer), Instance->GetPrivate()->HasCOInstanceFlags(UsedByPlayerOrNearIt), GFrameNumber);
+
+	if (HasCOInstanceFlags(PendingLODsUpdateSecondStage))
+	{
+		UE_LOG(LogMutable, Verbose, TEXT("LOD change: %d, %d -> %d, %d"), CurrentMinLOD, CurrentMaxLOD, MinLODToLoad, MaxLODToLoad);
+	}
+
+	UCustomizableObjectSystem::GetInstance()->GetPrivate()->InitUpdateSkeletalMesh(Instance, Priority);
+
+#if WITH_EDITOR
+	if (!UCustomizableObjectSystem::GetInstance()->IsCompilationDisabled())
+	{
+		if (Instance->SkeletalMeshStatus != ESkeletalMeshState::AsyncUpdatePending)
+		{
+			Instance->PreUpdateSkeletalMeshStatus = Instance->SkeletalMeshStatus;
+		}
+
+		Instance->SkeletalMeshStatus = ESkeletalMeshState::AsyncUpdatePending;
+	}
+#else
+	Instance->SkeletalMeshStatus = ESkeletalMeshState::AsyncUpdatePending;
+#endif
+}
+
+
+void UCustomizableInstancePrivateData::DoUpdateSkeletalMesh(UCustomizableObjectInstance* CustomizableInstance, bool bAsync, bool bIsCloseDistTick, bool bOnlyUpdateIfNotGenerated, bool bIgnoreCloseDist, bool bForceHighPriority)
+{
+	MUTABLE_CPUPROFILER_SCOPE(UCustomizableInstancePrivateData::DoUpdateSkeletalMesh);
+
+	check(IsInGameThread());
+	check(bAsync);
+
+	UCustomizableObjectSystem* System = UCustomizableObjectSystem::GetInstance();
+	
+	UCustomizableObject* CustomizableObject = CustomizableInstance->GetCustomizableObject();
+
+	if (CustomizableObject && !CustomizableObject->IsCompiled())
+	{
+#if WITH_EDITOR
+		if(!CustomizableObject->ConditionalAutoCompile())
+		{
+			return;
+		}
+#else
+		System->AddUncompiledCOWarning(CustomizableObject);
+		return;
+#endif
+
+	}
+
+	const bool bIsGenerated = HasCOInstanceFlags(Generated);
+
+	int32 NumGeneratedInstancesLimit = System->GetInstanceLODManagement()->GetNumGeneratedInstancesLimitFullLODs();
+	int32 NumGeneratedInstancesLimitLOD1 = System->GetInstanceLODManagement()->GetNumGeneratedInstancesLimitLOD1();
+	int32 NumGeneratedInstancesLimitLOD2 = System->GetInstanceLODManagement()->GetNumGeneratedInstancesLimitLOD2();
+
+	if (NumGeneratedInstancesLimit > 0 && System->GetPrivate()->GetCountAllocatedSkeletalMesh() > NumGeneratedInstancesLimit + NumGeneratedInstancesLimitLOD1 + NumGeneratedInstancesLimitLOD2)
+	{
+		if (!bIsGenerated) // Prevent generating more instances than the limit, but let updates to existing instances run normally
+		{
+			return;
+		}
+	}
+
+	const bool bShouldUpdateLODs = HasCOInstanceFlags(PendingLODsUpdate);
+
+	if (
+		!HasCOInstanceFlags(DiscardedByNumInstancesLimit)
+		&&
+		CustomizableObject && !CustomizableObject->IsLocked()
+		&&
+		(
+			!System->GetInstanceLODManagement()->IsOnlyUpdateCloseCustomizableObjectsEnabled() 
+			|| LastMinSquareDistFromComponentToPlayer <= FMath::Square(System->GetInstanceLODManagement()->GetOnlyUpdateCloseCustomizableObjectsDist())
+			|| bIgnoreCloseDist 
+			|| bShouldUpdateLODs
+			)
+		)
+	{
+		if (bShouldUpdateLODs)
+		{
+			SetCOInstanceFlags(PendingLODsUpdateSecondStage);
+		}
+
+		if ((!(bIsGenerated && bOnlyUpdateIfNotGenerated) && !(bIsCloseDistTick && bIsGenerated)) || bShouldUpdateLODs)
+		{
+			FMutableQueueElem::EQueuePriorityType Priority = FMutableQueueElem::EQueuePriorityType::Low;
+
+			const bool bIsDowngradeLODUpdate = HasCOInstanceFlags(PendingLODsDowngrade);
+			const bool bIsPlayerOrNearIt = HasCOInstanceFlags(UsedByPlayerOrNearIt);
+
+			if (bForceHighPriority)
+			{
+				Priority = FMutableQueueElem::EQueuePriorityType::High;
+			}
+			else if (!bIsGenerated || !CustomizableInstance->HasAnySkeletalMesh())
+			{
+				Priority = FMutableQueueElem::EQueuePriorityType::Med;
+			}
+			else if (bShouldUpdateLODs && bIsDowngradeLODUpdate)
+			{
+				Priority = FMutableQueueElem::EQueuePriorityType::Med_Low;
+			}
+			else if (bIsPlayerOrNearIt && bShouldUpdateLODs && !bIsDowngradeLODUpdate)
+			{
+				Priority = FMutableQueueElem::EQueuePriorityType::High;
+			}
+			else if (bShouldUpdateLODs && !bIsDowngradeLODUpdate)
+			{
+				Priority = FMutableQueueElem::EQueuePriorityType::Med;
+			}
+			else if (bIsPlayerOrNearIt)
+			{
+				Priority = FMutableQueueElem::EQueuePriorityType::High;
+			}
+
+			DoUpdateSkeletalMeshAsync(CustomizableInstance, Priority);
+
+			SetCOInstanceFlags(Generated); // Will be done in UpdateSkeletalMesh_PostBeginUpdate
+			LastUpdateTime = FPlatformTime::Seconds();
+
+		}
+
+		ClearCOInstanceFlags(PendingLODsUpdate);
+	}
+	else
+	{
+		if (HasCOInstanceFlags(Generated) && !HasCOInstanceFlags(Updating))
+		{
+			System->GetPrivate()->InitDiscardResourcesSkeletalMesh(CustomizableInstance);
+			ClearCOInstanceFlags((ECOInstanceFlags)(PendingLODsUpdate | PendingLODsUpdateSecondStage));
+		}
+	}
+}
+
+
+void UCustomizableInstancePrivateData::TickUpdateCloseCustomizableObjects( UCustomizableObjectInstance* Public )
+{
+	DoUpdateSkeletalMesh(Public, true, true, true, false);
+}
+
+
+void UCustomizableInstancePrivateData::UpdateInstanceIfNotGenerated(UCustomizableObjectInstance* Public, bool bAsync)
+{
+	if (!HasCOInstanceFlags(Generated))
+	{
+		DoUpdateSkeletalMesh(Public, bAsync, true, false, true);
+	}
+}
+
+
+USkeleton* UCustomizableInstancePrivateData::MergeSkeletons(UCustomizableObjectInstance* Public, const FMutableRefSkeletalMeshData* RefSkeletalMeshData, int32 ComponentIndex)
+{
+	MUTABLE_CPUPROFILER_SCOPE(BuildSkeletonData_MergeSkeletons);
+
+	const UCustomizableObject* CustomizableObject = Public->GetCustomizableObject();
+
+	check(CustomizableObject);
+	check(RefSkeletalMeshData);
+	
+	FCustomizableInstanceComponentData* ComponentData = GetComponentData(ComponentIndex);
+	check(ComponentData);
+
+	FReferencedSkeletons& ReferencedSkeletons = ComponentData->Skeletons;
+	
+	// No need to merge skeletons
+	if(ReferencedSkeletons.SkeletonsToMerge.Num() == 1)
+	{
+		const TObjectPtr<USkeleton> RefSkeleton = ReferencedSkeletons.SkeletonsToMerge[0];
+		ReferencedSkeletons.SkeletonsToMerge.Empty();
+
+		check(RefSkeleton);
+		return RefSkeleton;
+	}
+
+	FSkeletonMergeParams Params;
+	Exchange(Params.SkeletonsToMerge, ReferencedSkeletons.SkeletonsToMerge);
+
+#if !UE_BUILD_SHIPPING
+#if !MUTABLE_CLEAN_ENGINE_BRANCH
+	Params.bCheckSkeletonsCompatibility = true;
+#endif
+#endif
+
+	USkeleton* FinalSkeleton = USkeletalMergingLibrary::MergeSkeletons(Params);
+	if (!FinalSkeleton)
+	{
+		FString Msg = FString::Printf(TEXT("MergeSkeletons failed for Customizable Object [%s], Instance [%s]. Skeletons involved: "),
+			*CustomizableObject->GetName(),
+			*Public->GetName());
+		
+		const int32 SkeletonCount = Params.SkeletonsToMerge.Num();
+		for (int32 SkeletonIndex = 0; SkeletonIndex < SkeletonCount; ++SkeletonIndex)
+		{
+			Msg += FString::Printf(TEXT(" [%s]"), *Params.SkeletonsToMerge[SkeletonIndex]->GetName());
+		}
+		
+		UE_LOG(LogMutable, Error, TEXT("%s"), *Msg);
+	}
+	else
+	{
+		// The reference skeleton may be used by AnimBp as a target skeleton, add it as a compatible skeleton.
+		FinalSkeleton->AddCompatibleSkeleton(RefSkeletalMeshData->Skeleton.Get());
+	}
+	
+	return FinalSkeleton;
+}
+
+
+void UCustomizableInstancePrivateData::RefitPhysicsBodies(UPhysicsAsset* InOutPhysicsAsset, const mu::PhysicsBody* MutablePhysicsBody)
+{
+	if (!(InOutPhysicsAsset && MutablePhysicsBody))
+	{
+		return;
+	}
+
+	const int32 MeshPhysicsBodyCount = MutablePhysicsBody->GetBodyCount();
+
+	for (int32 B = 0; B < MeshPhysicsBodyCount; ++B)
+	{
+		const FString BodyBoneName( MutablePhysicsBody->GetBodyBoneName( B ) );
+
+		auto GetCollisionEnabledFormFlags = [](uint32 Flags) -> ECollisionEnabled::Type
+		{
+			return ECollisionEnabled::Type( Flags & 0xFF ); 
+		};
+
+		auto GetContributeToMassFromFlags = [](uint32 Flags) -> bool
+		{
+			return static_cast<bool>( (Flags >> 8 ) & 1 );
+		};
+		
+		TObjectPtr<USkeletalBodySetup>* Found = InOutPhysicsAsset->SkeletalBodySetups.FindByPredicate(
+		[&BodyBoneName](TObjectPtr<USkeletalBodySetup>& BodySetup)
+		{
+			return BodyBoneName.Equals(BodySetup->BoneName.ToString());
+		});
+
+		if (Found)
+		{
+			FKAggregateGeom& BodyAggGeom = (*Found)->AggGeom;
+			
+			const int32 NumSpheres = MutablePhysicsBody->GetSphereCount( B );
+			TArray<FKSphereElem>& AggSpheres = BodyAggGeom.SphereElems;
+			AggSpheres.Empty(NumSpheres);
+			for (int32 I = 0; I < NumSpheres; ++I)
+			{
+				uint32 Flags = MutablePhysicsBody->GetSphereFlags( B, I );
+				FString Name = MutablePhysicsBody->GetSphereName( B, I );
+
+				FVector3f Position;
+				float Radius;
+
+				MutablePhysicsBody->GetSphere( B, I, Position, Radius );
+				FKSphereElem& NewElem = AggSpheres.AddDefaulted_GetRef();
+				
+				NewElem.Center = FVector(Position);
+				NewElem.Radius = Radius;
+				NewElem.SetContributeToMass( GetContributeToMassFromFlags( Flags ) );
+				NewElem.SetCollisionEnabled( GetCollisionEnabledFormFlags( Flags ) );
+				NewElem.SetName(FName(*Name));
+			}
+			
+			const int32 NumBoxes = MutablePhysicsBody->GetBoxCount( B );
+			TArray<FKBoxElem>& AggBoxes = BodyAggGeom.BoxElems;
+			AggBoxes.Empty(NumBoxes);
+			for (int32 I = 0; I < NumBoxes; ++I)
+			{
+				uint32 Flags = MutablePhysicsBody->GetBoxFlags( B, I );
+				FString Name = MutablePhysicsBody->GetBoxName( B, I );
+
+				FVector3f Position;
+				FQuat4f Orientation;
+				FVector3f Size;
+				MutablePhysicsBody->GetBox( B, I, Position, Orientation, Size );
+
+				FKBoxElem& NewElem = AggBoxes.AddDefaulted_GetRef();
+				
+				NewElem.Center = FVector( Position );
+				NewElem.Rotation = FRotator( Orientation.Rotator() );
+				NewElem.X = Size.X;
+				NewElem.Y = Size.Y;
+				NewElem.Z = Size.Z;
+				NewElem.SetContributeToMass( GetContributeToMassFromFlags( Flags ) );
+				NewElem.SetCollisionEnabled( GetCollisionEnabledFormFlags( Flags ) );
+				NewElem.SetName( FName(*Name) );
+			}
+
+			const int32 NumConvexes = MutablePhysicsBody->GetConvexCount( B );
+			TArray<FKConvexElem>& AggConvexes = BodyAggGeom.ConvexElems;
+			AggConvexes.Empty();
+			for (int32 I = 0; I < NumConvexes; ++I)
+			{
+				uint32 Flags = MutablePhysicsBody->GetConvexFlags( B, I );
+				FString Name = MutablePhysicsBody->GetConvexName( B, I );
+
+				const FVector3f* Vertices;
+				const int32* Indices;
+				int32 NumVertices;
+				int32 NumIndices;
+				FTransform3f Transform;
+	
+				MutablePhysicsBody->GetConvex( B, I, Vertices, NumVertices, Indices, NumIndices, Transform );
+				
+				TArrayView<const FVector3f> VerticesView( Vertices, NumVertices );
+				TArrayView<const int32> IndicesView( Indices, NumIndices );
+			}
+
+			TArray<FKSphylElem>& AggSphyls = BodyAggGeom.SphylElems;
+			const int32 NumSphyls = MutablePhysicsBody->GetSphylCount( B );
+			AggSphyls.Empty(NumSphyls);
+
+			for (int32 I = 0; I < NumSphyls; ++I)
+			{
+				uint32 Flags = MutablePhysicsBody->GetSphylFlags( B, I );
+				FString Name = MutablePhysicsBody->GetSphylName( B, I );
+
+				FVector3f Position;
+				FQuat4f Orientation;
+				float Radius;
+				float Length;
+
+				MutablePhysicsBody->GetSphyl( B, I, Position, Orientation, Radius, Length );
+
+				FKSphylElem& NewElem = AggSphyls.AddDefaulted_GetRef();
+				
+				NewElem.Center = FVector( Position );
+				NewElem.Rotation = FRotator( Orientation.Rotator() );
+				NewElem.Radius = Radius;
+				NewElem.Length = Length;
+				
+				NewElem.SetContributeToMass( GetContributeToMassFromFlags( Flags ) );
+				NewElem.SetCollisionEnabled( GetCollisionEnabledFormFlags( Flags ) );
+				NewElem.SetName( FName(*Name) );
+			}	
+
+			TArray<FKTaperedCapsuleElem>& AggTaperedCapsules = BodyAggGeom.TaperedCapsuleElems;
+			const int32 NumTaperedCapsules = MutablePhysicsBody->GetTaperedCapsuleCount( B );
+			AggTaperedCapsules.Empty(NumTaperedCapsules);
+
+			for (int32 I = 0; I < NumTaperedCapsules; ++I)
+			{
+				uint32 Flags = MutablePhysicsBody->GetTaperedCapsuleFlags( B, I );
+				FString Name = MutablePhysicsBody->GetTaperedCapsuleName( B, I );
+
+				FVector3f Position;
+				FQuat4f Orientation;
+				float Radius0;
+				float Radius1;
+				float Length;
+
+				MutablePhysicsBody->GetTaperedCapsule( B, I, Position, Orientation, Radius0, Radius1, Length );
+
+				FKTaperedCapsuleElem& NewElem = AggTaperedCapsules.AddDefaulted_GetRef();
+				
+				NewElem.Center = FVector( Position );
+				NewElem.Rotation = FRotator( Orientation.Rotator() );
+				NewElem.Radius0 = Radius0;
+				NewElem.Radius1 = Radius1;
+				NewElem.Length = Length;
+				
+				NewElem.SetContributeToMass( GetContributeToMassFromFlags( Flags ) );
+				NewElem.SetCollisionEnabled( GetCollisionEnabledFormFlags( Flags ) );
+				NewElem.SetName( FName(*Name) );	
+			}	
+		}
+	}
+}
+
+UPhysicsAsset* UCustomizableInstancePrivateData::MergePhysicalAssets(UPhysicsAsset* BaseAsset, int32 ComponentId)
+{
+	UPhysicsAsset* result = nullptr;
+	bool bMergeSucced = false;
+
+	if (BaseAsset)
+	{
+		result = DuplicateObject(BaseAsset, NULL);
+	}
+
+	FCustomizableInstanceComponentData* ComponentData =	GetComponentData(ComponentId);
+	check(ComponentData);
+	
+	for (const UPhysicsAsset* PhysicsAsset : ComponentData->PhysicsAssets.PhysicsAssetsToMerge)
+	{
+		if (!result && PhysicsAsset)
+		{
+			result = DuplicateObject(PhysicsAsset, NULL);
+		}
+		else if (PhysicsAsset && PhysicsAsset != BaseAsset)
+		{
+			bMergeSucced = true;
+
+			for (const TObjectPtr<USkeletalBodySetup>& BodySetup : PhysicsAsset->SkeletalBodySetups)
+			{
+				if (!result->SkeletalBodySetups.Contains(BodySetup))
+				{
+					bool bFoundBone = false;
+
+					for (const TObjectPtr<USkeletalBodySetup>& OrigBodySetup : result->SkeletalBodySetups)
+					{
+						if (OrigBodySetup->BoneName == BodySetup->BoneName)
+						{
+							OrigBodySetup->AddCollisionFrom(BodySetup);
+							bFoundBone = true;
+							break;
+						}
+					}
+
+					if (!bFoundBone)
+					{
+						result->SkeletalBodySetups.Add(BodySetup);
+					}
+				}
+			}
+
+			result->ConstraintSetup.Append(PhysicsAsset->ConstraintSetup);
+			result->CollisionDisableTable.Append(PhysicsAsset->CollisionDisableTable);
+#if WITH_EDITORONLY_DATA
+			result->PhysicalAnimationProfiles.Append(PhysicsAsset->PhysicalAnimationProfiles);
+			result->ConstraintProfiles.Append(PhysicsAsset->ConstraintProfiles);
+#endif // WITH_EDITORONLY_DATA
+			// TODO: Merge the rest of the properties
+		}
+	}
+
+	if (bMergeSucced)
+	{
+			result->UpdateBodySetupIndexMap();
+			result->UpdateBoundsBodiesArray();
+	}
+	else if (!result)
+	{
+		result = NewObject<UPhysicsAsset>();
+	}
+
+	ComponentData->PhysicsAssets.PhysicsAssetsToMerge.Empty();
+
+	return result;
+}
+
+
+bool UCustomizableInstancePrivateData::UpdateSkeletalMesh_PostBeginUpdate0(UCustomizableObjectInstance* Public, const TSharedPtr<FMutableOperationData>& OperationData )
+{
+	MUTABLE_CPUPROFILER_SCOPE(UCustomizableInstancePrivateData::UpdateSkeletalMesh_PostBeginUpdate0)
+	
+	UCustomizableObject* CustomizableObject = Public->GetCustomizableObject();
+
+	Public->SkeletalMeshStatus = ESkeletalMeshState::Correct;
+
+	bool bEmptyMesh;
+	bool bMeshNeedsUpdate = MeshNeedsUpdate(Public, OperationData, bEmptyMesh);
+	SetCOInstanceFlags(CreatingSkeletalMesh);
+
+	if (bMeshNeedsUpdate)
+	{
+		TextureReuseCache.Empty(); // Sections may have changed, so invalidate the texture reuse cache because it's indexed by section
+	}
+#if WITH_EDITOR
+	else
+	{
+		Public->SkeletalMeshStatus = Public->PreUpdateSkeletalMeshStatus; // The mesh won't be updated, set the previous SkeletalMeshStatus to avoid losing error notifications
+	}
+#endif
+
+	if (bEmptyMesh)
+	{
+		Public->SkeletalMeshStatus = ESkeletalMeshState::UpdateError;
+		UE_LOG(LogMutable, Warning, TEXT("Updating skeletal mesh error for mesh %s"), *Public->GetName());
+	}
+
+	const int32 FirstComponent = OperationData->InstanceUpdateData.LODs[CurrentMinLOD].FirstComponent;
+	const int32 NumComponents = OperationData->InstanceUpdateData.LODs[CurrentMinLOD].ComponentCount;
+
+	bool bHasSkeletalMesh = false;
+
+	if (bMeshNeedsUpdate && !bEmptyMesh)
+	{
+		MUTABLE_CPUPROFILER_SCOPE(ConstructMesh);
+
+		bool bSuccess = true;
+
+		const int32 MaxNumComponents = CustomizableObject->ReferenceSkeletalMeshesData.Num();
+		
+		// Initialize the maximum number of SkeletalMeshes we could possibly have. 
+		Public->SkeletalMeshes.Init(nullptr, MaxNumComponents);
+
+		for (int32 ComponentIndex = 0; ComponentIndex < NumComponents; ++ComponentIndex)
+		{
+			const FInstanceUpdateData::FComponent& Component = OperationData->InstanceUpdateData.Components[FirstComponent+ComponentIndex];
+			
+			if (!Component.Mesh)
+			{
+				continue;
+			}
+
+			// ComponentIndex could be different than the final SkeletalMesh index if a component is not generated.
+			const uint16 SkeletalMeshIndex = Component.Id;
+
+			if (Component.SurfaceCount > 0)
+			{
+				Public->SkeletalMeshes[SkeletalMeshIndex] = NewObject<USkeletalMesh>(GetTransientPackage(), NAME_None, RF_Transient);
+				bHasSkeletalMesh = true;
+			}
+
+			USkeletalMesh* SkeletalMesh = Public->SkeletalMeshes[SkeletalMeshIndex];
+
+			const FMutableRefSkeletalMeshData* RefSkeletalMeshData = CustomizableObject->GetRefSkeletalMeshData(SkeletalMeshIndex);
+			
+			if (SkeletalMesh && !RefSkeletalMeshData)
+			{
+				bSuccess = false;
+				break;
+			}
+
+			if (SkeletalMesh)
+			{
+				INC_DWORD_STAT(STAT_MutableNumSkeletalMeshes);
+
+				// Mutable skeletal meshes are generated dynamically in-game and cannot be streamed from disk
+				SkeletalMesh->NeverStream = 1;
+
+				bSuccess = BuildSkeletalMeshSkeletonData(OperationData, SkeletalMesh, RefSkeletalMeshData, Public, ComponentIndex);
+				if (!bSuccess)
+				{
+					break;
+				}
+
+				SkeletalMesh->SetImportedBounds(RefSkeletalMeshData->Bounds);
+				SkeletalMesh->SetPhysicsAsset(RefSkeletalMeshData->PhysicsAsset.Get());
+				SkeletalMesh->SetEnablePerPolyCollision(RefSkeletalMeshData->Settings.bEnablePerPolyCollision);
+				SkeletalMesh->SetPostProcessAnimBlueprint(RefSkeletalMeshData->PostProcessAnimInst.Get());
+				SkeletalMesh->SetShadowPhysicsAsset(RefSkeletalMeshData->ShadowPhysicsAsset.Get());
+
+				if (!HasCOInstanceFlags(ReduceLODs))
+				{
+					SkeletalMesh->SetMinLod(FirstLODAvailable);
+					SkeletalMesh->SetQualityLevelMinLod(FirstLODAvailable);
+				}
+
+				// Merge Physics Assets coming from SubMeshes of the newly generated Mesh
+				if (SkeletalMesh->GetPhysicsAsset() && OperationData->InstanceUpdateData.LODs.Num() )
+				{
+					mu::MeshPtrConst MutableMesh = Component.Mesh;
+					mu::Ptr<const mu::PhysicsBody> MutablePhysicsBodies = MutableMesh->GetPhysicsBody();
+					const bool MutablePhysicsBodiesModified = MutablePhysicsBodies ? MutablePhysicsBodies->bBodiesModified : false;
+
+					if( HasCOInstanceFlags(ReplacePhysicsAssets) || MutablePhysicsBodiesModified )
+					{
+						bool bBaseAssetIsValid = true;
+
+						for (const TObjectPtr<USkeletalBodySetup>& BodySetup : SkeletalMesh->GetPhysicsAsset()->SkeletalBodySetups)
+						{
+							if (BodySetup->BoneName != NAME_None && SkeletalMesh->GetRefSkeleton().FindBoneIndex(BodySetup->BoneName) == INDEX_NONE)
+							{
+								bBaseAssetIsValid = false;
+								break;
+							}
+						}
+
+						UPhysicsAsset* MergedPhysicsAsset = MergePhysicalAssets(bBaseAssetIsValid ? SkeletalMesh->GetPhysicsAsset() : nullptr, SkeletalMeshIndex);
+					
+						if ( MutablePhysicsBodies )
+						{
+							RefitPhysicsBodies( MergedPhysicsAsset, MutablePhysicsBodies.get());
+						}
+
+						SkeletalMesh->SetPhysicsAsset(MergedPhysicsAsset);
+					}
+				}
+
+				// Update mesh sockets.
+				{
+					MUTABLE_CPUPROFILER_SCOPE(ConstructMesh_UpdateSockets);
+	
+					const uint32 SocketCount = RefSkeletalMeshData->Sockets.Num();
+
+					TArray<USkeletalMeshSocket*>& Sockets = SkeletalMesh->GetMeshOnlySocketList();
+					Sockets.Empty(SocketCount);
+					
+					for (uint32 SocketIndex = 0; SocketIndex < SocketCount; ++SocketIndex)
+					{
+						const FMutableRefSocket& RefSocket = RefSkeletalMeshData->Sockets[SocketIndex];
+
+						USkeletalMeshSocket* Socket = NewObject<USkeletalMeshSocket>(SkeletalMesh);
+						
+						Socket->SocketName = RefSocket.SocketName;
+						Socket->BoneName = RefSocket.BoneName;
+
+						Socket->RelativeLocation = RefSocket.RelativeLocation;
+						Socket->RelativeRotation = RefSocket.RelativeRotation;
+						Socket->RelativeScale = RefSocket.RelativeScale;
+
+						Socket->bForceAlwaysAnimated = RefSocket.bForceAlwaysAnimated;
+						Sockets.Add(Socket);
+					}
+				}
+			}
+		}
+
+		if (!bSuccess)
+		{
+			Public->SkeletalMeshes.Reset();
+			return bSuccess;
+		}
+	}
+	else if (bMeshNeedsUpdate)
+	{
+		Public->SkeletalMeshes.Reset();
+		
+		SetCOInstanceFlags(Generated);
+		ClearCOInstanceFlags((ECOInstanceFlags)(CreatingSkeletalMesh | PendingMeshUpdate));
+
+		return false;
+	}
+
+	if (bMeshNeedsUpdate && bHasSkeletalMesh)
+	{
+		MUTABLE_CPUPROFILER_SCOPE(ReleaseAndBuildMesh);
+		bool bSuccess = true;
+
+		for (int32 ComponentIndex = 0; ComponentIndex < NumComponents; ++ComponentIndex)
+		{
+			const FInstanceUpdateData::FComponent& Component = OperationData->InstanceUpdateData.Components[FirstComponent + ComponentIndex];
+
+			if (!Component.Mesh)
+			{
+				continue;
+			}
+
+			const uint16 SkeletalMeshIndex = Component.Id;
+
+			check(Public->SkeletalMeshes.IsValidIndex(SkeletalMeshIndex));
+			USkeletalMesh* SkeletalMesh = Public->SkeletalMeshes[SkeletalMeshIndex];
+
+			const FMutableRefSkeletalMeshData* RefSkeletalMeshData = CustomizableObject->GetRefSkeletalMeshData(SkeletalMeshIndex);
+
+			if (SkeletalMesh && !RefSkeletalMeshData)
+			{
+				bSuccess = false;
+				break;
+			}
+			
+			if (SkeletalMesh)
+			{
+				{
+					MUTABLE_CPUPROFILER_SCOPE(ReleaseAndBuildMesh_Empty);
+
+					if (FSkeletalMeshRenderData* Resource = Helper_GetResourceForRendering(SkeletalMesh))
+					{
+						Helper_GetLODData(Resource).Empty();
+					}
+					else
+					{
+						SkeletalMesh->AllocateResourceForRendering();
+					}
+					Helper_GetLODInfoArray(SkeletalMesh).Empty();
+				}
+
+				// Add LODs
+				{
+					MUTABLE_CPUPROFILER_SCOPE(ReleaseAndBuildMesh_AddLODs);
+
+					for (int32 LODIndex = 0; LODIndex <= OperationData->CurrentMaxLOD && LODIndex < OperationData->InstanceUpdateData.LODs.Num(); ++LODIndex)
+					{
+						if (HasCOInstanceFlags(ReduceLODs) && LODIndex < CurrentMinLOD)
+						{
+							continue;
+						}
+						
+						Helper_GetLODData(SkeletalMesh).Add((LODIndex <= CurrentMinLOD && Helper_GetLODData(SkeletalMesh).Num()) ? &Helper_GetLODData(SkeletalMesh)[0] : new Helper_LODDataType());
+
+						const FMutableRefLODData& LODData = RefSkeletalMeshData->LODData[LODIndex];
+						
+						FSkeletalMeshLODInfo& LODInfo = SkeletalMesh->AddLODInfo();
+						LODInfo.ScreenSize = LODData.LODInfo.ScreenSize;
+						LODInfo.LODHysteresis =  LODData.LODInfo.LODHysteresis;
+						LODInfo.bSupportUniformlyDistributedSampling =  LODData.LODInfo.bSupportUniformlyDistributedSampling;
+						LODInfo.bAllowCPUAccess =  LODData.LODInfo.bAllowCPUAccess;
+						
+#if WITH_EDITORONLY_DATA
+						LODInfo.BuildSettings.bUseFullPrecisionUVs = true;
+#endif
+					}
+				}
+
+				{
+					BuildSkeletalMeshElementData(OperationData, SkeletalMesh, Public, ComponentIndex);
+					bSuccess = BuildSkeletalMeshRenderData(OperationData, SkeletalMesh, Public, ComponentIndex);
+				}
+
+				// TODO: update the model resources cache		
+				//mu::MeshPtrConst MutableMesh = PrivateData->InstanceUpdateData.LODs[0].Components[0].Mesh;
+				//FMutableResourceCache& Cache = UCustomizableObjectSystem::GetInstance()->GetPrivate()->GetObjectCache(CustomizableObject);
+				//Cache.Meshes.Add(MutableMesh->GetId(), SkeletalMesh);
+
+				if (bSuccess)
+				{
+					BuildMorphTargetsData(OperationData, SkeletalMesh, Public, ComponentIndex);
+					BuildClothingData(OperationData, SkeletalMesh, Public, ComponentIndex);
+
+					ensure(Helper_GetLODData(SkeletalMesh).Num() > 0);
+					ensure(Helper_GetLODInfoArray(SkeletalMesh).Num() > 0);
+				}
+			}
+			
+			if (!bSuccess)
+			{
+				Public->SkeletalMeshes.Reset();
+				break;
+			}
+		}
+
+		return bSuccess;
+	}
+
+	SetCOInstanceFlags(Generated);
+	ClearCOInstanceFlags((ECOInstanceFlags)(CreatingSkeletalMesh | PendingMeshUpdate));
+
+	return !bMeshNeedsUpdate;
+}
+
+
+static TAutoConsoleVariable<int32> CVarPendingReleaseSkeletalMeshMode(
+	TEXT("mutable.PendingReleaseSkeletalMesh"),
+	1,
+	TEXT("Whether use pending for release or not \n")
+	TEXT(" 0: Use GC\n")
+	TEXT(" 1: Use PendingRelease\n")
+	TEXT(" Default use"),
+	ECVF_SetByConsole);
+
+
+UCustomizableInstancePrivateData::UCustomizableInstancePrivateData()
+{
+	MinSquareDistFromComponentToPlayer = FLT_MAX;
+	LastMinSquareDistFromComponentToPlayer = FLT_MAX;
+	
+	LastUpdateTime = 0.f;
+
+	CurrentMinLOD = -1;
+	CurrentMaxLOD = -1;
+
+	MinLODToLoad = 0;
+	MaxLODToLoad = INT32_MAX;
+
+	NumLODsAvailable = INT32_MAX;
+	FirstLODAvailable = 0;
+
+	NumMaxLODsToStream = MAX_MESH_LOD_COUNT;
+}
+
+
+UCustomizableObjectInstance* UCustomizableObjectInstance::Clone()
+{
+	MUTABLE_CPUPROFILER_SCOPE(UCustomizableObjectInstance::Clone);
+
+	// Default Outer is the transient package.
+	UCustomizableObjectInstance* NewInstance = NewObject<UCustomizableObjectInstance>();
+	NewInstance->CopyParametersFromInstance(this);
+
+	return NewInstance;
+}
+
+
+UCustomizableObjectInstance* UCustomizableObjectInstance::CloneStatic(UObject* Outer)
+{
+	UCustomizableObjectInstance* NewInstance = NewObject<UCustomizableObjectInstance>(Outer, GetClass());
+	NewInstance->CopyParametersFromInstance(this);
+	NewInstance->bShowOnlyRuntimeParameters = false;
+
+	return NewInstance;
+}
+
+
+void UCustomizableObjectInstance::CopyParametersFromInstance(UCustomizableObjectInstance* Instance)
+{
+	SetDescriptor(Instance->GetDescriptor());
+	SkeletalMeshStatus = Instance->SkeletalMeshStatus;
+}
+
+
+int32 UCustomizableObjectInstance::AddValueToIntRange(const FString& ParamName)
+{
+	return Descriptor.AddValueToIntRange(ParamName);
+}
+
+
+int32 UCustomizableObjectInstance::AddValueToFloatRange(const FString& ParamName)
+{
+	return Descriptor.AddValueToFloatRange(ParamName);
+}
+
+
+int32 UCustomizableObjectInstance::AddValueToProjectorRange(const FString& ParamName)
+{
+	return Descriptor.AddValueToProjectorRange(ParamName);
+}
+
+
+int32 UCustomizableObjectInstance::RemoveValueFromIntRange(const FString& ParamName)
+{
+	return Descriptor.RemoveValueFromIntRange(ParamName);
+}
+
+
+int32 UCustomizableObjectInstance::RemoveValueFromIntRange(const FString& ParamName, int32 RangeIndex)
+{
+	return Descriptor.RemoveValueFromIntRange(ParamName);
+
+}
+
+int32 UCustomizableObjectInstance::RemoveValueFromFloatRange(const FString& ParamName)
+{
+	return Descriptor.RemoveValueFromFloatRange(ParamName);
+}
+
+
+int32 UCustomizableObjectInstance::RemoveValueFromFloatRange(const FString& ParamName, const int32 RangeIndex)
+{
+	return Descriptor.RemoveValueFromFloatRange(ParamName, RangeIndex);
+}
+
+
+int32 UCustomizableObjectInstance::RemoveValueFromProjectorRange(const FString& ParamName)
+{
+	return Descriptor.RemoveValueFromProjectorRange(ParamName);
+}
+
+
+int32 UCustomizableObjectInstance::RemoveValueFromProjectorRange(const FString& ParamName, const int32 RangeIndex)
+{
+	return Descriptor.RemoveValueFromProjectorRange(ParamName, RangeIndex);
+}
+
+
+bool UCustomizableObjectInstance::CreateMultiLayerProjector(const FName& ProjectorParamName)
+{
+	if (const FMultilayerProjector* Result = MultilayerProjectors.Find(ProjectorParamName))
+	{
+		checkCode(Result->CheckInstanceParameters());
+	}
+	else
+	{
+		if (!FMultilayerProjector::AreInstanceParametersValid(*this, ProjectorParamName))
+		{
+			return false;
+		}
+		
+		const FMultilayerProjector MultilayerProjector(*this, ProjectorParamName);
+		MultilayerProjectors.Add(ProjectorParamName, MultilayerProjector);
+	}
+
+	return true;
+}
+
+
+void UCustomizableObjectInstance::RemoveMultilayerProjector(const FName& ProjectorParamName)
+{
+	MultilayerProjectors.Remove(ProjectorParamName);
+}
+
+
+int32 UCustomizableObjectInstance::MultilayerProjectorNumLayers(const FName& ProjectorParamName) const
+{
+	check(MultilayerProjectors.Contains(ProjectorParamName)); // Multilayer Projector not created.
+	
+	return MultilayerProjectors[ProjectorParamName].NumLayers();
+}
+
+
+void UCustomizableObjectInstance::MultilayerProjectorCreateLayer(const FName& ProjectorParamName, int32 Index)
+{
+	check(MultilayerProjectors.Contains(ProjectorParamName)); // Multilayer Projector not created.
+	
+	MultilayerProjectors[ProjectorParamName].CreateLayer(Index);
+}
+
+
+void UCustomizableObjectInstance::MultilayerProjectorRemoveLayerAt(const FName& ProjectorParamName, int32 Index)
+{
+	check(MultilayerProjectors.Contains(ProjectorParamName)); // Multilayer Projector not created.
+	
+	MultilayerProjectors[ProjectorParamName].RemoveLayerAt(Index);
+}
+
+
+FMultilayerProjectorLayer UCustomizableObjectInstance::MultilayerProjectorGetLayer(const FName& ProjectorParamName, int32 Index) const
+{
+	check(MultilayerProjectors.Contains(ProjectorParamName)); // Multilayer Projector not created.
+	
+	return MultilayerProjectors[ProjectorParamName].GetLayer(Index);
+}
+
+
+void UCustomizableObjectInstance::MultilayerProjectorUpdateLayer(const FName& ProjectorParamName, int32 Index, const FMultilayerProjectorLayer& Layer)
+{
+	check(MultilayerProjectors.Contains(ProjectorParamName)); // Multilayer Projector not created.
+	
+	MultilayerProjectors[ProjectorParamName].UpdateLayer(Index, Layer);
+}
+
+
+TArray<FName> UCustomizableObjectInstance::MultilayerProjectorGetVirtualLayers(const FName& ProjectorParamName) const
+{
+	check(MultilayerProjectors.Contains(ProjectorParamName)); // Multilayer Projector not created.
+	
+	return MultilayerProjectors[ProjectorParamName].GetVirtualLayers();
+}
+
+
+void UCustomizableObjectInstance::MultilayerProjectorCreateVirtualLayer(const FName& ProjectorParamName, const FName& Id)
+{
+	check(MultilayerProjectors.Contains(ProjectorParamName)); // Multilayer Projector not created.
+	
+	MultilayerProjectors[ProjectorParamName].CreateVirtualLayer(Id);
+}
+
+
+FMultilayerProjectorVirtualLayer UCustomizableObjectInstance::MultilayerProjectorFindOrCreateVirtualLayer(const FName& ProjectorParamName, const FName& Id)
+{
+	check(MultilayerProjectors.Contains(ProjectorParamName)); // Multilayer Projector not created.
+	
+	return MultilayerProjectors[ProjectorParamName].FindOrCreateVirtualLayer(Id);
+}
+
+
+void UCustomizableObjectInstance::MultilayerProjectorRemoveVirtualLayer(const FName& ProjectorParamName, const FName& Id)
+{
+	check(MultilayerProjectors.Contains(ProjectorParamName)); // Multilayer Projector not created.
+	
+	MultilayerProjectors[ProjectorParamName].RemoveVirtualLayer(Id);
+}
+
+FMultilayerProjectorVirtualLayer UCustomizableObjectInstance::MultilayerProjectorGetVirtualLayer(const FName& ProjectorParamName, const FName& Id) const
+{
+	check(MultilayerProjectors.Contains(ProjectorParamName)); // Multilayer Projector not created.
+	
+	return MultilayerProjectors[ProjectorParamName].GetVirtualLayer(Id);
+}
+
+void UCustomizableObjectInstance::MultilayerProjectorUpdateVirtualLayer(const FName& ProjectorParamName, const FName& Id, const FMultilayerProjectorVirtualLayer& Layer)
+{
+	check(MultilayerProjectors.Contains(ProjectorParamName)); // Multilayer Projector not created.
+	
+	MultilayerProjectors[ProjectorParamName].UpdateVirtualLayer(Id, Layer);
+}
+
+
+void UCustomizableObjectInstance::SaveDescriptor(FArchive &Ar)
+{
+	Descriptor.SaveDescriptor(Ar);
+}
+
+
+void UCustomizableObjectInstance::LoadDescriptor(FArchive &Ar)
+{
+	Descriptor.LoadDescriptor(Ar);
+}
+
+
+const FString& UCustomizableObjectInstance::GetIntParameterSelectedOption(const FString& ParamName, const int32 RangeIndex) const
+{
+	return Descriptor.GetIntParameterSelectedOption(ParamName, RangeIndex);
+}
+
+
+void UCustomizableObjectInstance::SetIntParameterSelectedOption(int32 IntParamIndex, const FString& SelectedOption, const int32 RangeIndex)
+{
+	Descriptor.SetIntParameterSelectedOption(IntParamIndex, SelectedOption, RangeIndex);
+}
+
+
+void UCustomizableObjectInstance::SetIntParameterSelectedOption(const FString& ParamName, const FString& SelectedOptionName, const int32 RangeIndex)
+{
+	Descriptor.SetIntParameterSelectedOption(ParamName, SelectedOptionName, RangeIndex);
+}
+
+
+float UCustomizableObjectInstance::GetFloatParameterSelectedOption(const FString& FloatParamName, const int32 RangeIndex) const
+{
+	return Descriptor.GetFloatParameterSelectedOption(FloatParamName, RangeIndex);
+}
+
+
+void UCustomizableObjectInstance::SetFloatParameterSelectedOption(const FString& FloatParamName, const float FloatValue, const int32 RangeIndex)
+{
+	return Descriptor.SetFloatParameterSelectedOption(FloatParamName, FloatValue, RangeIndex);
+}
+
+
+FLinearColor UCustomizableObjectInstance::GetColorParameterSelectedOption(const FString& ColorParamName) const
+{
+	return Descriptor.GetColorParameterSelectedOption(ColorParamName);
+}
+
+
+void UCustomizableObjectInstance::SetColorParameterSelectedOption(const FString & ColorParamName, const FLinearColor& ColorValue)
+{
+	Descriptor.SetColorParameterSelectedOption(ColorParamName, ColorValue);
+}
+
+
+bool UCustomizableObjectInstance::GetBoolParameterSelectedOption(const FString& BoolParamName) const
+{
+	return Descriptor.GetBoolParameterSelectedOption(BoolParamName);
+}
+
+
+void UCustomizableObjectInstance::SetBoolParameterSelectedOption(const FString& BoolParamName, const bool BoolValue)
+{
+	Descriptor.SetBoolParameterSelectedOption(BoolParamName, BoolValue);
+}
+
+
+void UCustomizableObjectInstance::SetVectorParameterSelectedOption(const FString& VectorParamName, const FLinearColor& VectorValue)
+{
+	Descriptor.SetVectorParameterSelectedOption(VectorParamName, VectorValue);
+}
+
+
+void UCustomizableObjectInstance::SetProjectorValue(const FString& ProjectorParamName,
+	const FVector& Pos, const FVector& Direction, const FVector& Up, const FVector& Scale,
+	const float Angle,
+	const int32 RangeIndex)
+{
+	Descriptor.SetProjectorValue(ProjectorParamName, Pos, Direction, Up, Scale, Angle, RangeIndex);
+}
+
+
+void UCustomizableObjectInstance::SetProjectorPosition(const FString& ProjectorParamName, const FVector3f& Pos, const int32 RangeIndex)
+{
+	Descriptor.SetProjectorPosition(ProjectorParamName, Pos, RangeIndex);
+}
+
+
+void UCustomizableObjectInstance::GetProjectorValue(const FString& ProjectorParamName,
+	FVector& OutPos, FVector& OutDir, FVector& OutUp, FVector& OutScale,
+	float& OutAngle, ECustomizableObjectProjectorType& OutType,
+	const int32 RangeIndex) const
+{
+	Descriptor.GetProjectorValue(ProjectorParamName, OutPos, OutDir, OutUp, OutScale, OutAngle, OutType, RangeIndex);
+}
+
+
+void UCustomizableObjectInstance::GetProjectorValueF(const FString& ProjectorParamName,
+	FVector3f& OutPos, FVector3f& OutDir, FVector3f& OutUp, FVector3f& OutScale,
+	float& OutAngle, ECustomizableObjectProjectorType& OutType,
+	int32 RangeIndex) const
+{
+	Descriptor.GetProjectorValueF(ProjectorParamName, OutPos, OutDir, OutUp, OutScale, OutAngle, OutType, RangeIndex);
+}
+
+
+FVector UCustomizableObjectInstance::GetProjectorPosition(const FString& ParamName, const int32 RangeIndex) const
+{
+	return Descriptor.GetProjectorPosition(ParamName, RangeIndex);
+}
+
+
+FVector UCustomizableObjectInstance::GetProjectorDirection(const FString& ParamName, const int32 RangeIndex) const
+{
+	return Descriptor.GetProjectorDirection(ParamName, RangeIndex);
+}
+
+
+FVector UCustomizableObjectInstance::GetProjectorUp(const FString& ParamName, const int32 RangeIndex) const
+{
+	return Descriptor.GetProjectorUp(ParamName, RangeIndex);
+}
+
+
+FVector UCustomizableObjectInstance::GetProjectorScale(const FString& ParamName, const int32 RangeIndex) const
+{
+	return Descriptor.GetProjectorScale(ParamName, RangeIndex);
+}
+
+
+float UCustomizableObjectInstance::GetProjectorAngle(const FString& ParamName, int32 RangeIndex) const
+{
+	return Descriptor.GetProjectorAngle(ParamName, RangeIndex);
+}
+
+
+ECustomizableObjectProjectorType UCustomizableObjectInstance::GetProjectorParameterType(const FString& ParamName, const int32 RangeIndex) const
+{
+	return Descriptor.GetProjectorParameterType(ParamName, RangeIndex);
+}
+
+
+FCustomizableObjectProjector UCustomizableObjectInstance::GetProjector(const FString& ParamName, const int32 RangeIndex) const
+{
+	return Descriptor.GetProjector(ParamName, RangeIndex);
+}
+
+
+int32 UCustomizableObjectInstance::FindIntParameterNameIndex(const FString& ParamName) const
+{
+	return Descriptor.FindIntParameterNameIndex(ParamName);
+}
+
+
+int32 UCustomizableObjectInstance::FindFloatParameterNameIndex(const FString& ParamName) const
+{
+	return Descriptor.FindFloatParameterNameIndex(ParamName);
+}
+
+
+int32 UCustomizableObjectInstance::FindBoolParameterNameIndex(const FString& ParamName) const
+{
+	return Descriptor.FindBoolParameterNameIndex(ParamName);
+}
+
+
+int32 UCustomizableObjectInstance::FindVectorParameterNameIndex(const FString& ParamName) const
+{
+	return Descriptor.FindVectorParameterNameIndex(ParamName);
+}
+
+
+int32 UCustomizableObjectInstance::FindProjectorParameterNameIndex(const FString& ParamName) const
+{
+	return Descriptor.FindProjectorParameterNameIndex(ParamName);
+}
+
+
+bool UCustomizableInstancePrivateData::MeshNeedsUpdate( UCustomizableObjectInstance* Public, const TSharedPtr<FMutableOperationData>& OperationData, bool& bOutEmptyMesh)
+{
+	mu::MeshPtrConst MutableMesh;
+	bool bMeshNeedsUpdate = HasCOInstanceFlags(PendingMeshUpdate);
+	bool bFoundEmptyMesh = false;
+	bool bFoundEmptyLOD = false;
+
+	// Check if the current meshes have the number of LODs the operation requires or they must be updated.
+	const int32 SkeletalMeshCount = Public->SkeletalMeshes.Num();
+
+	// All meshes have the same number of LODs, check the first valid SkeletalMesh
+	const int32 CurrentMinLODIndex = HasCOInstanceFlags(ReduceLODs) ? 0 : OperationData->CurrentMinLOD;
+
+	for(int32 SkeletalMeshIndex = 0; SkeletalMeshIndex < SkeletalMeshCount; ++SkeletalMeshIndex)
+	{
+		if(const USkeletalMesh* SkeletalMesh = Public->SkeletalMeshes[SkeletalMeshIndex])
+		{
+			bMeshNeedsUpdate = (SkeletalMesh->GetLODNum() - 1) != (OperationData->CurrentMaxLOD - CurrentMinLODIndex);
+			break;
+		}
+	}
+	
+	for (int32 LODIndex = OperationData->CurrentMinLOD; LODIndex <= OperationData->CurrentMaxLOD && LODIndex < OperationData->InstanceUpdateData.LODs.Num(); ++LODIndex)
+	{
+		bool bFoundNonEmptyMesh = false;
+
+		const FInstanceUpdateData::FLOD& LOD = OperationData->InstanceUpdateData.LODs[LODIndex];
+
+		bMeshNeedsUpdate = bMeshNeedsUpdate || (LOD.ComponentCount > SkeletalMeshCount);
+
+		for (int32 ComponentIndex = 0; ComponentIndex < LOD.ComponentCount; ++ComponentIndex)
+		{
+			const FInstanceUpdateData::FComponent& Component = OperationData->InstanceUpdateData.Components[LOD.FirstComponent + ComponentIndex];
+			if (Component.Mesh)
+			{
+				bMeshNeedsUpdate = bMeshNeedsUpdate || (Component.Mesh->GetId() != GetLastMeshId(Component.Id, LODIndex)) || !Public->SkeletalMeshes[ComponentIndex];
+
+				// Don't build a degenerated mesh if something went wrong
+				if (Component.Mesh->GetVertexCount() > 0)
+				{
+					bFoundNonEmptyMesh = true;
+				}
+				else
+				{
+					if (LOD.ComponentCount == 1)
+					{
+						bFoundEmptyMesh = true;
+					}
+				}
+			}
+			else
+			{
+				bMeshNeedsUpdate = bMeshNeedsUpdate || Public->SkeletalMeshes[ComponentIndex];
+			}
+		}
+
+		if (!bFoundNonEmptyMesh)
+		{
+			bFoundEmptyLOD = true;
+			UE_LOG(LogMutable, Warning, TEXT("Building instance: an LOD has no mesh geometry. This cannot be handled by Unreal."));
+		}
+	}
+
+	bOutEmptyMesh = bFoundEmptyLOD || bFoundEmptyMesh;
+
+	return bMeshNeedsUpdate;
+}
+
+
+void UCustomizableObjectInstance::SetRandomValues()
+{
+	Descriptor.SetRandomValues();
+
+}
+
+
+bool UCustomizableObjectInstance::LoadParametersFromProfile(int32 ProfileIndex)
+{
+	UCustomizableObject* CustomizableObject = GetCustomizableObject();
+	if (!CustomizableObject)
+	{
+		return false;
+	}
+
+#if WITH_EDITOR
+	if (ProfileIndex < 0 || ProfileIndex >= CustomizableObject->InstancePropertiesProfiles.Num() )
+	{
+		return false;
+	}
+	
+	// This could be done only when the instance changes.
+	MigrateProfileParametersToCurrentInstance(ProfileIndex);
+
+	const FProfileParameterDat& Profile = CustomizableObject->InstancePropertiesProfiles[ProfileIndex];
+
+	Descriptor.BoolParameters = Profile.BoolParameters;
+	Descriptor.IntParameters = Profile.IntParameters;
+	Descriptor.FloatParameters = Profile.FloatParameters;
+	Descriptor.TextureParameters = Profile.TextureParameters;
+	Descriptor.ProjectorParameters = Profile.ProjectorParameters;
+	Descriptor.VectorParameters = Profile.VectorParameters;
+#endif
+	return true;
+
+}
+
+bool UCustomizableObjectInstance::SaveParametersToProfile(int32 ProfileIndex)
+{
+	UCustomizableObject* CustomizableObject = GetCustomizableObject();
+	if (!CustomizableObject)
+	{
+		return false;
+	}
+
+#if WITH_EDITOR
+	bSelectedProfileDirty = ProfileIndex != SelectedProfileIndex;
+
+	if (ProfileIndex < 0 || ProfileIndex >= CustomizableObject->InstancePropertiesProfiles.Num())
+	{
+		return false;
+	}
+
+	FProfileParameterDat& Profile = CustomizableObject->InstancePropertiesProfiles[ProfileIndex];
+
+	Profile.BoolParameters = Descriptor.BoolParameters;
+	Profile.IntParameters = Descriptor.IntParameters;
+	Profile.FloatParameters = Descriptor.FloatParameters;
+	Profile.TextureParameters = Descriptor.TextureParameters;
+	Profile.ProjectorParameters = Descriptor.ProjectorParameters;
+	Profile.VectorParameters = Descriptor.VectorParameters;
+#endif
+	return true;
+}
+
+bool UCustomizableObjectInstance::MigrateProfileParametersToCurrentInstance(int32 ProfileIndex)
+{
+	UCustomizableObject* CustomizableObject = GetCustomizableObject();
+	if (!CustomizableObject)
+	{
+		return false;
+	}
+
+#if WITH_EDITOR
+	if (ProfileIndex < 0 || ProfileIndex >= CustomizableObject->InstancePropertiesProfiles.Num())
+	{
+		return false;
+	}
+
+	FProfileParameterDat& Profile = CustomizableObject->InstancePropertiesProfiles[ProfileIndex];
+	FProfileParameterDat TempProfile;
+
+	TempProfile.ProfileName = Profile.ProfileName;
+	TempProfile.BoolParameters = Descriptor.BoolParameters;
+	TempProfile.FloatParameters = Descriptor.FloatParameters;
+	TempProfile.IntParameters = Descriptor.IntParameters;
+	TempProfile.ProjectorParameters = Descriptor.ProjectorParameters;
+	TempProfile.TextureParameters = Descriptor.TextureParameters;
+	TempProfile.VectorParameters = Descriptor.VectorParameters;
+	
+
+	// Populate TempProfile with the parameters found in the profile.
+	// Any profile parameter missing will be discarded.
+	for (FCustomizableObjectBoolParameterValue& Parameter : TempProfile.BoolParameters)
+	{
+		using ParamValType = FCustomizableObjectBoolParameterValue;
+		ParamValType* Found = Profile.BoolParameters.FindByPredicate(
+			[&Parameter](const ParamValType& P) { return P.ParameterName == Parameter.ParameterName; });
+
+		if (Found)
+		{
+			Parameter.ParameterValue = Found->ParameterValue;
+		}
+	}
+
+	for (FCustomizableObjectIntParameterValue& Parameter : TempProfile.IntParameters)
+	{
+		using ParamValType = FCustomizableObjectIntParameterValue;
+		ParamValType* Found = Profile.IntParameters.FindByPredicate(
+			[&Parameter](const ParamValType& P) { return P.ParameterName == Parameter.ParameterName; });
+
+		if (Found)
+		{
+			Parameter.ParameterValueName = Found->ParameterValueName;
+		}
+	}
+
+	for (FCustomizableObjectFloatParameterValue& Parameter : TempProfile.FloatParameters)
+	{
+		using ParamValType = FCustomizableObjectFloatParameterValue;
+		ParamValType* Found = Profile.FloatParameters.FindByPredicate(
+			[&Parameter](const ParamValType& P) { return P.ParameterName == Parameter.ParameterName; });
+
+		if (Found)
+		{
+			Parameter.ParameterValue = Found->ParameterValue;
+			Parameter.ParameterRangeValues = Found->ParameterRangeValues;
+		}
+	}
+
+	for (FCustomizableObjectTextureParameterValue& Parameter : TempProfile.TextureParameters)
+	{
+		using ParamValType = FCustomizableObjectTextureParameterValue;
+		ParamValType* Found = Profile.TextureParameters.FindByPredicate(
+			[&Parameter](const ParamValType& P) { return P.ParameterName == Parameter.ParameterName; });
+
+		if (Found)
+		{
+			Parameter.ParameterValue = Found->ParameterValue;
+		}
+	}
+
+	for (FCustomizableObjectVectorParameterValue& Parameter : TempProfile.VectorParameters)
+	{
+		using ParamValType = FCustomizableObjectVectorParameterValue;
+		ParamValType* Found = Profile.VectorParameters.FindByPredicate(
+			[&Parameter](const ParamValType& P) { return P.ParameterName == Parameter.ParameterName; });
+
+		if (Found)
+		{
+			Parameter.ParameterValue = Found->ParameterValue;
+		}
+	}
+
+	for (FCustomizableObjectProjectorParameterValue& Parameter : TempProfile.ProjectorParameters)
+	{
+		using ParamValType = FCustomizableObjectProjectorParameterValue;
+		ParamValType* Found = Profile.ProjectorParameters.FindByPredicate(
+			[&Parameter](const ParamValType& P) { return P.ParameterName == Parameter.ParameterName; });
+
+		if (Found)
+		{
+			Parameter.RangeValues = Found->RangeValues;
+			Parameter.Value = Found->Value;
+		}
+	}
+
+	Profile = TempProfile;
+
+	//CustomizableObject->Modify();
+#endif
+
+	return true;
+}
+
+
+void UCustomizableObjectInstance::SetSelectedParameterProfileDirty()
+{
+	UCustomizableObject* CustomizableObject = GetCustomizableObject();
+	if (!CustomizableObject)
+	{
+		return;
+	}
+
+#if WITH_EDITOR
+	bSelectedProfileDirty = SelectedProfileIndex != INDEX_NONE;
+	
+	if (bSelectedProfileDirty)
+	{
+		CustomizableObject->Modify();
+	}
+#endif
+}
+
+bool UCustomizableObjectInstance::IsSelectedParameterProfileDirty() const
+{
+	
+#if WITH_EDITOR
+	return bSelectedProfileDirty && SelectedProfileIndex != INDEX_NONE;
+#else
+	return false;
+#endif
+}
+
+//int32 UCustomizableObjectInstance::GetIntParameterNumOptions(int32 IntParamIndex)
+//{
+//	if (CustomizableObject && IntParamIndex>=0 && IntParamIndex<IntParameters.Num())
+//	{
+//		int32 ObjectParamIndex = CustomizableObject->FindParameter(IntParameters[IntParamIndex].ParameterName);
+//		return CustomizableObject->GetIntParameterNumOptions(ObjectParamIndex);
+//	}
+//
+//	return 0;
+//}
+
+
+//void UCustomizableObjectInstance::SetIntParameter(int32 IntParamIndex, int32 i)
+//{
+//	if (CustomizableObject && IntParamIndex >= 0 && IntParamIndex < IntParameters.Num())
+//	{
+//		int32 ObjectParamIndex = CustomizableObject->FindParameter(IntParameters[IntParamIndex].ParameterName);
+//		int32 PossibleValue = CustomizableObject->GetIntParameterAvailableOptionValue(ObjectParamIndex, i);
+//		IntParameters[IntParamIndex].ParameterValue = PossibleValue;
+//	}
+//}
+
+
+//const FString & UCustomizableObjectInstance::GetIntParameterAvailableOption(int32 IntParamIndex, int32 K)
+//{
+//	if (CustomizableObject && IntParamIndex >= 0 && IntParamIndex<IntParameters.Num())
+//	{
+//		int32 ObjectParamIndex = CustomizableObject->FindParameter(IntParameters[IntParamIndex].ParameterName);
+//		if (ObjectParamIndex != INDEX_NONE)
+//		{
+//			return CustomizableObject->GetIntParameterAvailableOption(ObjectParamIndex,K);
+//		}
+//	}
+//
+//	return s_EmptyStringCI;
+//}
+
+
+void UCustomizableInstancePrivateData::DiscardResourcesAndSetReferenceSkeletalMesh(UCustomizableObjectInstance* Public )
+{
+	if (HasCOInstanceFlags(Generated))
+	{
+		for (int32 Component = 0; Component < Public->SkeletalMeshes.Num(); ++Component)
+		{
+			if (Public->SkeletalMeshes[Component] && Public->SkeletalMeshes[Component]->IsValidLowLevel())
+			{
+				Public->SkeletalMeshes[Component]->ReleaseResources();
+				Public->SkeletalMeshes[Component] = nullptr;
+			}
+		}
+
+		Public->ReleaseMutableResources(false);
+	}
+
+	ClearCOInstanceFlags(Generated);
+	ClearAllLastMeshIds();
+	
+	Public->SkeletalMeshes.Reset();
+
+	for (TObjectIterator<UCustomizableSkeletalComponent> It; It; ++It)
+	{
+		UCustomizableSkeletalComponent* CustomizableSkeletalComponent = *It;
+
+		if (CustomizableSkeletalComponent && CustomizableSkeletalComponent->CustomizableObjectInstance == Public)
+		{
+			UCustomizableObject* CustomizableObject = Public->GetCustomizableObject();
+			bool bReplaceDiscardedWithReferenceMesh = UCustomizableObjectSystem::GetInstance()->GetPrivate()->IsReplaceDiscardedWithReferenceMeshEnabled();
+			CustomizableSkeletalComponent->SetSkeletalMesh(CustomizableObject && bReplaceDiscardedWithReferenceMesh ? CustomizableObject->GetRefSkeletalMesh(CustomizableSkeletalComponent->ComponentIndex) : nullptr);
+		}
+	}
+}
+
+
+
+namespace
+{
+	inline float unpack_uint8(uint8 i)
+	{
+		float res = i;
+		res -= 127.5f;
+		res /= 127.5f;
+		return res;
+	}
+}
+
+
+void SetTexturePropertiesFromMutableImageProps(UTexture2D* Texture, const FMutableModelImageProperties& Props, bool bNeverStream)
+{
+#if !PLATFORM_DESKTOP
+	if (UCustomizableObjectSystem::GetInstance()->GetPrivate()->EnableMutableProgressiveMipStreaming <= 0)
+	{
+		Texture->NeverStream = true;
+	}
+	else
+	{
+		Texture->NeverStream = bNeverStream;
+	}
+
+#if !PLATFORM_ANDROID && !PLATFORM_IOS
+	Texture->bNotOfflineProcessed = true;
+#endif
+
+#else
+	Texture->NeverStream = bNeverStream;
+#endif
+
+	Texture->SRGB = Props.SRGB;
+	Texture->Filter = Props.Filter;
+	Texture->LODBias = Props.LODBias;
+
+#if WITH_EDITORONLY_DATA
+	Texture->bFlipGreenChannel = Props.FlipGreenChannel;
+#endif
+
+	Texture->LODGroup = Props.LODGroup;
+	Texture->AddressX = Props.AddressX;
+	Texture->AddressY = Props.AddressY;
+}
+
+
+void UCustomizableObjectInstance::Updated()
+{
+	UpdatedDelegate.Broadcast(this);
+	UpdatedNativeDelegate.Broadcast(this);
+}
+
+
+uint32 UCustomizableObjectInstance::GetLastUpdateDescriptorHash() const
+{
+	return LastUpdateDescriptorHash;
+}
+
+
+void CopyTextureProperties(UTexture2D* Texture, const UTexture2D* SourceTexture)
+{
+	MUTABLE_CPUPROFILER_SCOPE(CopyTextureProperties)
+		
+	Texture->NeverStream = SourceTexture->NeverStream;
+
+	Texture->SRGB = SourceTexture->SRGB;
+	Texture->Filter = SourceTexture->Filter;
+	Texture->LODBias = SourceTexture->LODBias;
+
+#if WITH_EDITOR
+	Texture->MipGenSettings = SourceTexture->MipGenSettings;
+	Texture->CompressionNone = SourceTexture->CompressionNone;
+#endif
+
+	Texture->LODGroup = SourceTexture->LODGroup;
+	Texture->AddressX = SourceTexture->AddressX;
+	Texture->AddressY = SourceTexture->AddressY;
+}
+
+
+// The memory allocated in the function and pointed by the returned pointer is owned by the caller and must be freed. 
+// If assigned to a UTexture2D, it will be freed by that UTexture2D
+FTexturePlatformData* UCustomizableInstancePrivateData::MutableCreateImagePlatformData(const mu::Image* MutableImage, int32 OnlyLOD, uint16 FullSizeX, uint16 FullSizeY)
+{
+	int32 SizeX = FMath::Max(MutableImage->GetSize().x(), FullSizeX);
+	int32 SizeY = FMath::Max(MutableImage->GetSize().y(), FullSizeY);
+
+	if (SizeX <= 0 || SizeY <= 0)
+	{
+		UE_LOG(LogMutable, Verbose, TEXT("Invalid parameters specified for UCustomizableInstancePrivateData::MutableCreateImagePlatformData()"));
+		return nullptr;
+	}
+
+	FTexturePlatformData* PlatformData = new FTexturePlatformData();
+
+	int32 FirstLOD = 0;
+	for (int32 l = 0; l < OnlyLOD; ++l)
+	{
+		if (SizeX <= 4 || SizeY <= 4)
+		{
+			break;
+		}
+		SizeX = FMath::Max(SizeX / 2, 1);
+		SizeY = FMath::Max(SizeY / 2, 1);
+		++FirstLOD;
+	}
+
+	int32 MaxSize = FMath::Max(SizeX, SizeY);
+	int32 FullLODCount = 1;
+	int32 MipsToSkip = 0;
+	
+	if (OnlyLOD < 0)
+	{
+		FullLODCount = FMath::CeilLogTwo(MaxSize) + 1;
+		MipsToSkip = FullLODCount - MutableImage->GetLODCount();
+		check(MipsToSkip >= 0);
+	}
+
+	if (MutableImage->GetLODCount() == 1)
+	{
+		MipsToSkip = 0;
+		FullLODCount = 1;
+		FirstLOD = 0;
+	}
+
+	int32 EndLOD = OnlyLOD < 0 ? FullLODCount : FirstLOD + 1;
+	
+	mu::EImageFormat MutableFormat = MutableImage->GetFormat();
+
+	int32 MaxPossibleSize = int32(FMath::Pow(2.f, float(FullLODCount - 1)));
+	check(SizeX == MaxPossibleSize || SizeY == MaxPossibleSize || FullLODCount == 1);
+
+
+	EPixelFormat PlatformFormat = PF_Unknown;
+	switch (MutableFormat)
+	{
+	case mu::EImageFormat::IF_RGB_UBYTE:
+		// performance penalty. can happen in states that remove compression.
+		PlatformFormat = PF_R8G8B8A8;	
+		UE_LOG(LogMutable, Warning, TEXT("Building instance: a texture was generated in a format not supported by the hardware (RGB), this results in an additional conversion, so a performance penalty."));
+		break; 
+
+	case mu::EImageFormat::IF_BGRA_UBYTE:			
+		// performance penalty. can happen with texture parameter images.
+		PlatformFormat = PF_R8G8B8A8;	
+		UE_LOG(LogMutable, Warning, TEXT("Building instance: a texture was generated in a format not supported by the hardware (BGRA), this results in an additional conversion, so a performance penalty."));
+		break;
+
+	// Good cases:
+	case mu::EImageFormat::IF_RGBA_UBYTE:		PlatformFormat = PF_R8G8B8A8;	break;
+	case mu::EImageFormat::IF_BC1:				PlatformFormat = PF_DXT1;		break;
+	case mu::EImageFormat::IF_BC2:				PlatformFormat = PF_DXT3;		break;
+	case mu::EImageFormat::IF_BC3:				PlatformFormat = PF_DXT5;		break;
+	case mu::EImageFormat::IF_BC4:				PlatformFormat = PF_BC4;		break;
+	case mu::EImageFormat::IF_BC5:				PlatformFormat = PF_BC5;		break;
+	case mu::EImageFormat::IF_L_UBYTE:			PlatformFormat = PF_G8;			break;
+	case mu::EImageFormat::IF_ASTC_4x4_RGB_LDR:	PlatformFormat = PF_ASTC_4x4;	break;
+	case mu::EImageFormat::IF_ASTC_4x4_RGBA_LDR:PlatformFormat = PF_ASTC_4x4;	break;
+	case mu::EImageFormat::IF_ASTC_4x4_RG_LDR:	PlatformFormat = PF_ASTC_4x4;	break;
+	default:
+		// Cannot prepare texture if it's not in the right format, this can happen if mutable is in debug mode
+		UE_LOG(LogMutable, Warning, TEXT("Building instance: a texture was generated in an unsupported format, which cannot be converted to Unreal."));
+		delete PlatformData;
+		return nullptr; 
+	}
+
+	PlatformData->SizeX = SizeX;
+	PlatformData->SizeY = SizeY;
+	PlatformData->PixelFormat = PlatformFormat;
+
+	// Allocate mipmaps.
+	const uint8_t* MutableData = MutableImage->GetMipData( FirstLOD );
+
+	if (!FMath::IsPowerOfTwo(SizeX) || !FMath::IsPowerOfTwo(SizeY))
+	{
+		EndLOD = FirstLOD + 1;
+	}
+
+	for (int32 MipLevelUE = FirstLOD; MipLevelUE < EndLOD; ++MipLevelUE)
+	{
+		int32 MipLevelMutable = MipLevelUE - MipsToSkip;		
+
+		PlatformData->Mips.Add(new FTexture2DMipMap());
+		FTexture2DMipMap* Mip = &PlatformData->Mips.Last(); //new(Texture->PlatformData->Mips) FTexture2DMipMap();
+		Mip->SizeX = SizeX;
+		Mip->SizeY = SizeY;
+
+		// Unlike Mutable, UE expects MIPs sizes to be at least the size of the compression block.
+		// For example, a 8x8 PF_DXT1 texture will have the following MIPs:
+		// Mutable    Unreal Engine
+		// 8x8        8x8
+		// 4x4        4x4
+		// 2x2        4x4
+		// 1x1        4x4
+		//
+		// Notice that even though Mutable reports MIP smaller than the block size, the actual data contains at least a block.
+		Mip->SizeX = FMath::Max(Mip->SizeX, GPixelFormats[PlatformFormat].BlockSizeX);
+		Mip->SizeY = FMath::Max(Mip->SizeY, GPixelFormats[PlatformFormat].BlockSizeY);
+
+		if(MipLevelUE >= MipsToSkip || OnlyLOD>=0)
+		{
+			check(MipLevelMutable >= 0);
+			check(MipLevelMutable < MutableImage->GetLODCount());
+
+			Mip->BulkData.Lock(LOCK_READ_WRITE);
+			Mip->BulkData.ClearBulkDataFlags(BULKDATA_SingleUse);
+
+			uint32 SourceDataSize = MutableImage->GetLODDataSize(MipLevelMutable);
+			uint32 DestDataSize = (MutableFormat == mu::EImageFormat::IF_RGB_UBYTE)
+				? (SourceDataSize/3*4)
+				: SourceDataSize;
+			void* pData = Mip->BulkData.Realloc(DestDataSize);
+
+			// Special inefficient cases
+			if (MutableFormat== mu::EImageFormat::IF_BGRA_UBYTE)
+			{
+				check(SourceDataSize==DestDataSize);
+
+				MUTABLE_CPUPROFILER_SCOPE(Innefficent_BGRA_Format_Conversion);
+
+				uint8_t* pDest = reinterpret_cast<uint8_t*>(pData);
+				for (size_t p = 0; p < SourceDataSize / 4; ++p)
+				{
+					pDest[p * 4 + 0] = MutableData[p * 4 + 2];
+					pDest[p * 4 + 1] = MutableData[p * 4 + 1];
+					pDest[p * 4 + 2] = MutableData[p * 4 + 0];
+					pDest[p * 4 + 3] = MutableData[p * 4 + 3];
+				}
+			}
+
+			else if (MutableFormat == mu::EImageFormat::IF_RGB_UBYTE)
+			{
+				MUTABLE_CPUPROFILER_SCOPE(Innefficent_RGB_Format_Conversion);
+
+				uint8_t* pDest = reinterpret_cast<uint8_t*>(pData);
+				for (size_t p = 0; p < SourceDataSize / 3; ++p)
+				{
+					pDest[p * 4 + 0] = MutableData[p * 3 + 0];
+					pDest[p * 4 + 1] = MutableData[p * 3 + 1];
+					pDest[p * 4 + 2] = MutableData[p * 3 + 2];
+					pDest[p * 4 + 3] = 255;
+				}
+			}
+
+			// Normal case
+			else
+			{
+				check(SourceDataSize == DestDataSize);
+				FMemory::Memcpy(pData, MutableData, SourceDataSize);
+			}
+
+			MutableData += SourceDataSize;
+
+			Mip->BulkData.Unlock();
+		}
+		else
+		{
+			Mip->BulkData.SetBulkDataFlags(BULKDATA_PayloadInSeperateFile);
+			Mip->BulkData.ClearBulkDataFlags(BULKDATA_PayloadAtEndOfFile);
+		}
+
+		SizeX /= 2;
+		SizeY /= 2;
+
+		SizeX = SizeX > 0 ? SizeX : 1;
+		SizeY = SizeY > 0 ? SizeY : 1;
+	}
+
+#if !(UE_BUILD_SHIPPING || UE_BUILD_TEST)
+	// Some consistency checks for dev builds
+	int32 BulkDataCount = 0;
+
+	for (int32 i = 0; i < PlatformData->Mips.Num(); ++i)
+	{
+		if (i > 0)
+		{
+			check(PlatformData->Mips[i].SizeX == PlatformData->Mips[i - 1].SizeX / 2 || PlatformData->Mips[i].SizeX == GPixelFormats[PlatformFormat].BlockSizeX);
+			check(PlatformData->Mips[i].SizeY == PlatformData->Mips[i - 1].SizeY / 2 || PlatformData->Mips[i].SizeY == GPixelFormats[PlatformFormat].BlockSizeY);
+		}
+
+		if (PlatformData->Mips[i].BulkData.GetBulkDataSize() > 0)
+		{
+			BulkDataCount++;
+		}
+	}
+
+	check(OnlyLOD >= 0 || (BulkDataCount == MutableImage->GetLODCount()));
+#endif
+
+	return PlatformData;
+}
+
+
+void UCustomizableInstancePrivateData::ConvertImage(UTexture2D* Texture, mu::ImagePtrConst MutableImage, const FMutableModelImageProperties& Props, int OnlyLOD, int32 ExtractChannel )
+{
+	MUTABLE_CPUPROFILER_SCOPE(UCustomizableInstancePrivateData::ConvertImage);
+
+	SetTexturePropertiesFromMutableImageProps(Texture, Props, false);
+
+	mu::EImageFormat MutableFormat = MutableImage->GetFormat();
+
+	// Extract a single channel, if requested.
+	if (ExtractChannel >= 0)
+	{
+		MutableImage = mu::ImagePixelFormat( 4, MutableImage.get(), mu::EImageFormat::IF_RGBA_UBYTE );
+
+		uint8_t Channel = uint8_t( FMath::Clamp(ExtractChannel,0,3) );
+		MutableImage = mu::ImageSwizzle( mu::EImageFormat::IF_L_UBYTE, &MutableImage, &Channel );
+		MutableFormat = mu::EImageFormat::IF_L_UBYTE;
+	}
+
+	// Hack: This format is unsupported in UE, but it shouldn't happen in production.
+	if (MutableFormat == mu::EImageFormat::IF_RGB_UBYTE)
+	{
+		UE_LOG(LogMutable, Warning, TEXT("Building instance: a texture was generated in RGB format, which is slow to convert to Unreal."));
+
+		// Expand the image.
+		mu::ImagePtr Converted = new mu::Image(MutableImage->GetSizeX(), MutableImage->GetSizeY(), MutableImage->GetLODCount(), mu::EImageFormat::IF_RGBA_UBYTE);
+
+		for (int LODIndex = 0; LODIndex < Converted->GetLODCount(); ++LODIndex)
+		{
+			int32 pixelCount = MutableImage->GetLODDataSize(LODIndex)/3;
+			const uint8_t* pSource = MutableImage->GetMipData(LODIndex);
+			uint8_t* pTarget = Converted->GetMipData(LODIndex);
+			for (int32 p = 0; p < pixelCount; ++p)
+			{
+				pTarget[4 * p + 0] = pSource[3 * p + 0];
+				pTarget[4 * p + 1] = pSource[3 * p + 1];
+				pTarget[4 * p + 2] = pSource[3 * p + 2];
+				pTarget[4 * p + 3] = 255;
+			}
+		}
+
+		MutableImage = Converted;
+		MutableFormat = mu::EImageFormat::IF_RGBA_UBYTE;
+	}
+	else if (MutableFormat == mu::EImageFormat::IF_BGRA_UBYTE)
+	{
+		UE_LOG(LogMutable, Warning, TEXT("Building instance: a texture was generated in BGRA format, which is slow to convert to Unreal."));
+
+		MUTABLE_CPUPROFILER_SCOPE(Swizzle);
+		// Swizzle the image.
+		// \TODO: Raise a warning?
+		mu::ImagePtr Converted = new mu::Image(MutableImage->GetSizeX(), MutableImage->GetSizeY(), 1, mu::EImageFormat::IF_RGBA_UBYTE);
+		int32 pixelCount = MutableImage->GetSizeX() * MutableImage->GetSizeY();
+
+		const uint8_t* pSource = MutableImage->GetData();
+		uint8_t* pTarget = Converted->GetData();
+		for (int32 p = 0; p<pixelCount; ++p)
+		{
+			pTarget[4 * p + 0] = pSource[4 * p + 2];
+			pTarget[4 * p + 1] = pSource[4 * p + 1];
+			pTarget[4 * p + 2] = pSource[4 * p + 0];
+			pTarget[4 * p + 3] = pSource[4 * p + 3];
+		}
+
+		MutableImage = Converted;
+		MutableFormat = mu::EImageFormat::IF_RGBA_UBYTE;
+	}
+
+	if (OnlyLOD >= 0)
+	{
+		OnlyLOD = FMath::Min( OnlyLOD, MutableImage->GetLODCount()-1 );
+	}
+
+	Texture->SetPlatformData(MutableCreateImagePlatformData(MutableImage.get(),OnlyLOD,0,0) );
+
+}
+
+
+void UCustomizableInstancePrivateData::BuildSkeletalMeshElementData(const TSharedPtr<FMutableOperationData>& OperationData, USkeletalMesh* SkeletalMesh, UCustomizableObjectInstance* CustomizableObjectInstance, int32 ComponentIndex)
+{
+	MUTABLE_CPUPROFILER_SCOPE(UCustomizableInstancePrivateData::BuildSkeletalMeshElementData);
+
+	// Set up unreal's default material
+	UMaterialInterface* UnrealMaterial = UMaterial::GetDefaultMaterial(MD_Surface);
+	SkeletalMesh->GetMaterials().SetNum(1);
+	SkeletalMesh->GetMaterials()[0] = UnrealMaterial;
+
+	int32 MeshLODIndex = HasCOInstanceFlags(ReduceLODs) ? 0 : CurrentMinLOD;
+
+	int32 LODCount = OperationData->InstanceUpdateData.LODs.Num();
+	for (int32 LODIndex = CurrentMinLOD; LODIndex <= CurrentMaxLOD && LODIndex < LODCount; ++LODIndex, ++MeshLODIndex)
+	{
+		const FInstanceUpdateData::FLOD& LOD = OperationData->InstanceUpdateData.LODs[LODIndex];
+		const mu::MeshPtrConst MutableMesh = OperationData->InstanceUpdateData.Components[LOD.FirstComponent + ComponentIndex].Mesh;
+		UnrealConversionUtils::BuildSkeletalMeshElementDataAtLOD(MeshLODIndex,MutableMesh,SkeletalMesh);
+	}
+}
+
+
+bool UCustomizableInstancePrivateData::BuildSkeletalMeshSkeletonData(const TSharedPtr<FMutableOperationData>& OperationData, USkeletalMesh* SkeletalMesh, const FMutableRefSkeletalMeshData* RefSkeletalMeshData, UCustomizableObjectInstance* Public, int32 ComponentIndex)
+{
+	MUTABLE_CPUPROFILER_SCOPE(BuildSkeletonData);
+
+	check(SkeletalMesh);
+	check(RefSkeletalMeshData);
+
+	const TObjectPtr<USkeleton> Skeleton = MergeSkeletons(Public, RefSkeletalMeshData, ComponentIndex);
+	if (!Skeleton)
+	{
+		return false;
+	}
+
+	SkeletalMesh->SetSkeleton(Skeleton);
+
+	SkeletalMesh->SetRefSkeleton(Skeleton->GetReferenceSkeleton());
+	FReferenceSkeleton& ReferenceSkeleton = SkeletalMesh->GetRefSkeleton();
+
+	// Check that the bones we need are present in the current Skeleton
+	FInstanceUpdateData::FSkeletonData& MutSkeletonData = OperationData->InstanceUpdateData.Skeletons[ComponentIndex];
+	const int32 MutBoneCount = MutSkeletonData.BoneNames.Num();
+	
+	TMap<FName, uint16> BoneToFinalBoneIndexMap;
+	BoneToFinalBoneIndexMap.Reserve(MutBoneCount);
+
+	{
+		MUTABLE_CPUPROFILER_SCOPE(BuildSkeletonData_CheckForMissingBones);
+
+		// Ensure all the required bones are present in the skeleton
+		for (int32 BoneIndex = 0; BoneIndex < MutBoneCount; ++BoneIndex)
+		{
+			FName BoneName = MutSkeletonData.BoneNames[BoneIndex];
+			check(BoneName != NAME_None);
+
+			const int32 SourceBoneIndex = ReferenceSkeleton.FindRawBoneIndex(BoneName);
+			if (SourceBoneIndex == INDEX_NONE)
+			{
+				// Merged skeleton is missing some bones! This happens if one of the skeletons involved in the merge is discarded due to being incompatible with the rest
+				// or if the source mesh is not in sync with the skeleton. 
+				UE_LOG(LogMutable, Warning, TEXT("Building instance: generated mesh has a bone [%s] not present in the reference mesh [%s]. Failing to generate mesh. "),
+					*BoneName.ToString(), *SkeletalMesh->GetName());
+				return false;
+			}
+
+			BoneToFinalBoneIndexMap.Add(BoneName, SourceBoneIndex);
+		}
+	}
+
+	{
+		MUTABLE_CPUPROFILER_SCOPE(BuildSkeletonData_FixBoneIndices);
+
+		// Fix up BoneMaps and ActiveBones indices
+		for (FInstanceUpdateData::FComponent& Component : OperationData->InstanceUpdateData.Components)
+		{
+			if (Component.Id != ComponentIndex)
+			{
+				continue;
+			}
+
+			for (uint16& BoneIndex : Component.BoneMap)
+			{
+				BoneIndex = BoneToFinalBoneIndexMap[MutSkeletonData.BoneNames[BoneIndex]];
+			}
+
+			for (uint16& BoneIndex : Component.ActiveBones)
+			{
+				BoneIndex = BoneToFinalBoneIndexMap[MutSkeletonData.BoneNames[BoneIndex]];
+			}
+			Component.ActiveBones.Sort();
+		}
+	}
+	
+	{
+		MUTABLE_CPUPROFILER_SCOPE(BuildSkeletonData_ApplyPose);
+		
+		const int32 RefRawBoneCount = ReferenceSkeleton.GetRawBoneNum();
+		const int32 MutBonePoseCount = MutSkeletonData.BoneMatricesWithScale.Num();
+
+		TArray<FMatrix44f>& RefBasesInvMatrix = SkeletalMesh->GetRefBasesInvMatrix();
+		RefBasesInvMatrix.Empty(RefRawBoneCount);
+
+		// Initialize the base matrices
+		if (RefRawBoneCount == MutBonePoseCount)
+		{
+			RefBasesInvMatrix.AddUninitialized(MutBonePoseCount);
+		}
+		else
+		{
+			// Bad case, some bone poses are missing, calculate the InvRefMatrices to ensure all transforms are there for the second step 
+			MUTABLE_CPUPROFILER_SCOPE(BuildSkeletonData_CalcInvRefMatrices0);
+			SkeletalMesh->CalculateInvRefMatrices();
+		}
+
+		// First step is to update the RefBasesInvMatrix for the bones.
+		for (const auto& BoneMatrix : MutSkeletonData.BoneMatricesWithScale)
+		{
+			const int32 RefSkelBoneIndex = BoneToFinalBoneIndexMap[BoneMatrix.Key];
+			RefBasesInvMatrix[RefSkelBoneIndex] = BoneMatrix.Value;
+		}
+
+		// The second step is to update the pose transforms in the ref skeleton from the BasesInvMatrix
+		FReferenceSkeletonModifier SkeletonModifier(ReferenceSkeleton, Skeleton);
+		for (int32 RefSkelBoneIndex = 0; RefSkelBoneIndex < ReferenceSkeleton.GetNum(); ++RefSkelBoneIndex)
+		{
+			int32 ParentBoneIndex = ReferenceSkeleton.GetParentIndex(RefSkelBoneIndex);
+			if (ParentBoneIndex >= 0)
+			{
+				FMatrix44f BoneBaseInvMatrix = RefBasesInvMatrix[RefSkelBoneIndex];
+				FMatrix44f ParentBaseInvMatrix = RefBasesInvMatrix[ParentBoneIndex];
+				FMatrix44f BonePoseMatrix = BoneBaseInvMatrix.Inverse() * ParentBaseInvMatrix;
+
+				FTransform3f BonePoseTransform;
+				BonePoseTransform.SetFromMatrix(BonePoseMatrix);
+
+				SkeletonModifier.UpdateRefPoseTransform(RefSkelBoneIndex, (FTransform)BonePoseTransform);
+			}
+		}
+
+		// Force a CalculateInvRefMatrices
+		RefBasesInvMatrix.Empty(RefRawBoneCount);
+	}
+
+	{
+		MUTABLE_CPUPROFILER_SCOPE(BuildSkeletonData_CalcInvRefMatrices);
+		SkeletalMesh->CalculateInvRefMatrices();
+	}
+
+
+	return true;
+}
+
+
+void UCustomizableInstancePrivateData::BuildMorphTargetsData(const TSharedPtr<FMutableOperationData>& OperationData, USkeletalMesh* SkeletalMesh, UCustomizableObjectInstance* CustomizableObjectInstance, int32 ComponentIndex)
+{
+	MUTABLE_CPUPROFILER_SCOPE(UCustomizableInstancePrivateData::BuildMorphTargetsData);
+
+	if (!(SkeletalMesh && OperationData->InstanceUpdateData.LODs.Num()))
+	{
+		return;
+	}
+
+	const UCustomizableObject* CustomizableObject = CustomizableObjectInstance->GetCustomizableObject();
+	const TArray<FMorphTargetInfo>& ContributingMorphTargetsInfo = CustomizableObject->ContributingMorphTargetsInfo;
+	const TArray<FMorphTargetVertexData>& MorphTargetReconstructionData = CustomizableObject->MorphTargetReconstructionData;
+
+	if (!(MorphTargetReconstructionData.Num() && ContributingMorphTargetsInfo.Num()))
+	{
+		return;
+	}
+
+	TArray<int32> SectionMorphTargetVertices;
+	SectionMorphTargetVertices.SetNumZeroed(ContributingMorphTargetsInfo.Num());
+
+	SkeletalMesh->GetMorphTargets().Empty(ContributingMorphTargetsInfo.Num());
+	for (const FMorphTargetInfo& MorphTargetInfo : ContributingMorphTargetsInfo)
+	{
+		UMorphTarget* NewMorphTarget = NewObject<UMorphTarget>(SkeletalMesh, MorphTargetInfo.Name);
+
+		NewMorphTarget->BaseSkelMesh = SkeletalMesh;
+		NewMorphTarget->GetMorphLODModels().SetNum(MorphTargetInfo.LodNum);
+		SkeletalMesh->GetMorphTargets().Add(NewMorphTarget);
+	}
+
+	const int32 SkeletalMeshMorphTargetsNum = SkeletalMesh->GetMorphTargets().Num();
+
+	int32 LODCount = OperationData->InstanceUpdateData.LODs.Num();
+	for (int32 LODIndex = CurrentMinLOD; LODIndex <= CurrentMaxLOD && LODIndex < LODCount; ++LODIndex)
+	{
+		const FInstanceUpdateData::FLOD& LOD = OperationData->InstanceUpdateData.LODs[LODIndex];
+		if (mu::MeshPtrConst MutableMesh = OperationData->InstanceUpdateData.Components[LOD.FirstComponent+ComponentIndex].Mesh)
+		{
+			const mu::FMeshBufferSet& MeshSet = MutableMesh->GetVertexBuffers();
+
+			int32 VertexMorphsInfoIndexBufferIndex, VertexMorphsInfoIndexBufferChannel;
+			MeshSet.FindChannel(mu::MBS_OTHER, 0, &VertexMorphsInfoIndexBufferIndex, &VertexMorphsInfoIndexBufferChannel);
+
+			int32 VertexMorphsCountBufferIndex, VertexMorphsCountBufferChannel;
+			MeshSet.FindChannel(mu::MBS_OTHER, 1, &VertexMorphsCountBufferIndex, &VertexMorphsCountBufferChannel);
+
+			if (VertexMorphsInfoIndexBufferIndex < 0 || VertexMorphsCountBufferIndex < 0)
+			{
+				continue;
+			}
+
+			const int32* const VertexMorphsInfoIndexBuffer = reinterpret_cast<const int32*>(MeshSet.GetBufferData(VertexMorphsInfoIndexBufferIndex));
+			TArrayView<const int32> VertexMorphsInfoIndexView(VertexMorphsInfoIndexBuffer, MeshSet.GetElementCount());
+
+			const int32* const VertexMorphsCountBuffer = reinterpret_cast<const int32*>(MeshSet.GetBufferData(VertexMorphsCountBufferIndex));
+			TArrayView<const int32> VertexMorphsCountView(VertexMorphsCountBuffer, MeshSet.GetElementCount());
+
+			const int32 SurfaceCount = MutableMesh->GetSurfaceCount();
+			for (int32 Section = 0; Section < SurfaceCount; ++Section)
+			{
+				// Reset SectionMorphTargets.
+				for (auto& Elem : SectionMorphTargetVertices)
+				{
+					Elem = 0;
+				}
+
+				int FirstVertex, VerticesCount, FirstIndex, IndiciesCount;
+				MutableMesh->GetSurface(Section, &FirstVertex, &VerticesCount, &FirstIndex, &IndiciesCount);
+
+				for (int32 VertexIdx = FirstVertex; VertexIdx < FirstVertex + VerticesCount; ++VertexIdx)
+				{
+					const int32 MorphCount = VertexMorphsCountView[VertexIdx];
+					if (MorphCount <= 0)
+					{
+						continue;
+					}
+
+					TArrayView<const FMorphTargetVertexData> MorphsVertexDataView(&(MorphTargetReconstructionData[VertexMorphsInfoIndexView[VertexIdx]]), MorphCount);
+					for (const FMorphTargetVertexData& SourceVertex : MorphsVertexDataView)
+					{
+						// check(SkeletalMeshMorphTargetsNum <= SourceVertex.MorphIndex);
+						if (SkeletalMeshMorphTargetsNum <= SourceVertex.MorphIndex)
+						{
+							continue;
+						}
+
+						FMorphTargetLODModel& DestMorphLODModel = SkeletalMesh->GetMorphTargets()[SourceVertex.MorphIndex]->GetMorphLODModels()[LODIndex];
+
+						DestMorphLODModel.Vertices.Emplace(
+							FMorphTargetDelta{ SourceVertex.PositionDelta, SourceVertex.TangentZDelta, static_cast<uint32>(VertexIdx) });
+
+						++SectionMorphTargetVertices[SourceVertex.MorphIndex];
+
+					}
+				}
+
+				const int32 SectionMorphTargetsNum = SectionMorphTargetVertices.Num();
+				for (int32 MorphIdx = 0; MorphIdx < SectionMorphTargetsNum; ++MorphIdx)
+				{
+					if (SectionMorphTargetVertices[MorphIdx] > 0)
+					{
+						FMorphTargetLODModel& MorphTargetLodModel = SkeletalMesh->GetMorphTargets()[MorphIdx]->GetMorphLODModels()[LODIndex];
+
+						MorphTargetLodModel.SectionIndices.Add(Section);
+						MorphTargetLodModel.NumVertices = SectionMorphTargetVertices[MorphIdx];
+					}
+				}
+			}
+		}
+	}
+
+	SkeletalMesh->InitMorphTargets();
+}
+
+namespace 
+{
+	// Only used to be able to create new clothing assets and assign a new guid to them without the factory.
+	class UCustomizableObjectClothingAsset : public UClothingAssetCommon
+	{
+	public:
+		void AssignNewGuid()
+		{
+			AssetGuid = FGuid::NewGuid();
+		}
+	};
+
+}
+
+
+void UCustomizableInstancePrivateData::BuildClothingData(const TSharedPtr<FMutableOperationData>&OperationData, USkeletalMesh * SkeletalMesh, UCustomizableObjectInstance * CustomizableObjectInstance, int32 ComponentIndex)
+{
+	MUTABLE_CPUPROFILER_SCOPE(UCustomizableInstancePrivateData::BuildClothingData);
+
+	if (!(SkeletalMesh && OperationData->InstanceUpdateData.LODs.Num()))
+	{
+		return;
+	}
+
+	const UCustomizableObject* CustomizableObject = CustomizableObjectInstance->GetCustomizableObject();
+	const TArray<FCustomizableObjectClothingAssetData>& ContributingClothingAssetsData = CustomizableObject->ContributingClothingAssetsData;
+	const TArray<FCustomizableObjectClothConfigData>& ClothSharedConfigsData = CustomizableObject->ClothSharedConfigsData;
+	const TArray<FCustomizableObjectMeshToMeshVertData>& ClothMeshToMeshVertData = CustomizableObject->ClothMeshToMeshVertData;
+
+	if (!(ContributingClothingAssetsData.Num() && ClothMeshToMeshVertData.Num()))
+	{
+		return;
+	}
+
+	// First we need to discover if any clothing asset is used for the instance. 
+
+	struct FSectionWithClothData
+	{
+		int32 ClothAssetIndex;
+		int32 ClothAssetLodIndex;
+		int32 Section;
+		int32 Lod;
+		int32 BaseVertex;
+		TArrayView<const uint16> SectionIndex16View;
+		TArrayView<const uint32> SectionIndex32View;
+		TArrayView<const int32> ClothingDataIndicesView;
+		TArray<FMeshToMeshVertData> MappingData;
+	};
+
+	TArray<FSectionWithClothData> SectionsWithCloth;
+	SectionsWithCloth.Reserve(32);
+
+	int32 MeshNumLods = 0;
+
+	int32 LODCount = OperationData->InstanceUpdateData.LODs.Num();
+	{
+		MUTABLE_CPUPROFILER_SCOPE(DiscoverSectionsWithCloth);
+
+		int32 MeshLODIndex = HasCOInstanceFlags(ReduceLODs) ? 0 : CurrentMinLOD;
+		for (int32 LODIndex = CurrentMinLOD; LODIndex <= CurrentMaxLOD && LODIndex < LODCount; ++LODIndex, ++MeshLODIndex)
+		{
+			const FInstanceUpdateData::FLOD& LOD = OperationData->InstanceUpdateData.LODs[LODIndex];
+			const int32 ComponentCount = LOD.ComponentCount;
+			if (!ComponentCount)
+			{
+				return;
+			}
+
+			if (mu::MeshPtrConst MutableMesh = OperationData->InstanceUpdateData.Components[LOD.FirstComponent+ComponentIndex].Mesh)
+			{
+				const mu::FMeshBufferSet& MeshSet = MutableMesh->GetVertexBuffers();
+
+				const mu::FMeshBufferSet& IndicesSet = MutableMesh->GetIndexBuffers();
+
+				// Semantics index may vary depending on whether realtime morph targets are enabled.
+				const int32 ClothingDataBufferIndex = [&MeshSet]()
+				{
+					int32 BufferIndex, Channel;
+					MeshSet.FindChannel(mu::MBS_OTHER, 2, &BufferIndex, &Channel);
+
+					if (BufferIndex > 0)
+					{
+						return BufferIndex;
+					}
+
+					MeshSet.FindChannel(mu::MBS_OTHER, 0, &BufferIndex, &Channel);
+
+					return BufferIndex;
+				}(); // lambda is invoked
+
+				if (ClothingDataBufferIndex < 0)
+				{
+					continue;
+				}
+
+				const int32* const ClothingDataBuffer = reinterpret_cast<const int32*>(MeshSet.GetBufferData(ClothingDataBufferIndex));
+
+				const int32 SurfaceCount = MutableMesh->GetSurfaceCount();
+				for (int32 Section = 0; Section < SurfaceCount; ++Section)
+				{
+					int FirstVertex, VerticesCount, FirstIndex, IndicesCount;
+					MutableMesh->GetSurface(Section, &FirstVertex, &VerticesCount, &FirstIndex, &IndicesCount);
+
+					if (VerticesCount == 0 || IndicesCount == 0)
+					{
+						continue;
+					}
+					
+					// A Section has cloth data on all its vertices or it does not have it in any.
+					// It can be determined if this section has clothing data just looking at the 
+					// first vertex of the section.
+					TArrayView<const int32> ClothingDataView(ClothingDataBuffer + FirstVertex, VerticesCount);
+
+					const int32 IndexCount = MutableMesh->GetIndexBuffers().GetElementCount();
+
+					TArrayView<const uint16> IndicesView16Bits;
+					TArrayView<const uint32> IndicesView32Bits;
+					
+					if (IndexCount)
+					{
+						 if (MutableMesh->GetIndexBuffers().GetElementSize(0) == 2)
+						 {
+						 	const uint16* IndexPtr = (const uint16*)MutableMesh->GetIndexBuffers().GetBufferData(0);
+							IndicesView16Bits = TArrayView<const uint16>(IndexPtr + FirstIndex, IndicesCount);
+						 }
+						 else
+						 {
+						 	const uint32* IndexPtr = (const uint32*)MutableMesh->GetIndexBuffers().GetBufferData(0);
+							IndicesView32Bits = TArrayView<const uint32>(IndexPtr + FirstIndex, IndicesCount);
+						 }
+					}
+					
+					if (!ClothingDataView.Num())
+					{
+						continue;
+					}
+
+					const int32 ClothDataIndex = ClothingDataView[0];
+					if (ClothDataIndex < 0)
+					{
+						continue;
+					}
+
+					const int32 ClothAssetIndex = ClothMeshToMeshVertData[ClothDataIndex].SourceAssetIndex;
+					const int32 ClothAssetLodIndex = ClothMeshToMeshVertData[ClothDataIndex].SourceAssetLodIndex;
+
+					// Defensive check, this indicates the clothing data might be stale and needs to be recompiled.
+					// Should never happen.
+					if (!ensure(ClothAssetIndex >= 0 && ClothAssetIndex < ContributingClothingAssetsData.Num()
+						&& ContributingClothingAssetsData[ClothAssetIndex].LodData.Num()))
+					{
+						continue;
+					}
+
+					SectionsWithCloth.Add
+							(FSectionWithClothData{ ClothAssetIndex, ClothAssetLodIndex, Section, MeshLODIndex, FirstVertex, IndicesView16Bits, IndicesView32Bits, ClothingDataView, TArray<FMeshToMeshVertData>() });
+				}
+			}
+		}
+
+		MeshNumLods = MeshLODIndex;
+	}
+
+	if (!SectionsWithCloth.Num())
+	{
+		return; // Nothing to do.
+	}
+	
+	TArray<FCustomizableObjectClothingAssetData> NewClothingAssetsData;
+	NewClothingAssetsData.SetNum( ContributingClothingAssetsData.Num() );
+
+	{
+
+		for (FSectionWithClothData& SectionWithCloth : SectionsWithCloth)
+		{
+			const FCustomizableObjectClothingAssetData& SrcAssetData = ContributingClothingAssetsData[SectionWithCloth.ClothAssetIndex];
+			FCustomizableObjectClothingAssetData& DstAssetData = NewClothingAssetsData[SectionWithCloth.ClothAssetIndex];
+		
+			// Only initilialize once, multiple sections with cloth could point to the same cloth asset.
+			if (!DstAssetData.LodMap.Num())
+			{
+				DstAssetData.LodMap.Init( -1, MeshNumLods );
+
+				DstAssetData.LodData.SetNum( SrcAssetData.LodData.Num() );
+				DstAssetData.UsedBoneNames = SrcAssetData.UsedBoneNames;
+				DstAssetData.UsedBoneIndices = SrcAssetData.UsedBoneIndices;
+				DstAssetData.ReferenceBoneIndex = SrcAssetData.ReferenceBoneIndex;
+				DstAssetData.Name = SrcAssetData.Name;
+			}
+
+			DstAssetData.LodMap[SectionWithCloth.Lod] = SectionWithCloth.ClothAssetLodIndex;
+		}
+	}
+
+	{
+		MUTABLE_CPUPROFILER_SCOPE(CopyMeshToMeshData)
+		
+		// Copy MeshToMeshData.
+		for (FSectionWithClothData& SectionWithCloth : SectionsWithCloth)
+		{
+			const int32 NumVertices = SectionWithCloth.ClothingDataIndicesView.Num();
+			
+			TArray<FMeshToMeshVertData>& ClothMappingData = SectionWithCloth.MappingData;
+			ClothMappingData.SetNum(NumVertices);
+
+			// Copy mesh to mesh data indexed by the index stored per vertex at compile time. 
+			for (int32 VertexIdx = 0; VertexIdx < NumVertices; ++VertexIdx)
+			{
+				// Possible Optimization: Gather consecutive indices in ClothingDataView and Memcpy the whole range.
+				// FMeshToMeshVertData and FCustomizableObjectMeshToMeshVertData have the same memory footprint and
+				// bytes in a FCustomizableObjectMeshToMeshVertData form a valid FMeshToMeshVertData (not the other way around).
+
+				static_assert(sizeof(FCustomizableObjectMeshToMeshVertData) == sizeof(FMeshToMeshVertData), "");
+				static_assert(TIsTrivial<FCustomizableObjectMeshToMeshVertData>::Value, "");
+				static_assert(TIsTrivial<FMeshToMeshVertData>::Value, "");
+
+				const int32 VertexDataIndex = SectionWithCloth.ClothingDataIndicesView[VertexIdx];
+				check(VertexDataIndex > 0);
+
+				const FCustomizableObjectMeshToMeshVertData& SrcData = ClothMeshToMeshVertData[VertexDataIndex];
+
+				FMeshToMeshVertData& DstData = ClothMappingData[VertexIdx];
+				FMemory::Memcpy(&DstData, &SrcData, sizeof(FMeshToMeshVertData));
+			}
+		}
+	}
+
+	// Indices remaps for {Section, AssetLod}, needed to recreate the lod transition data.
+	TArray<TArray<TArray<int32>>> PhysicsSectionLodsIndicesRemaps;	
+
+	check( SectionsWithCloth.Num() > 0 );
+	FSectionWithClothData* MaxSection = Algo::MaxElement(SectionsWithCloth, 
+			[](const FSectionWithClothData& A, const FSectionWithClothData& B) { return  A.Section < B.Section; } );
+	
+	PhysicsSectionLodsIndicesRemaps.SetNum(MaxSection->Section + 1);
+		
+	for (FSectionWithClothData& SectionLods : SectionsWithCloth)
+	{
+		const int32 SectionLodNum = PhysicsSectionLodsIndicesRemaps[SectionLods.Section].Num();
+
+		PhysicsSectionLodsIndicesRemaps[SectionLods.Section].SetNum(
+				FMath::Max( SectionLodNum, SectionLods.ClothAssetLodIndex + 1 ));
+	}
+
+	{
+		MUTABLE_CPUPROFILER_SCOPE(RemapPhysicsMesh)
+
+		for (FSectionWithClothData& SectionWithCloth : SectionsWithCloth)
+		{
+			const FCustomizableObjectClothingAssetData& SrcClothingAssetData = ContributingClothingAssetsData[SectionWithCloth.ClothAssetIndex];
+			FCustomizableObjectClothingAssetData& NewClothingAssetData = NewClothingAssetsData[SectionWithCloth.ClothAssetIndex];
+					
+			const FClothLODDataCommon& SrcLodData = SrcClothingAssetData.LodData[SectionWithCloth.ClothAssetLodIndex];
+			FClothLODDataCommon& NewLodData = NewClothingAssetData.LodData[SectionWithCloth.ClothAssetLodIndex];
+			
+			const int32 PhysicalMeshVerticesNum = SrcLodData.PhysicalMeshData.Vertices.Num();
+
+			if (!PhysicalMeshVerticesNum)
+			{
+				// Nothing to do.
+				continue;
+			}
+			
+			// Vertices not indexed in the mesh to mesh data generated for this section need to be removed.
+			TArray<uint8> VertexUtilizationBuffer;
+			VertexUtilizationBuffer.Init(0, PhysicalMeshVerticesNum);
+		
+			// Discover used vertices.
+			const int32 SectionVerticesNum = SectionWithCloth.ClothingDataIndicesView.Num();
+
+			TArray<uint8> RenderVertexUtilizationBuffer;
+			RenderVertexUtilizationBuffer.Init(0, SectionVerticesNum);
+
+			// Sometimes, at least when a clip morph is applied, vertices are not removed from the section
+			// and only the triangles (indices) that form the mesh are modified.
+
+			auto GenerateRenderUtilizationBuffer = [&RenderVertexUtilizationBuffer](const auto& IndicesView, int32 SectionBaseVertex )
+			{
+				const int32 IndicesCount = IndicesView.Num();
+				check(IndicesCount % 3 == 0);
+				for ( int32 I = 0; I < IndicesCount; I += 3 )
+				{
+					RenderVertexUtilizationBuffer[int32(IndicesView[I + 0]) - SectionBaseVertex] = 1;
+					RenderVertexUtilizationBuffer[int32(IndicesView[I + 1]) - SectionBaseVertex] = 1;
+					RenderVertexUtilizationBuffer[int32(IndicesView[I + 2]) - SectionBaseVertex] = 1;
+				}
+			};
+			
+			if (SectionWithCloth.SectionIndex16View.Num())
+			{
+				GenerateRenderUtilizationBuffer(SectionWithCloth.SectionIndex16View, SectionWithCloth.BaseVertex);
+			}
+			else
+			{
+				check(SectionWithCloth.SectionIndex32View.Num());
+				GenerateRenderUtilizationBuffer(SectionWithCloth.SectionIndex32View, SectionWithCloth.BaseVertex);
+			}
+
+			const TArray<FMeshToMeshVertData>& SectionClothMappingData = SectionWithCloth.MappingData;
+			for (int32 Idx = 0; Idx < SectionVerticesNum; ++Idx)
+			{
+				if (RenderVertexUtilizationBuffer[Idx])
+				{
+					const uint16* Indices = SectionClothMappingData[Idx].SourceMeshVertIndices;
+
+					VertexUtilizationBuffer[Indices[0]] = 1;
+					VertexUtilizationBuffer[Indices[1]] = 1;
+					VertexUtilizationBuffer[Indices[2]] = 1;
+				}
+			}
+
+			TArray<int32>& IndexMap = PhysicsSectionLodsIndicesRemaps[SectionWithCloth.Section][SectionWithCloth.ClothAssetLodIndex];
+			IndexMap.SetNumUninitialized(PhysicalMeshVerticesNum);
+
+			// Compute index remap and number of remaining physics vertices.
+			// -1 indicates the vertex has been removed.
+			int32 NewPhysicalMeshVerticesNum = 0;
+			for (int32 Idx = 0; Idx < PhysicalMeshVerticesNum; ++Idx)
+			{
+				IndexMap[Idx] = VertexUtilizationBuffer[Idx] ? NewPhysicalMeshVerticesNum++ : -1;
+			}
+
+			const bool bHasVerticesRemoved = NewPhysicalMeshVerticesNum < PhysicalMeshVerticesNum;
+			if (!bHasVerticesRemoved)
+			{
+				// If no vertices are removed the IndexMap is no longer needed. The lack of data in the map 
+				// can indicates that no vertex has been removed to subsequent operations.
+				IndexMap.Reset();
+			}
+			
+			const auto CopyIfUsed = [&VertexUtilizationBuffer, bHasVerticesRemoved](auto& Dst, const auto& Src)
+			{	
+				const int32 SrcNumElems = Src.Num();
+
+				if (!bHasVerticesRemoved)
+				{
+					for (int32 Idx = 0; Idx < SrcNumElems; ++Idx)
+					{
+						Dst[Idx] = Src[Idx];
+					}
+
+					return;
+				}
+
+				for (int32 Idx = 0, DstNumElems = 0; Idx < SrcNumElems; ++Idx)
+				{
+					if (VertexUtilizationBuffer[Idx])
+					{
+						Dst[DstNumElems++] = Src[Idx];
+					}
+				}
+			};
+
+			NewLodData.PhysicalMeshData.MaxBoneWeights = SrcLodData.PhysicalMeshData.MaxBoneWeights;
+
+			NewLodData.PhysicalMeshData.Vertices.SetNum(NewPhysicalMeshVerticesNum);
+			NewLodData.PhysicalMeshData.Normals.SetNum(NewPhysicalMeshVerticesNum);
+			NewLodData.PhysicalMeshData.BoneData.SetNum(NewPhysicalMeshVerticesNum);
+			NewLodData.PhysicalMeshData.InverseMasses.SetNum(NewPhysicalMeshVerticesNum);
+
+			CopyIfUsed(NewLodData.PhysicalMeshData.Vertices, SrcLodData.PhysicalMeshData.Vertices);
+			CopyIfUsed(NewLodData.PhysicalMeshData.Normals, SrcLodData.PhysicalMeshData.Normals);
+			CopyIfUsed(NewLodData.PhysicalMeshData.BoneData, SrcLodData.PhysicalMeshData.BoneData);
+			CopyIfUsed(NewLodData.PhysicalMeshData.InverseMasses, SrcLodData.PhysicalMeshData.InverseMasses);
+
+			const bool bNeedsTransitionUpData = SectionWithCloth.Lod - 1 >= 0;
+			if (bNeedsTransitionUpData)
+			{
+				NewLodData.TransitionUpSkinData.SetNum(SrcLodData.TransitionUpSkinData.Num() ? NewPhysicalMeshVerticesNum : 0);	
+				CopyIfUsed(NewLodData.TransitionUpSkinData, SrcLodData.TransitionUpSkinData);
+			}
+
+			check(SectionWithCloth.Section < PhysicsSectionLodsIndicesRemaps.Num());
+			const bool bNeedsTransitionDownData = SectionWithCloth.Lod + 1 < PhysicsSectionLodsIndicesRemaps[SectionWithCloth.Section].Num();
+			if (bNeedsTransitionDownData)
+			{
+				NewLodData.TransitionDownSkinData.SetNum(SrcLodData.TransitionDownSkinData.Num() ? NewPhysicalMeshVerticesNum : 0);
+				CopyIfUsed(NewLodData.TransitionDownSkinData, SrcLodData.TransitionDownSkinData);
+			}
+			
+			const TMap<uint32, FPointWeightMap>& SrcPhysWeightMaps = SrcLodData.PhysicalMeshData.WeightMaps;
+			TMap<uint32, FPointWeightMap>& NewPhysWeightMaps = NewLodData.PhysicalMeshData.WeightMaps;
+
+			for (const TPair<uint32, FPointWeightMap>& WeightMap : SrcPhysWeightMaps)
+			{
+				if (WeightMap.Value.Values.Num() > 0)
+				{
+					FPointWeightMap& NewWeightMap = NewLodData.PhysicalMeshData.AddWeightMap(WeightMap.Key);
+					NewWeightMap.Values.SetNum(NewPhysicalMeshVerticesNum);
+
+					CopyIfUsed(NewWeightMap.Values, WeightMap.Value.Values);
+				}
+			}
+			// Remap render mesh to mesh indices.
+			if (bHasVerticesRemoved)
+			{
+				for (FMeshToMeshVertData& VertClothData : SectionWithCloth.MappingData)
+				{
+					uint16* Indices = VertClothData.SourceMeshVertIndices;
+					Indices[0] = (uint16)IndexMap[Indices[0]];
+					Indices[1] = (uint16)IndexMap[Indices[1]];
+					Indices[2] = (uint16)IndexMap[Indices[2]];
+				}
+			}
+
+			// Remap and trim physics mesh vertices and self collision indices. 
+			
+			// Returns the final size of Dst.
+			const auto TrimAndRemapTriangles = [&IndexMap](TArray<uint32>& Dst, const TArray<uint32>& Src) -> int32
+			{
+				check(Src.Num() % 3 == 0);
+
+				const int32 SrcNumElems = Src.Num();
+				if (!IndexMap.Num())
+				{
+					//for (int32 Idx = 0; Idx < SrcNumElems; ++Idx)
+					//{
+					//	Dst[Idx] = Src[Idx];
+					//}
+
+					FMemory::Memcpy( Dst.GetData(), Src.GetData(), SrcNumElems*sizeof(uint32) );
+					return SrcNumElems;
+				}
+
+				int32 DstNumElems = 0;
+				for (int32 Idx = 0; Idx < SrcNumElems; Idx += 3)
+				{
+					const int32 Idx0 = IndexMap[Src[Idx + 0]];
+					const int32 Idx1 = IndexMap[Src[Idx + 1]];
+					const int32 Idx2 = IndexMap[Src[Idx + 2]];
+
+					// triangles are only copied if all vertices are used.
+					if (!((Idx0 < 0) | (Idx1 < 0) | (Idx2 < 0)))
+					{
+						Dst[DstNumElems + 0] = Idx0;
+						Dst[DstNumElems + 1] = Idx1;
+						Dst[DstNumElems + 2] = Idx2;
+
+						DstNumElems += 3;
+					}
+				}
+
+				return DstNumElems;
+			};
+
+			const TArray<uint32>& SrcPhysicalMeshIndices = SrcLodData.PhysicalMeshData.Indices;
+			TArray<uint32>& NewPhysicalMeshIndices = NewLodData.PhysicalMeshData.Indices;
+			NewPhysicalMeshIndices.SetNum(SrcPhysicalMeshIndices.Num());
+			NewPhysicalMeshIndices.SetNum(TrimAndRemapTriangles(NewPhysicalMeshIndices, SrcPhysicalMeshIndices), false);
+		
+			const auto TrimAndRemapIndices = [&IndexMap](TArray<uint32>& Dst, const TArray<uint32>& Src) -> int32
+			{	
+				const int32 SrcNumElems = Src.Num();
+				if (!IndexMap.Num())
+				{
+					//for (int32 Idx = 0; Idx < SrcNumElems; ++Idx)
+					//{
+					//	Dst[Idx] = Src[Idx];
+					//}
+
+					FMemory::Memcpy( Dst.GetData(), Src.GetData(), SrcNumElems*sizeof(uint32) );
+					return SrcNumElems;
+				}
+
+				int32 DstNumElems = 0;
+				for (int32 Idx = 0; Idx < SrcNumElems; ++Idx)
+				{
+					const int32 MappedIdx = IndexMap[Src[Idx]];
+
+					if (MappedIdx >= 0)
+					{
+						Dst[DstNumElems++] = MappedIdx;
+					}
+				}
+
+				return DstNumElems;
+			};
+
+			const TArray<uint32>& SrcSelfCollisionIndices = SrcLodData.PhysicalMeshData.SelfCollisionIndices;
+			TArray<uint32>& NewSelfCollisionIndices = NewLodData.PhysicalMeshData.SelfCollisionIndices;
+			NewSelfCollisionIndices.SetNum(SrcSelfCollisionIndices.Num());
+			NewSelfCollisionIndices.SetNum(TrimAndRemapIndices(NewSelfCollisionIndices, SrcSelfCollisionIndices), false);
+						
+			// TODO: Remap tethers. This is important with the chaos implementation to remove stretchiness to the cloth.
+			// Possible implementation, not checked thoroughly or benchmarked.
+			auto TrimAndRemapTethers = [&IndexMap](FClothTetherData& Dst, const FClothTetherData& Src)
+			{
+				if (!IndexMap.Num())
+				{
+					Dst.Tethers = Src.Tethers;
+					return;
+				}
+
+				Dst.Tethers.Reserve(Src.Tethers.Num());
+				for ( const TArray<TTuple<int32, int32, float>>& SrcTetherCluster : Src.Tethers )
+				{
+					TArray<TTuple<int32, int32, float>>& DstTetherCluster = Dst.Tethers.Emplace_GetRef();
+					DstTetherCluster.Reserve(SrcTetherCluster.Num());
+					for ( const TTuple<int32, int32, float>& Tether : SrcTetherCluster )
+					{
+						const int32 Index0 = IndexMap[Tether.Get<0>()];
+						const int32 Index1 = IndexMap[Tether.Get<1>()];
+						if ((Index0 > 0) & (Index1 > 0))
+						{
+							DstTetherCluster.Emplace(Index0, Index1, Tether.Get<2>());
+						}
+					}
+
+					if (!DstTetherCluster.Num())
+					{
+						constexpr bool bAllowShrinking = false; 
+						Dst.Tethers.RemoveAt( Dst.Tethers.Num() - 1, 1, bAllowShrinking );
+					}
+				}
+			};
+
+			// Not enabled yet.
+			//TrimAndRemapTethers( NewLodData.PhysicalMeshData.GeodesicTethers, SrcLodData.PhysicalMeshData.GeodesicTethers );
+			//TrimAndRemapTethers( NewLodData.PhysicalMeshData.EuclideanTethers, SrcLodData.PhysicalMeshData.EuclideanTethers );
+		}
+	}
+
+	// Try to find plausible values for LodTransitionData vertices that have lost the triangle to which are attached.
+	{
+		MUTABLE_CPUPROFILER_SCOPE(BuildLodTransitionData)
+		
+		for (const FSectionWithClothData& SectionWithCloth : SectionsWithCloth)
+		{
+
+			FCustomizableObjectClothingAssetData& NewClothingAssetData = NewClothingAssetsData[SectionWithCloth.ClothAssetIndex];
+			FClothLODDataCommon& NewLodData = NewClothingAssetData.LodData[SectionWithCloth.ClothAssetLodIndex];
+
+			const int32 PhysicalMeshVerticesNum = NewLodData.PhysicalMeshData.Vertices.Num();
+
+			if (!PhysicalMeshVerticesNum)
+			{
+				// Nothing to do.
+				continue;
+			}
+		
+			auto RemapTransitionMeshToMeshVertData = []( TArray<FMeshToMeshVertData>& InOutVertData, TArray<int32>& IndexMap )
+			{
+				for (FMeshToMeshVertData& VertData : InOutVertData)
+				{
+					uint16* Indices = VertData.SourceMeshVertIndices;
+					Indices[0] = (uint16)IndexMap[Indices[0]];
+					Indices[1] = (uint16)IndexMap[Indices[1]];
+					Indices[2] = (uint16)IndexMap[Indices[2]];
+				}
+			};
+
+			if (NewLodData.TransitionDownSkinData.Num() > 0)
+			{
+				check(SectionWithCloth.Section < PhysicsSectionLodsIndicesRemaps.Num() &&
+					  SectionWithCloth.ClothAssetLodIndex + 1 < PhysicsSectionLodsIndicesRemaps[SectionWithCloth.Section].Num() );
+				
+				TArray<int32>& IndexMap = PhysicsSectionLodsIndicesRemaps[SectionWithCloth.Section][SectionWithCloth.ClothAssetLodIndex + 1];
+
+				if (IndexMap.Num())
+				{
+					RemapTransitionMeshToMeshVertData( NewLodData.TransitionDownSkinData, IndexMap );
+				}
+			}
+			
+			if (NewLodData.TransitionUpSkinData.Num() > 0)
+			{
+				check( SectionWithCloth.ClothAssetLodIndex - 1 >= 0 )
+				
+				check(SectionWithCloth.Section < PhysicsSectionLodsIndicesRemaps.Num() &&
+					  SectionWithCloth.ClothAssetLodIndex - 1 < PhysicsSectionLodsIndicesRemaps[SectionWithCloth.Section].Num());
+				
+				TArray<int32>& IndexMap = PhysicsSectionLodsIndicesRemaps[SectionWithCloth.Section][SectionWithCloth.ClothAssetLodIndex - 1];
+			
+				if (IndexMap.Num())
+				{
+					RemapTransitionMeshToMeshVertData( NewLodData.TransitionUpSkinData, IndexMap );
+				}
+			}
+
+			struct FMeshPhysicsDesc
+			{
+				const TArray<FVector3f>& Vertices;
+				const TArray<FVector3f>& Normals;
+				const TArray<uint32>& Indices;
+			};	
+		
+			auto RebindVertex = [](const FMeshPhysicsDesc& Mesh, const FVector3f& InPosition, const FVector3f& InNormal, FMeshToMeshVertData& Out)
+			{
+				const FVector3f Normal = InNormal;
+
+				// We don't have the mesh tangent, find something plausible.
+				FVector3f Tan0, Tan1;
+				Normal.FindBestAxisVectors(Tan0, Tan1);
+				const FVector3f Tangent = Tan0;
+				
+				// Some of the math functions take as argument FVector, we'd want to be FVector3f. 
+				// This should be changed once support for the single type in the FMath functions is added. 
+				const FVector Position = (FVector)InPosition;
+				int32 BestBaseTriangleIdx = INDEX_NONE;
+				FVector::FReal BestDistanceSq = TNumericLimits<FVector::FReal>::Max();
+				
+				const int32 NumIndices = Mesh.Indices.Num();
+				check(NumIndices % 3 == 0);
+
+				for (int32 I = 0; I < NumIndices; I += 3)
+				{
+					const FVector& A = (FVector)Mesh.Vertices[Mesh.Indices[I + 0]];
+					const FVector& B = (FVector)Mesh.Vertices[Mesh.Indices[I + 1]];
+					const FVector& C = (FVector)Mesh.Vertices[Mesh.Indices[I + 2]];
+
+					FVector ClosestTrianglePoint = FMath::ClosestPointOnTriangleToPoint(Position, (FVector)A, (FVector)B, (FVector)C);
+
+					const FVector::FReal CurrentDistSq = (ClosestTrianglePoint - Position).SizeSquared();
+					if (CurrentDistSq < BestDistanceSq)
+					{
+						BestDistanceSq = CurrentDistSq;
+						BestBaseTriangleIdx = I;
+					}
+				}
+
+				check(BestBaseTriangleIdx >= 0);
+
+				auto ComputeBaryCoordsAndDist = [](const FVector3f& A, const FVector3f& B, const FVector3f& C, const FVector3f& P) -> FVector4f
+				{
+					FPlane4f TrianglePlane(A, B, C);
+
+					const FVector3f PointOnTriPlane = FVector3f::PointPlaneProject(P, TrianglePlane);
+					const FVector3f BaryCoords = (FVector3f)FMath::ComputeBaryCentric2D((FVector)PointOnTriPlane, (FVector)A, (FVector)B, (FVector)C);
+
+					return FVector4f(BaryCoords, TrianglePlane.PlaneDot((FVector3f)P));
+				};
+
+				const FVector3f& A = Mesh.Vertices[Mesh.Indices[BestBaseTriangleIdx + 0]];
+				const FVector3f& B = Mesh.Vertices[Mesh.Indices[BestBaseTriangleIdx + 1]];
+				const FVector3f& C = Mesh.Vertices[Mesh.Indices[BestBaseTriangleIdx + 2]];
+
+				Out.PositionBaryCoordsAndDist = ComputeBaryCoordsAndDist(A, B, C, (FVector3f)Position );
+				Out.NormalBaryCoordsAndDist = ComputeBaryCoordsAndDist(A, B, C, (FVector3f)Position + Normal );
+				Out.TangentBaryCoordsAndDist = ComputeBaryCoordsAndDist(A, B, C, (FVector3f)Position + Tangent );
+				Out.SourceMeshVertIndices[0] = (uint16)Mesh.Indices[BestBaseTriangleIdx + 0];
+				Out.SourceMeshVertIndices[1] = (uint16)Mesh.Indices[BestBaseTriangleIdx + 1]; 
+				Out.SourceMeshVertIndices[2] = (uint16)Mesh.Indices[BestBaseTriangleIdx + 2];
+			};
+		
+			auto RecreateTransitionData = [&PhysicsSectionLodsIndicesRemaps, &RebindVertex]( 
+				const FMeshPhysicsDesc& ToMesh, const FMeshPhysicsDesc& FromMesh, const TArray<int32>& IndexMap, TArray<FMeshToMeshVertData>& InOutTransitionData )
+			{
+				if (!IndexMap.Num())
+				{
+					return;
+				}
+
+				if (!InOutTransitionData.Num())
+				{
+					return;
+				}
+
+				const int32 TransitionDataNum = InOutTransitionData.Num();
+				
+				for (int32 I = 0; I < TransitionDataNum; ++I)
+				{
+					FMeshToMeshVertData& VertData = InOutTransitionData[I];
+					uint16* Indices = VertData.SourceMeshVertIndices;
+
+					// If any original indices are missing but the vertex is still alive rebind the vertex.
+					// In general, the number of rebinds should be small.
+
+					// Currently, if any index is missing we rebind to the closest triangle but it could be nice to use the remaining indices, 
+					// if any, to find the most appropriate triangle to bind to. 
+					const bool bNeedsRebind = (Indices[0] == 0xFFFF) | (Indices[1] == 0xFFFF) | (Indices[2] == 0xFFFF);
+
+					if (bNeedsRebind)
+					{
+						RebindVertex( ToMesh, FromMesh.Vertices[I], FromMesh.Normals[I], VertData );
+					}
+				}
+			};
+
+			const FMeshPhysicsDesc CurrentPhysicsMesh {
+				NewClothingAssetData.LodData[SectionWithCloth.ClothAssetLodIndex].PhysicalMeshData.Vertices,
+				NewClothingAssetData.LodData[SectionWithCloth.ClothAssetLodIndex].PhysicalMeshData.Normals,
+				NewClothingAssetData.LodData[SectionWithCloth.ClothAssetLodIndex].PhysicalMeshData.Indices };
+
+			const TArray< TArray<int32> >& SectionIndexRemaps = PhysicsSectionLodsIndicesRemaps[SectionWithCloth.Section];
+
+			if (SectionWithCloth.ClothAssetLodIndex < SectionIndexRemaps.Num() - 1)
+			{
+				const TArray<int32>& IndexMap = PhysicsSectionLodsIndicesRemaps[SectionWithCloth.Section][SectionWithCloth.ClothAssetLodIndex + 1];
+				
+				const FMeshPhysicsDesc TransitionDownTarget {  
+					NewClothingAssetData.LodData[SectionWithCloth.ClothAssetLodIndex + 1].PhysicalMeshData.Vertices,
+					NewClothingAssetData.LodData[SectionWithCloth.ClothAssetLodIndex + 1].PhysicalMeshData.Normals,
+					NewClothingAssetData.LodData[SectionWithCloth.ClothAssetLodIndex + 1].PhysicalMeshData.Indices };
+					
+				RecreateTransitionData( TransitionDownTarget, CurrentPhysicsMesh, IndexMap, NewLodData.TransitionDownSkinData );
+			}
+			
+			if (SectionWithCloth.ClothAssetLodIndex > 0)
+			{
+				const TArray<int32>& IndexMap = PhysicsSectionLodsIndicesRemaps[SectionWithCloth.Section][SectionWithCloth.ClothAssetLodIndex - 1];
+				
+				FMeshPhysicsDesc TransitionUpTarget{  
+					NewClothingAssetData.LodData[SectionWithCloth.ClothAssetLodIndex - 1].PhysicalMeshData.Vertices,
+					NewClothingAssetData.LodData[SectionWithCloth.ClothAssetLodIndex - 1].PhysicalMeshData.Normals,
+					NewClothingAssetData.LodData[SectionWithCloth.ClothAssetLodIndex - 1].PhysicalMeshData.Indices };
+		
+				RecreateTransitionData( TransitionUpTarget, CurrentPhysicsMesh, IndexMap, NewLodData.TransitionUpSkinData );
+			}
+		}
+	}
+
+	// From here up, could be moved to an async task similar to what is done with the other prepare tasks.  
+ 
+	// Create Clothing Assets.
+
+	// Based on FSkeletalMeshLODModel::GetClothMappingData().
+	TArray<TArray<FMeshToMeshVertData>> LodMappingData;
+	LodMappingData.SetNum(FMath::Max<int32>(CurrentMaxLOD + 1, LODCount));
+
+	TArray<TArray<FClothBufferIndexMapping>> LodClothingIndexMapping;
+	LodClothingIndexMapping.SetNum(FMath::Max<int32>(CurrentMaxLOD + 1, LODCount));
+	{
+		int32 NumSectionsWithClothProcessed = 0;
+		int32 MeshLODIndex = HasCOInstanceFlags(ReduceLODs) ? 0 : CurrentMinLOD;
+
+		for (int32 LODIndex = CurrentMinLOD; LODIndex <= CurrentMaxLOD && LODIndex < LODCount; ++LODIndex, ++MeshLODIndex)
+		{
+			const FInstanceUpdateData::FLOD& LOD = OperationData->InstanceUpdateData.LODs[LODIndex];
+			const int32 ComponentCount = LOD.ComponentCount;
+			if (!ComponentCount)
+			{
+				return;
+			}
+
+			if (mu::MeshPtrConst MutableMesh = OperationData->InstanceUpdateData.Components[LOD.FirstComponent+ComponentIndex].Mesh)
+			{
+				TArray<FMeshToMeshVertData>& MappingData = LodMappingData[LODIndex];
+				TArray<FClothBufferIndexMapping>& ClothingIndexMapping = LodClothingIndexMapping[LODIndex];
+				ClothingIndexMapping.Reserve(32);
+
+				const int32 SurfaceCount = MutableMesh->GetSurfaceCount();
+				for (int32 Section = 0; Section < SurfaceCount; ++Section)
+				{
+					// Check that is a valid surface.
+					int32 FirstVertex, VerticesCount, FirstIndex, IndicesCount;
+					MutableMesh->GetSurface(Section, &FirstVertex, &VerticesCount, &FirstIndex, &IndicesCount);
+
+					if (VerticesCount == 0 || IndicesCount == 0 )
+					{
+						continue;
+					}					
+			
+					// An entry is added for all sections.	
+					FClothBufferIndexMapping& ClothBufferIndexMapping = ClothingIndexMapping.AddZeroed_GetRef();
+
+					if (NumSectionsWithClothProcessed < SectionsWithCloth.Num())
+					{
+						const FSectionWithClothData& SectionWithCloth = SectionsWithCloth[NumSectionsWithClothProcessed];
+						// Section with cloth are sorted by {LOD, Section}
+						if (SectionWithCloth.Lod == LODIndex && SectionWithCloth.Section == Section)
+						{	
+							ClothBufferIndexMapping.BaseVertexIndex = SectionWithCloth.BaseVertex;
+							ClothBufferIndexMapping.MappingOffset = (uint32)MappingData.Num();
+							ClothBufferIndexMapping.LODBiasStride = (uint32)SectionWithCloth.MappingData.Num();
+						
+							MappingData += SectionWithCloth.MappingData; 
+
+							++NumSectionsWithClothProcessed;
+						}
+					}
+				}
+			}
+		}
+	}
+
+	FSkeletalMeshRenderData* RenderResource = SkeletalMesh->GetResourceForRendering();
+	{
+		MUTABLE_CPUPROFILER_SCOPE(InitClothRenderData)
+		// Based on FSkeletalMeshLODModel::GetClothMappingData().
+		int32 MeshLODIndex = HasCOInstanceFlags(ReduceLODs) ? 0 : CurrentMinLOD;
+		for (int32 LODIndex = CurrentMinLOD; LODIndex <= CurrentMaxLOD && LODIndex < LODCount; ++LODIndex, ++MeshLODIndex)
+		{
+			FSkeletalMeshLODRenderData& LODModel = RenderResource->LODRenderData[MeshLODIndex];
+	
+			if (LodMappingData[LODIndex].Num() > 0)
+			{
+				LODModel.ClothVertexBuffer.Init(LodMappingData[LODIndex], LodClothingIndexMapping[LODIndex]);
+			}
+		}
+	}
+
+	TArray<UCustomizableObjectClothingAsset*> NewClothingAssets;
+	NewClothingAssets.Init(nullptr, ContributingClothingAssetsData.Num());
+
+	{ 
+		MUTABLE_CPUPROFILER_SCOPE(CreateClothingConfigs)
+
+		check(NewClothingAssets.Num() == ClothingPhysicsAssets.Num());
+		for (int32 I = 0; I < NewClothingAssetsData.Num(); ++I)
+		{
+			// skip assets not set.
+			if (!NewClothingAssetsData[I].LodData.Num())
+			{
+				continue;
+			}
+		
+			NewClothingAssets[I] = NewObject<UCustomizableObjectClothingAsset>( SkeletalMesh, NewClothingAssetsData[I].Name );
+			
+			// The data can be moved to the actual asset since it will not be used anymore.
+			NewClothingAssets[I]->LodMap = MoveTemp(NewClothingAssetsData[I].LodMap);
+			NewClothingAssets[I]->LodData = MoveTemp(NewClothingAssetsData[I].LodData);
+			NewClothingAssets[I]->UsedBoneIndices = MoveTemp(NewClothingAssetsData[I].UsedBoneIndices);
+			NewClothingAssets[I]->UsedBoneNames = MoveTemp(NewClothingAssetsData[I].UsedBoneNames);
+			NewClothingAssets[I]->ReferenceBoneIndex = NewClothingAssetsData[I].ReferenceBoneIndex;
+			NewClothingAssets[I]->AssignNewGuid();
+			NewClothingAssets[I]->RefreshBoneMapping(SkeletalMesh);
+			NewClothingAssets[I]->CalculateReferenceBoneIndex();	
+			NewClothingAssets[I]->PhysicsAsset = ClothingPhysicsAssets[I];
+
+			for ( const FCustomizableObjectClothConfigData& ConfigData : ContributingClothingAssetsData[I].ConfigsData )
+			{
+				UClass* ClothConfigClass = FindObject<UClass>(nullptr, *ConfigData.ClassPath);
+				if (ClothConfigClass)
+				{
+					UClothConfigCommon* ClothConfig = NewObject<UClothConfigCommon>(NewClothingAssets[I], ClothConfigClass);
+					if (ClothConfig)
+					{
+						FMemoryReaderView MemoryReader( ConfigData.ConfigBytes );
+						ClothConfig->Serialize(MemoryReader);
+						NewClothingAssets[I]->ClothConfigs.Add(ConfigData.ConfigName, ClothConfig);
+					}
+				}
+			}
+
+			for (const FCustomizableObjectClothConfigData& ConfigData : ClothSharedConfigsData)
+			{	
+				UClass* ClothConfigClass = FindObject<UClass>(nullptr, *ConfigData.ClassPath);
+				
+				if (ClothConfigClass)
+				{
+					UClothSharedConfigCommon* ClothConfig = NewObject<UClothSharedConfigCommon>(NewClothingAssets[I], ClothConfigClass);
+				}
+			}
+
+			SkeletalMesh->GetMeshClothingAssets().AddUnique(NewClothingAssets[I]);	
+		}	
+	}
+
+	for ( FSectionWithClothData& SectionWithCloth : SectionsWithCloth )
+	{
+		FSkeletalMeshLODRenderData& LODModel = RenderResource->LODRenderData[SectionWithCloth.Lod];
+		FSkelMeshRenderSection& SectionData = LODModel.RenderSections[SectionWithCloth.Section];
+        
+		UClothingAssetCommon* NewClothingAsset = NewClothingAssets[SectionWithCloth.ClothAssetIndex];
+		if (!NewClothingAsset)
+		{
+			continue;
+		}
+       
+		SectionData.ClothMappingDataLODs.AddDefaulted(1);
+		SectionData.ClothMappingDataLODs[0] = MoveTemp(SectionWithCloth.MappingData);
+		
+		SectionData.CorrespondClothAssetIndex = SkeletalMesh->GetClothingAssetIndex(NewClothingAsset);
+		SectionData.ClothingData.AssetGuid = NewClothingAsset->GetAssetGuid();
+		SectionData.ClothingData.AssetLodIndex = SectionWithCloth.ClothAssetLodIndex;
+	}
+
+	SkeletalMesh->SetHasActiveClothingAssets( static_cast<bool>( SectionsWithCloth.Num() ) );
+}
+
+bool UCustomizableInstancePrivateData::BuildSkeletalMeshRenderData(const TSharedPtr<FMutableOperationData>& OperationData, USkeletalMesh* SkeletalMesh, UCustomizableObjectInstance* Public, int32 ComponentIndex)
+{
+	MUTABLE_CPUPROFILER_SCOPE(UCustomizableInstancePrivateData::BuildSkeletalMeshRenderData)
+
+	UCustomizableObject* CustomizableObject = Public->GetCustomizableObject();
+
+	SkeletalMesh->SetHasBeenSimplified(false);
+	SkeletalMesh->SetHasVertexColors(false);
+
+	int32 MeshLODIndex = HasCOInstanceFlags(ReduceLODs) ? 0 : CurrentMinLOD;
+
+	int32 LODCount = OperationData->InstanceUpdateData.LODs.Num();
+	for (int32 LODIndex = CurrentMinLOD; LODIndex <= CurrentMaxLOD && LODIndex < LODCount; ++LODIndex, ++MeshLODIndex)
+	{
+		MUTABLE_CPUPROFILER_SCOPE(LODLoop);
+
+		const FInstanceUpdateData::FLOD& LOD = OperationData->InstanceUpdateData.LODs[LODIndex];
+
+		// There could be components without a mesh in LODs
+		if (!LOD.ComponentCount)
+		{
+			UE_LOG(LogMutable, Warning, TEXT("Building instance: generated mesh [%s] has LOD [%d] with no component.")
+				, *SkeletalMesh->GetName()
+				, LODIndex);
+
+			// End with failure
+			return false;
+		}
+
+		FInstanceUpdateData::FComponent& Component = OperationData->InstanceUpdateData.Components[LOD.FirstComponent+ComponentIndex];
+
+		// There could be components without a mesh in LODs
+		if (!Component.Mesh)
+		{
+			UE_LOG(LogMutable, Warning, TEXT("Building instance: generated mesh [%s] has LOD [%d] of component [%d] with no mesh.")
+				, *SkeletalMesh->GetName()
+				, LODIndex
+				, Component.Id);
+
+			// End with failure
+			return false;
+		}
+
+		// Detect if the produced mesh is correct
+		bool bSuccess = true;
+
+		SetLastMeshId(Component.Id, LODIndex, Component.Mesh->GetId());
+
+		FSkeletalMeshLODRenderData& LODModel = Helper_GetLODData(SkeletalMesh)[MeshLODIndex];
+		const mu::FMeshBufferSet& MutableMeshVertexBuffers = Component.Mesh->GetVertexBuffers();
+
+		int32 NumVerticesLODModel = Component.Mesh->GetVertexCount();
+		int32 SurfaceCount = Component.Mesh->GetSurfaceCount();
+
+		mu::MESH_BUFFER_FORMAT boneIndexFormat = mu::MBF_NONE;
+		int boneIndexComponents = 0;
+		int boneIndexOffset = 0;
+		int boneIndexBuffer = -1;
+		int boneIndexChannel = -1;
+		MutableMeshVertexBuffers.FindChannel(mu::MBS_BONEINDICES, 0, &boneIndexBuffer, &boneIndexChannel);
+		if (boneIndexBuffer >= 0 || boneIndexChannel >= 0)
+		{
+			MutableMeshVertexBuffers.GetChannel(boneIndexBuffer, boneIndexChannel,
+				nullptr, nullptr, &boneIndexFormat, &boneIndexComponents, &boneIndexOffset);
+		}
+
+		int NumBoneInfluences = boneIndexComponents;
+
+		TArray<uint16>& BoneMap = Component.BoneMap;
+		UnrealConversionUtils::SetupRenderSections(
+			Component.Mesh,
+			SkeletalMesh,
+			MeshLODIndex,
+			NumBoneInfluences,
+			BoneMap);
+
+		UnrealConversionUtils::CopyMutableVertexBuffers(
+			SkeletalMesh,
+			NumVerticesLODModel,
+			NumBoneInfluences,
+			boneIndexBuffer,
+			MutableMeshVertexBuffers,
+			MeshLODIndex,
+			boneIndexFormat);
+
+
+		if (LODModel.DoesVertexBufferUse16BitBoneIndex() && !UCustomizableObjectSystem::GetInstance()->IsSupport16BitBoneIndexEnabled())
+		{
+			Public->SkeletalMeshStatus = ESkeletalMeshState::PostUpdateError;
+
+			const FString Msg = FString::Printf(TEXT("Customizable Object [%s] requires of Skinning - 'Support 16 Bit Bone Index' to be enabled. Please, update the Project Settings."),
+				*CustomizableObject->GetName());
+			UE_LOG(LogMutable, Error, TEXT("%s"), *Msg);
+
+#if WITH_EDITOR
+			FNotificationInfo Info(FText::FromString(Msg));
+			Info.bFireAndForget = true;
+			Info.FadeOutDuration = 1.0f;
+			Info.ExpireDuration = 10.0f;
+			FSlateNotificationManager::Get().AddNotification(Info);
+#endif
+		}
+
+
+		// Update active and required bones
+		{
+			MUTABLE_CPUPROFILER_SCOPE(UpdateSkeletalMesh_SurfaceLoop_UpdateActiveAndRequiredBones);
+			LODModel.ActiveBoneIndices.Append(Component.ActiveBones);
+			LODModel.RequiredBones.Append(Component.ActiveBones);
+		}
+
+		if (bSuccess)
+		{
+			bSuccess = UnrealConversionUtils::CopyMutableIndexBuffers(Component.Mesh,LODModel);
+		}
+
+		// Update LOD and streaming data
+		const FMutableRefLODRenderData& RefLODRenderData = CustomizableObject->GetRefSkeletalMeshData(Component.Id)->LODData[LODIndex].RenderData;
+		LODModel.bIsLODOptional = RefLODRenderData.bIsLODOptional;
+		LODModel.bStreamedDataInlined = RefLODRenderData.bStreamedDataInlined;
+
+		if (!bSuccess)
+		{
+			// End with failure
+			return false;
+		}
+
+	}
+	
+
+	return true;
+}
+
+FGraphEventRef UCustomizableInstancePrivateData::LoadAdditionalAssetsAsync(const TSharedPtr<FMutableOperationData>& OperationData, UCustomizableObjectInstance* Public, FStreamableManager &StreamableManager)
+{
+	UCustomizableObject* CustomizableObject = Public->GetCustomizableObject();
+
+	FGraphEventRef Result = nullptr;
+
+	TArray<FSoftObjectPath> AssetsToStream;
+
+	TArray<FInstanceUpdateData::FLOD>& LODs = OperationData->InstanceUpdateData.LODs;
+	TArray<FInstanceUpdateData::FComponent>& Components = OperationData->InstanceUpdateData.Components;
+
+	ObjectToInstanceIndexMap.Empty();
+	ReferencedMaterials.Empty();
+
+	const int32 NumClothingAssets = CustomizableObject->ContributingClothingAssetsData.Num();
+	ClothingPhysicsAssets.Reset(NumClothingAssets);
+	ClothingPhysicsAssets.SetNum(NumClothingAssets);
+
+	GatheredAnimBPs.Empty();
+	AnimBPGameplayTags.Reset();
+
+	
+	for (const FInstanceUpdateData::FSurface& Surface : OperationData->InstanceUpdateData.Surfaces)
+	{
+		uint32 MaterialIndex = Surface.MaterialIndex;
+
+		if (!ObjectToInstanceIndexMap.Contains(MaterialIndex))
+		{
+			UMaterialInterface* MatInterface = CustomizableObject->GetReferencedMaterialAssetPtr(MaterialIndex).Get();
+
+			ObjectToInstanceIndexMap.Add(MaterialIndex, ReferencedMaterials.Num());
+			ReferencedMaterials.Add(MatInterface);
+
+			check(ObjectToInstanceIndexMap.Num() == ReferencedMaterials.Num());
+
+			if (!MatInterface)
+			{
+				FSoftObjectPath MaterialPath = CustomizableObject->GetReferencedMaterialAssetPtr(MaterialIndex).ToSoftObjectPath();
+
+				if (MaterialPath.IsValid())
+				{
+					AssetsToStream.Add(MaterialPath);
+				}
+			}
+		}
+	}
+
+	const int32 ComponentCount = OperationData->InstanceUpdateData.Skeletons.Num();
+
+	// Load Skeletons required by the SubMeshes of the newly generated Mesh, will be merged later
+	for (FCustomizableInstanceComponentData& ComponentData : ComponentsData)
+	{
+		const uint16 ComponentIndex = ComponentData.ComponentIndex;
+
+		FInstanceUpdateData::FSkeletonData* SkeletonData = OperationData->InstanceUpdateData.Skeletons.FindByPredicate(
+			[&ComponentIndex](FInstanceUpdateData::FSkeletonData& S) { return S.ComponentIndex == ComponentIndex; });
+		check(SkeletonData);
+
+		for (const uint32 SkeletonId : SkeletonData->SkeletonIds)
+		{
+			TSoftObjectPtr<USkeleton> AssetPtr;
+			if (SkeletonId == 0)
+			{
+				FMutableRefSkeletalMeshData* RefSkeletalMeshData = CustomizableObject->GetRefSkeletalMeshData(ComponentIndex);
+				check(RefSkeletalMeshData);
+
+				AssetPtr = RefSkeletalMeshData->Skeleton;
+			}
+			else
+			{
+				AssetPtr = CustomizableObject->GetReferencedSkeletonAssetPtr(SkeletonId);
+			}
+
+			if (AssetPtr.IsNull())
+			{
+				continue;
+			}
+
+			if (USkeleton* Skeleton = AssetPtr.Get())
+			{
+				ComponentData.Skeletons.SkeletonsToMerge.Add(Skeleton);
+			}
+			else
+			{
+				// Add referenced skeletons to the assets to stream
+				ComponentData.Skeletons.SkeletonsToLoad.Add(SkeletonId);
+				AssetsToStream.Add(AssetPtr.ToSoftObjectPath());
+			}
+
+		}
+	}
+
+
+	// Load assets coming from SubMeshes of the newly generated Mesh
+	if (OperationData->InstanceUpdateData.LODs.Num())
+	{
+		const FInstanceUpdateData::FLOD& LOD = OperationData->InstanceUpdateData.LODs[CurrentMinLOD];
+		for (int32 ComponentIndex = 0; ComponentIndex < LOD.ComponentCount; ++ComponentIndex)
+		{
+			const FInstanceUpdateData::FComponent& Component = Components[LOD.FirstComponent+ComponentIndex];
+			mu::MeshPtrConst MutableMesh = Component.Mesh;
+
+			if (!MutableMesh)
+			{
+				continue;
+			}
+				
+			FCustomizableInstanceComponentData* ComponentData =	GetComponentData(Component.Id);
+
+			const bool bReplacePhysicsAssets = HasCOInstanceFlags(ReplacePhysicsAssets);
+
+			for (int32 TagIndex = 0; TagIndex < MutableMesh->GetTagCount(); ++TagIndex)
+			{
+				FString Tag = MutableMesh->GetTag(TagIndex);
+				if (bReplacePhysicsAssets && Tag.RemoveFromStart("__PhysicsAsset:"))
+				{
+					TSoftObjectPtr<UPhysicsAsset>* PhysicsAsset = CustomizableObject->PhysicsAssetsMap.Find(Tag);
+
+					if (!PhysicsAsset->IsNull())
+					{
+						if (PhysicsAsset->Get())
+						{
+							ComponentData->PhysicsAssets.PhysicsAssetsToMerge.Add(PhysicsAsset->Get());
+						}
+						else
+						{
+							ComponentData->PhysicsAssets.PhysicsAssetToLoad.Add(PhysicsAsset->ToSoftObjectPath().ToString());
+							AssetsToStream.Add(PhysicsAsset->ToSoftObjectPath());
+						}
+					}
+				}
+				else if (Tag.RemoveFromStart("__ClothPhysicsAsset:"))
+				{
+					FString AssetIndexString, AssetPath;
+
+					if (Tag.Split(TEXT("_AssetIdx_"), &AssetIndexString, &AssetPath) && AssetIndexString.IsNumeric())
+					{
+						
+						int32 AssetIndex = FCString::Atoi(*AssetIndexString);
+
+						TSoftObjectPtr<UPhysicsAsset>* PhysicsAssetPtr = CustomizableObject->PhysicsAssetsMap.Find(AssetPath);
+
+						// The entry should always be in the map
+						check(PhysicsAssetPtr);
+						if (!PhysicsAssetPtr->IsNull())
+						{
+							if (PhysicsAssetPtr->Get())
+							{
+								if( ClothingPhysicsAssets.IsValidIndex(AssetIndex) )
+								{
+									ClothingPhysicsAssets[AssetIndex] = PhysicsAssetPtr->Get();
+								}
+							}
+							else
+							{
+								ComponentData->ClothingPhysicsAssetsToStream.Emplace(AssetIndex, PhysicsAssetPtr->ToSoftObjectPath().ToString());
+								AssetsToStream.Add(PhysicsAssetPtr->ToSoftObjectPath());
+							}
+						}
+					}
+				}
+				if (Tag.RemoveFromStart("__AnimBP:"))
+				{
+					FString SlotIndexString, AssetPath;
+
+					if (Tag.Split(TEXT("_Slot_"), &SlotIndexString, &AssetPath))
+					{
+						if (SlotIndexString.IsNumeric())
+						{
+							int32 SlotIndex = FCString::Atoi(*SlotIndexString);
+
+							TSoftClassPtr<UAnimInstance>* AnimBPAsset = CustomizableObject->AnimBPAssetsMap.Find(AssetPath);
+
+							if (AnimBPAsset && !AnimBPAsset->IsNull())
+							{
+								if (!ComponentData->AnimSlotToBP.Contains(SlotIndex))
+								{
+									ComponentData->AnimSlotToBP.Add(SlotIndex, *AnimBPAsset);
+
+									if (AnimBPAsset->Get())
+									{
+										GatheredAnimBPs.Add(AnimBPAsset->Get());
+									}
+									else
+									{
+										AssetsToStream.Add(AnimBPAsset->ToSoftObjectPath());
+									}
+								}
+								else
+								{
+									// Two submeshes should not have the same animation slot index
+									FString ErrorMsg = FString::Printf(TEXT("Two submeshes have the same anim slot index [%d] in a Mutable Instance."), SlotIndex);
+									UE_LOG(LogMutable, Error, TEXT("%s"), *ErrorMsg);
+#if WITH_EDITOR
+									FMessageLogModule& MessageLogModule = FModuleManager::LoadModuleChecked<FMessageLogModule>("MessageLog");
+									MessageLogModule.RegisterLogListing(FName("Mutable"), FText::FromString(FString("Mutable")));
+									FMessageLog MessageLog("Mutable");
+
+									MessageLog.Notify(FText::FromString(ErrorMsg), EMessageSeverity::Error, true);
+#endif
+								}
+							}
+						}
+					}
+				}
+				else if (Tag.RemoveFromStart("__AnimBPTag:"))
+				{
+					AnimBPGameplayTags.AddTag(FGameplayTag::RequestGameplayTag(*Tag));
+				}
+			}
+		}
+	}
+
+	if (AssetsToStream.Num() > 0)
+	{
+		ClearCOInstanceFlags(AssetsLoaded);
+
+		check(!StreamingHandle);
+
+		Result = FGraphEvent::CreateGraphEvent();
+		StreamingHandle = StreamableManager.RequestAsyncLoad(AssetsToStream, FStreamableDelegate::CreateUObject(Public, &UCustomizableObjectInstance::AdditionalAssetsAsyncLoaded, Result));
+	}
+	else
+	{
+		SetCOInstanceFlags(AssetsLoaded);
+	}
+
+	return Result;
+}
+
+void UCustomizableObjectInstance::AdditionalAssetsAsyncLoaded( FGraphEventRef CompletionEvent )
+{
+	// TODO: Do we need this separated?
+	PrivateData->AdditionalAssetsAsyncLoaded(this);
+	CompletionEvent->DispatchSubsequents(); // TODO: we know it is game thread?
+
+	PrivateData->StreamingHandle = nullptr;
+}
+
+
+void UCustomizableInstancePrivateData::AdditionalAssetsAsyncLoaded(UCustomizableObjectInstance* Public)
+{
+	UCustomizableObject* CustomizableObject = Public->GetCustomizableObject();
+
+	SetCOInstanceFlags(AssetsLoaded);
+
+	// Loaded Materials
+	check(ObjectToInstanceIndexMap.Num() == ReferencedMaterials.Num());
+
+	for (TPair<uint32, uint32> Pair : ObjectToInstanceIndexMap)
+	{
+		ReferencedMaterials[Pair.Value] = CustomizableObject->GetReferencedMaterialAssetPtr(Pair.Key).Get();
+	}
+
+#if WITH_EDITOR
+	for (int32 i = 0; i < ReferencedMaterials.Num(); ++i)
+	{
+		if (!ReferencedMaterials[i])
+		{
+			const uint32* Key = ObjectToInstanceIndexMap.FindKey(i);
+			TSoftObjectPtr<UMaterialInterface> ErrorMaterial = Key != nullptr ? CustomizableObject->GetReferencedMaterialAssetPtr(*Key) : nullptr;
+
+			if (!ErrorMaterial.IsNull())
+			{
+				FMessageLogModule& MessageLogModule = FModuleManager::LoadModuleChecked<FMessageLogModule>("MessageLog");
+				MessageLogModule.RegisterLogListing(FName("Mutable"), FText::FromString(FString("Mutable")));
+				FMessageLog MessageLog("Mutable");
+
+				FString ErrorMsg = FString::Printf(TEXT("Mutable couldn't load the material [%s] and won't be rendered. If it has been deleted or renamed, please recompile all the mutable objects that use it."), *ErrorMaterial.GetAssetName());
+				UE_LOG(LogMutable, Error, TEXT("%s"), *ErrorMsg);
+				MessageLog.Notify(FText::FromString(ErrorMsg), EMessageSeverity::Error, true);
+			}
+			else
+			{
+				ensure(false); // Couldn't load the material, and we don't know which material
+			}
+		}
+	}
+#endif
+
+	
+	for (FCustomizableInstanceComponentData& ComponentData : ComponentsData)
+	{
+		// Loaded Skeletons
+		FReferencedSkeletons& Skeletons = ComponentData.Skeletons;
+		for (const int32& SkeletonIndex : Skeletons.SkeletonsToLoad)
+		{
+			Skeletons.SkeletonsToMerge.Add(CustomizableObject->GetReferencedSkeletonAssetPtr(SkeletonIndex).Get());
+		}
+		Skeletons.SkeletonsToLoad.Empty();
+
+		// Loaded PhysicsAssets
+		FReferencedPhysicsAssets& PhysicsAssets = ComponentData.PhysicsAssets;
+		for(const FString& Path : PhysicsAssets.PhysicsAssetToLoad)
+		{
+			const TSoftObjectPtr<UPhysicsAsset>* PhysicsAssetPtr = CustomizableObject->PhysicsAssetsMap.Find(Path);
+			PhysicsAssets.PhysicsAssetsToMerge.Add(PhysicsAssetPtr->Get());
+
+#if WITH_EDITOR
+			if (!PhysicsAssetPtr->Get())
+			{
+				if (!PhysicsAssetPtr->IsNull())
+				{
+					FMessageLogModule& MessageLogModule = FModuleManager::LoadModuleChecked<FMessageLogModule>("MessageLog");
+					MessageLogModule.RegisterLogListing(FName("Mutable"), FText::FromString(FString("Mutable")));
+					FMessageLog MessageLog("Mutable");
+
+					FString ErrorMsg = FString::Printf(TEXT("Mutable couldn't load the PhysicsAsset [%s] and won't be merged. If it has been deleted or renamed, please recompile all the mutable objects that use it."), *PhysicsAssetPtr->GetAssetName());
+					UE_LOG(LogMutable, Error, TEXT("%s"), *ErrorMsg);
+					MessageLog.Notify(FText::FromString(ErrorMsg), EMessageSeverity::Error, true);
+				}
+				else
+				{
+					ensure(false); // Couldn't load the PhysicsAsset, and we don't know which PhysicsAsset
+				}
+			}
+#endif
+		}
+		PhysicsAssets.PhysicsAssetToLoad.Empty();
+		
+		// Loaded Clothing PhysicsAssets 
+		for ( TPair<int32, FString>& AssetToStream : ComponentData.ClothingPhysicsAssetsToStream )
+		{
+			const int32 AssetIndex = AssetToStream.Key;
+
+			if (ClothingPhysicsAssets.IsValidIndex(AssetIndex))
+			{
+				const TSoftObjectPtr<UPhysicsAsset>* PhysicsAssetPtr = CustomizableObject->PhysicsAssetsMap.Find(AssetToStream.Value);
+				check(PhysicsAssetPtr) // Should always be found.
+			
+				ClothingPhysicsAssets[AssetIndex] = PhysicsAssetPtr->Get();
+			}
+		}
+		ComponentData.ClothingPhysicsAssetsToStream.Empty();
+
+		// Loaded anim BPs
+		for (TPair<int32, TSoftClassPtr<UAnimInstance>>& SlotAnimBP : ComponentData.AnimSlotToBP)
+		{
+			if (TSubclassOf<UAnimInstance> AnimBP = SlotAnimBP.Value.Get())
+			{
+				if (!GatheredAnimBPs.Contains(AnimBP))
+				{
+					GatheredAnimBPs.Add(AnimBP);
+				}
+			}
+#if WITH_EDITOR
+			else
+			{
+				FMessageLogModule& MessageLogModule = FModuleManager::LoadModuleChecked<FMessageLogModule>("MessageLog");
+				MessageLogModule.RegisterLogListing(FName("Mutable"), FText::FromString(FString("Mutable")));
+				FMessageLog MessageLog("Mutable");
+
+				FString ErrorMsg = FString::Printf(TEXT("Mutable couldn't load the AnimBlueprint [%s]. If it has been deleted or renamed, please recompile all the mutable objects that use it."), *SlotAnimBP.Value.GetAssetName());
+				UE_LOG(LogMutable, Error, TEXT("%s"), *ErrorMsg);
+				MessageLog.Notify(FText::FromString(ErrorMsg), EMessageSeverity::Error, true);
+			}
+#endif
+		}
+	}
+}
+
+
+void UpdateTextureRegionsMutable(UTexture2D* Texture, int32 MipIndex, const FUpdateTextureRegion2D& Region, uint32 SrcPitch, FByteBulkData* BulkData)
+{
+	if (Texture->GetResource())
+	{
+		struct FUpdateTextureRegionsData
+		{
+			FTexture2DResource* Texture2DResource;
+			int32 MipIndex;
+			FUpdateTextureRegion2D Region;
+			uint32 SrcPitch;
+		};
+
+		FUpdateTextureRegionsData* RegionData = new FUpdateTextureRegionsData;
+
+		RegionData->Texture2DResource = (FTexture2DResource*)Texture->GetResource();
+		RegionData->MipIndex = MipIndex;
+		RegionData->Region = Region;
+		RegionData->SrcPitch = SrcPitch;
+
+
+#if (ENGINE_MAJOR_VERSION == 4 && ENGINE_MINOR_VERSION > 22) || ENGINE_MAJOR_VERSION >= 5
+		ENQUEUE_RENDER_COMMAND(UpdateTextureRegionsMutable)(
+			[RegionData, BulkData](FRHICommandList& CmdList)
+#else
+		ENQUEUE_UNIQUE_RENDER_COMMAND_TWOPARAMETER(
+			UpdateTextureRegionsMutable,
+			FUpdateTextureRegionsData*, RegionData, RegionData,
+			FByteBulkData*, BulkData, BulkData,
+#endif
+			{
+				int32 CurrentFirstMip = RegionData->Texture2DResource->GetCurrentFirstMip();
+				uint8* SrcData = (uint8*)BulkData->LockReadOnly();
+
+				if (RegionData->MipIndex >= CurrentFirstMip)
+				{
+					RHIUpdateTexture2D(
+						RegionData->Texture2DResource->GetTexture2DRHI(),
+						RegionData->MipIndex - CurrentFirstMip,
+						RegionData->Region,
+						RegionData->SrcPitch,
+						SrcData);
+				}
+
+				BulkData->Unlock();
+				delete RegionData;
+			});
+	}
+}
+
+void UCustomizableInstancePrivateData::ReuseTexture(UTexture2D* Texture)
+{
+	if (Texture->GetPlatformData())
+	{
+		for (uint32 i = 0; i < (uint32)Texture->GetPlatformData()->Mips.Num(); i++)
+		{
+			FTexture2DMipMap& Mip = Texture->GetPlatformData()->Mips[i];
+
+			if (Mip.BulkData.GetElementCount() > 0)
+			{
+				FUpdateTextureRegion2D Region;
+
+				Region.DestX = 0;
+				Region.DestY = 0;
+				Region.SrcX = 0;
+				Region.SrcY = 0;
+				Region.Width = Mip.SizeX;
+				Region.Height = Mip.SizeY;
+
+				UpdateTextureRegionsMutable(Texture, i, Region, Mip.SizeX * sizeof(uint8) * 4, &Mip.BulkData);
+			}
+		}
+	}
+}
+
+
+void UCustomizableInstancePrivateData::BuildMaterials(const TSharedPtr<FMutableOperationData>& OperationData, UCustomizableObjectInstance* Public )
+{
+	MUTABLE_CPUPROFILER_SCOPE(UCustomizableInstancePrivateData::BuildMaterials)
+
+	UCustomizableObject* CustomizableObject = Public->GetCustomizableObject();
+
+	const int32 LODCount = OperationData->InstanceUpdateData.LODs.Num();
+
+	TArray<FGeneratedTexture> NewGeneratedTextures;
+
+	GeneratedMaterials.Reset();
+
+	// Prepare the data to store in order to regenerate resources for this instance (usually texture mips).
+	TSharedPtr<FMutableUpdateContext> UpdateContext = MakeShared<FMutableUpdateContext>();
+	UpdateContext->System = UCustomizableObjectSystem::GetInstance()->GetPrivate()->MutableSystem;
+	UpdateContext->Model = CustomizableObject->GetModel();
+	UpdateContext->Parameters = OperationData->MutableParameters;
+	UpdateContext->State = OperationData->State;
+
+	const bool bUseCurrentMinLODAsBaseLOD = HasCOInstanceFlags(ReduceLODs);
+	const bool bReuseTextures = HasCOInstanceFlags(ReuseTextures);
+
+	const FInstanceUpdateData::FLOD& FirstLOD = OperationData->InstanceUpdateData.LODs[CurrentMinLOD];
+	const int32 NumComponents = FirstLOD.ComponentCount;
+
+	for (int32 ComponentIndex = 0; ComponentIndex < NumComponents; ++ComponentIndex)
+	{
+		USkeletalMesh* SkeletalMesh = Public->SkeletalMeshes[ComponentIndex];
+
+		{
+			MUTABLE_CPUPROFILER_SCOPE(Init);
+
+			if (SkeletalMesh)
+			{
+				SkeletalMesh->GetMaterials().Reset();
+			}
+		}
+
+		// Maps serializations of FMutableMaterialPlaceholder to Created Dynamic Material instances, used to reuse materials across LODs
+		TMap<FString, FMutableMaterialPlaceholder*> ReuseMaterialCache;
+		// Stores the FMutableMaterialPlaceholders and keeps them in memory while the component's LODs are being processed
+		TArray<TSharedPtr<FMutableMaterialPlaceholder>> MutableMaterialPlaceholderArray;
+
+		MUTABLE_CPUPROFILER_SCOPE(LODLoop);
+
+		int32 MeshLODIndex = bUseCurrentMinLODAsBaseLOD ? 0 : CurrentMinLOD;
+
+		for (int32 LODIndex = CurrentMinLOD; LODIndex <= CurrentMaxLOD && LODIndex < LODCount; ++LODIndex, ++MeshLODIndex)
+		{
+			const FInstanceUpdateData::FLOD& LOD = OperationData->InstanceUpdateData.LODs[LODIndex];
+
+			if (SkeletalMesh && Helper_GetLODInfoArray(SkeletalMesh).Num() != 0)
+			{
+				Helper_GetLODInfoArray(SkeletalMesh)[MeshLODIndex].LODMaterialMap.Reset();
+			}
+
+			if (ComponentIndex<LOD.ComponentCount)
+			{
+				const FInstanceUpdateData::FComponent& Component = OperationData->InstanceUpdateData.Components[LOD.FirstComponent+ComponentIndex];
+				auto MutableMesh = Component.Mesh;
+
+				if (!MutableMesh)
+				{
+					// This case is dangerous. If the LOD has no mesh at all it will crash unreal.
+					// In theory this is checked earlier, and the LOD is removed.
+					continue;
+				}
+
+				for (int32 SurfaceIndex = 0; SurfaceIndex < Component.SurfaceCount; ++SurfaceIndex)
+				{
+					const FInstanceUpdateData::FSurface& Surface = OperationData->InstanceUpdateData.Surfaces[Component.FirstSurface + SurfaceIndex];
+
+					MutableMaterialPlaceholderArray.Add(TSharedPtr<FMutableMaterialPlaceholder>(new FMutableMaterialPlaceholder));
+					FMutableMaterialPlaceholder& MutableMaterialPlaceholder = *MutableMaterialPlaceholderArray.Last();
+
+					{
+						MUTABLE_CPUPROFILER_SCOPE(ParamLoop);
+
+						for (int32 VectorIndex = 0; VectorIndex < Surface.VectorCount; ++VectorIndex)
+						{
+							const FInstanceUpdateData::FVector& Vector = OperationData->InstanceUpdateData.Vectors[Surface.FirstVector + VectorIndex];
+
+							// Decoding Material Layer from Mutable parameter name
+							FString EncodingString = "-MutableLayerParam:";
+
+							int32 EncodingPosition = Vector.Name.Find(EncodingString);
+							int32 LayerIndex = -1;
+							
+							if (EncodingPosition == INDEX_NONE)
+							{
+								MutableMaterialPlaceholder.AddParam(
+									FMutableMaterialPlaceholder::FMutableMaterialPlaceHolderParam(FName(Vector.Name), -1, Vector.Vector));
+							}
+							else
+							{
+								//Getting layer index
+								int32 LayerPosition = Vector.Name.Len() - (EncodingPosition + EncodingString.Len());
+								FString IndexString = Vector.Name.RightChop(Vector.Name.Len() - LayerPosition);
+								LayerIndex = FCString::Atof(*IndexString);
+
+								//Getting parameter name
+								FString Sufix = EncodingString + FString::FromInt(LayerIndex);
+								FString NameCopy = Vector.Name;
+								NameCopy.RemoveFromEnd(Sufix);
+									
+								MutableMaterialPlaceholder.AddParam(
+									FMutableMaterialPlaceholder::FMutableMaterialPlaceHolderParam(FName(NameCopy), LayerIndex, Vector.Vector));
+							}
+						}
+
+						for (int32 ScalarIndex = 0; ScalarIndex < Surface.ScalarCount; ++ScalarIndex)
+						{
+							const FInstanceUpdateData::FScalar& Scalar = OperationData->InstanceUpdateData.Scalars[Surface.FirstScalar + ScalarIndex];
+							
+							// Decoding Material Layer from Mutable parameter name
+							FString EncodingString = "-MutableLayerParam:";
+
+							int32 EncodingPosition = Scalar.Name.Find(EncodingString);
+							int32 LayerIndex = -1;
+
+							if (EncodingPosition == INDEX_NONE)
+							{
+								MutableMaterialPlaceholder.AddParam(
+									FMutableMaterialPlaceholder::FMutableMaterialPlaceHolderParam(FName(Scalar.Name), -1, Scalar.Scalar));
+							}
+							else
+							{
+								//Getting layer index
+								int32 LayerPosition = Scalar.Name.Len() - (EncodingPosition + EncodingString.Len());
+								FString IndexString = Scalar.Name.RightChop(Scalar.Name.Len() - LayerPosition);
+								LayerIndex = FCString::Atof(*IndexString);
+
+								//Getting parameter name
+								FString Sufix = EncodingString + FString::FromInt(LayerIndex);
+								FString NameCopy = Scalar.Name;
+								NameCopy.RemoveFromEnd(Sufix);
+
+								MutableMaterialPlaceholder.AddParam(
+									FMutableMaterialPlaceholder::FMutableMaterialPlaceHolderParam(FName(NameCopy), LayerIndex, Scalar.Scalar));
+							}
+						}
+					}
+
+					{
+						MUTABLE_CPUPROFILER_SCOPE(ImageLoop);
+
+						// Get the cache of resources of all live instances of this object
+						FMutableResourceCache& Cache = UCustomizableObjectSystem::GetInstance()->GetPrivate()->GetObjectCache(CustomizableObject);
+
+						FString CurrentState = Public->GetCurrentState();
+						bool bNeverStream = OperationData->bNeverStream;
+
+						check((bNeverStream && OperationData->MipsToSkip == 0) ||
+							 (!bNeverStream && OperationData->MipsToSkip >= 0));
+
+						for (int32 ImageIndex = 0; ImageIndex < Surface.ImageCount; ++ImageIndex)
+						{
+							const FInstanceUpdateData::FImage& Image = OperationData->InstanceUpdateData.Images[Surface.FirstImage+ImageIndex];
+							FString KeyName = Image.Name;
+							mu::ImagePtrConst MutableImage = Image.Image;
+
+							UTexture2D* Texture = nullptr;
+
+							// \TODO: Change this key to a struct.
+							FString TextureReuseCacheRef = bReuseTextures ? FString::Printf(TEXT("%d-%d-%d-%d"), LODIndex, ComponentIndex, SurfaceIndex, ImageIndex) : FString();
+
+							// If the mutable image is null, it must be in the cache
+							FMutableImageCacheKey ImageCacheKey = { Image.ImageID };//, OperationData->MipsToSkip }; // Todo: Reenable extended FMutableImageCacheKey
+							if (!MutableImage)
+							{
+								TWeakObjectPtr<UTexture2D>* CachedPointerPtr = Cache.Images.Find(ImageCacheKey);
+								if (CachedPointerPtr)
+								{
+									ensure(!CachedPointerPtr->IsStale());
+									Texture = CachedPointerPtr->Get();
+								}
+
+								check(Texture);
+							}
+
+							// Find the additional information for this image
+							int32 ImageKey = FCString::Atoi(*KeyName);
+							if (ImageKey >= 0 && ImageKey < CustomizableObject->ImageProperties.Num())
+							{
+								const FMutableModelImageProperties& Props = CustomizableObject->ImageProperties[ImageKey];
+
+								if (MutableImage)
+								{
+									TWeakObjectPtr<UTexture2D>* ReusedTexture = bReuseTextures ? TextureReuseCache.Find(TextureReuseCacheRef) : nullptr;
+
+									if (ReusedTexture && (*ReusedTexture).IsValid() && !(*ReusedTexture)->HasAnyFlags(RF_BeginDestroyed))
+									{
+										Texture = (*ReusedTexture).Get();
+										check(Texture != nullptr);
+									}
+									else
+									{
+										ReusedTexture = nullptr;
+										Texture = CreateTexture();
+									}
+
+									if (Texture)
+									{
+										if (OperationData->ImageToPlatformDataMap.Contains(Image.ImageID))
+										{
+											SetTexturePropertiesFromMutableImageProps(Texture, Props, bNeverStream);
+
+											FTexturePlatformData* PlatformData = OperationData->ImageToPlatformDataMap[Image.ImageID];
+
+											if (ReusedTexture)
+											{
+												check(PlatformData->Mips.Num() == Texture->GetPlatformData()->Mips.Num());
+												check(PlatformData->Mips[0].SizeX == Texture->GetPlatformData()->Mips[0].SizeX);
+												check(PlatformData->Mips[0].SizeY == Texture->GetPlatformData()->Mips[0].SizeY);
+											}
+
+											if (Texture->GetPlatformData())
+											{
+												delete Texture->GetPlatformData();
+											}
+
+											Texture->SetPlatformData( PlatformData );
+											OperationData->ImageToPlatformDataMap.Remove(Image.ImageID);
+										}
+										else
+										{
+											UE_LOG(LogMutable, Error, TEXT("Required image was not generated in the mutable thread, and it is not cached."));
+										}
+
+										if (bNeverStream)
+										{
+											// To prevent LogTexture Error "Loading non-streamed mips from an external bulk file."
+											for (int32 i = 0; i < Texture->GetPlatformData()->Mips.Num(); ++i)
+											{
+												Texture->GetPlatformData()->Mips[i].BulkData.ClearBulkDataFlags(BULKDATA_PayloadInSeperateFile);
+											}
+										}
+
+										{
+											MUTABLE_CPUPROFILER_SCOPE(UpdateResource);
+#if !PLATFORM_DESKTOP && !PLATFORM_SWITCH // Switch does this in FTexture2DResource::InitRHI()
+											for (int32 i = 0; i < Texture->GetPlatformData()->Mips.Num(); ++i)
+											{
+												uint32 DataFlags = Texture->GetPlatformData()->Mips[i].BulkData.GetBulkDataFlags();
+												Texture->GetPlatformData()->Mips[i].BulkData.SetBulkDataFlags(DataFlags | BULKDATA_SingleUse);
+											}
+#endif
+
+											if (ReusedTexture)
+											{
+												// Must remove texture from cache since it will be reused with a different ImageID
+												for (TPair<FMutableImageCacheKey, TWeakObjectPtr<UTexture2D>>& CachedTexture : Cache.Images)
+												{
+													if (CachedTexture.Value == Texture)
+													{
+														Cache.Images.Remove(CachedTexture.Key);
+														break;
+													}
+												}
+
+												ReuseTexture(Texture);
+											}
+											else if (Texture)
+											{
+												//if (!bNeverStream) // No need to check bNeverStream. In that case, the texture won't use 
+												// the MutableMipDataProviderFactory anyway and it's needed for detecting Mutable textures elsewhere
+												{
+													UMutableTextureMipDataProviderFactory* MutableMipDataProviderFactory = Cast<UMutableTextureMipDataProviderFactory>(Texture->GetAssetUserDataOfClass(UMutableTextureMipDataProviderFactory::StaticClass()));
+													if (!MutableMipDataProviderFactory)
+													{
+														MutableMipDataProviderFactory = NewObject<UMutableTextureMipDataProviderFactory>();
+
+														if (MutableMipDataProviderFactory)
+														{
+															MutableMipDataProviderFactory->CustomizableObjectInstance = Public;
+															check(LODIndex<256 && ComponentIndex<256 && SurfaceIndex<256 && ImageIndex<256 );
+															MutableMipDataProviderFactory->ImageRef = { Image.ImageID, uint8(LODIndex), uint8(ComponentIndex), uint8(SurfaceIndex), uint8(ImageIndex) };
+															MutableMipDataProviderFactory->UpdateContext = UpdateContext;
+															Texture->AddAssetUserData(MutableMipDataProviderFactory);
+														}
+													}
+												}
+
+												Texture->UpdateResource();
+											}
+										}
+
+										// update the model resources cache
+										Cache.Images.Add(ImageCacheKey, Texture);
+									}
+									else
+									{
+										UE_LOG(LogMutable, Error, TEXT("Texture creation failed."));
+									}
+								}
+
+								FGeneratedTexture TextureData;
+								TextureData.Id = Image.ImageID;
+								TextureData.Name = Props.TextureParameterName;
+								TextureData.Texture = Texture;
+								NewGeneratedTextures.Add(TextureData);
+
+								// Decoding Material Layer from Mutable parameter name
+								FString ImageName = Image.Name;
+								FString EncodingString = "-MutableLayerParam:";
+
+								int32 EncodingPosition = ImageName.Find(EncodingString);
+								int32 LayerIndex = -1;
+
+								if (EncodingPosition == INDEX_NONE)
+								{
+									MutableMaterialPlaceholder.AddParam(
+										FMutableMaterialPlaceholder::FMutableMaterialPlaceHolderParam(*Props.TextureParameterName, -1, TextureData));
+								}
+								else
+								{
+									//Getting layer index
+									int32 LayerPosition = ImageName.Len() - (EncodingPosition + EncodingString.Len());
+									FString IndexString = ImageName.RightChop(ImageName.Len() - LayerPosition);
+									LayerIndex = FCString::Atof(*IndexString);
+
+									MutableMaterialPlaceholder.AddParam(
+										FMutableMaterialPlaceholder::FMutableMaterialPlaceHolderParam(*Props.TextureParameterName, LayerIndex, TextureData));
+								}
+							}
+							else
+							{
+								// This means the compiled model (maybe coming from derived data) has images that the asset doesn't know about.
+								UE_LOG(LogMutable, Error, TEXT("CustomizableObject derived data out of sync with asset for [%s]. Try recompiling it."), *CustomizableObject->GetName());
+							}
+
+							if (bReuseTextures)
+							{
+								if (Texture)
+								{
+									TextureReuseCache.Add(TextureReuseCacheRef, Texture);
+								}
+								else
+								{
+									TextureReuseCache.Remove(TextureReuseCacheRef);
+								}
+							}
+						}
+					}
+
+					const uint32 ReferencedMaterialIndex = ObjectToInstanceIndexMap[Surface.MaterialIndex];
+					UMaterialInterface* MaterialTemplate = ReferencedMaterials[ReferencedMaterialIndex];
+					MutableMaterialPlaceholder.ParentMaterial = MaterialTemplate;
+
+					FString MaterialPlaceholderSerialization = MutableMaterialPlaceholder.GetSerialization();
+					FMutableMaterialPlaceholder** FoundMaterialPlaceholder = ReuseMaterialCache.Find(MaterialPlaceholderSerialization);
+
+					if (FoundMaterialPlaceholder)
+					{
+						int32 MatIndex = (*FoundMaterialPlaceholder)->MatIndex;
+						check(MatIndex >= 0);
+						int32 LODMaterialIndex = Helper_GetLODInfoArray(SkeletalMesh)[MeshLODIndex].LODMaterialMap.Add(MatIndex);
+						Helper_GetLODRenderSections(SkeletalMesh, MeshLODIndex)[SurfaceIndex].MaterialIndex = LODMaterialIndex;
+					}
+					else
+					{
+						ReuseMaterialCache.Add(MaterialPlaceholderSerialization, &MutableMaterialPlaceholder);
+
+						FGeneratedMaterial Material;
+						UMaterialInstanceDynamic* MaterialInstance = nullptr;
+
+						if (MaterialTemplate)
+						{
+							MUTABLE_CPUPROFILER_SCOPE(CreateMaterial);
+
+							UMaterialInterface* ActualMaterialInterface = MaterialTemplate;
+
+							if (MutableMaterialPlaceholder.Params.Num() != 0)
+							{
+								MaterialInstance = UMaterialInstanceDynamic::Create(MaterialTemplate, GetTransientPackage());
+								ActualMaterialInterface = MaterialInstance;
+							}
+
+							if (SkeletalMesh)
+							{
+								int32 MatIndex = SkeletalMesh->GetMaterials().Add(ActualMaterialInterface);
+								SkeletalMesh->GetMaterials()[MatIndex].MaterialSlotName = CustomizableObject->ReferencedMaterialSlotNames[Surface.MaterialIndex];
+								int32 LODMaterialIndex = Helper_GetLODInfoArray(SkeletalMesh)[MeshLODIndex].LODMaterialMap.Add(MatIndex);
+								Helper_GetLODRenderSections(SkeletalMesh, MeshLODIndex)[SurfaceIndex].MaterialIndex = LODMaterialIndex;
+
+								MutableMaterialPlaceholder.MatIndex = MatIndex;
+							}
+
+							if (MaterialInstance)
+							{
+								for (const FMutableMaterialPlaceholder::FMutableMaterialPlaceHolderParam& Param : MutableMaterialPlaceholder.Params)
+								{
+									switch (Param.Type)
+									{
+									case FMutableMaterialPlaceholder::EPlaceHolderParamType::Vector:
+										if (Param.LayerIndex < 0)
+										{
+											MaterialInstance->SetVectorParameterValue(Param.ParamName, Param.Vector);
+										}
+										else
+										{
+											FMaterialParameterInfo ParameterInfo = FMaterialParameterInfo(Param.ParamName, EMaterialParameterAssociation::LayerParameter, Param.LayerIndex);
+											MaterialInstance->SetVectorParameterValueByInfo(ParameterInfo, Param.Vector);
+										}
+
+										break;
+
+									case FMutableMaterialPlaceholder::EPlaceHolderParamType::Scalar:
+										if (Param.LayerIndex < 0)
+										{
+											MaterialInstance->SetScalarParameterValue(FName(Param.ParamName), Param.Scalar);
+										}
+										else
+										{
+											FMaterialParameterInfo ParameterInfo = FMaterialParameterInfo(Param.ParamName, EMaterialParameterAssociation::LayerParameter, Param.LayerIndex);
+											MaterialInstance->SetScalarParameterValueByInfo(ParameterInfo, Param.Scalar);
+										}
+
+										break;
+
+									case FMutableMaterialPlaceholder::EPlaceHolderParamType::Texture:
+										if (Param.LayerIndex < 0)
+										{
+											MaterialInstance->SetTextureParameterValue(Param.ParamName, Param.Texture.Texture);
+										}
+										else
+										{
+											FMaterialParameterInfo ParameterInfo = FMaterialParameterInfo(Param.ParamName, EMaterialParameterAssociation::LayerParameter, Param.LayerIndex);
+											MaterialInstance->SetTextureParameterValueByInfo(ParameterInfo, Param.Texture.Texture);
+										}
+
+										Material.Textures.Add(Param.Texture);
+
+										break;
+									}
+								}
+							}
+
+							GeneratedMaterials.Add(Material);
+						}
+					}
+				}
+			}
+		}
+	}
+
+	{
+		MUTABLE_CPUPROFILER_SCOPE(Exchange);
+
+		FCustomizableObjectSystemPrivate* CustomizableObjectSystem = UCustomizableObjectSystem::GetInstance()->GetPrivate();
+		TexturesToRelease.Empty();
+
+		for (FGeneratedTexture& Texture : GeneratedTextures)
+		{
+			if (CustomizableObjectSystem->RemoveTextureReference(Texture.Id))
+			{
+				if (CustomizableObjectSystem->bReleaseTexturesImmediately)
+				{
+					TexturesToRelease.Add(Texture.Id, Texture); // Texture count is zero, so prepare to release it
+				}
+			}
+		}
+
+		for (FGeneratedTexture& Texture : NewGeneratedTextures)
+		{
+			CustomizableObjectSystem->AddTextureReference(Texture.Id);
+
+			if (CustomizableObjectSystem->bReleaseTexturesImmediately)
+			{
+				TexturesToRelease.Remove(Texture.Id); // Texture count in the end is not zero, so do not release it
+			}
+		}
+
+		Exchange(GeneratedTextures, NewGeneratedTextures);
+	}
+}
+
+
+void UCustomizableInstancePrivateData::ProcessTextureCoverageQueries(const TSharedPtr<FMutableOperationData>& OperationData, UCustomizableObject* CustomizableObject, const FString& ImageKeyName, FTexturePlatformData *PlatformData, UMaterialInterface* Material)
+{
+	if (OperationData->TextureCoverageQueries_MutableThreadParams.Num() && PlatformData->Mips.Num())
+	{
+		int32 ImageKey = FCString::Atoi(*ImageKeyName);
+		int32 ScannedMipLevel = 0;
+		check(PlatformData->Mips.IsValidIndex(ScannedMipLevel));
+
+		if (ImageKey >= 0 && ImageKey < CustomizableObject->ImageProperties.Num())
+		{
+			const FMutableModelImageProperties& Props = CustomizableObject->ImageProperties[ImageKey];
+
+			FTextureCoverageQueryData* TextureCoverageQueryData = OperationData->TextureCoverageQueries_MutableThreadParams.Find(Props.TextureParameterName);
+
+			if (TextureCoverageQueryData)
+			{
+				FTexture2DMipMap& ScannedMip = PlatformData->Mips[ScannedMipLevel];
+				uint32 TotalTexels = ScannedMip.SizeX * ScannedMip.SizeY;
+				uint32 CoveredTexels = 0;
+				uint32 MaskedOutCoveredTexels = 0;
+
+				auto MaskOutCache = CustomizableObject->MaskOutCache.Get();
+
+				if (!MaskOutCache)
+				{
+					UE_LOG(LogMutable, Error, TEXT("The CustomizableObject->MaskOutCache has to be manually loaded by the programmer by calling CustomizableObject->LoadMaskOutCache()."));
+				}
+
+				FString* MaskOutTexturePath = MaskOutCache ? MaskOutCache->Materials.Find(Material->GetPathName()) : nullptr;
+				FMaskOutTexture* MaskOutTexture = MaskOutTexturePath ? MaskOutCache->Textures.Find(*MaskOutTexturePath) : nullptr;
+
+				// Textures must be uncompressed, the CustomizableInstance's state should be set to one with
+				// texture compression disabled for Texture Coverage Queries
+				ensure(PlatformData->PixelFormat == PF_R8G8B8A8 || PlatformData->PixelFormat == PF_B8G8R8A8);
+
+				const void* pData = ScannedMip.BulkData.Lock(LOCK_READ_ONLY);
+				uint8_t* pDest = (uint8_t*)pData;
+
+				for (size_t px = 0; px < ScannedMip.SizeX; ++px)
+				{
+					for (size_t py = 0; py < ScannedMip.SizeY; ++py)
+					{
+						if (pDest[(py * ScannedMip.SizeX + px) * 4 + 3] > 0) // Check if alpha channel of PF_R8G8B8A8 is not zero
+						{
+							CoveredTexels++;
+
+							if (MaskOutTexture)
+							{
+								float u = float(px) / ScannedMip.SizeX;
+								float v = float(py) / ScannedMip.SizeY;
+
+								size_t pu = FMath::RoundToInt(u * MaskOutTexture->GetSizeX());
+								size_t pv = FMath::RoundToInt(v * MaskOutTexture->GetSizeY());
+
+								if (MaskOutTexture->GetTexelReference(pv * MaskOutTexture->GetSizeX() + pu) == 0)
+								{
+									MaskedOutCoveredTexels++;
+								}
+							}
+						}
+					}
+				}
+
+				ScannedMip.BulkData.Unlock();
+
+				auto& ResultData = OperationData->TextureCoverageQueries_MutableThreadResults.FindOrAdd(Props.TextureParameterName);
+				ResultData.CoveredTexels += CoveredTexels;
+				ResultData.MaskedOutCoveredTexels += MaskedOutCoveredTexels;
+				ResultData.TotalTexels += TotalTexels;
+			}
+		}
+	}
+}
+
+
+void UCustomizableInstancePrivateData::UpdateParameterDecorationsEngineResources(const TSharedPtr<FMutableOperationData>& OperationData)
+{
+	// This should run in the unreal thread.
+	ParameterDecorations.SetNum(OperationData->ParametersUpdateData.Parameters.Num());
+	for (int32 ParamIndex = 0; ParamIndex < OperationData->ParametersUpdateData.Parameters.Num(); ++ParamIndex)
+	{
+		int32 DescCount = OperationData->ParametersUpdateData.Parameters[ParamIndex].Images.Num();
+		ParameterDecorations[ParamIndex].Images.SetNum(DescCount);
+		for (int32 DescIndex = 0; DescIndex < DescCount; ++DescIndex)
+		{
+			UTexture2D* Tex = nullptr;
+
+			auto Image = OperationData->ParametersUpdateData.Parameters[ParamIndex].Images[DescIndex];
+			if (Image)
+			{
+				check(Image->GetSizeX()>0 && Image->GetSizeY()>0);
+
+				Tex = NewObject<UTexture2D>(UTexture2D::StaticClass());
+				FMutableModelImageProperties Props;
+				Props.Filter = TF_Bilinear;
+				Props.SRGB = true;
+				Props.LODBias = 0;
+
+				ConvertImage(Tex, Image, Props);
+
+#if !PLATFORM_DESKTOP
+				for (int32 i = 0; i < Tex->GetPlatformData()->Mips.Num(); ++i)
+				{
+					uint32 DataFlags = Tex->GetPlatformData()->Mips[i].BulkData.GetBulkDataFlags();
+					Tex->GetPlatformData()->Mips[i].BulkData.SetBulkDataFlags(DataFlags | BULKDATA_SingleUse);
+				}
+#endif
+
+				Tex->UpdateResource();
+			}
+
+			ParameterDecorations[ParamIndex].Images[DescIndex] = Tex;
+		}
+	}
+
+	OperationData->ParametersUpdateData.Clear();
+
+	// Relevancy
+	RelevantParameters = OperationData->RelevantParametersInProgress;
+}
+
+
+void UCustomizableObjectInstance::SetReplacePhysicsAssets(bool bReplaceEnabled)
+{
+	bReplaceEnabled ? GetPrivate()->SetCOInstanceFlags(ReplacePhysicsAssets) : GetPrivate()->ClearCOInstanceFlags(ReplacePhysicsAssets);
+}
+
+
+void UCustomizableObjectInstance::SetReuseInstanceTextures(bool bTextureReuseEnabled)
+{
+	bTextureReuseEnabled ? GetPrivate()->SetCOInstanceFlags(ReuseTextures) : GetPrivate()->ClearCOInstanceFlags(ReuseTextures);
+}
+
+
+void UCustomizableObjectInstance::AddQueryTextureCoverage(const FString& TextureName, const FString* MaskOutChannelName)
+{
+	FTextureCoverageQueryData& QueryData = GetPrivate()->TextureCoverageQueries.FindOrAdd(TextureName);
+
+	if (MaskOutChannelName)
+	{
+		QueryData.MaskOutChannelName = *MaskOutChannelName;
+	}
+}
+
+
+void UCustomizableObjectInstance::RemoveQueryTextureCoverage(const FString& TextureName)
+{
+	GetPrivate()->TextureCoverageQueries.Remove(TextureName);
+}
+
+
+float UCustomizableObjectInstance::GetQueryResultTextureCoverage(const FString& TextureName)
+{
+	FTextureCoverageQueryData* QueryData = GetPrivate()->TextureCoverageQueries.Find(TextureName);
+
+	return QueryData ? QueryData->GetCoverage() : 0.f;
+}
+
+
+float UCustomizableObjectInstance::GetQueryResultTextureCoverageMasked(const FString& TextureName)
+{
+	FTextureCoverageQueryData* QueryData = GetPrivate()->TextureCoverageQueries.Find(TextureName);
+
+	return QueryData ? QueryData->GetMaskedOutCoverage() : 0.f;
+}
+
+
+void UCustomizableObjectInstance::SetIsBeingUsedByComponentInPlay(bool bIsUsedByComponentInPlay)
+{
+	bIsUsedByComponentInPlay ? GetPrivate()->SetCOInstanceFlags(UsedByComponentInPlay) : GetPrivate()->ClearCOInstanceFlags(UsedByComponentInPlay);
+}
+
+
+bool UCustomizableObjectInstance::GetIsBeingUsedByComponentInPlay() const
+{
+	return GetPrivate()->HasCOInstanceFlags(UsedByComponentInPlay);
+}
+
+
+void UCustomizableObjectInstance::SetIsDiscardedBecauseOfTooManyInstances(bool bIsDiscarded)
+{
+	bIsDiscarded ? GetPrivate()->SetCOInstanceFlags(DiscardedByNumInstancesLimit) : GetPrivate()->ClearCOInstanceFlags(DiscardedByNumInstancesLimit);
+}
+
+
+bool UCustomizableObjectInstance::GetIsDiscardedBecauseOfTooManyInstances() const
+{
+	return GetPrivate()->HasCOInstanceFlags(DiscardedByNumInstancesLimit);
+}
+
+
+void UCustomizableObjectInstance::SetIsPlayerOrNearIt(bool bIsPlayerorNearIt)
+{
+	bIsPlayerorNearIt ? GetPrivate()->SetCOInstanceFlags(UsedByPlayerOrNearIt) : GetPrivate()->ClearCOInstanceFlags(UsedByPlayerOrNearIt);
+}
+
+
+float UCustomizableObjectInstance::GetMinSquareDistToPlayer() const
+{
+	return GetPrivate()->MinSquareDistFromComponentToPlayer;
+}
+
+void UCustomizableObjectInstance::SetMinSquareDistToPlayer(float NewValue)
+{
+	GetPrivate()->MinSquareDistFromComponentToPlayer = NewValue;
+}
+
+
+void UCustomizableObjectInstance::SetMinMaxLODToLoad(int32 NewMinLOD, int32 NewMaxLOD, bool bLimitLODUpgrades)
+{
+	GetPrivate()->SetMinMaxLODToLoad(NewMinLOD, NewMaxLOD, bLimitLODUpgrades);
+}
+
+
+void UCustomizableObjectInstance::SetUseCurrentMinLODAsBaseLOD(bool bIsBaseLOD)
+{
+	bIsBaseLOD ? GetPrivate()->SetCOInstanceFlags(ReduceLODs) : GetPrivate()->ClearCOInstanceFlags(ReduceLODs);
+}
+
+
+bool UCustomizableObjectInstance::GetUseCurrentMinLODAsBaseLOD() const
+{
+	return GetPrivate()->HasCOInstanceFlags(ReduceLODs);
+}
+
+
+int32 UCustomizableObjectInstance::GetNumLODsAvailable() const
+{
+	return GetPrivate()->GetNumLODsAvailable();
+}
+
+
+USkeletalMesh* UCustomizableObjectInstance::GetSkeletalMesh(int32 ComponentIndex) const
+{
+	return SkeletalMeshes.IsValidIndex(ComponentIndex) ? SkeletalMeshes[ComponentIndex] : nullptr;
+}
+
+
+bool UCustomizableObjectInstance::HasAnySkeletalMesh() const
+{
+	for (int32 Index = 0; Index < SkeletalMeshes.Num(); ++Index)
+	{
+		if (SkeletalMeshes[Index])
+		{
+			return true;
+		}
+	}
+
+	return false;
+}
+
+
+TArray<FCustomizableObjectBoolParameterValue>& UCustomizableObjectInstance::GetBoolParameters()
+{
+	return Descriptor.GetBoolParameters();
+}
+
+
+const TArray<FCustomizableObjectBoolParameterValue>& UCustomizableObjectInstance::GetBoolParameters() const
+{
+	return Descriptor.GetBoolParameters();
+}
+
+
+TArray<FCustomizableObjectIntParameterValue>& UCustomizableObjectInstance::GetIntParameters()
+{
+	return Descriptor.GetIntParameters();
+}
+
+
+const TArray<FCustomizableObjectIntParameterValue>& UCustomizableObjectInstance::GetIntParameters() const
+{
+	return Descriptor.GetIntParameters();
+}
+
+
+TArray<FCustomizableObjectFloatParameterValue>& UCustomizableObjectInstance::GetFloatParameters()
+{
+	return Descriptor.GetFloatParameters();
+}
+
+
+const TArray<FCustomizableObjectFloatParameterValue>& UCustomizableObjectInstance::GetFloatParameters() const
+{
+	return Descriptor.GetFloatParameters();
+}
+
+
+TArray<FCustomizableObjectTextureParameterValue>& UCustomizableObjectInstance::GetTextureParameters()
+{
+	return Descriptor.GetTextureParameters();
+}
+
+
+const TArray<FCustomizableObjectTextureParameterValue>& UCustomizableObjectInstance::GetTextureParameters() const
+{
+	return Descriptor.GetTextureParameters();
+}
+
+
+TArray<FCustomizableObjectVectorParameterValue>& UCustomizableObjectInstance::GetVectorParameters()
+{
+	return Descriptor.GetVectorParameters();
+}
+
+const TArray<FCustomizableObjectVectorParameterValue>& UCustomizableObjectInstance::GetVectorParameters() const
+{
+	return Descriptor.GetVectorParameters();
+}
+
+
+TArray<FCustomizableObjectProjectorParameterValue>& UCustomizableObjectInstance::GetProjectorParameters()
+{
+	return Descriptor.GetProjectorParameters();
+}
+
+
+const TArray<FCustomizableObjectProjectorParameterValue>& UCustomizableObjectInstance::GetProjectorParameters() const
+{
+	return Descriptor.GetProjectorParameters();
+}
+
+
+bool UCustomizableObjectInstance::HasAnyParameters() const
+{
+	return Descriptor.HasAnyParameters();	
+}
+
+
+TSubclassOf<UAnimInstance> UCustomizableObjectInstance::GetAnimBP(int32 ComponentIndex, int32 SlotIndex) const
+{
+	FCustomizableInstanceComponentData* ComponentData =	GetPrivate()->GetComponentData(ComponentIndex);
+	
+	if (!ComponentData)
+	{
+		FString ErrorMsg = FString::Printf(TEXT("Tried to access and invalid component index [%d] in a Mutable Instance."), ComponentIndex);
+		UE_LOG(LogMutable, Error, TEXT("%s"), *ErrorMsg);
+#if WITH_EDITOR
+		FMessageLogModule& MessageLogModule = FModuleManager::LoadModuleChecked<FMessageLogModule>("MessageLog");
+		MessageLogModule.RegisterLogListing(FName("Mutable"), FText::FromString(FString("Mutable")));
+		FMessageLog MessageLog("Mutable");
+
+		MessageLog.Notify(FText::FromString(ErrorMsg), EMessageSeverity::Error, true);
+#endif
+
+		return nullptr;
+	}
+
+	TSoftClassPtr<UAnimInstance>* Result = ComponentData->AnimSlotToBP.Find(SlotIndex);
+
+	return Result ? Result->Get() : nullptr;
+}
+
+const FGameplayTagContainer& UCustomizableObjectInstance::GetAnimationGameplayTags() const
+{
+	return GetPrivate()->AnimBPGameplayTags;
+}
+
+void UCustomizableObjectInstance::ForEachAnimInstance(int32 ComponentIndex, FEachComponentAnimInstanceClassDelegate Delegate) const
+{
+	// allow us to log out both bad states with one pass
+	bool bAnyErrors = false;
+
+	if (!Delegate.IsBound())
+	{
+		FString ErrorMsg = FString::Printf(TEXT("Attempting to iterate over AnimInstances with an unbound delegate - does nothing!"), ComponentIndex);
+		UE_LOG(LogMutable, Warning, TEXT("%s"), *ErrorMsg);
+#if WITH_EDITOR
+		FMessageLogModule& MessageLogModule = FModuleManager::LoadModuleChecked<FMessageLogModule>("MessageLog");
+		MessageLogModule.RegisterLogListing(FName("Mutable"), FText::FromString(FString("Mutable")));
+		FMessageLog MessageLog("Mutable");
+
+		MessageLog.Notify(FText::FromString(ErrorMsg), EMessageSeverity::Warning, true);
+#endif
+		bAnyErrors = true;
+	}
+
+	const FCustomizableInstanceComponentData* ComponentData =	GetPrivate()->GetComponentData(ComponentIndex);
+	
+	if (!ComponentData)
+	{
+		FString ErrorMsg = FString::Printf(TEXT("Tried to access and invalid component index [%d] in a Mutable Instance."), ComponentIndex);
+		UE_LOG(LogMutable, Error, TEXT("%s"), *ErrorMsg);
+#if WITH_EDITOR
+		FMessageLogModule& MessageLogModule = FModuleManager::LoadModuleChecked<FMessageLogModule>("MessageLog");
+		MessageLogModule.RegisterLogListing(FName("Mutable"), FText::FromString(FString("Mutable")));
+		FMessageLog MessageLog("Mutable");
+
+		MessageLog.Notify(FText::FromString(ErrorMsg), EMessageSeverity::Error, true);
+#endif
+
+		bAnyErrors = true;
+	}
+
+	if (bAnyErrors)
+	{
+		return;
+	}
+
+	for (const TPair<int32, TSoftClassPtr<UAnimInstance>>& MapElem : ComponentData->AnimSlotToBP)
+	{
+		const int32 Index = MapElem.Key;
+		const TSoftClassPtr<UAnimInstance>& AnimBP = MapElem.Value;
+
+		// if this _can_ resolve to a real AnimBP
+		if (!AnimBP.IsNull())
+		{
+			// force a load right now - we don't know whether we would have loaded already - this could be called in editor
+			const TSubclassOf<UAnimInstance> LiveAnimBP = AnimBP.LoadSynchronous();
+			if (LiveAnimBP)
+			{
+				Delegate.Execute(Index, LiveAnimBP);
+			}
+		}
+	}
+}
+
+void UCustomizableObjectInstance::ForEachAnimInstance(int32 ComponentIndex, FEachComponentAnimInstanceClassNativeDelegate Delegate) const
+{
+	// allow us to log out both bad states with one pass
+	bool bAnyErrors = false;
+
+	if (!Delegate.IsBound())
+	{
+		FString ErrorMsg = FString::Printf(TEXT("Attempting to iterate over AnimInstances with an unbound delegate - does nothing!"), ComponentIndex);
+		UE_LOG(LogMutable, Warning, TEXT("%s"), *ErrorMsg);
+#if WITH_EDITOR
+		FMessageLogModule& MessageLogModule = FModuleManager::LoadModuleChecked<FMessageLogModule>("MessageLog");
+		MessageLogModule.RegisterLogListing(FName("Mutable"), FText::FromString(FString("Mutable")));
+		FMessageLog MessageLog("Mutable");
+
+		MessageLog.Notify(FText::FromString(ErrorMsg), EMessageSeverity::Warning, true);
+#endif
+		bAnyErrors = true;
+	}
+
+	const FCustomizableInstanceComponentData* ComponentData =	GetPrivate()->GetComponentData(ComponentIndex);
+	
+	if (!ComponentData)
+	{
+		FString ErrorMsg = FString::Printf(TEXT("Tried to access and invalid component index [%d] in a Mutable Instance."), ComponentIndex);
+		UE_LOG(LogMutable, Error, TEXT("%s"), *ErrorMsg);
+#if WITH_EDITOR
+		FMessageLogModule& MessageLogModule = FModuleManager::LoadModuleChecked<FMessageLogModule>("MessageLog");
+		MessageLogModule.RegisterLogListing(FName("Mutable"), FText::FromString(FString("Mutable")));
+		FMessageLog MessageLog("Mutable");
+
+		MessageLog.Notify(FText::FromString(ErrorMsg), EMessageSeverity::Error, true);
+#endif
+
+		bAnyErrors = true;
+	}
+
+	if (bAnyErrors)
+	{
+		return;
+	}
+
+	for (const TPair<int32, TSoftClassPtr<UAnimInstance>>& MapElem : ComponentData->AnimSlotToBP)
+	{
+		const int32 Index = MapElem.Key;
+		const TSoftClassPtr<UAnimInstance>& AnimBP = MapElem.Value;
+
+		// if this _can_ resolve to a real AnimBP
+		if (!AnimBP.IsNull())
+		{
+			// force a load right now - we don't know whether we would have loaded already - this could be called in editor
+			const TSubclassOf<UAnimInstance> LiveAnimBP = AnimBP.LoadSynchronous();
+			if (LiveAnimBP)
+			{
+				Delegate.Execute(Index, LiveAnimBP);
+			}
+		}
+	}
+}
