[mutable] Move Mutable plugin from the restricted folder to the experimental plugins folder.

diff --git a/Engine/Plugins/Experimental/Mutable/Source/CustomizableObject/Private/CustomizableObjectInstanceDescriptor.cpp b/Engine/Plugins/Experimental/Mutable/Source/CustomizableObject/Private/CustomizableObjectInstanceDescriptor.cpp
new file mode 100644
--- /dev/null
+++ b/Engine/Plugins/Experimental/Mutable/Source/CustomizableObject/Private/CustomizableObjectInstanceDescriptor.cpp
@@ -0,0 +1,1442 @@
+// Copyright Epic Games, Inc. All Rights Reserved.
+
+#include "CustomizableObjectInstanceDescriptor.h"
+
+#include "CustomizableObjectPrivate.h"
+#include "CustomizableObjectSystem.h"
+#include "MutableRuntime/Public/Model.h"
+
+
+const FString EmptyString;
+
+
+FString GetAvailableOptionsString(const UCustomizableObject& CustomizableObject, const int32 ParameterIndexInObject)
+{
+	FString OptionsString;
+	const int32 NumOptions = CustomizableObject.GetIntParameterNumOptions(ParameterIndexInObject);
+
+	for (int32 k = 0; k < NumOptions; ++k)
+	{
+		OptionsString += CustomizableObject.GetIntParameterAvailableOption(ParameterIndexInObject, k);
+
+		if (k < NumOptions - 1)
+		{
+			OptionsString += FString(", ");
+		}
+	}
+
+	return OptionsString;
+}
+
+
+FCustomizableObjectInstanceDescriptor::FCustomizableObjectInstanceDescriptor(UCustomizableObject &Object)
+{
+	CustomizableObject = &Object;
+	
+	CreateParametersLookupTable();
+}
+
+
+FCustomizableObjectInstanceDescriptor::FCustomizableObjectInstanceDescriptor(const FCustomizableObjectInstanceDescriptor& Other)
+	: FCustomizableObjectInstanceDescriptor(*Other.CustomizableObject)
+{
+	BoolParameters = Other.BoolParameters;
+	IntParameters =	Other.IntParameters;
+	IntParametersLookupTable = Other.IntParametersLookupTable;
+	FloatParameters = Other.FloatParameters;
+	TextureParameters =	Other.TextureParameters;
+	VectorParameters = Other.VectorParameters;
+	ProjectorParameters = Other.ProjectorParameters;
+	State = Other.State;
+}
+
+
+void FCustomizableObjectInstanceDescriptor::SaveDescriptor(FArchive& Ar)
+{
+	check(CustomizableObject);
+
+	// This is a non-portable but very compact descriptor if bUseCompactDescriptor is true. It assumes the compiled objects are the same
+	// on both ends of the serialisation. That's why we iterate the parameters in the actual compiled
+	// model, instead of the arrays in this class.
+
+	bool bUseCompactDescriptor = UCustomizableObjectSystem::GetInstance()->IsCompactSerializationEnabled();
+
+	Ar << bUseCompactDescriptor;
+
+	// Not sure if this is needed, but it is small.
+	Ar << State;
+
+	int32 ModelParameterCount = CustomizableObject->GetParameterCount();
+
+	if (!bUseCompactDescriptor)
+	{
+		Ar << ModelParameterCount;
+	}
+
+	for (int32 ModelParameterIndex = 0; ModelParameterIndex < ModelParameterCount; ++ModelParameterIndex)
+	{
+		const FString & Name = CustomizableObject->GetParameterName(ModelParameterIndex);
+		EMutableParameterType Type = CustomizableObject->GetParameterType(ModelParameterIndex);
+
+		if (!bUseCompactDescriptor)
+		{
+			check(Ar.IsSaving());
+			Ar << const_cast<FString &>(Name);
+		}
+
+		switch (Type)
+		{
+		case EMutableParameterType::Bool:
+		{
+			bool Value = false;
+			for (const FCustomizableObjectBoolParameterValue& P: BoolParameters)
+			{
+				if (P.ParameterName == Name)
+				{
+					Value = P.ParameterValue;
+					break;
+				}
+			}
+			Ar << Value;
+			break;
+		}
+
+		case EMutableParameterType::Float:
+		{
+			float Value = 0.0f;
+			TArray<float> RangeValues;
+			for (const FCustomizableObjectFloatParameterValue& P : FloatParameters)
+			{
+				if (P.ParameterName == Name)
+				{
+					Value = P.ParameterValue;
+					RangeValues = P.ParameterRangeValues;
+					break;
+				}
+			}
+			Ar << Value;
+			Ar << RangeValues;
+			break;
+		}
+
+		case EMutableParameterType::Int:
+		{
+			int32 Value = 0;
+			FString ValueName;
+
+			TArray<int32> Values;
+			TArray<FString> ValueNames;
+
+			bool bIsParamMultidimensional = false;
+
+			int32 IntParameterIndex = FindIntParameterNameIndex(Name);
+			if (IntParameterIndex >= 0)
+				{					
+				const FCustomizableObjectIntParameterValue & P = IntParameters[IntParameterIndex];
+					Value = CustomizableObject->FindIntParameterValue(ModelParameterIndex,P.ParameterValueName);
+
+				int32 ParameterIndexInObject = CustomizableObject->FindParameter(IntParameters[IntParameterIndex].ParameterName);
+				bIsParamMultidimensional = IsParamMultidimensional(ParameterIndexInObject);
+
+				if (bIsParamMultidimensional)
+				{
+					for (int i = 0; i < P.ParameterRangeValueNames.Num(); ++i)
+					{
+						ValueNames.Add(P.ParameterRangeValueNames[i]);
+						Values.Add(CustomizableObject->FindIntParameterValue(ModelParameterIndex, P.ParameterRangeValueNames[i]));
+					}
+				}
+
+				if (!bUseCompactDescriptor)
+				{
+					ValueName = P.ParameterValueName;
+				}
+			}
+
+			if (bUseCompactDescriptor)
+			{
+				Ar << Value;
+
+				if (bIsParamMultidimensional)
+				{
+					Ar << Values;
+				}
+			}
+			else
+			{
+				Ar << ValueName;
+
+				if (bIsParamMultidimensional)
+				{
+					Ar << ValueNames;
+				}
+			}
+
+			break;
+		}
+
+		case EMutableParameterType::Color:
+		{
+			FLinearColor Value(FLinearColor::Black);
+			for (const FCustomizableObjectVectorParameterValue& P : VectorParameters)
+			{
+				if (P.ParameterName == Name)
+				{
+					Value = P.ParameterValue;
+					break;
+				}
+			}
+			Ar << Value;
+
+			break;
+		}
+
+		case EMutableParameterType::Texture:
+		{
+			uint64 Value = 0;
+			for (const FCustomizableObjectTextureParameterValue& P : TextureParameters)
+			{
+				if (P.ParameterName == Name)
+				{
+					Value = P.ParameterValue;
+					break;
+				}
+			}
+			Ar << Value;
+
+			break;
+		}
+
+		case EMutableParameterType::Projector:
+		{
+			FCustomizableObjectProjector Value;
+			TArray<FCustomizableObjectProjector> RangeValues;
+
+			for (const FCustomizableObjectProjectorParameterValue& P : ProjectorParameters)
+			{
+				if (P.ParameterName == Name)
+				{
+					Value = P.Value;
+					RangeValues = P.RangeValues;
+					break;
+				}
+			}
+			Ar << Value;
+			Ar << RangeValues;
+
+			break;
+		}
+
+		default:
+			// Parameter type replication not implemented.
+			check(false);
+		}
+	}
+}
+
+
+void FCustomizableObjectInstanceDescriptor::LoadDescriptor(FArchive& Ar)
+{
+	check(CustomizableObject);
+
+	// This is a non-portable but very compact descriptor if bUseCompactDescriptor is true. It assumes the compiled objects are the same
+	// on both ends of the serialisation. That's why we iterate the parameters in the actual compiled
+	// model, instead of the arrays in this class.
+
+	bool bUseCompactDescriptor = UCustomizableObjectSystem::GetInstance()->IsCompactSerializationEnabled();
+
+	Ar << bUseCompactDescriptor;
+
+	// Not sure if this is needed, but it is small.
+	Ar << State;
+	
+	int32 ModelParameterCount = CustomizableObject->GetParameterCount();
+
+	if (!bUseCompactDescriptor)
+	{
+		Ar << ModelParameterCount;
+	}
+
+	for (int32 i = 0; i < ModelParameterCount; ++i)
+	{
+		FString Name;
+		EMutableParameterType Type;
+		int32 ModelParameterIndex = -1;
+
+		if (bUseCompactDescriptor)
+		{
+			ModelParameterIndex = i;
+			Name = CustomizableObject->GetParameterName(ModelParameterIndex);
+			Type = CustomizableObject->GetParameterType(ModelParameterIndex);
+		}
+		else
+		{
+			Ar << Name;
+			Type = CustomizableObject->GetParameterTypeByName(Name);
+		}
+
+		switch (Type)
+		{
+		case EMutableParameterType::Bool:
+		{
+			bool Value = false;
+			Ar << Value;
+			for (FCustomizableObjectBoolParameterValue& P : BoolParameters)
+			{
+				if (P.ParameterName == Name)
+				{
+					P.ParameterValue = Value;
+					break;
+				}
+			}
+			break;
+		}
+
+		case EMutableParameterType::Float:
+		{
+			float Value = 0.0f;
+			TArray<float> RangeValues;
+			Ar << Value;
+			Ar << RangeValues;
+			for (FCustomizableObjectFloatParameterValue& P : FloatParameters)
+			{
+				if (P.ParameterName == Name)
+				{
+					P.ParameterValue = Value;
+					P.ParameterRangeValues = RangeValues;
+					break;
+				}
+			}
+			break;
+		}
+
+		case EMutableParameterType::Int:
+		{
+			int32 Value = 0;
+			FString ValueName;
+
+			TArray<int32> Values;
+			TArray<FString> ValueNames;
+
+			const int32 IntParameterIndex = FindIntParameterNameIndex(Name);
+			const int32 ParameterIndexInObject = CustomizableObject->FindParameter(IntParameters[IntParameterIndex].ParameterName);
+			const bool bIsParamMultidimensional = IsParamMultidimensional(ParameterIndexInObject);
+
+			if (bUseCompactDescriptor)
+			{
+				Ar << Value;
+
+				if (bIsParamMultidimensional)
+				{
+					Ar << Values;
+				}
+			}
+			else
+			{
+				Ar << ValueName;
+
+				if (bIsParamMultidimensional)
+				{
+					Ar << ValueNames;
+				}
+			}
+
+			for (FCustomizableObjectIntParameterValue& P : IntParameters)
+			{
+				if (P.ParameterName == Name)
+				{
+					if (bUseCompactDescriptor)
+					{
+						P.ParameterValueName = CustomizableObject->FindIntParameterValueName(ModelParameterIndex, Value);
+						//check((P.ParameterValueName.IsEmpty() && ValueName.Equals(FString("None"))) || P.ParameterValueName.Equals(ValueName));
+						P.ParameterRangeValueNames.SetNum(Values.Num());
+
+						for (int ParamIndex = 0; ParamIndex < Values.Num(); ++ParamIndex)
+						{
+							P.ParameterRangeValueNames[ParamIndex] = CustomizableObject->FindIntParameterValueName(ModelParameterIndex, Values[ParamIndex]);
+						}
+					}
+					else
+					{
+						P.ParameterValueName = ValueName;
+						P.ParameterRangeValueNames = ValueNames;
+					}
+
+					break;
+				}
+			}
+
+			break;
+		}
+
+		case EMutableParameterType::Color:
+		{
+			FLinearColor Value(FLinearColor::Black);
+			Ar << Value;
+			for (FCustomizableObjectVectorParameterValue& P : VectorParameters)
+			{
+				if (P.ParameterName == Name)
+				{
+					P.ParameterValue = Value;
+					break;
+				}
+			}
+
+			break;
+		}
+
+		case EMutableParameterType::Texture:
+		{
+			uint64 Value = 0;
+			Ar << Value;
+			for (FCustomizableObjectTextureParameterValue& P : TextureParameters)
+			{
+				if (P.ParameterName == Name)
+				{
+					P.ParameterValue = Value;
+					break;
+				}
+			}
+
+			break;
+		}
+
+		case EMutableParameterType::Projector:
+		{
+			FCustomizableObjectProjector Value;
+			TArray<FCustomizableObjectProjector> RangeValues;
+			Ar << Value;
+			Ar << RangeValues;
+
+			for (FCustomizableObjectProjectorParameterValue& P : ProjectorParameters)
+			{
+				if (P.ParameterName == Name)
+				{
+					P.Value = Value;
+					P.RangeValues = RangeValues;
+					break;
+				}
+			}
+
+			break;
+		}
+
+		default:
+			// Parameter type replication not implemented.
+			check(false);
+		}
+	}
+
+	CreateParametersLookupTable();
+}
+
+
+UCustomizableObject* FCustomizableObjectInstanceDescriptor::GetCustomizableObject() const
+{
+	return CustomizableObject;
+}
+
+
+TArray<FCustomizableObjectBoolParameterValue>& FCustomizableObjectInstanceDescriptor::GetBoolParameters()
+{
+	return BoolParameters;
+}
+
+
+const TArray<FCustomizableObjectBoolParameterValue>& FCustomizableObjectInstanceDescriptor::GetBoolParameters() const
+{
+	return const_cast<FCustomizableObjectInstanceDescriptor*>(this)->GetBoolParameters();
+}
+
+
+TArray<FCustomizableObjectIntParameterValue>& FCustomizableObjectInstanceDescriptor::GetIntParameters()
+{
+	return IntParameters;
+}
+
+
+const TArray<FCustomizableObjectIntParameterValue>& FCustomizableObjectInstanceDescriptor::GetIntParameters() const
+{
+	return const_cast<FCustomizableObjectInstanceDescriptor*>(this)->GetIntParameters();	
+}
+
+
+TArray<FCustomizableObjectFloatParameterValue>& FCustomizableObjectInstanceDescriptor::GetFloatParameters()
+{
+	return FloatParameters;	
+}
+
+
+const TArray<FCustomizableObjectFloatParameterValue>& FCustomizableObjectInstanceDescriptor::GetFloatParameters() const
+{
+	return const_cast<FCustomizableObjectInstanceDescriptor*>(this)->GetFloatParameters();	
+}
+
+
+TArray<FCustomizableObjectTextureParameterValue>& FCustomizableObjectInstanceDescriptor::GetTextureParameters()
+{
+	return TextureParameters;	
+}
+
+
+const TArray<FCustomizableObjectTextureParameterValue>& FCustomizableObjectInstanceDescriptor::GetTextureParameters() const
+{
+	return const_cast<FCustomizableObjectInstanceDescriptor*>(this)->GetTextureParameters();	
+}
+
+
+TArray<FCustomizableObjectVectorParameterValue>& FCustomizableObjectInstanceDescriptor::GetVectorParameters()
+{
+	return VectorParameters;	
+}
+
+
+const TArray<FCustomizableObjectVectorParameterValue>& FCustomizableObjectInstanceDescriptor::GetVectorParameters() const
+{
+	return const_cast<FCustomizableObjectInstanceDescriptor*>(this)->GetVectorParameters();
+}
+
+
+TArray<FCustomizableObjectProjectorParameterValue>& FCustomizableObjectInstanceDescriptor::GetProjectorParameters()
+{
+	return ProjectorParameters;	
+}
+
+
+const TArray<FCustomizableObjectProjectorParameterValue>& FCustomizableObjectInstanceDescriptor::GetProjectorParameters() const
+{
+	return const_cast<FCustomizableObjectInstanceDescriptor*>(this)->GetProjectorParameters();	
+}
+
+
+bool FCustomizableObjectInstanceDescriptor::HasAnyParameters() const
+{
+	return !BoolParameters.IsEmpty() ||
+		!IntParameters.IsEmpty() ||
+		!FloatParameters.IsEmpty() || 
+		!TextureParameters.IsEmpty() ||
+		!ProjectorParameters.IsEmpty() ||
+		!VectorParameters.IsEmpty();
+}
+
+
+const FString& FCustomizableObjectInstanceDescriptor::GetIntParameterSelectedOption(const FString& ParamName, const int32 RangeIndex) const
+{
+	const int32 ParameterIndexInObject = CustomizableObject->FindParameter(ParamName);
+
+	const int32 IntParamIndex = FindIntParameterNameIndex(ParamName);
+
+	if (ParameterIndexInObject >= 0 && IntParamIndex >= 0)
+	{
+		if (RangeIndex == -1)
+		{
+			check(!IsParamMultidimensional(ParameterIndexInObject)); // This param is multidimensional, it must have a RangeIndex of 0 or more
+			return IntParameters[IntParamIndex].ParameterValueName;
+		}
+		else
+		{
+			check(IsParamMultidimensional(ParameterIndexInObject)); // This param is not multidimensional, it must have a RangeIndex of -1
+
+			if (IntParameters[IntParamIndex].ParameterRangeValueNames.IsValidIndex(RangeIndex))
+			{
+				return IntParameters[IntParamIndex].ParameterRangeValueNames[RangeIndex];
+			}
+		}
+	}
+
+	return EmptyString;
+}
+
+
+void FCustomizableObjectInstanceDescriptor::SetIntParameterSelectedOption(const int32 IntParamIndex, const FString& SelectedOption, const int32 RangeIndex)
+{
+	check(IntParameters.IsValidIndex(IntParamIndex));
+
+	if (IntParameters.IsValidIndex(IntParamIndex))
+	{
+		const int32 ParameterIndexInObject = CustomizableObject->FindParameter(IntParameters[IntParamIndex].ParameterName);
+		if (ParameterIndexInObject >= 0)
+		{
+			const bool bValid = SelectedOption == TEXT("None") || CustomizableObject->FindIntParameterValue(ParameterIndexInObject, SelectedOption) >= 0;
+			if (!bValid)
+			{
+				const FString Message = FString::Printf(
+					TEXT("LogMutable: Tried to set the invalid value [%s] to parameter [%d, %s]! Value index=[%d]. Correct values=[%s]."), 
+					*SelectedOption, ParameterIndexInObject,
+					*IntParameters[IntParamIndex].ParameterName, 
+					CustomizableObject->FindIntParameterValue(ParameterIndexInObject, SelectedOption), 
+					*GetAvailableOptionsString(*CustomizableObject, ParameterIndexInObject)
+				);
+				UE_LOG(LogMutable, Error, TEXT("%s"), *Message);
+			}
+			
+			if (RangeIndex == -1)
+			{
+				check(!IsParamMultidimensional(ParameterIndexInObject)); // This param is multidimensional, it must have a RangeIndex of 0 or more
+				IntParameters[IntParamIndex].ParameterValueName = SelectedOption;
+			}
+			else
+			{
+				check(IsParamMultidimensional(ParameterIndexInObject)); // This param is not multidimensional, it must have a RangeIndex of -1
+
+				if (!IntParameters[IntParamIndex].ParameterRangeValueNames.IsValidIndex(RangeIndex))
+				{
+					const int32 InsertionIndex = IntParameters[IntParamIndex].ParameterRangeValueNames.Num();
+					const int32 NumInsertedElements = RangeIndex + 1 - IntParameters[IntParamIndex].ParameterRangeValueNames.Num();
+					IntParameters[IntParamIndex].ParameterRangeValueNames.InsertDefaulted(InsertionIndex, NumInsertedElements);
+				}
+
+				check(IntParameters[IntParamIndex].ParameterRangeValueNames.IsValidIndex(RangeIndex));
+				IntParameters[IntParamIndex].ParameterRangeValueNames[RangeIndex] = SelectedOption;
+			}
+		}
+	}
+}
+
+
+void FCustomizableObjectInstanceDescriptor::SetIntParameterSelectedOption(const FString& ParamName, const FString& SelectedOptionName, const int32 RangeIndex)
+{
+	const int32 IntParamIndex = FindIntParameterNameIndex(ParamName);
+	SetIntParameterSelectedOption(IntParamIndex, SelectedOptionName, RangeIndex);
+}
+
+
+float FCustomizableObjectInstanceDescriptor::GetFloatParameterSelectedOption(const FString& FloatParamName, const int32 RangeIndex) const
+{
+	const int32 ParameterIndexInObject = CustomizableObject->FindParameter(FloatParamName);
+
+	const int32 FloatParamIndex = FindFloatParameterNameIndex(FloatParamName);
+
+	if (ParameterIndexInObject >= 0 && FloatParamIndex >= 0)
+	{
+		if (RangeIndex == -1)
+		{
+			check(!IsParamMultidimensional(ParameterIndexInObject)); // This param is multidimensional, it must have a RangeIndex of 0 or more
+			return FloatParameters[FloatParamIndex].ParameterValue;
+		}
+		else
+		{
+			check(IsParamMultidimensional(ParameterIndexInObject)); // This param is not multidimensional, it must have a RangeIndex of -1
+
+			if (FloatParameters[FloatParamIndex].ParameterRangeValues.IsValidIndex(RangeIndex))
+			{
+				return FloatParameters[FloatParamIndex].ParameterRangeValues[RangeIndex];
+			}
+		}
+	}
+
+	return -1.0f; 
+}
+
+
+void FCustomizableObjectInstanceDescriptor::SetFloatParameterSelectedOption(const FString& FloatParamName, const float FloatValue, const int32 RangeIndex)
+{
+	const int32 ParameterIndexInObject = CustomizableObject->FindParameter(FloatParamName);
+
+	const int32 FloatParamIndex = FindFloatParameterNameIndex(FloatParamName);
+
+	if (ParameterIndexInObject >= 0 && FloatParamIndex >= 0)
+	{
+		if (RangeIndex == -1)
+		{
+			check(!IsParamMultidimensional(ParameterIndexInObject)); // This param is multidimensional, it must have a RangeIndex of 0 or more
+			FloatParameters[FloatParamIndex].ParameterValue = FloatValue;
+		}
+		else
+		{
+			check(IsParamMultidimensional(ParameterIndexInObject)); // This param is not multidimensional, it must have a RangeIndex of -1
+
+			if (!FloatParameters[FloatParamIndex].ParameterRangeValues.IsValidIndex(RangeIndex))
+			{
+				const int32 InsertionIndex = FloatParameters[FloatParamIndex].ParameterRangeValues.Num();
+				const int32 NumInsertedElements = RangeIndex + 1 - FloatParameters[FloatParamIndex].ParameterRangeValues.Num();
+				FloatParameters[FloatParamIndex].ParameterRangeValues.InsertDefaulted(InsertionIndex, NumInsertedElements);
+			}
+
+			check(FloatParameters[FloatParamIndex].ParameterRangeValues.IsValidIndex(RangeIndex));
+			FloatParameters[FloatParamIndex].ParameterRangeValues[RangeIndex] = FloatValue;
+		}
+	}
+}
+
+
+FLinearColor FCustomizableObjectInstanceDescriptor::GetColorParameterSelectedOption(const FString& ColorParamName) const
+{
+	const int32 ParameterIndexInObject = CustomizableObject->FindParameter(ColorParamName);
+
+	const int32 ColorParamIndex = FindVectorParameterNameIndex(ColorParamName);
+
+	if (ParameterIndexInObject >= 0 && ColorParamIndex >= 0)
+	{
+		return VectorParameters[ColorParamIndex].ParameterValue;
+	}
+
+	return FLinearColor::Black;
+}
+
+
+void FCustomizableObjectInstanceDescriptor::SetColorParameterSelectedOption(const FString& ColorParamName, const FLinearColor& ColorValue)
+{
+	SetVectorParameterSelectedOption(ColorParamName, ColorValue);
+}
+
+
+bool FCustomizableObjectInstanceDescriptor::GetBoolParameterSelectedOption(const FString& BoolParamName) const
+{
+	const int32 ParameterIndexInObject = CustomizableObject->FindParameter(BoolParamName);
+
+	const int32 BoolParamIndex = FindBoolParameterNameIndex(BoolParamName);
+
+	if (ParameterIndexInObject >= 0 && BoolParamIndex >= 0)
+	{
+		return BoolParameters[BoolParamIndex].ParameterValue;
+	}
+
+	return false;
+}
+
+
+void FCustomizableObjectInstanceDescriptor::SetBoolParameterSelectedOption(const FString& BoolParamName, const bool BoolValue)
+{
+	const int32 ParameterIndexInObject = CustomizableObject->FindParameter(BoolParamName);
+
+	const int32 BoolParamIndex = FindBoolParameterNameIndex(BoolParamName);
+
+	if (ParameterIndexInObject >= 0 && BoolParamIndex >= 0)
+	{
+		BoolParameters[BoolParamIndex].ParameterValue = BoolValue;
+	}
+}
+
+
+void FCustomizableObjectInstanceDescriptor::SetVectorParameterSelectedOption(const FString& VectorParamName, const FLinearColor& VectorValue)
+{
+	const int32 ParameterIndexInObject = CustomizableObject->FindParameter(VectorParamName);
+
+	const int32 VectorParamIndex = FindVectorParameterNameIndex(VectorParamName);
+
+	if (ParameterIndexInObject >= 0	&& VectorParamIndex >= 0)
+	{
+		VectorParameters[VectorParamIndex].ParameterValue = VectorValue;
+	}
+}
+
+
+void FCustomizableObjectInstanceDescriptor::SetProjectorValue(const FString& ProjectorParamName,
+	const FVector& Pos, const FVector& Direction, const FVector& Up, const FVector& Scale,
+	const float Angle,
+	const int32 RangeIndex)
+{
+	const int32 ParameterIndexInObject = CustomizableObject->FindParameter(ProjectorParamName);
+
+	const int32 ProjectorParamIndex = FindProjectorParameterNameIndex(ProjectorParamName);
+
+	if (ParameterIndexInObject >= 0 && ProjectorParamIndex >= 0)
+	{
+		FCustomizableObjectProjector ProjectorData;
+		ProjectorData.Position = static_cast<FVector3f>(Pos);
+		ProjectorData.Direction = static_cast<FVector3f>(Direction);
+		ProjectorData.Up = static_cast<FVector3f>(Up);
+		ProjectorData.Scale = static_cast<FVector3f>(Scale);
+		ProjectorData.Angle = Angle;
+		ProjectorData.ProjectionType = ProjectorParameters[ProjectorParamIndex].Value.ProjectionType;
+
+		if (RangeIndex == -1)
+		{
+			check(!IsParamMultidimensional(ParameterIndexInObject)); // This param is multidimensional, it must have a RangeIndex of 0 or more
+			ProjectorParameters[ProjectorParamIndex].Value = ProjectorData;
+		}
+		else
+		{
+			check(IsParamMultidimensional(ParameterIndexInObject)); // This param is not multidimensional, it must have a RangeIndex of -1
+
+			if (!ProjectorParameters[ProjectorParamIndex].RangeValues.IsValidIndex(RangeIndex))
+			{
+				const int32 InsertionIndex = ProjectorParameters[ProjectorParamIndex].RangeValues.Num();
+				const int32 NumInsertedElements = RangeIndex + 1 - ProjectorParameters[ProjectorParamIndex].RangeValues.Num();
+				ProjectorParameters[ProjectorParamIndex].RangeValues.InsertDefaulted(InsertionIndex, NumInsertedElements);
+			}
+
+			check(ProjectorParameters[ProjectorParamIndex].RangeValues.IsValidIndex(RangeIndex));
+			ProjectorParameters[ProjectorParamIndex].RangeValues[RangeIndex] = ProjectorData;
+		}
+	}
+}
+
+
+void FCustomizableObjectInstanceDescriptor::SetProjectorPosition(const FString& ProjectorParamName, const FVector3f& Pos, const int32 RangeIndex)
+{
+	const int32 ParameterIndexInObject = CustomizableObject->FindParameter(ProjectorParamName);
+
+	const int32 ProjectorParamIndex = FindProjectorParameterNameIndex(ProjectorParamName);
+
+	if (ParameterIndexInObject >= 0	&& ProjectorParamIndex >= 0)
+	{
+		FCustomizableObjectProjector ProjectorData = ProjectorParameters[ProjectorParamIndex].Value;
+		ProjectorData.Position = static_cast<FVector3f>(Pos);
+
+		if (RangeIndex == -1)
+		{
+			check(!IsParamMultidimensional(ParameterIndexInObject)); // This param is multidimensional, it must have a RangeIndex of 0 or more
+			ProjectorParameters[ProjectorParamIndex].Value = ProjectorData;
+		}
+		else
+		{
+			check(IsParamMultidimensional(ParameterIndexInObject)); // This param is not multidimensional, it must have a RangeIndex of -1
+
+			if (!ProjectorParameters[ProjectorParamIndex].RangeValues.IsValidIndex(RangeIndex))
+			{
+				const int32 InsertionIndex = ProjectorParameters[ProjectorParamIndex].RangeValues.Num();
+				const int32 NumInsertedElements = RangeIndex + 1 - ProjectorParameters[ProjectorParamIndex].RangeValues.Num();
+				ProjectorParameters[ProjectorParamIndex].RangeValues.InsertDefaulted(InsertionIndex, NumInsertedElements);
+			}
+
+			check(ProjectorParameters[ProjectorParamIndex].RangeValues.IsValidIndex(RangeIndex));
+			ProjectorParameters[ProjectorParamIndex].RangeValues[RangeIndex] = ProjectorData;
+		}
+	}
+}
+
+
+void FCustomizableObjectInstanceDescriptor::GetProjectorValue(const FString& ProjectorParamName,
+	FVector& OutPos, FVector& OutDirection, FVector& OutUp, FVector& OutScale,
+	float& OutAngle, ECustomizableObjectProjectorType& OutType,
+	const int32 RangeIndex) const
+{
+	const int32 ParameterIndexInObject = CustomizableObject->FindParameter(ProjectorParamName);
+
+	const int32 ProjectorParamIndex = FindProjectorParameterNameIndex(ProjectorParamName);
+
+	if (ParameterIndexInObject >= 0 && ProjectorParamIndex >= 0)
+	{
+		const FCustomizableObjectProjector* Projector;
+
+		if (RangeIndex == -1)
+		{
+			check(!IsParamMultidimensional(ParameterIndexInObject)); // This param is multidimensional, it must have a RangeIndex of 0 or more
+
+			Projector = &ProjectorParameters[ProjectorParamIndex].Value;
+		}
+		else
+		{
+			check(IsParamMultidimensional(ParameterIndexInObject)); // This param is not multidimensional, it must have a RangeIndex of -1
+			check(ProjectorParameters[ProjectorParamIndex].RangeValues.IsValidIndex(RangeIndex));
+
+			Projector = &ProjectorParameters[ProjectorParamIndex].RangeValues[RangeIndex];
+		}
+
+		if (Projector)
+		{
+			OutPos = static_cast<FVector>(Projector->Position);
+			OutDirection = static_cast<FVector>(Projector->Direction);
+			OutUp = static_cast<FVector>(Projector->Up);
+			OutScale = static_cast<FVector>(Projector->Scale);
+			OutAngle = Projector->Angle;
+			OutType = Projector->ProjectionType;
+		}
+	}
+}
+
+
+void FCustomizableObjectInstanceDescriptor::GetProjectorValueF(const FString& ProjectorParamName,
+	FVector3f& OutPos, FVector3f& OutDirection, FVector3f& OutUp, FVector3f& OutScale,
+	float& OutAngle, ECustomizableObjectProjectorType& OutType,
+	const int32 RangeIndex) const
+{
+	const int32 ParameterIndexInObject = CustomizableObject->FindParameter(ProjectorParamName);
+
+	const int32 ProjectorParamIndex = FindProjectorParameterNameIndex(ProjectorParamName);
+
+	if (ParameterIndexInObject >= 0 && ProjectorParamIndex >= 0)
+	{
+		const FCustomizableObjectProjector* Projector;
+
+		if (RangeIndex == -1)
+		{
+			check(!IsParamMultidimensional(ParameterIndexInObject)); // This param is multidimensional, it must have a RangeIndex of 0 or more
+
+			Projector = &ProjectorParameters[ProjectorParamIndex].Value;
+		}
+		else
+		{
+			check(IsParamMultidimensional(ParameterIndexInObject)); // This param is not multidimensional, it must have a RangeIndex of -1
+			check(ProjectorParameters[ProjectorParamIndex].RangeValues.IsValidIndex(RangeIndex));
+
+			Projector = &ProjectorParameters[ProjectorParamIndex].RangeValues[RangeIndex];
+		}
+
+		if (Projector)
+		{
+			OutPos = Projector->Position;
+			OutDirection = Projector->Direction;
+			OutUp = Projector->Up;
+			OutScale = Projector->Scale;
+			OutAngle = Projector->Angle;
+			OutType = Projector->ProjectionType;
+		}
+	}
+}
+
+
+FVector FCustomizableObjectInstanceDescriptor::GetProjectorPosition(const FString& ParamName, const int32 RangeIndex) const
+{
+	const int32 ParameterIndexInObject = CustomizableObject->FindParameter(ParamName);
+
+	const int32 ProjectorParamIndex = FindProjectorParameterNameIndex(ParamName);
+
+	if (ParameterIndexInObject >= 0 && ProjectorParamIndex >= 0)
+	{
+		if (RangeIndex == -1)
+		{
+			return static_cast<FVector>(ProjectorParameters[ProjectorParamIndex].Value.Position);
+		}
+		else if (ProjectorParameters[ProjectorParamIndex].RangeValues.IsValidIndex(RangeIndex))
+		{
+			return static_cast<FVector>(ProjectorParameters[ProjectorParamIndex].RangeValues[RangeIndex].Position);
+		}
+	}
+
+	return FVector(-0.0,-0.0,-0.0);
+}
+
+
+FVector FCustomizableObjectInstanceDescriptor::GetProjectorDirection(const FString& ParamName, const int32 RangeIndex) const
+{
+	const int32 ParameterIndexInObject = CustomizableObject->FindParameter(ParamName);
+
+	const int32 ProjectorParamIndex = FindProjectorParameterNameIndex(ParamName);
+
+	if (ParameterIndexInObject >= 0 && ProjectorParamIndex >= 0)
+	{
+		if (RangeIndex == -1)
+		{
+			return static_cast<FVector>(ProjectorParameters[ProjectorParamIndex].Value.Direction);
+		}
+		else if (ProjectorParameters[ProjectorParamIndex].RangeValues.IsValidIndex(RangeIndex))
+		{
+			return static_cast<FVector>(ProjectorParameters[ProjectorParamIndex].RangeValues[RangeIndex].Direction);
+		}
+	}
+
+	return FVector(-0.0, -0.0, -0.0);
+}
+
+
+FVector FCustomizableObjectInstanceDescriptor::GetProjectorUp(const FString& ParamName, const int32 RangeIndex) const
+{
+	const int32 ParameterIndexInObject = CustomizableObject->FindParameter(ParamName);
+
+	const int32 ProjectorParamIndex = FindProjectorParameterNameIndex(ParamName);
+
+	if (ParameterIndexInObject >= 0 && ProjectorParamIndex >= 0)
+	{
+		if (RangeIndex == -1)
+		{
+			return static_cast<FVector>(ProjectorParameters[ProjectorParamIndex].Value.Up);
+		}
+		else if (ProjectorParameters[ProjectorParamIndex].RangeValues.IsValidIndex(RangeIndex))
+		{
+			return static_cast<FVector>(ProjectorParameters[ProjectorParamIndex].RangeValues[RangeIndex].Up);
+		}
+	}
+
+	return FVector(-0.0, -0.0, -0.0);	
+}
+
+
+FVector FCustomizableObjectInstanceDescriptor::GetProjectorScale(const FString& ParamName, const int32 RangeIndex) const
+{
+	const int32 ParameterIndexInObject = CustomizableObject->FindParameter(ParamName);
+
+	const int32 ProjectorParamIndex = FindProjectorParameterNameIndex(ParamName);
+
+	if ((ParameterIndexInObject >= 0) && (ProjectorParamIndex >= 0))
+	{
+		if (RangeIndex == -1)
+		{
+			return static_cast<FVector>(ProjectorParameters[ProjectorParamIndex].Value.Scale);
+		}
+		else if (ProjectorParameters[ProjectorParamIndex].RangeValues.IsValidIndex(RangeIndex))
+		{
+			return static_cast<FVector>(ProjectorParameters[ProjectorParamIndex].RangeValues[RangeIndex].Scale);
+		}
+	}
+
+	return FVector(-0.0, -0.0, -0.0);
+}
+
+
+float FCustomizableObjectInstanceDescriptor::GetProjectorAngle(const FString& ParamName, const int32 RangeIndex) const
+{
+	const int32 ParameterIndexInObject = CustomizableObject->FindParameter(ParamName);
+
+	const int32 ProjectorParamIndex = FindProjectorParameterNameIndex(ParamName);
+
+	if (ParameterIndexInObject >= 0 && ProjectorParamIndex >= 0)
+	{
+		if (RangeIndex == -1)
+		{
+			return ProjectorParameters[ProjectorParamIndex].Value.Angle;
+		}
+		else if (ProjectorParameters[ProjectorParamIndex].RangeValues.IsValidIndex(RangeIndex))
+		{
+			return ProjectorParameters[ProjectorParamIndex].RangeValues[RangeIndex].Angle;
+		}
+	}
+
+	return 0.0;
+}
+
+
+ECustomizableObjectProjectorType FCustomizableObjectInstanceDescriptor::GetProjectorParameterType(const FString& ParamName, const int32 RangeIndex) const
+{
+	const int32 ParameterIndexInObject = CustomizableObject->FindParameter(ParamName);
+
+	const int32 ProjectorParamIndex = FindProjectorParameterNameIndex(ParamName);
+
+	if (ParameterIndexInObject >= 0 && ProjectorParamIndex >= 0)
+	{
+		if (RangeIndex == -1)
+		{
+			return ProjectorParameters[ProjectorParamIndex].Value.ProjectionType;
+		}
+		else if (ProjectorParameters[ProjectorParamIndex].RangeValues.IsValidIndex(RangeIndex))
+		{
+			return ProjectorParameters[ProjectorParamIndex].RangeValues[RangeIndex].ProjectionType;
+		}
+	}
+
+	return ECustomizableObjectProjectorType::Planar;
+}
+
+
+FCustomizableObjectProjector FCustomizableObjectInstanceDescriptor::GetProjector(const FString& ParamName, const int32 RangeIndex) const
+{
+	const int32 Index = FindProjectorParameterNameIndex(ParamName);
+
+	if (Index != -1)
+	{
+		if (RangeIndex == -1)
+		{
+			return ProjectorParameters[Index].Value;
+		}
+		else if (ProjectorParameters[Index].RangeValues.IsValidIndex(RangeIndex))
+		{
+			return ProjectorParameters[Index].RangeValues[RangeIndex];
+		}
+	}
+
+	return FCustomizableObjectProjector();
+}
+
+
+int32 FCustomizableObjectInstanceDescriptor::FindIntParameterNameIndex(const FString& ParamName) const
+{
+	if (const int32* Found = IntParametersLookupTable.Find(ParamName))
+	{
+		return *Found;
+	}
+	else
+	{
+		return INDEX_NONE;
+	}
+}
+
+
+int32 FCustomizableObjectInstanceDescriptor::FindFloatParameterNameIndex(const FString& ParamName) const
+{
+	for (int32 i = 0; i < FloatParameters.Num(); ++i)
+	{
+		if (FloatParameters[i].ParameterName == ParamName)
+		{
+			return i;
+		}
+	}
+
+	return -1;
+}
+
+
+int32 FCustomizableObjectInstanceDescriptor::FindBoolParameterNameIndex(const FString& ParamName) const
+{
+	for (int32 i = 0; i < BoolParameters.Num(); ++i)
+	{
+		if (BoolParameters[i].ParameterName == ParamName)
+		{
+			return i;
+		}
+	}
+
+	return -1;
+}
+
+
+
+int32 FCustomizableObjectInstanceDescriptor::FindVectorParameterNameIndex(const FString& ParamName) const
+{
+	for (int32 i = 0; i < VectorParameters.Num(); ++i)
+	{
+		if (VectorParameters[i].ParameterName == ParamName)
+		{
+			return i;
+		}
+	}
+
+	return -1;
+}
+
+
+int32 FCustomizableObjectInstanceDescriptor::FindProjectorParameterNameIndex(const FString& ParamName) const
+{
+	for (int32 i = 0; i < ProjectorParameters.Num(); ++i)
+	{
+		if (ProjectorParameters[i].ParameterName == ParamName)
+		{
+			return i;
+		}
+	}
+
+	return -1;
+}
+
+
+bool FCustomizableObjectInstanceDescriptor::IsParamMultidimensional(const FString& ParamName) const
+{
+	const int32 ParameterIndex = CustomizableObject->FindParameter(ParamName);
+	return IsParamMultidimensional(ParameterIndex);
+}
+
+
+bool FCustomizableObjectInstanceDescriptor::IsParamMultidimensional(const int32 ParamIndex) const
+{
+	const mu::ParametersPtr MutableParameters = CustomizableObject->GetPrivate()->GetModel()->NewParameters();
+	check(ParamIndex < MutableParameters->GetCount());
+	const mu::RangeIndexPtr RangeIdxPtr = MutableParameters->NewRangeIndex(ParamIndex);
+
+	return RangeIdxPtr.get() != nullptr;
+}
+
+
+int32 FCustomizableObjectInstanceDescriptor::CurrentParamRange(const FString& ParamName) const
+{
+	const int32 ParameterIndexInObject = CustomizableObject->FindParameter(ParamName);
+
+	const int32 ProjectorParamIndex = FindProjectorParameterNameIndex(ParamName);
+
+	if (ParameterIndexInObject >= 0 && ProjectorParamIndex >= 0)
+	{
+		return ProjectorParameters[ProjectorParamIndex].RangeValues.Num();
+	}
+	else
+	{
+		return 0;
+	}
+}
+
+
+int32 FCustomizableObjectInstanceDescriptor::AddValueToIntRange(const FString& ParamName)
+{
+	const int32 intParameterIndex = FindIntParameterNameIndex(ParamName);
+	if (intParameterIndex != -1)
+	{
+		FCustomizableObjectIntParameterValue& intParameter = IntParameters[intParameterIndex];
+		const int32 ParamIndexInObject = CustomizableObject->FindParameter(intParameter.ParameterName);
+		// TODO: Define the default option in the editor instead of taking the first available, like it's currently defined for GetProjectorDefaultValue()
+		const FString DefaultValue = CustomizableObject->GetIntParameterAvailableOption(ParamIndexInObject, 0);
+		intParameter.ParameterRangeValueNames.Add(DefaultValue);
+		return intParameter.ParameterRangeValueNames.Num() - 1;
+	}
+	return -1;
+}
+
+
+int32 FCustomizableObjectInstanceDescriptor::AddValueToFloatRange(const FString& ParamName)
+{
+	const int32 floatParameterIndex = FindFloatParameterNameIndex(ParamName);
+	if (floatParameterIndex != -1)
+	{
+		FCustomizableObjectFloatParameterValue& floatParameter = FloatParameters[floatParameterIndex];
+		// TODO: Define the default float in the editor instead of [0.5f], like it's currently defined for GetProjectorDefaultValue()
+		floatParameter.ParameterRangeValues.Add(0.5f);
+		return floatParameter.ParameterRangeValues.Num() - 1;
+	}
+	return -1;
+}
+
+
+int32 FCustomizableObjectInstanceDescriptor::AddValueToProjectorRange(const FString& ParamName)
+{
+	const int32 projectorParameterIndex = FindProjectorParameterNameIndex(ParamName);
+	if (projectorParameterIndex != -1)
+	{
+		FCustomizableObjectProjectorParameterValue& ProjectorParameter = ProjectorParameters[projectorParameterIndex];
+		const FCustomizableObjectProjector Projector = GetProjectorDefaultValue(CustomizableObject->FindParameter(ParamName));
+		ProjectorParameter.RangeValues.Add(Projector);
+		return ProjectorParameter.RangeValues.Num() - 1;
+	}
+
+	return -1;
+}
+
+
+int32 FCustomizableObjectInstanceDescriptor::RemoveValueFromIntRange(const FString& ParamName)
+{
+	const int32 IntParameterIndex = FindIntParameterNameIndex(ParamName);
+	if (IntParameterIndex != -1)
+	{
+		FCustomizableObjectIntParameterValue& IntParameter = IntParameters[IntParameterIndex];
+		if (IntParameter.ParameterRangeValueNames.Num() > 0)
+		{
+			IntParameter.ParameterRangeValueNames.Pop();
+			return IntParameter.ParameterRangeValueNames.Num() - 1;
+		}
+	}
+	return -1;
+}
+
+
+int32 FCustomizableObjectInstanceDescriptor::RemoveValueFromIntRange(const FString& ParamName, const int32 RangeIndex)
+{
+	const int32 IntParameterIndex = FindIntParameterNameIndex(ParamName);
+	if (IntParameterIndex != -1)
+	{
+		FCustomizableObjectIntParameterValue& IntParameter = IntParameters[IntParameterIndex];
+		if (IntParameter.ParameterRangeValueNames.Num() > 0)
+		{
+			IntParameter.ParameterRangeValueNames.RemoveAt(RangeIndex);
+			return IntParameter.ParameterRangeValueNames.Num() - 1;
+		}
+	}
+	return -1;
+}
+
+
+int32 FCustomizableObjectInstanceDescriptor::RemoveValueFromFloatRange(const FString& ParamName)
+{
+	const int32 FloatParameterIndex = FindFloatParameterNameIndex(ParamName);
+	if (FloatParameterIndex != -1)
+	{
+		FCustomizableObjectFloatParameterValue& FloatParameter = FloatParameters[FloatParameterIndex];
+		if (FloatParameter.ParameterRangeValues.Num() > 0)
+		{
+			FloatParameter.ParameterRangeValues.Pop();
+			return FloatParameter.ParameterRangeValues.Num() - 1;
+		}
+	}
+	return -1;
+}
+
+
+int32 FCustomizableObjectInstanceDescriptor::RemoveValueFromFloatRange(const FString& ParamName, const int32 RangeIndex)
+{
+	const int32 FloatParameterIndex = FindFloatParameterNameIndex(ParamName);
+	if (FloatParameterIndex != -1)
+	{
+		FCustomizableObjectFloatParameterValue& FloatParameter = FloatParameters[FloatParameterIndex];
+		if (FloatParameter.ParameterRangeValues.Num() > 0)
+		{
+			FloatParameter.ParameterRangeValues.RemoveAt(RangeIndex);
+			return FloatParameter.ParameterRangeValues.Num() - 1;
+		}
+	}
+	return -1;
+}
+
+
+int32 FCustomizableObjectInstanceDescriptor::RemoveValueFromProjectorRange(const FString& ParamName)
+{
+	const int32 ProjectorParameterIndex = FindProjectorParameterNameIndex(ParamName);
+	if (ProjectorParameterIndex != -1)
+	{
+		FCustomizableObjectProjectorParameterValue& ProjectorParameter = ProjectorParameters[ProjectorParameterIndex];
+		if (ProjectorParameter.RangeValues.Num() > 0)
+		{
+			ProjectorParameter.RangeValues.Pop();
+			return ProjectorParameter.RangeValues.Num() - 1;
+		}
+	}
+	return -1;
+}
+
+
+int32 FCustomizableObjectInstanceDescriptor::RemoveValueFromProjectorRange(const FString& ParamName, const int32 RangeIndex)
+{
+	const int32 ProjectorParameterIndex = FindProjectorParameterNameIndex(ParamName);
+	if (ProjectorParameterIndex != -1)
+	{
+		FCustomizableObjectProjectorParameterValue& ProjectorParameter = ProjectorParameters[ProjectorParameterIndex];
+		if (ProjectorParameter.RangeValues.Num() > 0)
+		{
+			ProjectorParameter.RangeValues.RemoveAt(RangeIndex);
+			return ProjectorParameter.RangeValues.Num() - 1;
+		}
+	}
+	return -1;
+}
+
+
+FCustomizableObjectProjector FCustomizableObjectInstanceDescriptor::GetProjectorDefaultValue(int32 const ParamIndex) const
+{
+	const mu::ParametersPtr MutableParameters = CustomizableObject->GetPrivate()->GetModel()->NewParameters();
+	check(ParamIndex < MutableParameters->GetCount());
+
+	FCustomizableObjectProjector Projector;
+
+	mu::PROJECTOR_TYPE type;
+	
+	MutableParameters->GetProjectorValue(ParamIndex,
+		&type,
+		&Projector.Position[0], &Projector.Position[1], &Projector.Position[2],
+		&Projector.Direction[0], &Projector.Direction[1], &Projector.Direction[2],
+		&Projector.Up[0], &Projector.Up[1], &Projector.Up[2],
+		&Projector.Scale[0], &Projector.Scale[1], &Projector.Scale[2],
+		&Projector.Angle,
+		nullptr);
+
+	switch (type)
+	{
+	case mu::PROJECTOR_TYPE::PLANAR:
+		Projector.ProjectionType = ECustomizableObjectProjectorType::Planar;
+		break;
+
+	case mu::PROJECTOR_TYPE::CYLINDRICAL:
+		// Unapply strange swizzle for scales.
+		// TODO: try to avoid this
+		Projector.ProjectionType = ECustomizableObjectProjectorType::Cylindrical;
+		Projector.Direction = -Projector.Direction;
+		Projector.Up = -Projector.Up;
+		Projector.Scale[2] = -Projector.Scale[0];
+		Projector.Scale[0] = Projector.Scale[1] = Projector.Scale[1] * 2.0f;
+		break;
+
+	case mu::PROJECTOR_TYPE::WRAPPING:
+		Projector.ProjectionType = ECustomizableObjectProjectorType::Wrapping;
+		break;
+	default:
+		unimplemented()
+	}
+
+	return Projector;
+}
+
+
+int32 FCustomizableObjectInstanceDescriptor::GetState() const
+{
+	return State;
+}
+
+
+FString FCustomizableObjectInstanceDescriptor::GetCurrentState() const
+{
+	return CustomizableObject->GetStateName(GetState());
+}
+
+
+void FCustomizableObjectInstanceDescriptor::SetState(const int32 InState)
+{
+	State = InState;
+}
+
+
+void FCustomizableObjectInstanceDescriptor::SetCurrentState(const FString& StateName)
+{
+	SetState(CustomizableObject->FindState(StateName));
+}
+
+
+void FCustomizableObjectInstanceDescriptor::SetRandomValues()
+{
+	for (int32 i = 0; i < FloatParameters.Num(); ++i)
+	{
+		FloatParameters[i].ParameterValue = FMath::SRand();
+	}
+
+	for (int32 i = 0; i < BoolParameters.Num(); ++i)
+	{
+		BoolParameters[i].ParameterValue = FMath::Rand() % 2 == 0;
+	}
+
+	for (int32 i = 0; i < IntParameters.Num(); ++i)
+	{
+		const int32 ParameterIndexInCO = CustomizableObject->FindParameter(IntParameters[i].ParameterName);
+
+		// TODO: Randomize multidimensional parameters
+		if (ParameterIndexInCO >= 0 && !IsParamMultidimensional(ParameterIndexInCO))
+		{
+			const int32 NumValues = CustomizableObject->GetIntParameterNumOptions(ParameterIndexInCO);
+			if (NumValues > 0)
+			{
+				const int32 Index = FMath::Rand() % NumValues;
+				FString Option = CustomizableObject->GetIntParameterAvailableOption(ParameterIndexInCO, Index);
+				SetIntParameterSelectedOption(i, Option);
+			}
+		}
+	}
+}
+
+
+void FCustomizableObjectInstanceDescriptor::CreateParametersLookupTable()
+{
+	IntParametersLookupTable.Reset();
+	IntParametersLookupTable.Reserve(IntParameters.Num());
+
+	for (int32 Index = 0, Num = IntParameters.Num(); Index < Num; ++Index)
+	{
+		const FCustomizableObjectIntParameterValue & Value = IntParameters[Index];
+
+#if WITH_EDITOR
+		if (IntParametersLookupTable.Contains(Value.ParameterName))
+		{
+			int ExistIndex = IntParametersLookupTable[Value.ParameterName];
+
+			UE_LOG(LogMutable, Warning,
+				TEXT("Name '%s' is already in IntParametersLookupTable (%s/%s/%s/#%d)"),
+				*Value.ParameterName, *Value.ParameterName, *Value.ParameterValueName, *Value.Uid, ExistIndex);
+		}
+#endif
+
+		IntParametersLookupTable.Add(Value.ParameterName, Index);
+	}
+}
+
+
+uint32 GetTypeHash(const FCustomizableObjectInstanceDescriptor& Key)
+{
+	FCustomizableObjectIntParameterValue a;
+	
+	uint32 Hash = GetTypeHash(Key.CustomizableObject);
+
+	for (const FCustomizableObjectBoolParameterValue& Value : Key.BoolParameters)
+	{
+		Hash = HashCombine(Hash, GetTypeHash(Value));
+	}
+
+	for (const FCustomizableObjectIntParameterValue& Value : Key.IntParameters)
+	{
+		Hash = HashCombine(Hash, GetTypeHash(Value));
+	}
+	
+	for (const FCustomizableObjectFloatParameterValue& Value : Key.FloatParameters)
+	{
+		Hash = HashCombine(Hash, GetTypeHash(Value));
+	}
+
+	for (const FCustomizableObjectTextureParameterValue& Value : Key.TextureParameters)
+	{
+		Hash = HashCombine(Hash, GetTypeHash(Value));
+	}
+
+	for (const FCustomizableObjectVectorParameterValue& Value : Key.VectorParameters)
+	{
+		Hash = HashCombine(Hash, GetTypeHash(Value));
+	}
+
+	for (const FCustomizableObjectProjectorParameterValue& Value : Key.ProjectorParameters)
+	{
+		Hash = HashCombine(Hash, GetTypeHash(Value));
+	}
+	
+	Hash = HashCombine(Hash, GetTypeHash(Key.State));
+
+	return Hash;
+}
diff --git a/Engine/Plugins/Experimental/Mutable/Source/CustomizableObject/Private/CustomizableObjectMipDataProvider.cpp b/Engine/Plugins/Experimental/Mutable/Source/CustomizableObject/Private/CustomizableObjectMipDataProvider.cpp
new file mode 100644
--- /dev/null
+++ b/Engine/Plugins/Experimental/Mutable/Source/CustomizableObject/Private/CustomizableObjectMipDataProvider.cpp
@@ -0,0 +1,283 @@
+// Copyright Epic Games, Inc. All Rights Reserved.
+
+#include "CustomizableObjectMipDataProvider.h"
+#include "MutableRuntime/Public/Mesh.h"
+#include "CustomizableObjectSystemPrivate.h"
+
+UMutableTextureMipDataProviderFactory::UMutableTextureMipDataProviderFactory(const FObjectInitializer& ObjectInitializer)
+	: Super(ObjectInitializer)
+{
+
+}
+
+
+FMutableTextureMipDataProvider::FMutableTextureMipDataProvider(const UTexture* Texture, UCustomizableObjectInstance* InCustomizableObjectInstance, const FMutableImageReference& InImageRef)
+	: FTextureMipDataProvider(Texture, ETickState::Init, ETickThread::Async),
+	CustomizableObjectInstance(InCustomizableObjectInstance), ImageRef(InImageRef)
+{
+	check(ImageRef.ImageID > 0);
+}
+
+
+void FMutableTextureMipDataProvider::Init(const FTextureUpdateContext& Context, const FTextureUpdateSyncOptions& SyncOptions)
+{
+	AdvanceTo(ETickState::GetMips, ETickThread::Async);
+}
+
+namespace impl
+{
+	void Task_Mutable_UpdateImage(TSharedPtr<FMutableImageOperationData> OperationData)
+	{
+		MUTABLE_CPUPROFILER_SCOPE(Task_Mutable_UpdateImage);
+
+		// TODO we cannot access all this from the mips steaming thread
+		// Get any external texture that may be requested for this update
+		//MUTABLE_CPUPROFILER_SCOPE(GetExternalImages);
+		//if (System->GetPrivate()->ImageProvider)
+		//{
+		//	for (const FCustomizableObjectTextureParameterValue& Param : CandidateInstance->TextureParameters)
+		//	{
+		//		System->GetPrivate()->ImageProvider->CacheImage(Param.ParameterValue);
+		//	}
+		//}
+
+
+		// This runs in a worker thread.
+		check(OperationData.IsValid());
+		check(OperationData->UpdateContext->System.get());
+		check(OperationData->UpdateContext->Model.get());
+		check(OperationData->UpdateContext->Parameters.get());
+
+		mu::SystemPtr System = OperationData->UpdateContext->System;
+		mu::ModelPtr Model = OperationData->UpdateContext->Model;
+
+		// For now, we are forcing the recreation of mutable-side instances with every update.
+		mu::Instance::ID InstanceID = System->NewInstance(Model.get());
+		UE_LOG(LogMutable, Log, TEXT("Creating instance with id [%d] "), InstanceID)
+
+		const mu::Instance* Instance = nullptr;
+
+		// Main instance generation step
+		{
+			// LOD mask, set to all ones to build all LODs
+			uint32 LODMask = 0xFFFFFFFF;
+
+			Instance = System->BeginUpdate(InstanceID, OperationData->UpdateContext->Parameters, OperationData->UpdateContext->State, LODMask);
+			check(Instance);
+		}
+
+
+		// Generate the required image
+		{
+			MUTABLE_CPUPROFILER_SCOPE(RequestedImage);
+
+			const FMutableImageReference& ImageRef = OperationData->RequestedImage;
+
+			// This ID may be different than the ID obtained the first time the image was generated, because the mutable
+			// runtime cannot remember all the resources it has built, and only remembers a fixed amount.
+			mu::RESOURCE_ID MipImageID = Instance->GetImageId(ImageRef.LOD, ImageRef.Component, ImageRef.Surface, ImageRef.Image);
+
+			// TODO: Why do we need to do this again? The full size should be stored in the initial image creation.
+			mu::FImageDesc ImageDesc;
+			mu::ImagePtrConst Image;
+			{
+				MUTABLE_CPUPROFILER_SCOPE(GetImage);
+
+				Image = System->GetImage(InstanceID, MipImageID, OperationData->MipsToSkip);
+			}
+
+			check(Image);
+
+			OperationData->Result = Image;
+		}
+
+		// End update
+		{
+			MUTABLE_CPUPROFILER_SCOPE(EndUpdate);
+			System->EndUpdate(InstanceID);
+		}
+
+		{
+			// The request could be cancelled in parallel from CancelCounterSafely and its value be changed
+			// between reading it and actually running Decrement() and RescheduleCallback(), so lock
+			FScopeLock Lock(&OperationData->CounterTaskLock);
+
+			if (OperationData->Counter) // If the request has been cancelled the counter will be null
+			{
+				// Make the FMutableTextureMipDataProvider continue
+				OperationData->Counter->Decrement();
+
+				if (OperationData->Counter->GetValue() == 0)
+				{
+					OperationData->RescheduleCallback();
+				}
+			}
+		}
+	}
+
+} // namespace
+
+
+int32 FMutableTextureMipDataProvider::GetMips(const FTextureUpdateContext& Context, int32 StartingMipIndex, const FTextureMipInfoArray& MipInfos, const FTextureUpdateSyncOptions& SyncOptions)
+{
+	const UTexture2D* Texture = Cast<UTexture2D>(Context.Texture);
+	check(Texture);
+	check(!Texture->NeverStream);
+	const TIndirectArray<FTexture2DMipMap>& OwnerMips = Texture->GetPlatformMips();
+
+	int32 NumMips = OwnerMips.Num();
+	check(ImageRef.ImageID > 0);
+
+	// Maximum value to skip, will be minimized by the first mip level requested
+	int32 MipsToSkip = 256;
+
+	for (int32 MipIndex = StartingMipIndex; MipIndex < CurrentFirstLODIdx; ++MipIndex)
+	{
+		check(Context.MipsView.IsValidIndex(MipIndex) && MipInfos.IsValidIndex(MipIndex));
+
+		const FTexture2DMipMap& MipMap = *Context.MipsView[MipIndex];
+		FByteBulkData* AuxBulkData = const_cast<FByteBulkData*>(&MipMap.BulkData);
+
+		const FTextureMipInfo& MipInfo = MipInfos[MipIndex];
+		void* Dest = MipInfo.DestData;
+
+		if (AuxBulkData->GetBulkDataSize() > 0)
+		{
+			// Mips are already generated, no need for Mutable progressive mip streaming, just normal CPU->GPU streaming
+			AuxBulkData->GetCopy(&Dest, false);
+		}
+		else // Generate a Mip Request to Mutable
+		{
+			check(UCustomizableObjectSystem::GetInstance()->GetPrivate()->EnableMutableProgressiveMipStreaming == 1);
+
+			if (!OperationData.IsValid())
+			{
+				OperationData = MakeShared<FMutableImageOperationData>();
+			}
+
+			const int32 LODBias = static_cast<int32>(Context.MipsView.GetData() - Texture->GetPlatformData()->Mips.GetData());
+			check(LODBias >= 0);
+			check(NumMips == Context.MipsView.Num() + LODBias);
+
+			OperationData->Levels.Add({ MipIndex + LODBias, Dest, (int32)MipInfo.SizeX, (int32)MipInfo.SizeY, (int32)MipInfo.DataSize, MipInfo.Format });
+			MipsToSkip = FMath::Min(MipsToSkip, MipIndex + LODBias);
+		}
+	}
+
+	FCustomizableObjectSystemPrivate* CustomizableObjectSystem = UCustomizableObjectSystem::GetInstance()->GetPrivate();
+	if (CustomizableObjectSystem)
+	{
+		if (OperationData.IsValid())
+		{
+			OperationData->RequestedImage = ImageRef;
+			OperationData->UpdateContext = UpdateContext;
+			OperationData->MipsToSkip = MipsToSkip;
+
+			check(SyncOptions.Counter);
+			// Increment to stop PollMips from running until the Mutable request task finishes. 
+			// If a request completes immediately, then it will call the callback
+			// but that won't do anything because the tick would not try to acquire the lock since it is already locked.
+			SyncOptions.Counter->Increment();
+			OperationData->Counter = SyncOptions.Counter;
+			OperationData->RescheduleCallback = SyncOptions.RescheduleCallback;
+
+			TSharedPtr<FMutableImageOperationData> LocalOperationData = OperationData;
+			UpdateImageMutableTaskEvent = CustomizableObjectSystem->AddMutableThreadTask(
+				TEXT("Mutable_MipUpdate"),
+				[LocalOperationData]()
+				{
+					impl::Task_Mutable_UpdateImage(LocalOperationData);
+				},
+				UE::Tasks::ETaskPriority::BackgroundHigh
+					);
+		}
+
+		AdvanceTo(ETickState::PollMips, ETickThread::Async);
+	}
+	else
+	{
+		AdvanceTo(ETickState::CleanUp, ETickThread::Async);
+	}
+
+	return CurrentFirstLODIdx;
+}
+
+
+bool FMutableTextureMipDataProvider::PollMips(const FTextureUpdateSyncOptions& SyncOptions)
+{
+	if (!bRequestAborted && UpdateImageMutableTaskEvent)
+	{
+		if (OperationData && OperationData->Result && OperationData->Levels.Num())
+		{
+			// The counter must be zero meaning the Mutable image operation has finished
+			check(SyncOptions.Counter->GetValue() == 0);
+
+			mu::ImagePtrConst Mip = OperationData->Result;
+			int32 MipIndex = 0;
+			check(Mip->GetSizeX() == OperationData->Levels[0].SizeX);
+			check(Mip->GetSizeY() == OperationData->Levels[0].SizeY);
+
+			for (FMutableMipUpdateLevel& Level : OperationData->Levels)
+			{
+				// Check Mip DataSize for consistency, but skip if 0 because it's optional and might be zero in cooked mips
+				check(Level.DataSize == 0 || Mip->GetLODDataSize(MipIndex) == Level.DataSize);
+				void* Dest = Level.Dest;
+				FMemory::Memcpy(Dest, Mip->GetMipData(MipIndex), Mip->GetLODDataSize(MipIndex));
+				++MipIndex;
+			}
+		}
+	}
+	else if (bRequestAborted)
+	{
+		OperationData = nullptr;
+		AdvanceTo(ETickState::CleanUp, ETickThread::Async);
+		return false;
+	}
+
+	OperationData = nullptr;
+	AdvanceTo(ETickState::Done, ETickThread::None);
+	return true;
+}
+
+
+void FMutableTextureMipDataProvider::CleanUp(const FTextureUpdateSyncOptions& SyncOptions)
+{
+	check(!SyncOptions.Counter || SyncOptions.Counter->GetValue() == 0);
+	AdvanceTo(ETickState::Done, ETickThread::None);
+}
+
+
+void FMutableTextureMipDataProvider::Cancel(const FTextureUpdateSyncOptions& SyncOptions)
+{
+	bRequestAborted = true;
+
+	CancelCounterSafely();
+}
+
+
+FTextureMipDataProvider::ETickThread FMutableTextureMipDataProvider::GetCancelThread() const
+{
+	//return OperationData ? FTextureMipDataProvider::ETickThread::Async : FTextureMipDataProvider::ETickThread::None;
+	return FTextureMipDataProvider::ETickThread::None;
+}
+
+
+void FMutableTextureMipDataProvider::AbortPollMips()
+{
+	bRequestAborted = true;
+
+	CancelCounterSafely();
+}
+
+
+void FMutableTextureMipDataProvider::CancelCounterSafely()
+{
+	// The Counter could be read in parallel from Task_Mutable_UpdateImage, so lock
+	FScopeLock Lock(&OperationData->CounterTaskLock);
+
+	if (OperationData->Counter)
+	{
+		OperationData->Counter->Set(0);
+		OperationData->Counter = nullptr;
+	}
+}
diff --git a/Engine/Plugins/Experimental/Mutable/Source/CustomizableObject/Private/CustomizableObjectModule.cpp b/Engine/Plugins/Experimental/Mutable/Source/CustomizableObject/Private/CustomizableObjectModule.cpp
new file mode 100644
--- /dev/null
+++ b/Engine/Plugins/Experimental/Mutable/Source/CustomizableObject/Private/CustomizableObjectModule.cpp
@@ -0,0 +1,71 @@
+// Copyright Epic Games, Inc. All Rights Reserved.
+
+#include "Modules/ModuleManager.h"
+
+#include "ICustomizableObjectModule.h"
+#include "CustomizableObjectSystem.h"
+
+#include "Interfaces/IPluginManager.h"
+#include "Misc/ConfigCacheIni.h"
+
+#include "Runtime/Launch/Resources/Version.h"
+
+#include "MutableRuntime/Public/MutableMemory.h"
+#include "MutableRuntime/Public/System.h"
+
+/**
+ * Customizable Object module implementation (private)
+ */
+class FCustomizableObjectModule : public ICustomizableObjectModule
+{
+public:
+
+	// IModuleInterface 
+	void StartupModule() override;
+	void ShutdownModule() override;
+
+	// ICustomizableObjectModule interface
+	FString GetPluginVersion() const override;
+	bool AreExtraBoneInfluencesEnabled() const override;
+};
+
+
+IMPLEMENT_MODULE( FCustomizableObjectModule, CustomizableObject );
+
+void FCustomizableObjectModule::StartupModule()
+{
+}
+
+
+void FCustomizableObjectModule::ShutdownModule()
+{
+}
+
+
+FString FCustomizableObjectModule::GetPluginVersion() const
+{
+	FString Version = "x.x";
+	TSharedPtr<IPlugin> Plugin = IPluginManager::Get().FindPlugin("Mutable");
+	if (Plugin.IsValid() && Plugin->IsEnabled())
+	{
+		Version = Plugin->GetDescriptor().VersionName;
+	}
+	return Version;
+}
+
+bool FCustomizableObjectModule::AreExtraBoneInfluencesEnabled() const
+{
+	bool bAreExtraBoneInfluencesEnabled = false;
+
+	FConfigFile* PluginConfig = GConfig->FindConfigFileWithBaseName("Mutable");
+	if (PluginConfig)
+	{
+		bool bValue = false;
+		if (PluginConfig->GetBool(TEXT("Features"), TEXT("bExtraBoneInfluencesEnabled"), bValue))
+		{
+			bAreExtraBoneInfluencesEnabled = bValue;
+		}
+	}
+
+	return bAreExtraBoneInfluencesEnabled;
+}
diff --git a/Engine/Plugins/Experimental/Mutable/Source/CustomizableObject/Private/CustomizableObjectPrivate.h b/Engine/Plugins/Experimental/Mutable/Source/CustomizableObject/Private/CustomizableObjectPrivate.h
new file mode 100644
--- /dev/null
+++ b/Engine/Plugins/Experimental/Mutable/Source/CustomizableObject/Private/CustomizableObjectPrivate.h
@@ -0,0 +1,32 @@
+// Copyright Epic Games, Inc. All Rights Reserved.
+
+#pragma once
+
+#include "MutableRuntime/Public/Model.h"
+#include "MutableRuntime/Public/Serialisation.h"
+
+
+
+class FCustomizableObjectPrivateData
+{
+private:
+
+	mu::ModelPtr MutableModel;
+
+public:
+
+	void SetModel(mu::Model* Model);
+	mu::Model* GetModel();
+	const mu::Model* GetModel() const;
+
+	// See UCustomizableObjectSystem::LockObject.
+	bool bLocked = false;
+
+#if WITH_EDITOR
+	bool bModelCompiledForCook = false;
+
+	TArray<FString> CachedPlatformNames;
+#endif
+
+};
+
