[mutable] Move Mutable plugin from the restricted folder to the experimental plugins folder.

diff --git a/Engine/Plugins/Experimental/Mutable/Source/CustomizableObject/Private/CustomizableObjectSystem.cpp b/Engine/Plugins/Experimental/Mutable/Source/CustomizableObject/Private/CustomizableObjectSystem.cpp
new file mode 100644
--- /dev/null
+++ b/Engine/Plugins/Experimental/Mutable/Source/CustomizableObject/Private/CustomizableObjectSystem.cpp
@@ -0,0 +1,2776 @@
+// Copyright Epic Games, Inc. All Rights Reserved.
+
+#include "CustomizableObjectSystem.h"
+
+#include "CustomizableObjectSystemPrivate.h"
+#include "CustomizableObject.h"
+#include "CustomizableObjectPrivate.h"
+#include "CustomizableInstancePrivateData.h"
+#include "UnrealMutableImageProvider.h"
+#include "UnrealMutableModelDiskStreamer.h"
+#include "LogBenchmarkUtil.h"
+#include "ICustomizableObjectModule.h"
+#include "CustomizableInstanceLODManagement.h"
+#include "UnrealBakeHelpers.h"
+
+#include "AssetRegistry/AssetRegistryModule.h"
+#include "Engine/SkeletalMesh.h"
+#include "Animation/Skeleton.h"
+#include "LogInformationUtil.h"
+#include "TimerManager.h"
+#include "SceneManagement.h"
+#include "GameFramework/PlayerController.h"
+#include "HAL/IConsoleManager.h"
+#include "Kismet/GameplayStatics.h"
+
+#if WITH_EDITOR
+#include "Editor.h"
+#include "Misc/ConfigCacheIni.h"
+#include "Misc/MessageDialog.h"
+#include "Logging/MessageLog.h"
+#include "Framework/Notifications/NotificationManager.h"
+#endif
+
+#include "MutableTrace.h"
+#include "MutableRuntime/Public/System.h"
+#include "MutableRuntime/Private/SystemPrivate.h"
+#include "MutableRuntime/Public/Parameters.h"
+
+
+DEFINE_STAT(STAT_MutableNumSkeletalMeshes);
+DEFINE_STAT(STAT_MutableNumCachedSkeletalMeshes);
+DEFINE_STAT(STAT_MutableNumAllocatedSkeletalMeshes);
+DEFINE_STAT(STAT_MutableNumInstancesLOD0);
+DEFINE_STAT(STAT_MutableNumInstancesLOD1);
+DEFINE_STAT(STAT_MutableNumInstancesLOD2);
+DEFINE_STAT(STAT_MutableSkeletalMeshResourceMemory);
+DEFINE_STAT(STAT_MutableNumTextures);
+DEFINE_STAT(STAT_MutableNumCachedTextures);
+DEFINE_STAT(STAT_MutableNumAllocatedTextures);
+DEFINE_STAT(STAT_MutableTextureResourceMemory);
+DEFINE_STAT(STAT_MutableTextureGeneratedMemory);
+DEFINE_STAT(STAT_MutableTextureCacheMemory);
+DEFINE_STAT(STAT_MutableTextureParameterDecorationMemory);
+DEFINE_STAT(STAT_MutablePendingInstanceUpdates);
+DEFINE_STAT(STAT_MutableAbandonedInstanceUpdates);
+DEFINE_STAT(STAT_MutableInstanceBuildTime);
+DEFINE_STAT(STAT_MutableInstanceBuildTimeAvrg);
+DEFINE_STAT(STAT_MutableStreamingOps);
+DEFINE_STAT(STAT_MutableStreamingCache);
+
+// These stats are provided by the mutable runtime
+DEFINE_STAT(STAT_MutableProfile_LiveInstanceCount);
+DEFINE_STAT(STAT_MutableProfile_StreamingCacheBytes);
+DEFINE_STAT(STAT_MutableProfile_InstanceUpdateCount);
+
+DECLARE_CYCLE_STAT(TEXT("MutablePendingRelease Time"), STAT_MutablePendingRelease, STATGROUP_Game);
+DECLARE_CYCLE_STAT(TEXT("MutableTask"), STAT_MutableTask, STATGROUP_Game);
+
+UCustomizableObjectSystem* FCustomizableObjectSystemPrivate::SSystem = nullptr;
+
+
+static TAutoConsoleVariable<int32> CVarStreamingMemory(
+	TEXT("b.MutableStreamingMemory"),
+	-1,
+	TEXT("If different than 0, limit the amount of memory (in KB) to use to cache streaming data when building characters. 0 means no limit, -1 means use default (40Mb in PC and 20Mb in consoles)."),
+	ECVF_Scalability);
+
+
+bool FMutableQueue::IsEmpty() const
+{
+	return Array.Num() == 0;
+}
+
+
+int FMutableQueue::Num() const
+{
+	return Array.Num();
+}
+
+
+void FMutableQueue::Enqueue(const FMutableQueueElem& TaskToEnqueue)
+{
+	for (FMutableQueueElem& QueueElem : Array)
+	{
+		if (QueueElem.Operation->CustomizableObjectInstance == TaskToEnqueue.Operation->CustomizableObjectInstance)
+		{
+			QueueElem.Operation = TaskToEnqueue.Operation;
+			QueueElem.PriorityType = FMath::Min(QueueElem.PriorityType, TaskToEnqueue.PriorityType);
+			QueueElem.Priority = FMath::Min(QueueElem.Priority, TaskToEnqueue.Priority);
+			
+			return;
+		}
+	}
+
+	Array.HeapPush(TaskToEnqueue);
+}
+
+
+void FMutableQueue::Dequeue(FMutableQueueElem* DequeuedTask)
+{
+	Array.HeapPop(*DequeuedTask);
+}
+
+
+void FMutableQueue::ChangePriorities()
+{
+	for (FMutableQueueElem& Elem : Array)
+	{
+		if (Elem.Operation->CustomizableObjectInstance.IsValid() && Elem.Operation->CustomizableObjectInstance->GetPrivate())
+		{
+			Elem.Priority = Elem.Operation->CustomizableObjectInstance->GetPrivate()->MinSquareDistFromComponentToPlayer;
+		}
+	}
+}
+
+
+void FMutableQueue::UpdatePriority(const UCustomizableObjectInstance* Instance)
+{
+	for (FMutableQueueElem& Elem : Array)
+	{
+		if (Elem.Operation->CustomizableObjectInstance == Instance)
+		{
+			Elem.Priority = Instance->GetPrivate()->MinSquareDistFromComponentToPlayer;
+			break;
+		}
+	}
+}
+
+
+const FMutableQueueElem* FMutableQueue::Get(const UCustomizableObjectInstance* Instance) const
+{
+	for (const FMutableQueueElem& Elem : Array)
+	{
+		if (Elem.Operation->CustomizableObjectInstance == Instance)
+		{
+			return &Elem;
+		}
+	}
+
+	return nullptr;
+}
+
+
+void FMutableQueue::Sort()
+{
+	Array.Heapify();
+}
+
+
+UCustomizableObjectSystem* UCustomizableObjectSystem::GetInstance()
+{
+	if (!FCustomizableObjectSystemPrivate::SSystem)
+	{
+		UE_LOG(LogMutable, Log, TEXT("Creating system."));
+
+		check(IsInGameThread());
+
+		FCustomizableObjectSystemPrivate::SSystem = NewObject<UCustomizableObjectSystem>(UCustomizableObjectSystem::StaticClass());
+		checkf(!GUObjectArray.IsDisregardForGC(FCustomizableObjectSystemPrivate::SSystem), TEXT("Mutable was initialized too early in the UE4 init process, for instance, in the constructor of a default UObject."));
+		FCustomizableObjectSystemPrivate::SSystem->AddToRoot();
+		checkf(!GUObjectArray.IsDisregardForGC(FCustomizableObjectSystemPrivate::SSystem), TEXT("Mutable was initialized too early in the UE4 init process, for instance, in the constructor of a default UObject."));
+		FCustomizableObjectSystemPrivate::SSystem->InitSystem();
+
+		//FCoreUObjectDelegates::PurgePendingReleaseSkeletalMesh.AddUObject(FCustomizableObjectSystemPrivate::SSystem, &UCustomizableObjectSystem::PurgePendingReleaseSkeletalMesh);
+	}
+
+	return FCustomizableObjectSystemPrivate::SSystem;
+}
+
+
+FString UCustomizableObjectSystem::GetPluginVersion() const
+{
+	// Bridge the call from the module. This implementation is available from blueprint.
+	return ICustomizableObjectModule::Get().GetPluginVersion();
+}
+
+
+
+void UCustomizableObjectSystem::LogShowData(bool bFullInfo, bool ShowMaterialInfo) const
+{
+	LogInformationUtil::ResetCounters();
+
+	AActor* ParentActor;
+	TArray<UCustomizableObjectInstance*> ArrayData;
+
+	for (TObjectIterator<UCustomizableSkeletalComponent> It; It; ++It)
+	{
+		UCustomizableSkeletalComponent* CustomizableSkeletalComponent = *It;
+
+		if ((CustomizableSkeletalComponent != nullptr) && (CustomizableSkeletalComponent->CustomizableObjectInstance != nullptr))
+		{
+			ParentActor = CustomizableSkeletalComponent->GetAttachmentRootActor();
+
+			if (ParentActor != nullptr)
+			{
+				ArrayData.AddUnique(CustomizableSkeletalComponent->CustomizableObjectInstance);
+			}
+		}
+	}
+
+	ArrayData.Sort([](UCustomizableObjectInstance& A, UCustomizableObjectInstance& B)
+	{
+		return (A.GetPrivate()->LastMinSquareDistFromComponentToPlayer < B.GetPrivate()->LastMinSquareDistFromComponentToPlayer);
+	});
+
+	int i;
+	const int Max = ArrayData.Num();
+
+	if (bFullInfo)
+	{
+		for (i = 0; i < Max; ++i)
+		{
+			LogInformationUtil::LogShowInstanceDataFull(ArrayData[i], ShowMaterialInfo);
+		}
+	}
+	else
+	{
+		FString LogData = "\n\n";
+		for (i = 0; i < Max; ++i)
+		{
+			LogInformationUtil::LogShowInstanceData(ArrayData[i], LogData);
+		}
+		UE_LOG(LogMutable, Warning, TEXT("%s"), *LogData);
+
+		UWorld* World = GWorld;
+
+		if (World)
+		{
+			APlayerController* PlayerController = World->GetFirstPlayerController();
+			if (PlayerController)
+			{
+				PlayerController->ClientMessage(LogData);
+			}
+		}
+	}
+}
+
+
+bool UCustomizableObjectSystem::IsCreated()
+{
+	return FCustomizableObjectSystemPrivate::SSystem != 0;
+}
+
+
+void UCustomizableObjectSystem::InitSystem()
+{
+	// Everything initialized in Init() instead of constructor to prevent the default UCustomizableObjectSystem from registering a tick function
+	Private = MakeShareable(new FCustomizableObjectSystemPrivate());
+	Private->NewCompilerFunc = nullptr;
+
+	Private->bReplaceDiscardedWithReferenceMesh = false;
+
+	const IConsoleVariable* CVarSupport16BitBoneIndex = IConsoleManager::Get().FindConsoleVariable(TEXT("r.GPUSkin.Support16BitBoneIndex"));
+	Private->bSupport16BitBoneIndex = CVarSupport16BitBoneIndex ? CVarSupport16BitBoneIndex->GetBool() : false;
+
+	Private->CurrentMutableOperation = nullptr;
+	Private->CurrentInstanceBeingUpdated = nullptr;
+
+#if !UE_SERVER
+	Private->TickDelegate = FTickerDelegate::CreateUObject(this, &UCustomizableObjectSystem::Tick);
+	Private->TickDelegateHandle = FTSTicker::GetCoreTicker().AddTicker(Private->TickDelegate, 0.f);
+#endif // !UE_SERVER
+
+	Private->TotalBuildMs = 0;
+	Private->TotalBuiltInstances = 0;
+	Private->NumInstances = 0;
+	Private->TextureMemoryUsed = 0;
+
+#if !(UE_BUILD_SHIPPING || UE_BUILD_TEST)
+	SET_DWORD_STAT(STAT_MutableNumSkeletalMeshes, 0);
+	SET_DWORD_STAT(STAT_MutableNumTextures, 0);
+	SET_DWORD_STAT(STAT_MutableInstanceBuildTime, 0);
+	SET_DWORD_STAT(STAT_MutableInstanceBuildTimeAvrg, 0);
+#endif
+
+	mu::SettingsPtr pSettings = new mu::Settings;
+	pSettings->SetProfile(false);
+	pSettings->SetStreamingCache(MUTABLE_STREAMING_CACHE);
+	Private->MutableSystem = new mu::System(pSettings);
+
+	Private->LastStreamingMemorySize = MUTABLE_STREAMING_CACHE;
+	Private->Streamer = new FUnrealMutableModelBulkStreamer();
+	Private->MutableSystem->SetStreamingInterface(Private->Streamer);
+
+	// Set up the external image provider, for image parameters.
+	FUnrealMutableImageProvider* Provider = new FUnrealMutableImageProvider();
+	Private->ImageProvider = Provider;
+	Private->MutableSystem->SetImageParameterGenerator(Provider);
+
+#if WITH_EDITOR
+	// Register an example provider of texture parameters, used by editor preview.
+	PreviewExternalImageProvider = NewObject<UCustomizableObjectImageProviderArray>();
+	RegisterImageProvider(PreviewExternalImageProvider);
+
+	if (!IsRunningGame())
+	{
+		FEditorDelegates::PreBeginPIE.AddUObject(this, &UCustomizableObjectSystem::OnPreBeginPIE);
+	}
+
+#endif
+
+	if (FParse::Param(FCommandLine::Get(), TEXT("MutablePortableSerialization")))
+	{
+		Private->bCompactSerialization = false;
+	}
+
+	DefaultInstanceLODManagement = NewObject<UCustomizableInstanceLODManagement>();
+	CurrentInstanceLODManagement = DefaultInstanceLODManagement;
+}
+
+
+void UCustomizableObjectSystem::BeginDestroy()
+{
+#if WITH_EDITOR
+	if (RecompileCustomizableObjectsCompiler)
+	{
+		RecompileCustomizableObjectsCompiler->ForceFinishCompilation();
+		delete RecompileCustomizableObjectsCompiler;
+	}
+
+	if (!IsRunningGame())
+	{
+		FEditorDelegates::PreBeginPIE.RemoveAll(this);
+	}
+
+#endif
+
+	// It could be null, for the default object.
+	if (Private.IsValid())
+	{
+
+#if !UE_SERVER
+		FTSTicker::GetCoreTicker().RemoveTicker(Private->TickDelegateHandle);
+#endif // !UE_SERVER
+
+		// Discard pending game thread tasks
+		Private->PendingTasks.Empty();
+
+		// Complete pending taskgraph tasks
+		Private->WaitForMutableTasks();
+
+		// Clear the ongoing operation
+		Private->CurrentMutableOperation = nullptr;
+
+		// Deallocate streaming
+		Private->Streamer->EndStreaming();
+
+		Private->CurrentInstanceBeingUpdated = nullptr;
+
+		while (!Private->MutableOperationQueue.IsEmpty())
+		{
+			FMutableQueueElem AuxOperation;
+			Private->MutableOperationQueue.Dequeue(&AuxOperation);
+		}
+
+		FCustomizableObjectSystemPrivate::SSystem = nullptr;
+
+		Private = nullptr;
+	}
+
+	Super::BeginDestroy();
+}
+
+
+FString UCustomizableObjectSystem::GetDesc()
+{
+	return TEXT("Customizable Object System Singleton");
+}
+
+
+FCustomizableObjectCompilerBase* (*FCustomizableObjectSystemPrivate::NewCompilerFunc)() = nullptr;
+
+
+FCustomizableObjectCompilerBase* UCustomizableObjectSystem::GetNewCompiler()
+{
+	if (Private->NewCompilerFunc!=nullptr)
+	{
+		return Private->NewCompilerFunc();
+	}
+	else
+	{
+		return nullptr;
+	}
+}
+
+
+void UCustomizableObjectSystem::SetNewCompilerFunc(FCustomizableObjectCompilerBase* (*InNewCompilerFunc)())
+{
+	Private->NewCompilerFunc = InNewCompilerFunc;
+}
+
+
+void FCustomizableObjectSystemPrivate::CreatedTexture(UTexture2D* Texture)
+{
+	// TODO: Do not keep this array unless we are gathering stats!
+	TextureTrackerArray.Add(Texture);
+
+	INC_DWORD_STAT(STAT_MutableNumTextures);
+}
+
+
+int32 GVarEnableMutableAnimInfoDebugging = 0;
+static FAutoConsoleVariableRef CVarEnableMutableAnimInfoDebugging(
+	TEXT("mutable.EnableMutableAnimInfoDebugging"), GVarEnableMutableAnimInfoDebugging,
+	TEXT("If set to 1 or greater print on screen the animation info of the pawn's Customizable Object Instance. Anim BPs, slots and tags will be displayed."),
+	ECVF_Default);
+
+
+void FCustomizableObjectSystemPrivate::UpdateStats()
+{
+#if !(UE_BUILD_SHIPPING || UE_BUILD_TEST)
+	bool bLogEnabled = true;
+#else
+	bool bLogEnabled = LogBenchmarkUtil::isLoggingActive();
+#endif
+	if (bLogEnabled)
+	{
+		CountAllocatedSkeletalMesh = 0;
+
+		int CountLOD0 = 0;
+		int CountLOD1 = 0;
+		int CountLOD2 = 0;
+		int CountTotal = 0;
+		int CountPending = 0;
+		for (TObjectIterator<UCustomizableObjectInstance> CustomizableObjectInstance; CustomizableObjectInstance; ++CustomizableObjectInstance)
+		{
+			if ( IsValidChecked(*CustomizableObjectInstance) && CustomizableObjectInstance->GetPrivate())
+			{
+				++CountTotal;
+
+				if (CustomizableObjectInstance->SkeletalMeshStatus == ESkeletalMeshState::AsyncUpdatePending)
+				{
+					++CountPending;
+				}
+
+				for (int32 ComponentIndex = 0; ComponentIndex < CustomizableObjectInstance->SkeletalMeshes.Num(); ++ComponentIndex)
+				{
+					if (CustomizableObjectInstance->SkeletalMeshes[ComponentIndex] && CustomizableObjectInstance->SkeletalMeshes[ComponentIndex]->GetResourceForRendering())
+					{
+						CountAllocatedSkeletalMesh++;
+
+						if (CustomizableObjectInstance->GetPrivate()->GetMinLODToLoad() < 1)
+						{
+							++CountLOD0;
+						}
+						else if (CustomizableObjectInstance->GetPrivate()->GetMinLODToLoad() < 2)
+						{
+							++CountLOD1;
+						}
+						else
+						{
+							++CountLOD2;
+						}
+					}
+				}
+			}
+		}
+
+		NumInstances = CountLOD0 + CountLOD1 + CountLOD2;
+		TotalInstances = CountTotal;
+		NumPendingInstances = CountPending;
+		//SET_DWORD_STAT(STAT_MutableSkeletalMeshResourceMemory, Size / 1024.f);
+		SET_DWORD_STAT(STAT_MutablePendingInstanceUpdates, MutableOperationQueue.Num());
+
+		uint64 Size = 0;
+		uint32 CountAllocated = 0;
+		for (TWeakObjectPtr<UTexture2D>& Tracker : TextureTrackerArray)
+		{
+			if (Tracker.IsValid() && Tracker->GetResource())
+			{
+				CountAllocated++;
+
+				if (Tracker->GetResource()->TextureRHI)
+				{
+					Size += Tracker->CalcTextureMemorySizeEnum(TMC_AllMips);
+				}
+			}
+		}
+
+		TextureMemoryUsed = int64_t(Size / 1024);
+
+		uint64 SizeDecoration = 0;
+		uint64 SizeGenerated = 0;
+		for (TObjectIterator<UCustomizableObjectInstance> CustomizableObjectInstance; CustomizableObjectInstance; ++CustomizableObjectInstance)
+		{
+			if (IsValidChecked(*CustomizableObjectInstance) && CustomizableObjectInstance->GetPrivate() && CustomizableObjectInstance->HasAnySkeletalMesh())
+			{
+				bool bHasResourceForRendering = false;
+				for (int32 MeshIndex = 0; !bHasResourceForRendering && MeshIndex < CustomizableObjectInstance->SkeletalMeshes.Num(); ++MeshIndex)
+				{
+					bHasResourceForRendering = CustomizableObjectInstance->GetSkeletalMesh(MeshIndex) && CustomizableObjectInstance->GetSkeletalMesh(MeshIndex)->GetResourceForRendering();
+				}
+
+				if (bHasResourceForRendering)
+				{
+					for (const FParameterDecorations& parameterDecoration : CustomizableObjectInstance->GetPrivate()->ParameterDecorations)
+					{
+						for (const UTexture2D* parameterDecorationImage : parameterDecoration.Images)
+						{
+							if (parameterDecorationImage && parameterDecorationImage->IsValidLowLevel())
+							{
+								SizeDecoration += parameterDecorationImage->CalcTextureMemorySizeEnum(TMC_AllMips);
+							}
+						}
+					}
+
+					for (const FGeneratedTexture& generatedTextures : CustomizableObjectInstance->GetPrivate()->GeneratedTextures)
+					{
+						if (generatedTextures.Texture)
+						{
+							SizeGenerated += generatedTextures.Texture->CalcTextureMemorySizeEnum(TMC_AllMips);
+						}
+					}
+				}
+			}
+		}
+
+		if (LogBenchmarkUtil::isLoggingActive())
+		{
+			LogBenchmarkUtil::updateStat("customizable_objects", (int32)CountAllocatedSkeletalMesh);
+			LogBenchmarkUtil::updateStat("pending_instance_updates", MutableOperationQueue.Num());
+			LogBenchmarkUtil::updateStat("allocated_textures", (int32)CountAllocated);
+			LogBenchmarkUtil::updateStat("texture_resource_memory", (long double)Size / 1048576.0L);
+			LogBenchmarkUtil::updateStat("texture_generated_memory", (long double)SizeGenerated / 1048576.0L);
+		}
+
+#if !(UE_BUILD_SHIPPING || UE_BUILD_TEST)
+		SET_DWORD_STAT(STAT_MutableNumInstancesLOD0, CountLOD0);
+		SET_DWORD_STAT(STAT_MutableNumInstancesLOD1, CountLOD1);
+		SET_DWORD_STAT(STAT_MutableNumInstancesLOD2, CountLOD2);
+		SET_DWORD_STAT(STAT_MutableNumAllocatedSkeletalMeshes, CountAllocatedSkeletalMesh);
+		SET_DWORD_STAT(STAT_MutablePendingInstanceUpdates, MutableOperationQueue.Num());
+		SET_DWORD_STAT(STAT_MutableNumAllocatedTextures, CountAllocated);
+		SET_DWORD_STAT(STAT_MutableTextureResourceMemory, Size / 1024.f);
+		SET_DWORD_STAT(STAT_MutableTextureParameterDecorationMemory, SizeDecoration / 1024.f);
+		SET_DWORD_STAT(STAT_MutableTextureGeneratedMemory, SizeGenerated / 1024.f);
+#endif
+
+		if (GVarEnableMutableAnimInfoDebugging > 0)
+		{
+			if (GEngine)
+			{
+				bool bFoundPlayer = false;
+				int32 MsgIndex = 15820; // Arbitrary big value to prevent collisions with other on-screen messages
+
+				for (TObjectIterator<UCustomizableSkeletalComponent> CustomizableSkeletalComponent; CustomizableSkeletalComponent; ++CustomizableSkeletalComponent)
+				{
+					AActor* ParentActor = CustomizableSkeletalComponent->GetAttachmentRootActor();
+					UCustomizableObjectInstance* Instance = CustomizableSkeletalComponent->CustomizableObjectInstance;
+					APawn* PlayerPawn = UGameplayStatics::GetPlayerPawn(CustomizableSkeletalComponent->GetWorld(), 0);
+					
+					if (ParentActor && ParentActor == PlayerPawn && Instance)
+					{
+						bFoundPlayer = true;
+
+						FString TagString;
+						const FGameplayTagContainer& Tags = Instance->GetAnimationGameplayTags();
+
+						for (const FGameplayTag& Tag : Tags)
+						{
+							TagString += !TagString.IsEmpty() ? FString(TEXT(", ")) : FString();
+							TagString += Tag.ToString();
+						}												
+												
+						GEngine->AddOnScreenDebugMessage(MsgIndex++, .0f, FColor::Green, TEXT("Animation tags: ") + TagString);
+
+						FCustomizableInstanceComponentData* ComponentData = Instance->GetPrivate()->GetComponentData(CustomizableSkeletalComponent->ComponentIndex);
+
+						if (ComponentData)
+						{
+							for (TPair<int32, TSoftClassPtr<UAnimInstance>>& Entry : ComponentData->AnimSlotToBP)
+							{
+								FString AnimBPSlot;
+
+								AnimBPSlot += FString::Printf(TEXT("%d"), Entry.Key) + FString("-") + Entry.Value.GetAssetName();
+								GEngine->AddOnScreenDebugMessage(MsgIndex++, .0f, FColor::Green, AnimBPSlot);
+							}
+						}
+
+						GEngine->AddOnScreenDebugMessage(MsgIndex++, .0f, FColor::Green, TEXT("Slots-AnimBP: "));
+						GEngine->AddOnScreenDebugMessage(MsgIndex++, .0f, FColor::Cyan, TEXT("Player Pawn Mutable Animation info "));
+					}
+				}
+
+				if (!bFoundPlayer)
+				{
+					GEngine->AddOnScreenDebugMessage(MsgIndex, .0f, FColor::Yellow, TEXT("Mutable Animation info: N/A"));
+				}
+			}
+		}
+	}
+}
+
+
+void FCustomizableObjectSystemPrivate::AddReferencedObjects(FReferenceCollector& Collector)
+{
+	if (CurrentInstanceBeingUpdated)
+	{
+		Collector.AddReferencedObject(CurrentInstanceBeingUpdated);
+	}
+}
+
+
+int32 FCustomizableObjectSystemPrivate::EnableMutableProgressiveMipStreaming = 0;
+
+
+// Warning! If this is enabled, do not get references to the textures generated by Mutable! They are owned by Mutable and could become invalid at any moment
+static FAutoConsoleVariableRef CVarEnableMutableProgressiveMipStreaming(
+	TEXT("mutable.EnableMutableProgressiveMipStreaming"), FCustomizableObjectSystemPrivate::EnableMutableProgressiveMipStreaming,
+	TEXT("If set to 1 or greater use progressive Mutable Mip streaming for Mutable textures. If disabled, all mips will always be generated and spending memory. In that case, on Desktop platforms they will be stored in CPU memory, on other platforms textures will be non-streaming."),
+	ECVF_Default);
+
+
+void FCustomizableObjectSystemPrivate::InitUpdateSkeletalMesh(UCustomizableObjectInstance* InCustomizableObjectInstance, FMutableQueueElem::EQueuePriorityType Priority)
+{
+	MUTABLE_CPUPROFILER_SCOPE(FCustomizableObjectSystemPrivate::InitUpdateSkeletalMesh);
+
+	check(IsInGameThread());
+
+	if (!InCustomizableObjectInstance ||
+		!InCustomizableObjectInstance->IsValidLowLevel())
+	{
+		return;
+	}
+	
+	if (UCustomizableObject* CustomizableObject = InCustomizableObjectInstance->GetCustomizableObject();
+		CustomizableObject &&
+		!CustomizableObject->IsLocked())
+	{
+		FString CurrentState = InCustomizableObjectInstance->GetCurrentState();
+		FParameterUIData* State = CustomizableObject->StateUIDataMap.Find(CurrentState);
+		bool bNeverStream = State ? State->bDontCompressRuntimeTextures : false;
+		const bool bUseMipmapStreaming = !bNeverStream;
+		int32 MipsToSkip = 0; // 0 means all mips
+
+		if (bUseMipmapStreaming && EnableMutableProgressiveMipStreaming)
+		{
+			MipsToSkip = 255; // This means skip all possible mips until only UTexture::GetStaticMinTextureResidentMipCount() are left
+		}
+
+		const TSharedPtr<FMutableOperation> Operation = MakeShared<FMutableOperation>(FMutableOperation::CreateInstanceUpdate(InCustomizableObjectInstance, bNeverStream, MipsToSkip));
+
+		// Skip the update if the Instance has the same descriptor (pending or applied).
+		// If there is a hash collision, at worse, we would add an unnecessary update.
+		{
+			const uint32 DescriptorHash = Operation->InstanceDescriptorHash;
+
+			if (const FMutableQueueElem* QueueElem = MutableOperationQueue.Get(InCustomizableObjectInstance))
+			{
+				if (const TSharedPtr<FMutableOperation>& QueuedOperation = QueueElem->Operation;
+					QueuedOperation &&
+					QueuedOperation->Type == FMutableOperation::EOperationType::Update &&
+					QueuedOperation->InstanceDescriptorHash == DescriptorHash)
+				{
+					return;
+				}
+			}
+		
+			if (CurrentMutableOperation &&
+				CurrentMutableOperation->Type == FMutableOperation::EOperationType::Update &&
+				CurrentMutableOperation->InstanceDescriptorHash == DescriptorHash)
+			{
+				return;
+			}
+			
+			if (InCustomizableObjectInstance->GetLastUpdateDescriptorHash() == DescriptorHash)
+			{
+				return;
+			}
+		}
+
+		MutableOperationQueue.Enqueue(FMutableQueueElem::Create(Operation, Priority, InCustomizableObjectInstance->GetPrivate()->MinSquareDistFromComponentToPlayer));
+	}
+}
+
+
+void FCustomizableObjectSystemPrivate::InitDiscardResourcesSkeletalMesh(UCustomizableObjectInstance* InCustomizableObjectInstance)
+{
+	check(IsInGameThread());
+
+	if (InCustomizableObjectInstance && InCustomizableObjectInstance->IsValidLowLevel())
+	{
+		const TSharedRef<FMutableOperation> Operation = MakeShared<FMutableOperation>(FMutableOperation::CreateInstanceDiscard(InCustomizableObjectInstance));
+		MutableOperationQueue.Enqueue(FMutableQueueElem::Create(Operation, FMutableQueueElem::EQueuePriorityType::High, InCustomizableObjectInstance->GetPrivate()->MinSquareDistFromComponentToPlayer));
+	}
+}
+
+
+bool UCustomizableObjectSystem::IsReplaceDiscardedWithReferenceMeshEnabled() const
+{
+	if (Private.IsValid())
+	{
+		return Private->IsReplaceDiscardedWithReferenceMeshEnabled();
+	}
+
+	return false;
+}
+
+
+void UCustomizableObjectSystem::SetReplaceDiscardedWithReferenceMeshEnabled(bool bIsEnabled)
+{
+	if (Private.IsValid())
+	{
+		Private->SetReplaceDiscardedWithReferenceMeshEnabled(bIsEnabled);
+	}
+}
+
+
+void UCustomizableObjectSystem::ClearResourceCacheProtected()
+{
+	check(IsInGameThread());
+
+	ProtectedCachedTextures.Reset(0);
+	GetPrivate()->ProtectedObjectCachedImages.Reset(0);
+}
+
+
+#if WITH_EDITOR
+bool UCustomizableObjectSystem::LockObject(const class UCustomizableObject* InObject)
+{
+	check(!InObject->GetPrivate()->bLocked);
+
+	// If the current instance is for this object, make the lock fail by returning false
+	if (Private->CurrentInstanceBeingUpdated &&
+		Private->CurrentInstanceBeingUpdated->GetCustomizableObject() == InObject)
+	{
+		UE_LOG(LogMutable, Warning, TEXT("---- failed to lock object %s"), *InObject->GetName());
+		return false;
+	}
+
+	// Ensure that we don't try to handle any further streaming operations for this object
+	if (GetPrivate()->Streamer)
+	{
+		GetPrivate()->Streamer->CancelStreamingForObject(InObject);
+	}
+
+	InObject->GetPrivate()->bLocked = true;
+
+	// Clear the cache for the instance, since we will remake it
+	FMutableResourceCache& Cache = GetPrivate()->GetObjectCache(InObject);
+	Cache.Clear();
+
+	return true;
+}
+
+
+void UCustomizableObjectSystem::UnlockObject(const class UCustomizableObject* Obj)
+{
+	check(Obj->GetPrivate()->bLocked);
+	Obj->GetPrivate()->bLocked = false;
+}
+
+
+void UCustomizableObjectSystem::EditorSettingsChanged(const FEditorCompileSettings& InEditorSettings)
+{
+	EditorSettings = InEditorSettings;
+}
+
+
+bool UCustomizableObjectSystem::IsCompilationDisabled() const
+{
+	return EditorSettings.bDisableCompilation;
+}
+
+
+bool UCustomizableObjectSystem::IsAutoCompileEnabled() const
+{
+	return EditorSettings.bEnableAutomaticCompilation;
+}
+
+
+bool UCustomizableObjectSystem::IsAutoCompilationSync() const
+{
+	return EditorSettings.bCompileObjectsSynchronously;
+}
+
+#endif
+
+
+void UCustomizableObjectSystem::PurgePendingReleaseSkeletalMesh()
+{
+	MUTABLE_CPUPROFILER_SCOPE(UCustomizableObjectSystem::PurgePendingReleaseSkeletalMesh);
+
+	double	CurTime = FPlatformTime::Seconds();
+	static double TimeToDelete = 1.0;
+
+	for( int32 i = PendingReleaseSkeletalMesh.Num() - 1; i >= 0; --i )
+	{
+		FPendingReleaseSkeletalMeshInfo& info = PendingReleaseSkeletalMesh[ i ];
+		
+		if (info.SkeletalMesh != nullptr)
+		{
+			if (( CurTime - info.TimeStamp ) >= TimeToDelete)
+			{
+				if (info.SkeletalMesh->GetSkeleton())
+				{
+					info.SkeletalMesh->GetSkeleton()->ClearCacheData();
+				}
+				info.SkeletalMesh->GetRefSkeleton().Empty();
+				info.SkeletalMesh->GetMaterials().Empty();
+				info.SkeletalMesh->GetRefBasesInvMatrix().Empty();
+				info.SkeletalMesh->ReleaseResources();
+				info.SkeletalMesh->ReleaseResourcesFence.Wait();
+
+				// TODO: This was done for version 4.18
+				//info.SkeletalMesh->GetImportedResource()->LODModels.Empty();
+				//info.SkeletalMesh->LODInfo.Empty();
+
+				PendingReleaseSkeletalMesh.RemoveAt( i );
+			}
+		}
+	}
+}
+
+void UCustomizableObjectSystem::AddPendingReleaseSkeletalMesh( USkeletalMesh* SkeletalMesh )
+{
+	check( SkeletalMesh != nullptr );
+
+	FPendingReleaseSkeletalMeshInfo	info;
+	info.SkeletalMesh = SkeletalMesh;
+	info.TimeStamp = FPlatformTime::Seconds();
+
+	PendingReleaseSkeletalMesh.Add( info );
+}
+
+
+void UCustomizableObjectSystem::ClearCurrentMutableOperation()
+{
+	Private->CurrentInstanceBeingUpdated = nullptr;
+	Private->CurrentMutableOperation = nullptr;
+	ClearResourceCacheProtected();
+}
+
+
+void FCustomizableObjectSystemPrivate::UpdateStreamingLimit()
+{
+	// This must run on game thread, and when the mutable thread is not running
+	check(IsInGameThread());
+
+	int32 VarValue = CVarStreamingMemory.GetValueOnGameThread();
+	uint64_t MemoryBytes = VarValue < 0 ? MUTABLE_STREAMING_CACHE : uint64_t(VarValue) * 1024;
+	if (MemoryBytes != LastStreamingMemorySize)
+	{
+		LastStreamingMemorySize = MemoryBytes;
+		MutableSystem->SetStreamingCache(MemoryBytes);
+	}
+}
+
+
+// Asynchronous tasks performed during the creation or update of a mutable instance. 
+// Check the documentation before modifying and keep it up to date.
+// https://docs.google.com/drawings/d/109NlsdKVxP59K5TuthJkleVG3AROkLJr6N03U4bNp4s
+// When it says "mutable thread" it means any task pool thread, but with the guarantee that no other thread is using the mutable runtime.
+// Naming: Task_<thread>_<description>
+namespace impl
+{
+
+	void Subtask_Mutable_UpdateParameterDecorations(const TSharedPtr<FMutableOperationData>& OperationData, const mu::ModelPtr& MutableModel, const mu::Parameters* MutableParameters)
+	{
+		MUTABLE_CPUPROFILER_SCOPE(UCustomizableInstancePrivateData::UpdateParameterDecorations)
+
+		OperationData->ParametersUpdateData.Clear();
+		OperationData->RelevantParametersInProgress.Empty();
+
+		check(MutableParameters);
+
+		// This must run in the mutable thread.
+		auto MutableSystem = UCustomizableObjectSystem::GetInstance()->GetPrivate()->MutableSystem;
+
+		// Decorations
+		int32 NumParameter = MutableParameters->GetCount();
+
+		{
+			MUTABLE_CPUPROFILER_SCOPE(ParameterDecorations)
+
+				// Generate all the images that decorate the parameters UI, only in the editor.
+				for (int32 p = 0; p < NumParameter; ++p)
+				{
+					FParameterDecorationsUpdateData::FParameterDesc ParamData;
+
+					for (int32 i = 0; i < MutableParameters->GetAdditionalImageCount(p); ++i)
+					{
+						mu::ImagePtrConst Image = MutableSystem->BuildParameterAdditionalImage(
+							MutableModel.get(),
+							MutableParameters,
+							p, i);
+
+						ParamData.Images.Add(Image);
+					};
+
+					OperationData->ParametersUpdateData.Parameters.Add(ParamData);
+				}
+		}
+
+		// Update the parameter relevancy.
+		{
+			MUTABLE_CPUPROFILER_SCOPE(ParameterRelevancy)
+
+			TArray<bool> relevant;
+			relevant.SetNumZeroed(NumParameter);
+			MutableSystem->GetParameterRelevancy(MutableModel.get(), MutableParameters, relevant.GetData());
+
+			for (int32 p = 0; p < NumParameter; ++p)
+			{
+				if (relevant[p])
+				{
+					OperationData->RelevantParametersInProgress.Add(p);
+				}
+			}
+		}
+
+	}
+
+
+	void Subtask_Mutable_BeginUpdate_GetMesh(const TSharedPtr<FMutableOperationData>& OperationData, mu::ModelPtr Model, const mu::Parameters* MutableParameters, int32 State)
+	{
+		MUTABLE_CPUPROFILER_SCOPE(Subtask_Mutable_BeginUpdate);
+
+		// This runs in the mutable thread.
+		check(MutableParameters);
+
+#ifdef WITH_EDITOR
+		// \TODO: Review this to make sure it is really needed.
+		// NOTE: this delegate can be called outside the main thread, make sure operations done in the
+		//       callbacks are safe, avoid using it outside editor functionality
+		//Public->UpdateBeginDelegate.ExecuteIfBound();
+#endif
+
+		OperationData->InstanceUpdateData.Clear();
+
+		mu::System* System = UCustomizableObjectSystem::GetInstance()->GetPrivate()->MutableSystem.get();
+
+		// For now, we are forcing the recreation of mutable-side instances with every update.
+		OperationData->InstanceID = System->NewInstance(Model.get());
+		UE_LOG(LogMutable, Log, TEXT("Creating instance with id [%d] "), OperationData->InstanceID)
+
+		const mu::Instance* Instance = nullptr;
+
+		// Main instance generation step
+		{
+			// LOD mask, set to all ones to build  all LODs
+			uint32 LODMask = 0xFFFFFFFF;
+
+			Instance = System->BeginUpdate(OperationData->InstanceID, MutableParameters, State, LODMask);
+
+			if (!Instance)
+			{
+				const mu::Error MutableError = mu::GetError();
+				if (MutableError == mu::Error::Unsupported)
+				{
+					UE_LOG(LogMutable, Log, TEXT("The necessary functionality is not supported in this version of Mutable"));
+				}
+
+				return;
+			}
+		}
+
+		OperationData->NumLODsAvailable = Instance->GetLODCount();
+
+		if (OperationData->CurrentMinLOD >= OperationData->NumLODsAvailable)
+		{
+			OperationData->CurrentMinLOD = OperationData->NumLODsAvailable - 1;
+			OperationData->CurrentMaxLOD = OperationData->CurrentMinLOD;
+		}
+		else if (OperationData->CurrentMaxLOD >= OperationData->NumLODsAvailable)
+		{
+			OperationData->CurrentMaxLOD = OperationData->NumLODsAvailable - 1;
+		}
+
+		// Generate the mesh and gather all the required resource Ids
+		OperationData->InstanceUpdateData.LODs.SetNum(Instance->GetLODCount());
+		for (int32 MutableLODIndex = 0; MutableLODIndex < Instance->GetLODCount(); ++MutableLODIndex)
+		{
+			// Skip LODs outside the reange we want to generate
+			if (MutableLODIndex < OperationData->CurrentMinLOD || MutableLODIndex > OperationData->CurrentMaxLOD)
+			{
+				continue;
+			}
+
+			FInstanceUpdateData::FLOD& LOD = OperationData->InstanceUpdateData.LODs[MutableLODIndex];
+			LOD.FirstComponent = OperationData->InstanceUpdateData.Components.Num();
+			LOD.ComponentCount = Instance->GetComponentCount(MutableLODIndex);
+
+			for (int32 ComponentIndex = 0; ComponentIndex < LOD.ComponentCount; ++ComponentIndex)
+			{
+				OperationData->InstanceUpdateData.Components.Push(FInstanceUpdateData::FComponent());
+				FInstanceUpdateData::FComponent& Component = OperationData->InstanceUpdateData.Components.Last();
+				Component.Id = Instance->GetComponentId(MutableLODIndex, ComponentIndex);
+				Component.FirstSurface = OperationData->InstanceUpdateData.Surfaces.Num();
+				Component.SurfaceCount = 0;
+
+				// Mesh
+				if (Instance->GetMeshCount(MutableLODIndex, ComponentIndex) > 0)
+				{
+					MUTABLE_CPUPROFILER_SCOPE(GetMesh);
+
+					mu::RESOURCE_ID MeshId = Instance->GetMeshId(MutableLODIndex, ComponentIndex, 0);
+
+					// Mesh cache is not enabled yet.
+					//auto CachedMesh = Cache.Meshes.Find(meshId);
+					//if (CachedMesh && CachedMesh->IsValid(false, true))
+					//{
+					//	UE_LOG(LogMutable, Verbose, TEXT("Mesh resource with id [%d] can be cached."), meshId);
+					//	INC_DWORD_STAT(STAT_MutableNumCachedSkeletalMeshes);
+					//}
+
+					Component.Mesh = System->GetMesh(OperationData->InstanceID, MeshId);
+				}
+
+				if (!Component.Mesh)
+				{
+					continue;
+				}
+
+				// Materials and images
+				for (int32 MeshSurfaceIndex = 0; MeshSurfaceIndex < Component.Mesh->GetSurfaceCount(); ++MeshSurfaceIndex)
+				{
+					int32 InstanceSurfaceIndex = Instance->FindSurfaceById(MutableLODIndex, ComponentIndex, Component.Mesh->GetSurfaceId(MeshSurfaceIndex));
+					check(Component.Mesh->GetVertexCount() == 0 || InstanceSurfaceIndex >= 0);
+
+					if (InstanceSurfaceIndex >= 0)
+					{
+						OperationData->InstanceUpdateData.Surfaces.Push({});
+						FInstanceUpdateData::FSurface& Surface = OperationData->InstanceUpdateData.Surfaces.Last();
+						++Component.SurfaceCount;
+
+						Surface.MaterialIndex = Instance->GetSurfaceCustomId(MutableLODIndex, ComponentIndex, InstanceSurfaceIndex);
+
+						// Images
+						Surface.FirstImage = OperationData->InstanceUpdateData.Images.Num();
+						Surface.ImageCount = Instance->GetImageCount(MutableLODIndex, ComponentIndex, InstanceSurfaceIndex);
+						for (int32 ImageIndex = 0; ImageIndex < Surface.ImageCount; ++ImageIndex)
+						{
+							MUTABLE_CPUPROFILER_SCOPE(GetImageId);
+
+							OperationData->InstanceUpdateData.Images.Push({});
+							FInstanceUpdateData::FImage& Image = OperationData->InstanceUpdateData.Images.Last();
+							Image.Name = Instance->GetImageName(MutableLODIndex, ComponentIndex, InstanceSurfaceIndex, ImageIndex);
+							Image.ImageID = Instance->GetImageId(MutableLODIndex, ComponentIndex, InstanceSurfaceIndex, ImageIndex);
+							Image.FullImageSizeX = 0;
+							Image.FullImageSizeY = 0;
+						}
+
+						// Vectors
+						Surface.FirstVector = OperationData->InstanceUpdateData.Vectors.Num();
+						Surface.VectorCount = Instance->GetVectorCount(MutableLODIndex, ComponentIndex, InstanceSurfaceIndex);
+						for (int32 VectorIndex = 0; VectorIndex < Surface.VectorCount; ++VectorIndex)
+						{
+							MUTABLE_CPUPROFILER_SCOPE(GetVector);
+							OperationData->InstanceUpdateData.Vectors.Push({});
+							FInstanceUpdateData::FVector& Vector = OperationData->InstanceUpdateData.Vectors.Last();
+							Vector.Name = Instance->GetVectorName(MutableLODIndex, ComponentIndex, InstanceSurfaceIndex, VectorIndex);
+							Instance->GetVector(MutableLODIndex, ComponentIndex, InstanceSurfaceIndex, VectorIndex, &Vector.Vector.R, &Vector.Vector.G, &Vector.Vector.B, &Vector.Vector.A);
+						}
+
+						// Scalars
+						Surface.FirstScalar = OperationData->InstanceUpdateData.Scalars.Num();
+						Surface.ScalarCount = Instance->GetScalarCount(MutableLODIndex, ComponentIndex, InstanceSurfaceIndex);
+						for (int32 ScalarIndex = 0; ScalarIndex < Surface.ScalarCount; ++ScalarIndex)
+						{
+							MUTABLE_CPUPROFILER_SCOPE(GetScalar);
+							OperationData->InstanceUpdateData.Scalars.Push({});
+							FInstanceUpdateData::FScalar& Scalar = OperationData->InstanceUpdateData.Scalars.Last();
+							Scalar.Name = Instance->GetScalarName(MutableLODIndex, ComponentIndex, InstanceSurfaceIndex, ScalarIndex);
+							Scalar.Scalar = Instance->GetScalar(MutableLODIndex, ComponentIndex, InstanceSurfaceIndex, ScalarIndex);
+						}
+					}
+				}
+			}
+		}
+
+	}
+
+	
+	void Subtask_Mutable_GetImages(const TSharedPtr<FMutableOperationData>& OperationData, mu::ModelPtr Model, const mu::Parameters* MutableParameters, int32 State)
+	{
+		MUTABLE_CPUPROFILER_SCOPE(Subtask_Mutable_GetImages);
+
+		// This runs in the mutable thread.
+
+		check(MutableParameters);
+
+		mu::System* System = UCustomizableObjectSystem::GetInstance()->GetPrivate()->MutableSystem.get();
+
+		// Generate all the required resources, that are not cached
+		TArray<mu::RESOURCE_ID> ImagesInThisInstance;
+		for ( FInstanceUpdateData::FImage& Image: OperationData->InstanceUpdateData.Images )
+		{
+			MUTABLE_CPUPROFILER_SCOPE(GetImage);
+
+			// This should only be done when using progressive images, since GetImageDesc does some actual processing.
+			{
+				mu::FImageDesc ImageDesc;
+				System->GetImageDesc(OperationData->InstanceID, Image.ImageID, ImageDesc);
+				Image.FullImageSizeX = ImageDesc.m_size[0];
+				Image.FullImageSizeY = ImageDesc.m_size[1];
+			}
+
+			bool bCached = false;
+
+			// See if it is cached from this same instance (can happen with LODs)
+			bCached = ImagesInThisInstance.Contains(Image.ImageID);
+
+			// See if it is cached from another instance
+			if (!bCached)
+			{
+				bCached = UCustomizableObjectSystem::GetInstance()->GetPrivate()->ProtectedObjectCachedImages.Contains(Image.ImageID);
+			}
+
+			if (bCached)
+			{
+				UE_LOG(LogMutable, Verbose, TEXT("Texture resource with id [%d] is cached."), Image.ImageID);
+				INC_DWORD_STAT(STAT_MutableNumCachedTextures);
+			}
+			else
+			{
+				MUTABLE_CPUPROFILER_SCOPE(GetImage);
+
+				int32 MaxSize = FMath::Max(Image.FullImageSizeX, Image.FullImageSizeY);
+				int32 FullLODCount = FMath::CeilLogTwo(MaxSize) + 1;
+				int32 MinMipsInImage = FMath::Min(FullLODCount, UTexture::GetStaticMinTextureResidentMipCount());
+				int32 MaxMipsToSkip = FullLODCount - MinMipsInImage;
+				int32 MipsToSkip = FMath::Min(MaxMipsToSkip, OperationData->MipsToSkip);
+				Image.Image = System->GetImage(OperationData->InstanceID, Image.ImageID, MipsToSkip);
+
+				// We need one mip or the complete chain. Otherwise there was a bug.
+				int32 FullMipCount = Image.Image->GetMipmapCount(Image.Image->GetSizeX(), Image.Image->GetSizeY());
+				int32 RealMipCount = Image.Image->GetLODCount();
+				if ( RealMipCount!=1 && RealMipCount!=FullMipCount )
+				{
+					MUTABLE_CPUPROFILER_SCOPE(GetImage_MipFix);
+
+					UE_LOG(LogMutable, Warning, TEXT("Mutable generated an incomplete mip chain for image %d."), Image.ImageID);
+
+					// Force the right number of mips. The missing data will be black.
+					mu::ImagePtr NewImage = new mu::Image( Image.Image->GetSizeX(), Image.Image->GetSizeY(), FullMipCount, Image.Image->GetFormat() );
+					if (NewImage->GetDataSize() >= Image.Image->GetDataSize())
+					{
+						FMemory::Memcpy(NewImage->GetData(), Image.Image->GetData(), Image.Image->GetDataSize());
+					}
+					Image.Image = NewImage;
+				}
+
+				ImagesInThisInstance.Add(Image.ImageID);
+			}
+		}
+
+	}
+
+
+	void Subtask_Mutable_PrepareTextures(const TSharedPtr<FMutableOperationData>& OperationData)
+	{
+		// This runs in a worker thread
+
+		for (const FInstanceUpdateData::FSurface& Surface : OperationData->InstanceUpdateData.Surfaces)
+		{
+			for (int32 ImageIndex = 0; ImageIndex<Surface.ImageCount; ++ImageIndex)
+			{
+				const FInstanceUpdateData::FImage& Image = OperationData->InstanceUpdateData.Images[Surface.FirstImage+ImageIndex];
+
+				FString KeyName = Image.Name;
+				mu::ImagePtrConst MutableImage = Image.Image;
+
+				// If the image is null, it must be in the cache (or repeated in this instance), and we don't need to do anything here.
+				if (MutableImage)
+				{
+					FTexturePlatformData* PlatformData = UCustomizableInstancePrivateData::MutableCreateImagePlatformData(MutableImage.get(), -1, Image.FullImageSizeX, Image.FullImageSizeY);
+					OperationData->ImageToPlatformDataMap.Add(Image.ImageID, PlatformData);
+					OperationData->PendingTextureCoverageQueries.Add({ KeyName, Surface.MaterialIndex, PlatformData });
+				}
+			}
+		}
+	}
+
+
+	void Subtask_Mutable_PrepareSkeletonData(const TSharedPtr<FMutableOperationData>& OperationData)
+	{
+		// This runs in a worker thread
+
+		MUTABLE_CPUPROFILER_SCOPE(PrepareSkeletonData);
+
+		const int32 LODCount = OperationData->InstanceUpdateData.LODs.Num();
+		const FInstanceUpdateData::FLOD& MinLOD = OperationData->InstanceUpdateData.LODs[OperationData->CurrentMinLOD];
+		const int32 ComponentCount = MinLOD.ComponentCount;
+
+		// Add SkeletonData for each component
+		OperationData->InstanceUpdateData.Skeletons.AddDefaulted(ComponentCount);
+
+		// Helpers
+		TMap<FName, int16> BoneNameToIndex;
+		TArray<const char*> ParentChain;
+		TArray<bool> UsedBones;
+
+		for (int32 ComponentIndex = 0; ComponentIndex < ComponentCount; ++ComponentIndex)
+		{
+			FInstanceUpdateData::FComponent& MinLODComponent = OperationData->InstanceUpdateData.Components[MinLOD.FirstComponent+ComponentIndex];
+
+			// Set the ComponentIndex
+			OperationData->InstanceUpdateData.Skeletons[ComponentIndex].ComponentIndex = MinLODComponent.Id;
+
+			// Fill the data used to generate the RefSkeletalMesh
+			TArray<uint32>& SkeletonIds = OperationData->InstanceUpdateData.Skeletons[ComponentIndex].SkeletonIds;
+			TArray<FName>& BoneNames = OperationData->InstanceUpdateData.Skeletons[ComponentIndex].BoneNames;
+			TMap<FName, FMatrix44f>& BoneMatricesWithScale = OperationData->InstanceUpdateData.Skeletons[ComponentIndex].BoneMatricesWithScale;
+
+			// Use first valid LOD bone count as a potential total number of bones, used for pre-allocating data arrays
+			if (MinLODComponent.Mesh && MinLODComponent.Mesh->GetSkeleton())
+			{
+				const int32 TotalPossibleBones = MinLODComponent.Mesh->GetSkeleton()->GetBoneCount();
+
+				// Helpers
+				BoneNameToIndex.Empty(TotalPossibleBones);
+
+				// Out Data
+				BoneNames.Reserve(TotalPossibleBones);
+				BoneMatricesWithScale.Reserve(TotalPossibleBones);
+			}
+
+			for (int32 LODIndex = OperationData->CurrentMinLOD; LODIndex <= OperationData->CurrentMaxLOD && LODIndex < LODCount; ++LODIndex)
+			{
+				MUTABLE_CPUPROFILER_SCOPE(PrepareSkeletonData_LODs);
+
+				const FInstanceUpdateData::FLOD& CurrentLOD = OperationData->InstanceUpdateData.LODs[LODIndex];
+				FInstanceUpdateData::FComponent& CurrentLODComponent = OperationData->InstanceUpdateData.Components[CurrentLOD.FirstComponent + ComponentIndex];
+				mu::MeshPtrConst Mesh = CurrentLODComponent.Mesh;
+
+				if (!Mesh || !Mesh->GetSkeleton())
+				{
+					continue;
+				}
+
+				const mu::Skeleton* Skeleton = Mesh->GetSkeleton().get();
+				const int32 BoneCount = Skeleton->GetBoneCount();
+
+				{
+					MUTABLE_CPUPROFILER_SCOPE(PrepareSkeletonData_UsedBones);
+
+					const mu::FMeshBufferSet& MutableMeshVertexBuffers = Mesh->GetVertexBuffers();
+
+					const int32 NumVerticesLODModel = Mesh->GetVertexCount();
+					const int32 SurfaceCount = Mesh->GetSurfaceCount();
+
+					mu::MESH_BUFFER_FORMAT boneIndexFormat = mu::MBF_NONE;
+					int boneIndexComponents = 0;
+					int boneIndexOffset = 0;
+					int boneIndexBuffer = -1;
+					int boneIndexChannel = -1;
+					MutableMeshVertexBuffers.FindChannel(mu::MBS_BONEINDICES, 0, &boneIndexBuffer, &boneIndexChannel);
+					if (boneIndexBuffer >= 0 || boneIndexChannel >= 0)
+					{
+						MutableMeshVertexBuffers.GetChannel(boneIndexBuffer, boneIndexChannel,
+							nullptr, nullptr, &boneIndexFormat, &boneIndexComponents, &boneIndexOffset);
+					}
+
+					int32 elemSize = MutableMeshVertexBuffers.GetElementSize(boneIndexBuffer);
+
+					const uint8_t* BufferStart = MutableMeshVertexBuffers.GetBufferData(boneIndexBuffer) + boneIndexOffset;
+
+					UsedBones.Empty(BoneCount);
+					UsedBones.AddDefaulted(BoneCount);
+
+					for (int32 Surface = 0; Surface < SurfaceCount; Surface++)
+					{
+						int32 is, ic, vs, VertexCount;
+						Mesh->GetSurface(Surface, &vs, &VertexCount, &is, &ic);
+
+						if (VertexCount == 0 || ic == 0)
+						{
+							continue;
+						}
+
+						const uint8_t* VertexBoneIndexPtr = BufferStart + vs * elemSize;
+
+						if (boneIndexFormat == mu::MBF_UINT8)
+						{
+							for (int32 v = 0; v < VertexCount; ++v)
+							{
+								for (int32 i = 0; i < boneIndexComponents; ++i)
+								{
+									size_t SectionBoneIndex = VertexBoneIndexPtr[i];
+									UsedBones[SectionBoneIndex] = true;
+								}
+								VertexBoneIndexPtr += elemSize;
+							}
+						}
+						else if (boneIndexFormat == mu::MBF_UINT16)
+						{
+							for (int32 v = 0; v < VertexCount; ++v)
+							{
+								for (int32 i = 0; i < boneIndexComponents; ++i)
+								{
+									size_t SectionBoneIndex = ((uint16_t*)VertexBoneIndexPtr)[i];
+									UsedBones[SectionBoneIndex] = true;
+								}
+								VertexBoneIndexPtr += elemSize;
+							}
+						}
+						else if (boneIndexFormat == mu::MBF_UINT32)
+						{
+							for (int32 v = 0; v < VertexCount; ++v)
+							{
+								for (int32 i = 0; i < boneIndexComponents; ++i)
+								{
+									size_t SectionBoneIndex = ((uint32_t*)VertexBoneIndexPtr)[i];
+									UsedBones[SectionBoneIndex] = true;
+								}
+								VertexBoneIndexPtr += elemSize;
+							}
+						}
+						else
+						{
+							// Unsupported bone index format in generated mutable mesh
+							check(false);
+						}
+
+					}
+				}
+
+
+				{
+					MUTABLE_CPUPROFILER_SCOPE(PrepareSkeletonData_ActiveBones);
+
+					TArray<uint16>& ComponentBoneMap = CurrentLODComponent.BoneMap;
+					TArray<uint16>& ComponentActiveBones = CurrentLODComponent.ActiveBones;
+
+					ComponentBoneMap.Reserve(BoneCount);
+					ComponentActiveBones.Reserve(BoneCount);
+
+					ParentChain.SetNumZeroed(BoneCount);
+
+					// Add new bones and their poses to the final hierarchy of active bones
+					auto AddBone = [&](const char* strName) {
+
+						const FName BoneName = strName;
+						const int16 FinalBoneIndex = BoneNames.Num();
+
+						// Add bone
+						BoneNames.Add(BoneName);
+
+						BoneNameToIndex.Add(BoneName, FinalBoneIndex);
+						ComponentActiveBones.Add(FinalBoneIndex);
+
+						// Add bone pose
+						const int32 BonePoseIndex = Mesh->FindBonePose(strName);
+						if (BonePoseIndex != INDEX_NONE)
+						{
+							FTransform3f transform;
+							Mesh->GetBoneTransform(BonePoseIndex, transform);
+							BoneMatricesWithScale.Add(BoneName, transform.Inverse().ToMatrixWithScale());
+						}
+
+						return FinalBoneIndex;
+					};
+
+					for (int32 BoneIndex = 0; BoneIndex < BoneCount; ++BoneIndex)
+					{
+						const char* BoneName = Skeleton->GetBoneName(BoneIndex);
+
+						int16* FinalBoneIndexPtr = BoneNameToIndex.Find(BoneName);
+						if (!UsedBones[BoneIndex])
+						{
+							if (!FinalBoneIndexPtr)
+							{
+								// Add bone even if it doesn't have weights
+								AddBone(BoneName);
+
+								// Add SkeletonId 
+								SkeletonIds.AddUnique(Skeleton->GetBoneId(BoneIndex));
+							}
+
+							// Add root as a placeholder
+							ComponentBoneMap.Add(0);
+							continue;
+						}
+
+						if (!FinalBoneIndexPtr) // New bone!
+						{
+							// Ensure parent chain is valid
+							int16 FinalParentIndex = INDEX_NONE;
+
+							const int16 MutParentIndex = Skeleton->GetBoneParent(BoneIndex);
+							if (MutParentIndex != INDEX_NONE) // Bone has a parent, ensure the parent exists
+							{
+								// Ensure parent chain until root exists
+								uint16 ParentChainCount = 0;
+
+								// Find parents to add
+								int16 NextMutParentIndex = MutParentIndex;
+								while (NextMutParentIndex != INDEX_NONE)
+								{
+									const char* ParentName = Skeleton->GetBoneName(NextMutParentIndex);
+									if (int16* Found = BoneNameToIndex.Find(ParentName))
+									{
+										FinalParentIndex = *Found;
+										break;
+									}
+
+									ParentChain[ParentChainCount] = ParentName;
+									++ParentChainCount;
+
+									// Add SkeletonId 
+									SkeletonIds.AddUnique(Skeleton->GetBoneId(NextMutParentIndex));
+
+									NextMutParentIndex = Skeleton->GetBoneParent(NextMutParentIndex);
+								}
+
+								// Add parent bones to the list and to the active bones array
+								for (int16 ParentChainIndex = ParentChainCount - 1; ParentChainIndex >= 0; --ParentChainIndex)
+								{
+									// Add the parent
+									AddBone(ParentChain[ParentChainIndex]);
+								}
+							}
+
+							// Add the used bone to the hierarchy
+							const int16 NewBoneIndex = AddBone(BoneName);
+							ComponentBoneMap.Add(NewBoneIndex);
+
+							// Add SkeletonId 
+							SkeletonIds.AddUnique(Skeleton->GetBoneId(BoneIndex));
+						}
+						else
+						{
+							int16 FinalBoneIndex = *FinalBoneIndexPtr;
+							ComponentBoneMap.Add(FinalBoneIndex);
+
+							if (ComponentActiveBones.Find(FinalBoneIndex) != INDEX_NONE)
+							{
+								continue;
+							}
+
+							int16 MutParentIndex = Skeleton->GetBoneParent(BoneIndex);
+							while (MutParentIndex != INDEX_NONE)
+							{
+								FName ParentName = Skeleton->GetBoneName(MutParentIndex);
+								const int16* FinalParentIndex = BoneNameToIndex.Find(ParentName);
+								check(FinalParentIndex);
+
+								if (ComponentActiveBones.Find(*FinalParentIndex) != INDEX_NONE)
+								{
+									break;
+								}
+
+								ComponentActiveBones.Add(*FinalParentIndex);
+
+								MutParentIndex = Skeleton->GetBoneParent(MutParentIndex);
+							}
+
+							ComponentActiveBones.Add(FinalBoneIndex);
+						}
+					}
+
+					ComponentBoneMap.Shrink();
+					ComponentActiveBones.Shrink();
+				}
+			}
+
+			BoneNames.Shrink();
+		}
+	}
+
+
+	void Task_Mutable_Update_GetMesh(TSharedPtr<FMutableOperationData> OperationData, mu::ModelPtr Model, mu::ParametersPtrConst Parameters, bool bBuildParameterDecorations, int32 State)
+	{
+		// This runs in a worker thread.
+		MUTABLE_CPUPROFILER_SCOPE(Task_Mutable_GetMesh)
+
+#if WITH_EDITOR
+		uint32 StartCycles = FPlatformTime::Cycles();
+#endif
+
+		check(OperationData.IsValid());
+
+		if (bBuildParameterDecorations)
+		{
+			// This cannot happen inside the beginupdate-endupdate cycle.
+			Subtask_Mutable_UpdateParameterDecorations(OperationData, Model, Parameters.get());
+		}
+		else
+		{
+			OperationData->ParametersUpdateData.Clear();
+			OperationData->RelevantParametersInProgress.Reset();
+		}
+
+		Subtask_Mutable_BeginUpdate_GetMesh(OperationData, Model, Parameters.get(), State);
+
+		// TODO: Not strictly mutable: move to another worker thread task to free mutable access?
+		Subtask_Mutable_PrepareSkeletonData(OperationData);
+
+#if WITH_EDITOR
+		uint32 EndCycles = FPlatformTime::Cycles();
+		OperationData->MutableRuntimeCycles = EndCycles - StartCycles;
+#endif
+
+	}
+
+
+	void Task_Mutable_Update_GetImages(TSharedPtr<FMutableOperationData> OperationData, mu::ModelPtr Model, mu::ParametersPtrConst Parameters, int32 State)
+	{
+		// This runs in a worker thread.
+		MUTABLE_CPUPROFILER_SCOPE(Task_Mutable_GetImages)
+
+#if WITH_EDITOR
+		uint32 StartCycles = FPlatformTime::Cycles();
+#endif
+
+		check(OperationData.IsValid());
+
+		Subtask_Mutable_GetImages(OperationData, Model, Parameters.get(), State);
+
+		// TODO: Not strictly mutable: move to another worker thread task to free mutable access?
+		Subtask_Mutable_PrepareTextures(OperationData);
+
+#if WITH_EDITOR
+		uint32 EndCycles = FPlatformTime::Cycles();
+		OperationData->MutableRuntimeCycles += EndCycles - StartCycles;
+#endif
+
+	}
+
+
+	void Task_Mutable_ReleaseInstance(TSharedPtr<FMutableOperationData> OperationData, mu::SystemPtr MutableSystem)
+	{
+		MUTABLE_CPUPROFILER_SCOPE(Task_Mutable_ReleaseInstance)
+
+		// This runs in a worker thread.
+		check(OperationData.IsValid());
+
+		if (OperationData->InstanceID > 0)
+		{
+			MUTABLE_CPUPROFILER_SCOPE(EndUpdate);
+			MutableSystem->EndUpdate(OperationData->InstanceID);
+			OperationData->InstanceUpdateData.Clear();
+
+			MutableSystem->ReleaseInstance(OperationData->InstanceID);
+			OperationData->InstanceID = 0;
+		}
+
+	}
+
+
+	void Task_Game_ReleasePlatformData(TSharedPtr<FMutableReleasePlatformOperationData> OperationData)
+	{
+		MUTABLE_CPUPROFILER_SCOPE(Task_Game_ReleasePlatformData)
+
+		TMap<uint32, FTexturePlatformData*>& ImageToPlatformDataMap = OperationData->ImageToPlatformDataMap;
+		for (TPair<uint32, FTexturePlatformData*> Pair : ImageToPlatformDataMap)
+		{
+			delete Pair.Value; // If this is not null then it must mean it hasn't been used, otherwise they would have taken ownership and nulled it
+		}
+		ImageToPlatformDataMap.Reset();
+	}
+
+
+	void Task_Game_Callbacks(TSharedPtr<FMutableOperationData> OperationData, const TWeakObjectPtr<UCustomizableObjectInstance>& CustomizableObjectInstancePtr)
+	{
+		MUTABLE_CPUPROFILER_SCOPE(Task_Game_Callbacks)
+
+		check(IsInGameThread());
+
+		UCustomizableObjectSystem* System = UCustomizableObjectSystem::GetInstance();
+		if (!System || !System->IsValidLowLevel() || System->HasAnyFlags(RF_BeginDestroyed))
+		{
+			return;
+		}
+
+		UCustomizableObjectInstance* CustomizableObjectInstance = CustomizableObjectInstancePtr.Get();
+
+		// TODO: Review checks.
+		if (!CustomizableObjectInstance || !CustomizableObjectInstance->IsValidLowLevel() )
+		{
+			System->ClearCurrentMutableOperation();
+			return;
+		}
+
+		// Actual work
+
+		// This used to be CustomizableObjectInstance::UpdateSkeletalMesh_PostBeginUpdate3
+		{
+			MUTABLE_CPUPROFILER_SCOPE(UpdateSkeletalMesh_PostBeginUpdate3);
+
+#if !WITH_EDITOR
+			for (int32 ComponentIndex = 0; ComponentIndex < CustomizableObjectInstance->SkeletalMeshes.Num(); ++ComponentIndex)
+			{
+				if (CustomizableObjectInstance->SkeletalMeshes[ComponentIndex])
+				{
+					CustomizableObjectInstance->SkeletalMeshes[ComponentIndex]->RebuildSocketMap();
+				}
+			}
+#endif
+
+			for (TObjectIterator<UCustomizableSkeletalComponent> It; It; ++It)
+			{
+				UCustomizableSkeletalComponent* CustomizableSkeletalComponent = *It;
+
+				if (CustomizableSkeletalComponent && CustomizableSkeletalComponent->CustomizableObjectInstance == CustomizableObjectInstance
+					&& CustomizableObjectInstance->SkeletalMeshes.IsValidIndex(CustomizableSkeletalComponent->ComponentIndex))
+				{
+					MUTABLE_CPUPROFILER_SCOPE(UpdateSkeletalMesh_SetSkeletalMesh);
+
+					const bool bIsCreatingSkeletalMesh = CustomizableObjectInstance->GetPrivate()->HasCOInstanceFlags(CreatingSkeletalMesh);
+					CustomizableSkeletalComponent->SetSkeletalMesh(CustomizableObjectInstance->SkeletalMeshes[CustomizableSkeletalComponent->ComponentIndex], false, bIsCreatingSkeletalMesh);
+
+					if (CustomizableObjectInstance->GetPrivate()->HasCOInstanceFlags(ReplacePhysicsAssets))
+					{
+						CustomizableSkeletalComponent->SetPhysicsAsset(
+							CustomizableObjectInstance->SkeletalMeshes[CustomizableSkeletalComponent->ComponentIndex] ? 
+							CustomizableObjectInstance->SkeletalMeshes[CustomizableSkeletalComponent->ComponentIndex]->GetPhysicsAsset() : nullptr);
+					}
+				}
+			}
+
+			CustomizableObjectInstance->GetPrivate()->SetCOInstanceFlags(Generated);
+			CustomizableObjectInstance->GetPrivate()->ClearCOInstanceFlags(CreatingSkeletalMesh);
+
+			if (CustomizableObjectInstance->bEditorPropertyChanged)
+			{
+				CustomizableObjectInstance->bEditorPropertyChanged = false;
+			}
+
+			for (TObjectIterator<UCustomizableSkeletalComponent> It; It; ++It)
+			{
+				UCustomizableSkeletalComponent* CustomizableSkeletalComponent = *It;
+
+				if (CustomizableSkeletalComponent && CustomizableSkeletalComponent->CustomizableObjectInstance == CustomizableObjectInstance)
+				{
+					MUTABLE_CPUPROFILER_SCOPE(UpdateSkeletalMesh_ComponentUpdatedDelegate);
+
+					CustomizableSkeletalComponent->UpdatedDelegate.ExecuteIfBound();
+				}
+			}
+
+			{
+				MUTABLE_CPUPROFILER_SCOPE(UpdateSkeletalMesh_UpdatedDelegate);
+
+				CustomizableObjectInstance->Updated();
+
+#if WITH_EDITOR
+				CustomizableObjectInstance->InstanceUpdated = true;
+
+				// \TODO: Review if we can avoid all this editor code here.
+				for (int32 MeshIndex = 0; MeshIndex < CustomizableObjectInstance->SkeletalMeshes.Num(); ++MeshIndex)
+				{
+					if (USkeletalMesh* SkeletalMesh = CustomizableObjectInstance->GetSkeletalMesh(MeshIndex))
+					{
+						FUnrealBakeHelpers::BakeHelper_RegenerateImportedModel(SkeletalMesh);
+					}
+				}
+#endif //WITH_EDITOR
+			}
+
+			FCustomizableObjectSystemPrivate* CustomizableObjectSystem = UCustomizableObjectSystem::GetInstance()->GetPrivate();
+			if (CustomizableObjectSystem->bReleaseTexturesImmediately)
+			{
+				FMutableResourceCache& Cache = CustomizableObjectSystem->GetObjectCache(CustomizableObjectInstance->GetCustomizableObject());
+
+				for (TPair<uint32, FGeneratedTexture>& Item : CustomizableObjectInstance->GetPrivate()->TexturesToRelease)
+				{
+					UCustomizableInstancePrivateData::ReleaseMutableTexture(Item.Value.Id, Item.Value.Texture, Cache);
+				}
+
+				CustomizableObjectInstance->GetPrivate()->TexturesToRelease.Empty();
+			}
+		}
+
+
+		// TODO: T2927
+		if (LogBenchmarkUtil::isLoggingActive())
+		{
+			// Log the time stat		
+			double deltaSeconds = FPlatformTime::Seconds() - System->GetPrivate()->CurrentMutableOperation->StartUpdateTime;
+			int32 deltaMs = int32(deltaSeconds * 1000);
+			int64 streamingCache = System->GetPrivate()->LastStreamingMemorySize / 1024;
+
+			LogBenchmarkUtil::updateStat("customizable_instance_build_time", deltaMs);
+			LogBenchmarkUtil::updateStat("mutable_streaming_cache_memory", (long double)streamingCache / 1024.0);
+			System->GetPrivate()->TotalBuildMs += deltaMs;
+			System->GetPrivate()->TotalBuiltInstances++;
+			SET_DWORD_STAT(STAT_MutableInstanceBuildTime, deltaMs);
+			SET_DWORD_STAT(STAT_MutableInstanceBuildTimeAvrg, System->GetPrivate()->TotalBuildMs / System->GetPrivate()->TotalBuiltInstances);
+			SET_DWORD_STAT(STAT_MutableStreamingCache, streamingCache);
+
+			mu::System* MutableSystem = System->GetPrivate()->MutableSystem.get();
+			if (MutableSystem)
+			{
+				SET_DWORD_STAT(STAT_MutableProfile_LiveInstanceCount, MutableSystem->GetProfileMetric(mu::System::ProfileMetric::LiveInstanceCount));
+				SET_DWORD_STAT(STAT_MutableProfile_StreamingCacheBytes, MutableSystem->GetProfileMetric(mu::System::ProfileMetric::StreamingCacheBytes));
+				SET_DWORD_STAT(STAT_MutableProfile_InstanceUpdateCount, MutableSystem->GetProfileMetric(mu::System::ProfileMetric::InstanceUpdateCount));
+			}
+			else
+			{
+				SET_DWORD_STAT(STAT_MutableProfile_LiveInstanceCount, 0);
+				SET_DWORD_STAT(STAT_MutableProfile_StreamingCacheBytes, 0);
+				SET_DWORD_STAT(STAT_MutableProfile_InstanceUpdateCount, 0);
+			}
+		}
+
+		// End Update
+		System->ClearCurrentMutableOperation();
+	}
+
+
+	void Task_Game_ConvertResources(TSharedPtr<FMutableOperationData> OperationData, const TWeakObjectPtr<UCustomizableObjectInstance>& CustomizableObjectInstancePtr)
+	{
+		MUTABLE_CPUPROFILER_SCOPE(Task_Game_ConvertResources)
+
+		check(IsInGameThread());
+
+		UCustomizableObjectSystem* System = UCustomizableObjectSystem::GetInstance();
+		if (!System || !System->IsValidLowLevel() || System->HasAnyFlags(RF_BeginDestroyed))
+		{
+			return;
+		}
+
+		UCustomizableObjectInstance* CustomizableObjectInstance = CustomizableObjectInstancePtr.Get();
+
+		// TODO: Review checks.
+		bool bCancel = false;
+		if (!CustomizableObjectInstance || !CustomizableObjectInstance->IsValidLowLevel())
+		{
+			bCancel = true;
+		}
+
+		// Actual work
+
+
+		// Process the pending texture coverage queries
+		if (!bCancel)
+		{
+			MUTABLE_CPUPROFILER_SCOPE(GameTextureQueries);
+			for (const FPendingTextureCoverageQuery& Query : OperationData->PendingTextureCoverageQueries)
+			{
+				UMaterialInterface* Material = nullptr;
+				const uint32* InstanceIndex = CustomizableObjectInstance->GetPrivate()->ObjectToInstanceIndexMap.Find(Query.MaterialIndex);
+				if (InstanceIndex && CustomizableObjectInstance->GetPrivate()->ReferencedMaterials.IsValidIndex(*InstanceIndex))
+				{
+					Material = CustomizableObjectInstance->GetPrivate()->ReferencedMaterials[*InstanceIndex];
+				}
+
+				UCustomizableInstancePrivateData::ProcessTextureCoverageQueries(OperationData, CustomizableObjectInstance->GetCustomizableObject(), Query.KeyName, Query.PlatformData, Material);
+			}
+			OperationData->PendingTextureCoverageQueries.Empty();
+		}
+
+		// Process texture coverage queries because it's safe to do now that the Mutable thread is stopped
+		if (!bCancel)
+		{
+			UCustomizableInstancePrivateData* PrivateInstance = CustomizableObjectInstance->GetPrivate();
+			if (OperationData->TextureCoverageQueries_MutableThreadResults.Num() > 0)
+			{
+				for (auto& Result : OperationData->TextureCoverageQueries_MutableThreadResults)
+				{
+					FTextureCoverageQueryData* FinalResultData = PrivateInstance->TextureCoverageQueries.Find(Result.Key);
+					*FinalResultData = Result.Value;
+				}
+
+				OperationData->TextureCoverageQueries_MutableThreadResults.Empty();
+			}
+
+#if WITH_EDITOR
+			CustomizableObjectInstance->LastUpdateMutableRuntimeCycles = OperationData->MutableRuntimeCycles;
+#endif
+
+			// Convert Step
+			//-------------------------------------------------------------
+
+			// \TODO: Bring that code here instead of keeping it in the UCustomizableObjectInstance
+			if (CustomizableObjectInstance->GetPrivate()->UpdateSkeletalMesh_PostBeginUpdate0(CustomizableObjectInstance, OperationData))
+			{
+				// This used to be CustomizableObjectInstance::UpdateSkeletalMesh_PostBeginUpdate1
+				{
+					MUTABLE_CPUPROFILER_SCOPE(UpdateSkeletalMesh_PostBeginUpdate1);
+
+					// \TODO: Bring here
+					CustomizableObjectInstance->GetPrivate()->BuildMaterials(OperationData, CustomizableObjectInstance);
+				}
+
+				// This used to be CustomizableObjectInstance::UpdateSkeletalMesh_PostBeginUpdate2
+				{
+					MUTABLE_CPUPROFILER_SCOPE(UpdateSkeletalMesh_PostBeginUpdate2);
+
+					for (int32 Component = 0; Component < CustomizableObjectInstance->SkeletalMeshes.Num(); ++Component)
+					{
+						if (CustomizableObjectInstance->SkeletalMeshes[Component] && CustomizableObjectInstance->SkeletalMeshes[Component]->GetLODInfoArray().Num())
+						{
+							MUTABLE_CPUPROFILER_SCOPE(UpdateSkeletalMesh_PostEditChangeProperty);
+
+							CustomizableObjectInstance->GetPrivate()->PostEditChangePropertyWithoutEditor(CustomizableObjectInstance->SkeletalMeshes[Component]);
+						}
+					}
+				}
+			}
+		}
+
+		// Next Task: Release Mutable. We need this regardless if we cancel or not
+		//-------------------------------------------------------------		
+		{
+			mu::SystemPtr MutableSystem = System->GetPrivate()->MutableSystem;
+			System->GetPrivate()->AddMutableThreadTask(
+				TEXT("Task_Mutable_ReleaseInstance"),
+				[OperationData, MutableSystem]() {Task_Mutable_ReleaseInstance(OperationData, MutableSystem); },
+				UE::Tasks::ETaskPriority::BackgroundNormal);
+		}
+
+
+		// Next Task: Release Platform Data
+		//-------------------------------------------------------------
+		if (!bCancel)
+		{
+			TSharedPtr<FMutableReleasePlatformOperationData> ReleaseOperationData = MakeShared<FMutableReleasePlatformOperationData>();
+			ReleaseOperationData->ImageToPlatformDataMap = MoveTemp(OperationData->ImageToPlatformDataMap);
+			System->GetPrivate()->AddAnyThreadTask(
+				TEXT("Mutable_ReleasePlatformData"),
+				[ReleaseOperationData]() { Task_Game_ReleasePlatformData(ReleaseOperationData); },
+				UE::Tasks::ETaskPriority::BackgroundNormal
+			);
+
+
+			// Unlock step
+			//-------------------------------------------------------------
+			if (CustomizableObjectInstance->GetCustomizableObject())
+			{
+				// Unlock the resource cache for the object used by this instance to avoid
+				// the destruction of resources that we may want to reuse.
+				System->ClearResourceCacheProtected();
+
+				// Free cached external images for this operation
+				{
+					MUTABLE_CPUPROFILER_SCOPE(ClearExternalImages);
+					if (System->GetPrivate()->ImageProvider)
+					{
+						System->GetPrivate()->ImageProvider->ClearCache();
+					}
+				}
+			}
+
+			// Next Task: Callbacks
+			//-------------------------------------------------------------
+			System->GetPrivate()->AddGameThreadTask(
+				{
+				FMutableTaskDelegate::CreateLambda(
+					[OperationData,CustomizableObjectInstancePtr]()
+					{
+						Task_Game_Callbacks(OperationData,CustomizableObjectInstancePtr);
+					}),
+					{}
+				});
+		}
+	}
+
+
+	/** "Lock Cached Resources" */
+	void Task_Game_LockCache(TSharedPtr<FMutableOperationData> OperationData, const TWeakObjectPtr<UCustomizableObjectInstance>& CustomizableObjectInstancePtr, mu::Ptr<const mu::Parameters> Parameters, bool bBuildParameterDecorations)
+	{
+		check(IsInGameThread());
+
+		UCustomizableObjectSystem* System = UCustomizableObjectSystem::GetInstance();
+		if (!System)
+		{
+			return;
+		}
+
+		UCustomizableObjectInstance* ObjectInstance = CustomizableObjectInstancePtr.Get();
+
+		if (!ObjectInstance)
+		{
+			System->ClearCurrentMutableOperation();
+			return;
+		}
+
+		const UCustomizableObject* CustomizableObject = ObjectInstance->GetCustomizableObject(); 
+		if (!CustomizableObject)
+		{
+			System->ClearCurrentMutableOperation();
+			return;
+		}
+		
+		// Process the parameter decorations if requested
+		if (bBuildParameterDecorations)
+		{
+			ObjectInstance->GetPrivate()->UpdateParameterDecorationsEngineResources(OperationData);
+		}
+
+
+		// Selectivelly lock the resource cache for the object used by this instance to avoid the destruction of resources that we may want to reuse.
+		// When protecting textures there mustn't be any left from a previous update
+		check(System->ProtectedCachedTextures.Num() == 0);
+
+		FMutableResourceCache& Cache = System->GetPrivate()->GetObjectCache(CustomizableObject);
+
+		System->ProtectedCachedTextures.Reset(Cache.Images.Num());
+		System->GetPrivate()->ProtectedObjectCachedImages.Reset(Cache.Images.Num());
+
+		for (const FInstanceUpdateData::FImage& Image : OperationData->InstanceUpdateData.Images)
+		{
+			TWeakObjectPtr<UTexture2D>* TexturePtr = Cache.Images.Find({ Image.ImageID });
+
+			if (TexturePtr && TexturePtr->Get() && System->GetPrivate()->TextureHasReferences(Image.ImageID))
+			{
+				System->ProtectedCachedTextures.Add(TexturePtr->Get());
+				System->GetPrivate()->ProtectedObjectCachedImages.Add(Image.ImageID);
+			}
+		}
+
+		// Get any external texture that may be requested for this update
+		{
+			MUTABLE_CPUPROFILER_SCOPE(GetExternalImages);
+			if (System->GetPrivate()->ImageProvider)
+			{
+				for (const FCustomizableObjectTextureParameterValue& Param : ObjectInstance->GetTextureParameters())
+				{
+					System->GetPrivate()->ImageProvider->CacheImage(Param.ParameterValue);
+				}
+			}
+		}
+
+
+		// Copy data generated in the mutable thread over to the instance
+		ObjectInstance->GetPrivate()->PrepareForUpdate(OperationData);
+
+		// Task: Mutable GetImages
+		//-------------------------------------------------------------
+		FGraphEventRef Mutable_GetImagesTask;
+		{
+			// Task inputs
+			mu::ModelPtr Model = CustomizableObject->GetPrivate()->GetModel();
+			int32 State = CustomizableObjectInstancePtr->GetState();
+
+			Mutable_GetImagesTask = System->GetPrivate()->AddMutableThreadTask(
+					TEXT("Task_Mutable_GetImages"),
+					[OperationData, Parameters, Model, State]()
+					{
+						impl::Task_Mutable_Update_GetImages(OperationData, Model, Parameters, State);
+					},
+					UE::Tasks::ETaskPriority::BackgroundHigh);
+		}
+
+
+		// Next Task: Load Unreal Assets
+		//-------------------------------------------------------------
+		FGraphEventRef Game_LoadUnrealAssets = ObjectInstance->GetPrivate()->LoadAdditionalAssetsAsync(OperationData, ObjectInstance, UCustomizableObjectSystem::GetInstance()->GetStreamableManager());
+		if (Game_LoadUnrealAssets) Game_LoadUnrealAssets->SetDebugName(TEXT("LoadAdditionalAssetsAsync"));
+
+		// Next-next Task: Convert Resources
+		//-------------------------------------------------------------
+		System->GetPrivate()->AddGameThreadTask(
+			{
+			FMutableTaskDelegate::CreateLambda(
+				[OperationData,CustomizableObjectInstancePtr]()
+				{
+					Task_Game_ConvertResources(OperationData,CustomizableObjectInstancePtr);
+				}),
+#ifdef MUTABLE_USE_NEW_TASKGRAPH
+				{},
+#endif
+				Game_LoadUnrealAssets,
+				Mutable_GetImagesTask
+			});
+	}
+
+
+	/** "Start Update" */
+	void Task_Game_StartUpdate(TSharedPtr<FMutableOperation> Operation)
+	{
+		check(Operation->Type == FMutableOperation::EOperationType::Update);
+
+		UCustomizableObjectSystem* System = UCustomizableObjectSystem::GetInstance();
+
+		if (!Operation->CustomizableObjectInstance.IsValid() || !Operation->CustomizableObjectInstance->IsValidLowLevel()) // Only start if it hasn't been already destroyed (i.e. GC after finish PIE)
+		{
+			System->ClearCurrentMutableOperation();
+			return;
+		}
+
+
+		if (Operation->CustomizableObjectInstance->GetPrivate()->HasCOInstanceFlags(PendingLODsUpdate))
+		{
+			Operation->CustomizableObjectInstance->GetPrivate()->ClearCOInstanceFlags(PendingLODsUpdate);
+			// TODO: Is anything needed for this now?
+			//Operation->CustomizableObjectInstance->ReleaseMutableInstanceId(); // To make mutable regenerate the LODs even if the instance parameters have not changed
+		}
+
+		float LastMinSquareDistFromComponentToPlayer = Operation->CustomizableObjectInstance->GetPrivate()->LastMinSquareDistFromComponentToPlayer;
+
+		UCustomizableObjectInstance* CandidateInstance = Operation->CustomizableObjectInstance.Get();
+
+		bool bCancel = false;
+
+		// If the object is locked (for instance, compiling) we skip any instance update.
+		if (!CandidateInstance)
+		{
+			bCancel = true;
+		}
+
+		UCustomizableObject* CustomizableObject = CandidateInstance->GetCustomizableObject(); 
+		if (CustomizableObject->GetPrivate()->bLocked)
+		{
+			bCancel = true;
+		}
+		
+		// Only update resources if the instance is in range (it could have got far from the player since the task was queued)
+		if (
+			System->CurrentInstanceLODManagement->IsOnlyUpdateCloseCustomizableObjectsEnabled()
+			&& LastMinSquareDistFromComponentToPlayer > FMath::Square(System->CurrentInstanceLODManagement->GetOnlyUpdateCloseCustomizableObjectsDist())
+			&& LastMinSquareDistFromComponentToPlayer != FLT_MAX // This means it is the first frame so it has to be updated
+			)
+		{
+			bCancel = true;
+		}
+
+		mu::Ptr<const mu::Parameters> Parameters = Operation->GetParameters();
+		if (!Parameters)
+		{
+			bCancel = true;
+		}
+
+		if (bCancel)
+		{
+			Operation->CustomizableObjectInstance->GetPrivate()->ClearCOInstanceFlags(Updating);
+			System->ClearCurrentMutableOperation();
+			return;
+		}
+
+		if (LogBenchmarkUtil::isLoggingActive())
+		{
+			Operation->StartUpdateTime = FPlatformTime::Seconds();
+		}
+
+		System->GetPrivate()->CurrentInstanceBeingUpdated = CandidateInstance;
+
+		UCustomizableInstancePrivateData* PrivateInstance = CandidateInstance->GetPrivate();
+
+		PrivateInstance->CurrentMinLOD = FMath::Max(PrivateInstance->MinLODToLoad, (int32)PrivateInstance->FirstLODAvailable);
+		PrivateInstance->CurrentMinLOD = FMath::Min(PrivateInstance->CurrentMinLOD, PrivateInstance->FirstLODAvailable + PrivateInstance->NumMaxLODsToStream);
+
+		PrivateInstance->CurrentMaxLOD = FMath::Max(PrivateInstance->MaxLODToLoad, PrivateInstance->CurrentMinLOD);
+
+
+		// Prepare streaming for the current customizable object
+		System->GetPrivate()->Streamer->PrepareStreamingForObject(CustomizableObject);
+
+
+		// Task: Mutable Update and GetMesh
+		//-------------------------------------------------------------
+		TSharedPtr<FMutableOperationData> CurrentOperationData = MakeShared<FMutableOperationData>();
+		CurrentOperationData->TextureCoverageQueries_MutableThreadParams = PrivateInstance->TextureCoverageQueries;
+		CurrentOperationData->TextureCoverageQueries_MutableThreadResults.Empty();
+		CurrentOperationData->CurrentMinLOD = PrivateInstance->CurrentMinLOD;
+		CurrentOperationData->CurrentMaxLOD = PrivateInstance->CurrentMaxLOD;
+		CurrentOperationData->bNeverStream = Operation->bNeverStream;
+		CurrentOperationData->MipsToSkip = Operation->MipsToSkip;
+		CurrentOperationData->MutableParameters = Parameters;
+		CurrentOperationData->State = Operation->CustomizableObjectInstance->GetState();
+
+		FGraphEventRef Mutable_GetMeshTask;
+		{
+			// Task inputs
+			TWeakObjectPtr<UCustomizableObjectInstance> CustomizableObjectInstancePtr = CandidateInstance;
+			TSharedPtr<FMutableOperation> CurrentMutableOperation = System->GetPrivate()->CurrentMutableOperation;
+			bool bBuildParameterDecorations = CurrentMutableOperation->IsBuildParameterDecorations();
+			mu::ModelPtr Model = CustomizableObject->GetPrivate()->GetModel();
+			int32 State = CustomizableObjectInstancePtr->GetState();
+
+			Mutable_GetMeshTask = System->GetPrivate()->AddMutableThreadTask(
+				TEXT("Task_Mutable_Update_GetMesh"),
+				[CurrentOperationData, bBuildParameterDecorations, Parameters, Model, State]()
+				{
+					impl::Task_Mutable_Update_GetMesh(CurrentOperationData, Model, Parameters, bBuildParameterDecorations, State);
+				},
+				UE::Tasks::ETaskPriority::BackgroundHigh);
+		}
+
+
+		// Task: Lock cache
+		//-------------------------------------------------------------
+		{
+			// Task inputs
+			TWeakObjectPtr<UCustomizableObjectInstance> CustomizableObjectInstancePtr = CandidateInstance;
+			bool bBuildParameterDecorations = System->GetPrivate()->CurrentMutableOperation->IsBuildParameterDecorations();
+
+			System->GetPrivate()->AddGameThreadTask(
+				{
+				FMutableTaskDelegate::CreateLambda(
+					[CurrentOperationData, CustomizableObjectInstancePtr, bBuildParameterDecorations, Parameters]()
+					{
+						impl::Task_Game_LockCache(CurrentOperationData, CustomizableObjectInstancePtr, Parameters, bBuildParameterDecorations);
+					}),
+				Mutable_GetMeshTask
+				});
+		}
+	}
+
+} // namespace impl
+
+
+void UCustomizableObjectSystem::AdvanceCurrentOperation() 
+{
+	TRACE_CPUPROFILER_EVENT_SCOPE(Mutable_AdvanceCurrentOperation);
+
+	// See if we can clear the last reference to the mutable-thread task
+	Private->ClearMutableTaskIfDone();
+
+	// See if we have a game-thread task to process
+	FMutableTask* PendingTask = Private->PendingTasks.Peek();
+	if (PendingTask)
+	{
+		if (PendingTask->AreDependenciesComplete())
+		{
+			PendingTask->ClearDependencies();
+			PendingTask->Function.Execute();
+			Private->PendingTasks.Pop();
+		}
+
+		// Don't do anything else until the pending work is completed.
+		return;
+	}
+
+	// It is safe to do this now.
+	Private->UpdateStreamingLimit();
+
+	// If we don't have an ongoing operation, don't do anything.
+	if (!Private->CurrentMutableOperation.IsValid())
+	{
+		return;
+	}
+
+	// If we reach here it means:
+	// - we have an ongoing operations
+	// - we have no pending work for the ongoing operation
+	// - so we are starting it.
+
+	switch (Private->CurrentMutableOperation->Type)
+	{
+	case FMutableOperation::EOperationType::Discard:
+	{
+		TRACE_CPUPROFILER_EVENT_SCOPE(Mutable_OperationDiscard);
+
+		// \TODO: Discards could be done in any case, concurrently with update operations. Should they be
+		// in their own "queue"?
+
+		UCustomizableObjectInstance* COI = Private->CurrentMutableOperation->CustomizableObjectInstance.Get();
+
+		// Only discard resources if the instance is still out range (it could have got closer to the player since the task was queued)
+		if (!CurrentInstanceLODManagement->IsOnlyUpdateCloseCustomizableObjectsEnabled()
+			|| !COI
+			|| COI->GetPrivate()->LastMinSquareDistFromComponentToPlayer > FMath::Square(CurrentInstanceLODManagement->GetOnlyUpdateCloseCustomizableObjectsDist()))
+		{
+			if (COI && COI->IsValidLowLevel())
+			{
+				COI->GetPrivate()->DiscardResourcesAndSetReferenceSkeletalMesh(COI);
+				COI->GetPrivate()->ClearCOInstanceFlags(Updating);
+				COI->SkeletalMeshStatus = ESkeletalMeshState::Correct;
+			}
+		}
+		else
+		{
+			COI->GetPrivate()->ClearCOInstanceFlags(Updating);
+		}
+
+		if (COI && !COI->HasAnySkeletalMesh())
+		{
+			// To solve the problem in the Mutable demo where PIE just after editor start made all instances appear as reference mesh until editor restart
+			COI->GetPrivate()->ClearCOInstanceFlags(Generated);
+		}
+
+		ClearCurrentMutableOperation();
+
+		break;
+	}
+
+	case FMutableOperation::EOperationType::Update:
+	{
+		TRACE_CPUPROFILER_EVENT_SCOPE(Mutable_OperationUpdate);
+
+		// Start the first task of the update process. See namespace impl comments above.
+		impl::Task_Game_StartUpdate(Private->CurrentMutableOperation);
+		break;
+	}
+
+	default:
+		check(false);
+		break;
+
+	}
+
+}
+
+
+bool UCustomizableObjectSystem::Tick(float DeltaTime)
+{
+	MUTABLE_CPUPROFILER_SCOPE(UCustomizableObjectSystem::Tick)
+	
+	// Building instances is not enable in servers. If at some point relevant collision or animation data is necessary for server logic this will need to be changed.
+#if UE_SERVER
+	return true;
+#endif
+
+	if (!Private.IsValid())
+	{
+		return true;
+	}
+
+	UWorld* World = GWorld;
+	if (World)
+	{
+		EWorldType::Type WorldType = World->WorldType;
+
+		if (WorldType != EWorldType::PIE && WorldType != EWorldType::Game && WorldType != EWorldType::Editor && WorldType != EWorldType::GamePreview)
+		{
+			return true;
+		}
+	}
+
+	// \TODO: Review: We should never compile an object from this tick, so this could be removed
+#if WITH_EDITOR
+	FAssetRegistryModule& AssetRegistryModule = FModuleManager::LoadModuleChecked<FAssetRegistryModule>("AssetRegistry");
+
+	if (AssetRegistryModule.Get().IsLoadingAssets())
+	{
+		return true; // Assets are still being loaded, so subobjects won't be found, compiled objects incomplete and thus updates wrong
+	}
+#endif
+
+
+	MUTABLE_CPUPROFILER_SCOPE(TickCustomizableObjectSystem)
+
+
+	// Get a new operation if we aren't working on one
+	if (!Private->CurrentMutableOperation)
+	{
+		// Update the streaming limit if it has changed. It is safe to do this now.
+		Private->UpdateStreamingLimit();
+
+		// Decide the next mutable operation to perform.
+		if (!Private->MutableOperationQueue.IsEmpty())
+		{
+			FMutableQueueElem AuxElem;
+			Private->MutableOperationQueue.Dequeue(&AuxElem);
+			Private->CurrentMutableOperation = AuxElem.Operation;
+		}
+		else
+		{
+			Private->CurrentMutableOperation = nullptr;
+		}
+	}
+
+	// Advance the current operation
+	if (Private->CurrentMutableOperation)
+	{
+		AdvanceCurrentOperation();
+	}
+
+
+	// Reset the instance relevancy
+	// \TODO: Review
+	// TODO: This should be done only when requiring a new job. And for the current operation instance, in case it needs to be cancelled.
+	CurrentInstanceLODManagement->UpdateInstanceDistsAndLODs();
+
+	for (TObjectIterator<UCustomizableObjectInstance> CustomizableObjectInstance; CustomizableObjectInstance; ++CustomizableObjectInstance)
+	{
+		if (IsValidChecked(*CustomizableObjectInstance) && CustomizableObjectInstance->GetPrivate())
+		{
+			if (CustomizableObjectInstance->GetPrivate()->HasCOInstanceFlags(UsedByComponentInPlay))
+			{
+				CustomizableObjectInstance->GetPrivate()->TickUpdateCloseCustomizableObjects(*CustomizableObjectInstance);
+			}
+			else if (CustomizableObjectInstance->GetPrivate()->HasCOInstanceFlags(UsedByComponent))
+			{
+				CustomizableObjectInstance->GetPrivate()->UpdateInstanceIfNotGenerated(*CustomizableObjectInstance, true);
+			}
+
+			CustomizableObjectInstance->GetPrivate()->ClearCOInstanceFlags((ECOInstanceFlags)(UsedByComponent | UsedByComponentInPlay));
+		}
+	}
+
+	// Update the queue. TODO: This should be done only when requiring a new job. And for the current operation instance, in case it needs to be cancelled.
+	{
+		Private->MutableOperationQueue.ChangePriorities();
+
+		for (TObjectIterator<UCustomizableObjectInstance> CustomizableObjectInstance; CustomizableObjectInstance; ++CustomizableObjectInstance)
+		{
+			if (IsValidChecked(*CustomizableObjectInstance) && CustomizableObjectInstance->GetPrivate())
+			{
+				CustomizableObjectInstance->GetPrivate()->LastMinSquareDistFromComponentToPlayer = CustomizableObjectInstance->GetPrivate()->MinSquareDistFromComponentToPlayer;
+				CustomizableObjectInstance->GetPrivate()->MinSquareDistFromComponentToPlayer = FLT_MAX;
+			}
+		}
+
+		Private->MutableOperationQueue.Sort();
+	}
+
+	// TODO: T2927
+	if (LogBenchmarkUtil::isLoggingActive())
+	{
+		uint64 SizeCache = 0;
+		for (const UTexture2D* cachedTextures : ProtectedCachedTextures)
+		{
+			if (cachedTextures)
+			{
+				SizeCache += cachedTextures->CalcTextureMemorySizeEnum(TMC_AllMips);
+			}
+		}
+		SET_DWORD_STAT(STAT_MutableTextureCacheMemory, SizeCache / 1024.f);
+	}
+
+	Private->UpdateStats();
+
+#if WITH_EDITOR
+	TickRecompileCustomizableObjects();
+#endif
+
+	return true;
+}
+
+
+TArray<FCustomizableObjectExternalTexture> UCustomizableObjectSystem::GetTextureParameterValues()
+{
+	TArray<FCustomizableObjectExternalTexture> Result;
+
+	for (const TWeakObjectPtr<UCustomizableSystemImageProvider> Provider : Private->ImageProvider->ImageProviders)
+	{
+		if (Provider.IsValid())
+		{
+			Provider->GetTextureParameterValues(Result);
+		}
+	}
+
+	return Result;
+}
+
+
+void UCustomizableObjectSystem::RegisterImageProvider(UCustomizableSystemImageProvider* Provider)
+{
+	Private->ImageProvider->ImageProviders.Add(Provider);
+}
+
+
+void UCustomizableObjectSystem::UnregisterImageProvider(UCustomizableSystemImageProvider* Provider)
+{
+	Private->ImageProvider->ImageProviders.Remove(Provider);
+}
+
+#if WITH_EDITOR
+UCustomizableObjectImageProviderArray* UCustomizableObjectSystem::GetEditorExternalImageProvider()
+{ 
+	return PreviewExternalImageProvider; 
+}
+#endif
+
+
+FMutableOperation FMutableOperation::CreateInstanceUpdate(UCustomizableObjectInstance* InCustomizableObjectInstance, bool bInNeverStream, int32 InMipsToSkip)
+{
+	FMutableOperation Op;
+	Op.Type = EOperationType::Update;
+	Op.bNeverStream = bInNeverStream;
+	Op.MipsToSkip = InMipsToSkip;
+	Op.CustomizableObjectInstance = InCustomizableObjectInstance;
+	Op.InstanceDescriptorHash = GetTypeHash(InCustomizableObjectInstance->GetDescriptor());
+	Op.bStarted = false;
+	Op.bBuildParameterDecorations = InCustomizableObjectInstance->bBuildParameterDecorations;
+	Op.Parameters = InCustomizableObjectInstance->GetPrivate()->ReloadParametersFromObject(InCustomizableObjectInstance, false);
+
+	InCustomizableObjectInstance->GetCustomizableObject()->ApplyStateForcedValuesToParameters(InCustomizableObjectInstance->GetState(), Op.Parameters.get());
+
+	if (!Op.Parameters)
+	{
+		// Cancel the update because the parameters aren't valid, probably because the object is not compiled
+		Op.CustomizableObjectInstance = nullptr;
+	}
+
+	return Op;
+}
+
+
+FMutableOperation FMutableOperation::CreateInstanceDiscard(UCustomizableObjectInstance* InCustomizableObjectInstance)
+{
+	check(InCustomizableObjectInstance);
+
+	FMutableOperation Op;
+	Op.Type = EOperationType::Discard;
+	Op.CustomizableObjectInstance = InCustomizableObjectInstance;
+
+	Op.CustomizableObjectInstance->GetPrivate()->SetCOInstanceFlags(Updating);
+
+	return Op;
+}
+
+
+void FMutableOperation::MutableIsDisabledCase()
+{
+	for (TObjectIterator<UCustomizableSkeletalComponent> It; It; ++It)
+	{
+		UCustomizableSkeletalComponent* CustomizableSkeletalComponent = *It;
+
+		if (CustomizableSkeletalComponent && CustomizableSkeletalComponent->CustomizableObjectInstance == CustomizableObjectInstance)
+		{
+			CustomizableSkeletalComponent->SetSkeletalMesh(CustomizableObjectInstance->GetCustomizableObject()->GetRefSkeletalMesh(CustomizableSkeletalComponent->ComponentIndex), false);
+		}
+	}
+
+	CustomizableObjectInstance->GetPrivate()->SetCOInstanceFlags(Generated);
+	CustomizableObjectInstance->GetPrivate()->ClearCOInstanceFlags(CreatingSkeletalMesh);
+
+	if (CustomizableObjectInstance->bEditorPropertyChanged)
+	{
+		CustomizableObjectInstance->bEditorPropertyChanged = false;
+	}
+
+	for (TObjectIterator<UCustomizableSkeletalComponent> It; It; ++It)
+	{
+		UCustomizableSkeletalComponent* CustomizableSkeletalComponent = *It;
+
+		if (CustomizableSkeletalComponent && CustomizableSkeletalComponent->CustomizableObjectInstance == CustomizableObjectInstance)
+		{
+			CustomizableSkeletalComponent->UpdatedDelegate.ExecuteIfBound();
+		}
+	}
+
+	CustomizableObjectInstance->Updated();
+
+#if WITH_EDITOR
+	CustomizableObjectInstance->InstanceUpdated = true;
+#endif
+}
+
+
+UCustomizableSystemImageProvider::ValueType UCustomizableObjectImageProviderArray::HasTextureParameterValue(int64 ID)
+{
+	if (ID >= FirstId && ID < FirstId + Textures.Num())
+	{
+		return UCustomizableSystemImageProvider::ValueType::Unreal;
+	}
+
+	return UCustomizableSystemImageProvider::ValueType::None;
+}
+
+
+UTexture2D* UCustomizableObjectImageProviderArray::GetTextureParameterValue(int64 ID)
+{
+	if (ID >= FirstId && ID < FirstId + Textures.Num())
+	{
+		return Textures[ID - FirstId];
+	}
+
+	return nullptr;
+}
+
+
+void UCustomizableObjectImageProviderArray::GetTextureParameterValues(TArray<FCustomizableObjectExternalTexture>& OutValues)
+{
+	for (int i = 0; i < Textures.Num(); ++i)
+	{
+		if (Textures[i])
+		{
+			FCustomizableObjectExternalTexture Data;
+			Data.Name = Textures[i]->GetName();
+			Data.Value = FirstId + i;
+			OutValues.Add(Data);
+		}
+	}
+}
+
+
+void UCustomizableObjectImageProviderArray::InvalidateIds()
+{
+	FirstId += Textures.Num();
+}
+
+
+#if WITH_EDITOR
+void UCustomizableObjectImageProviderArray::PostEditChangeProperty(struct FPropertyChangedEvent& PropertyChangedEvent)
+{
+	TexturesChangeDelegate.Broadcast();
+}
+#endif
+
+
+int32 UCustomizableObjectSystem::GetNumInstances() const
+{
+	return Private->NumInstances;
+}
+
+int32 UCustomizableObjectSystem::GetNumPendingInstances() const
+{
+	return Private->NumPendingInstances;
+}
+
+int32 UCustomizableObjectSystem::GetTotalInstances() const
+{
+	return Private->TotalInstances;
+}
+
+int32 UCustomizableObjectSystem::GetTextureMemoryUsed() const
+{
+	return int32(Private->TextureMemoryUsed);
+}
+
+int32 UCustomizableObjectSystem::GetAverageBuildTime() const
+{
+	return Private->TotalBuiltInstances == 0 ? 0 : Private->TotalBuildMs / Private->TotalBuiltInstances;
+}
+
+
+bool UCustomizableObjectSystem::IsCompactSerializationEnabled() const
+{
+	return Private->bCompactSerialization;
+}
+
+
+bool UCustomizableObjectSystem::IsSupport16BitBoneIndexEnabled() const
+{
+	return Private->bSupport16BitBoneIndex;
+}
+
+bool UCustomizableObjectSystem::IsProgressiveMipStreamingEnabled() const
+{
+	return Private->EnableMutableProgressiveMipStreaming != 0;
+}
+
+
+void UCustomizableObjectSystem::AddUncompiledCOWarning(UCustomizableObject* InObject)
+{
+	if (!InObject) return;
+
+	FString Msg;
+	Msg += FString::Printf(TEXT("Warning: Customizable Object [%s] not compiled. Please compile and save the object."), *InObject->GetName());
+	GEngine->AddOnScreenDebugMessage((uint64)((PTRINT)InObject), 10.0f, FColor::Red, Msg);
+
+#if WITH_EDITOR
+	if (UncompiledCustomizableObjectIds.Find(InObject->GetVersionId()) == INDEX_NONE)
+	{
+		UncompiledCustomizableObjectIds.Add(InObject->GetVersionId());
+
+		FMessageLog MessageLog("Mutable");
+		MessageLog.Warning(FText::FromString(Msg));
+
+		if (!UncompiledCustomizableObjectsNotificationPtr.IsValid())
+		{
+			FNotificationInfo Info(FText::FromString("Uncompiled Customizable Object/s found. Please, check the Message Log - Mutable for more information."));
+			Info.bFireAndForget = true;
+			Info.bUseThrobber = true;
+			Info.FadeOutDuration = 1.0f;
+			Info.ExpireDuration = 5.0f;
+
+			UncompiledCustomizableObjectsNotificationPtr = FSlateNotificationManager::Get().AddNotification(Info);
+		}
+	}
+#endif
+}
+
+void UCustomizableObjectSystem::EnableBenchmark()
+{
+	LogBenchmarkUtil::startLogging();
+}
+
+void UCustomizableObjectSystem::EndBenchmark()
+{
+	LogBenchmarkUtil::shutdownAndSaveResults();
+}
+
+
+void UCustomizableObjectSystem::SetReleaseMutableTexturesImmediately(bool bReleaseTextures)
+{
+	Private->bReleaseTexturesImmediately = bReleaseTextures;
+}
+
+
+#if WITH_EDITOR
+
+void UCustomizableObjectSystem::OnPreBeginPIE(const bool bIsSimulatingInEditor)
+{
+	if(!EditorSettings.bCompileRootObjectsOnStartPIE || IsRunningGame() || IsCompilationDisabled())
+	{
+		return;
+	}
+	
+	// Find root customizable objects
+	FARFilter AssetRegistryFilter;
+	AssetRegistryFilter.ClassPaths.Add(FTopLevelAssetPath(TEXT("/Script/CustomizableObject"), TEXT("CustomizableObject")));
+	AssetRegistryFilter.TagsAndValues.Add(FName("IsRoot"), FString::FromInt(1));
+
+	TArray<FAssetData> OutAssets;
+	const FAssetRegistryModule& AssetRegistryModule = FModuleManager::LoadModuleChecked<FAssetRegistryModule>("AssetRegistry");
+	AssetRegistryModule.Get().GetAssets(AssetRegistryFilter, OutAssets);
+
+	TArray<FAssetData> TempObjectsToRecompile;
+	for (const FAssetData& Asset : OutAssets)
+	{
+		// If it is referenced by PIE it should be loaded
+		if (!Asset.IsAssetLoaded())
+		{
+			continue;
+		}
+		
+		const UCustomizableObject* Object = Cast<UCustomizableObject>(Asset.GetAsset());
+		if(!Object || Object->IsCompiled() || Object->IsLocked() || Object->bIsChildObject)
+		{
+			continue;
+		}
+		
+		// Add uncompiled objects to the objects to cook list
+		TempObjectsToRecompile.Add(Asset);
+	}
+
+	if (!TempObjectsToRecompile.IsEmpty())
+	{
+		FText Msg = FText::FromString(TEXT("Warning: one or more Customizable Objects used in PIE are uncompiled.\n\nDo you want to compile them?"));
+		if (FMessageDialog::Open(EAppMsgType::OkCancel, Msg) == EAppReturnType::Ok)
+		{
+			ObjectsToRecompile.Empty(TempObjectsToRecompile.Num());
+			RecompileCustomizableObjects(TempObjectsToRecompile);
+		}
+	}
+}
+
+void UCustomizableObjectSystem::StartNextRecompile()
+{
+	if (ObjectsToRecompile.Num() % 10 == 0)
+	{
+		CollectGarbage(GARBAGE_COLLECTION_KEEPFLAGS);
+	}
+
+	FAssetData Itr = ObjectsToRecompile.Pop();
+	UCustomizableObject* CustomizableObject = Cast<UCustomizableObject>(Itr.GetAsset());
+
+	if (CustomizableObject)
+	{
+		FText UpdateMsg = FText::FromString(FString::Printf(TEXT("Compiling Customizable Objects:\n%s"), *CustomizableObject->GetName()));
+		FSlateNotificationManager::Get().UpdateProgressNotification(RecompileNotificationHandle, NumObjectsCompiled, TotalNumObjectsToRecompile, UpdateMsg);
+
+		// Use default options
+		FCompilationOptions Options = CustomizableObject->CompileOptions;
+		Options.bSilentCompilation = true;
+		RecompileCustomizableObjectsCompiler->Compile(*CustomizableObject, Options, true);
+	}
+}
+
+void UCustomizableObjectSystem::RecompileCustomizableObjectAsync(const FAssetData& InAssetData,
+	const UCustomizableObject* InObject)
+{
+	if(IsRunningGame() || IsCompilationDisabled())
+	{
+		return;
+	}
+	
+	if((InObject && InObject->IsLocked()) || ObjectsToRecompile.Find((InAssetData)) != INDEX_NONE)
+	{
+		return;
+	}
+	
+	if(!ObjectsToRecompile.IsEmpty())
+	{
+		ObjectsToRecompile.Add(InAssetData);
+	}
+	else
+	{
+		RecompileCustomizableObjects({InAssetData});
+	}
+}
+
+void UCustomizableObjectSystem::RecompileCustomizableObjects(const TArray<FAssetData>& InObjects)
+{
+	if (IsRunningGame() || IsCompilationDisabled())
+	{
+		return;
+	}
+
+	if (InObjects.Num())
+	{
+		if (!RecompileCustomizableObjectsCompiler)
+		{
+			RecompileCustomizableObjectsCompiler = GetNewCompiler();
+
+			if (!RecompileCustomizableObjectsCompiler)
+			{
+				return;
+			}
+		}
+
+		ObjectsToRecompile.Append(InObjects);
+
+		TotalNumObjectsToRecompile = ObjectsToRecompile.Num();
+		NumObjectsCompiled = 0;
+
+		if(RecompileNotificationHandle.IsValid())
+		{
+			++TotalNumObjectsToRecompile;
+			FSlateNotificationManager::Get().UpdateProgressNotification(RecompileNotificationHandle, NumObjectsCompiled, TotalNumObjectsToRecompile);
+		}
+		else
+		{
+			RecompileNotificationHandle = FSlateNotificationManager::Get().StartProgressNotification(FText::FromString(TEXT("Compiling Customizable Objects")), TotalNumObjectsToRecompile);
+			StartNextRecompile();
+		}
+	}
+}
+
+
+void UCustomizableObjectSystem::TickRecompileCustomizableObjects()
+{
+	bool bUpdated = false;
+	
+	if (RecompileCustomizableObjectsCompiler)
+	{
+		bUpdated = RecompileCustomizableObjectsCompiler->Tick() || RecompileCustomizableObjectsCompiler->GetCompilationState() == ECustomizableObjectCompilationState::Failed;
+	}
+
+	if (bUpdated)
+	{
+		NumObjectsCompiled++;
+
+		if (!ObjectsToRecompile.IsEmpty())
+		{
+			StartNextRecompile();
+		}
+		else // All objects compiled, clean up
+		{
+			// Delete compiler
+			delete RecompileCustomizableObjectsCompiler;
+			RecompileCustomizableObjectsCompiler = nullptr;
+
+			// Remove progress bar
+			FSlateNotificationManager::Get().UpdateProgressNotification(RecompileNotificationHandle, NumObjectsCompiled, TotalNumObjectsToRecompile);
+			FSlateNotificationManager::Get().CancelProgressNotification(RecompileNotificationHandle);
+			RecompileNotificationHandle.Reset();
+
+			CollectGarbage(GARBAGE_COLLECTION_KEEPFLAGS);
+		}
+	}
+}
+
+
+uint64 UCustomizableObjectSystem::GetMaxChunkSizeForPlatform(const ITargetPlatform* TargetPlatform)
+{
+	const FString& PlatformName = TargetPlatform ? TargetPlatform->IniPlatformName() : FPlatformProperties::IniPlatformName();
+
+	if(const int64* CachedMaxChunkSize = PlatformMaxChunkSize.Find(PlatformName))
+	{
+		return *CachedMaxChunkSize;
+	}
+
+	int64 MaxChunkSize = -1;
+
+	if (!FParse::Value(FCommandLine::Get(), TEXT("ExtraFlavorChunkSize="), MaxChunkSize) || MaxChunkSize < 0)
+	{
+		FConfigFile PlatformIniFile;
+		FConfigCacheIni::LoadLocalIniFile(PlatformIniFile, TEXT("Game"), true, *PlatformName);
+		FString ConfigString;
+		if (PlatformIniFile.GetString(TEXT("/Script/UnrealEd.ProjectPackagingSettings"), TEXT("MaxChunkSize"), ConfigString))
+		{
+			MaxChunkSize = FCString::Atoi64(*ConfigString);
+		}
+	}
+
+	// If no limit is specified default it to MUTABLE_STREAMED_DATA_MAXCHUNKSIZE 
+	if(MaxChunkSize <= 0)
+	{
+		MaxChunkSize = MUTABLE_STREAMED_DATA_MAXCHUNKSIZE;
+	}
+
+	PlatformMaxChunkSize.Add(PlatformName, MaxChunkSize);
+
+	return MaxChunkSize;
+}
+
+#endif // WITH_EDITOR
