[mutable] Move Mutable plugin from the restricted folder to the experimental plugins folder.

diff --git a/Engine/Plugins/Experimental/Mutable/Source/CustomizableObject/Private/CustomizableObjectSystemPrivate.h b/Engine/Plugins/Experimental/Mutable/Source/CustomizableObject/Private/CustomizableObjectSystemPrivate.h
new file mode 100644
--- /dev/null
+++ b/Engine/Plugins/Experimental/Mutable/Source/CustomizableObject/Private/CustomizableObjectSystemPrivate.h
@@ -0,0 +1,757 @@
+// Copyright Epic Games, Inc. All Rights Reserved.
+
+#pragma once
+
+#include "CustomizableObjectSystem.h"
+#include "CustomizableObjectInstance.h"
+#include "CustomizableObject.h"
+#include "HAL/ThreadSafeBool.h"
+#include "Containers/Ticker.h"
+
+#include "MutableRuntime/Public/System.h"
+#include "MutableRuntime/Public/Image.h"
+
+// This define could come from MutableRuntime/Public/System.h
+#ifdef MUTABLE_USE_NEW_TASKGRAPH
+	#include "Tasks/Task.h"
+#else
+	#include "Async/TaskGraphInterfaces.h"
+#endif
+
+
+//! An operation to be performed by Mutable. This could be creating or updating an instance, releasing resources, changing an LOD, etc.
+//! Operations may be done in several tasks in several thread or across frames.
+class FMutableOperation 
+{
+	/** Instance parameters at the time of the operation request. */
+	mu::ParametersPtr Parameters; 
+
+	bool bBuildParameterDecorations = false;
+	bool bMeshNeedsUpdate = false;
+
+	/** Protected constructor. */
+	FMutableOperation() {}
+
+public:
+
+
+	static FMutableOperation CreateInstanceUpdate(UCustomizableObjectInstance* COInstance, bool bInNeverStream, int32 MipsToSkip);
+	static FMutableOperation CreateInstanceDiscard(UCustomizableObjectInstance* COInstance);
+
+
+	enum class EOperationType
+	{
+		// Create mutable resources for a new instance or to update an exiting one.
+		Update,
+
+		// Discard the resources of an instance.
+		Discard
+	};
+
+	// Type of the operation
+	EOperationType Type;
+
+	bool bStarted = false;
+
+	bool bNeverStream = false;
+
+	/** If this is non-zero, the instance images will be generated with a certain amount of mips starting from the smallest. Otherwise, the full images will be generated. 
+	 * This is used for both Update and UpdateImage operations. 
+	 */
+	int32 MipsToSkip = 0;
+
+	// Weak reference to the instance we are operating on.
+	// It is weak because we don't want to lock it in case it becomes irrelevant in the game while operations are pending and it needs to be destroyed.
+	TWeakObjectPtr<UCustomizableObjectInstance> CustomizableObjectInstance;
+
+	/** Hash of the UCustomizableObjectInstance::Descriptor (state of the instance parameters + state). */
+	uint32 InstanceDescriptorHash;
+
+	//! This is used to calculate stats.
+	double StartUpdateTime = 0.0;
+
+	//!
+	bool IsBuildParameterDecorations() const
+	{
+		return bBuildParameterDecorations;
+	}
+
+	// \TODO: This should be intercepted earlier.
+	// In case Mutable Compilation has been disabled, do all the corresponding steps to assign the reference mesh of the customizable object
+	void MutableIsDisabledCase();
+
+	/** Read-only access to the mutable instance parameters for this operation. */
+	mu::ParametersPtrConst GetParameters() const
+	{
+		return Parameters;
+	}
+};
+
+
+
+struct FMutableQueueElem
+{
+	// Priority for mutable update queue, Low is the normal distance-based priority, High is normally used for discards and Mid for LOD downgrades
+	enum EQueuePriorityType { High, Med, Med_Low, Low };
+
+	EQueuePriorityType PriorityType;
+	float Priority;
+	TSharedPtr<FMutableOperation> Operation;
+	bool bIsDiscardResources;
+
+	static FMutableQueueElem Create(TSharedPtr<FMutableOperation> InOperation, EQueuePriorityType NewPriorityType, double NewPriority)
+	{
+		FMutableQueueElem MutableQueueElem;
+
+		MutableQueueElem.PriorityType = NewPriorityType;
+		MutableQueueElem.Priority = NewPriority;
+		MutableQueueElem.Operation = InOperation;
+		MutableQueueElem.bIsDiscardResources = (InOperation->Type==FMutableOperation::EOperationType::Discard);
+
+		return MutableQueueElem;
+	}
+
+	friend bool operator <(const FMutableQueueElem& A, const FMutableQueueElem& B)
+	{
+		if (A.PriorityType < B.PriorityType)
+		{
+			return true;
+		}
+		else if (A.PriorityType > B.PriorityType)
+		{
+			return false;
+		}
+		else
+		{
+			if (A.bIsDiscardResources && !B.bIsDiscardResources)
+			{
+				return true;
+			}
+			else if (!A.bIsDiscardResources && B.bIsDiscardResources)
+			{
+				return false;
+			}
+
+			return A.Priority < B.Priority;
+		}
+	}
+};
+
+
+/** Instances operations queue.
+ *
+ * The queue will only contain a single operation per UCustomizableObjectInstance.
+ * If there is already an operation it will be replaced. */
+class FMutableQueue
+{
+	TArray<FMutableQueueElem> Array;
+
+public:
+	bool IsEmpty() const;
+
+	int Num() const;
+
+	void Enqueue(const FMutableQueueElem& TaskToEnqueue);
+
+	void Dequeue(FMutableQueueElem* DequeuedTask);
+
+	void ChangePriorities();
+
+	void UpdatePriority(const UCustomizableObjectInstance* Instance);
+
+	const FMutableQueueElem* Get(const UCustomizableObjectInstance* Instance) const;
+	
+	void Sort();
+};
+
+
+struct FMutableImageCacheKey
+{
+	mu::RESOURCE_ID Resource = 0;
+	//int32 SkippedMips = 0; // Todo: Reenable extended FMutableImageCacheKey
+
+	inline bool operator==(const FMutableImageCacheKey& o) const
+	{
+		return Resource == o.Resource;//&& SkippedMips == o.SkippedMips; // Todo: Reenable extended FMutableImageCacheKey
+	}
+};
+
+
+inline uint32 GetTypeHash(const FMutableImageCacheKey& Key)
+{
+	//return HashCombine(GetTypeHash(Key.Resource), GetTypeHash(Key.SkippedMips)); // Todo: Reenable extended FMutableImageCacheKey
+	return GetTypeHash(Key.Resource);
+}
+
+
+// Cache of weak references to generated resources for one single model
+struct FMutableResourceCache
+{
+	TWeakObjectPtr<const UCustomizableObject> Object;
+	TMap<mu::RESOURCE_ID, TWeakObjectPtr<USkeletalMesh> > Meshes;
+	TMap<FMutableImageCacheKey, TWeakObjectPtr<UTexture2D> > Images;
+
+	void Clear()
+	{
+		Meshes.Reset();
+		Images.Reset();
+	}
+};
+
+
+DECLARE_DELEGATE(FMutableTaskDelegate);
+struct FMutableTask
+{
+	/** Actual function to perform in this task. */
+	FMutableTaskDelegate Function;
+
+	/** We can have 2 types of depencies:
+		From the new task system: */
+#ifdef MUTABLE_USE_NEW_TASKGRAPH
+	UE::Tasks::FTask Dependency;
+#endif
+
+	/** From the traditional event graph system, still used to wait for async loads. */
+	FGraphEventRef GraphDependency0;
+	FGraphEventRef GraphDependency1;
+
+	/** Check if all the dependencies of this task have been completed. */
+	inline bool AreDependenciesComplete() const
+	{
+#ifdef MUTABLE_USE_NEW_TASKGRAPH
+		return (!Dependency.IsValid() || Dependency.IsCompleted())
+			&&
+			(!GraphDependency0 || GraphDependency0->IsComplete())
+			&&
+			(!GraphDependency1 || GraphDependency1->IsComplete());
+#else
+		return 
+			(!GraphDependency0 || GraphDependency0->IsComplete())
+			&&
+			(!GraphDependency1 || GraphDependency1->IsComplete());
+#endif
+	}
+
+	/** Free the handles for any dependency of this task. */
+	inline void ClearDependencies()
+	{
+#ifdef MUTABLE_USE_NEW_TASKGRAPH
+		Dependency = {};
+#endif
+		GraphDependency0 = nullptr;
+		GraphDependency1 = nullptr;
+	}
+};
+
+
+// Mutable data generated during the update steps.
+// We keep it from begin to end update, and it is used in several steps.
+// TODO: Flatten this structure into 4 arrays and use indices in fixed-size structs instead of subarrays
+struct FInstanceUpdateData
+{
+	struct FImage
+	{
+		FString Name;
+		mu::RESOURCE_ID ImageID;
+		uint16 FullImageSizeX, FullImageSizeY;
+		mu::ImagePtrConst Image;
+		TWeakObjectPtr<UTexture2D> Cached;
+	};
+
+	struct FVector
+	{
+		FString Name;
+		FLinearColor Vector;
+	};
+
+	struct FScalar
+	{
+		FString Name;
+		float Scalar;
+	};
+
+	struct FSurface
+	{
+		/** Range in the Images array */
+		uint16 FirstImage = 0;
+		uint16 ImageCount = 0;
+
+		/** Range in the Vectors array */
+		uint16 FirstVector = 0;
+		uint16 VectorCount = 0;
+
+		/** Range in the Scalar array */
+		uint16 FirstScalar = 0;
+		uint16 ScalarCount = 0;
+
+		uint32 MaterialIndex = 0;
+	};
+
+	struct FComponent
+	{
+		uint16 Id = 0;
+		
+		mu::MeshPtrConst Mesh;
+
+		/** Range in the Surfaces array */
+		uint16 FirstSurface = 0;
+		uint16 SurfaceCount = 0;
+
+		// \TODO: Flatten
+		TArray<uint16> ActiveBones;
+		/** Range in the external Bones array */
+		//uint32 FirstActiveBone;
+		//uint32 ActiveBoneCount;
+
+		// \TODO: Flatten
+		TArray<uint16> BoneMap;
+		/** Range in the external Bones array */
+		//uint32 FirstBoneMap;
+		//uint32 BoneMapCount;
+	};
+
+	struct FLOD
+	{
+		/** Range in the Components array */
+		uint16 FirstComponent = 0;
+		uint16 ComponentCount = 0;
+	};
+
+	TArray<FLOD> LODs;
+	TArray<FComponent> Components;
+	TArray<FSurface> Surfaces;
+	TArray<FImage> Images;
+	TArray<FVector> Vectors;
+	TArray<FScalar> Scalars;
+
+	struct FSkeletonData
+	{
+		int16 ComponentIndex = INDEX_NONE;
+
+		TArray<uint32> SkeletonIds;
+
+		TArray<FName> BoneNames;
+		TMap<FName, FMatrix44f> BoneMatricesWithScale;
+	};
+
+	TArray<FSkeletonData> Skeletons;
+
+	/** */
+	void Clear()
+	{
+		LODs.Empty();
+		Components.Empty();
+		Surfaces.Empty();
+		Images.Empty();
+		Scalars.Empty();
+		Vectors.Empty();
+		Skeletons.Empty();
+	}
+};
+
+
+// Mutable data generated during the update steps.
+struct FParameterDecorationsUpdateData
+{
+	struct FParameterDesc
+	{
+		TArray<mu::ImagePtrConst> Images;
+	};
+
+	//! Information about additional images used for parameter UI decoration
+	TArray<FParameterDesc> Parameters;
+
+	void Clear()
+	{
+		Parameters.Empty();
+	}
+};
+
+
+struct FTextureCoverageQueryData
+{
+	uint32 CoveredTexels;
+	uint32 MaskedOutCoveredTexels;
+	uint32 TotalTexels;
+	FString MaskOutChannelName;
+
+	FTextureCoverageQueryData() : CoveredTexels(0), MaskedOutCoveredTexels(0), TotalTexels(0) { }
+
+	float GetCoverage() const { return TotalTexels > 0 ? (float(CoveredTexels) / TotalTexels) : 0.f; }
+	float GetMaskedOutCoverage() const { return CoveredTexels > 0 ? (float(MaskedOutCoveredTexels) / CoveredTexels) : 0.f; }
+};
+
+
+struct FPendingTextureCoverageQuery
+{
+	FString KeyName;
+	uint32 MaterialIndex = 0;
+	FTexturePlatformData* PlatformData = nullptr;
+};
+
+
+/** Runtime data used during a mutable instance update */
+struct FMutableOperationData
+{
+	FInstanceUpdateData InstanceUpdateData;
+	FParameterDecorationsUpdateData ParametersUpdateData;
+	TArray<int> RelevantParametersInProgress;
+
+	/** This option comes from the operation request */
+	bool bNeverStream = false;
+	/** This option comes from the operation request. It is used to reduce the number of mipmaps that mutable must generate for images.  */
+	int32 MipsToSkip = 0;
+
+	mu::Instance::ID InstanceID = 0;
+	int32 CurrentMinLOD = 0;
+	int32 CurrentMaxLOD = 0;
+	int32 NumLODsAvailable = 0;
+
+	TMap<FString, FTextureCoverageQueryData> TextureCoverageQueries_MutableThreadParams;
+	TMap<FString, FTextureCoverageQueryData> TextureCoverageQueries_MutableThreadResults;
+
+	TMap<uint32, FTexturePlatformData*> ImageToPlatformDataMap;
+
+	/** This list of queries is generated in the update mutable task, and consumed later in the game thread. */
+	TArray<FPendingTextureCoverageQuery> PendingTextureCoverageQueries;
+
+	/** */
+	mu::Ptr<const mu::Parameters> MutableParameters;
+	int32 State = 0;
+
+#if WITH_EDITOR
+	/** Used for profiling in the editor. */
+	uint32 MutableRuntimeCycles = 0;
+#endif 
+};
+
+
+/** Runtime data used during a mutable instance update */
+struct FMutableReleasePlatformOperationData
+{
+	TMap<uint32, FTexturePlatformData*> ImageToPlatformDataMap;
+};
+
+
+class FCustomizableObjectSystemPrivate : public FGCObject
+{
+public:
+
+	// Singleton for the unreal mutable system.
+	static UCustomizableObjectSystem* SSystem;
+
+	// Pointer to the lower level mutable system that actually does the work.
+	mu::SystemPtr MutableSystem;
+
+	// Store the last streaming memory size in bytes, to change it when it is safe.
+	uint64_t LastStreamingMemorySize = 0;
+
+	// This object is responsible for streaming data to the MutableSystem.
+	// Non-owned reference
+	class FUnrealMutableModelBulkStreamer* Streamer = nullptr;
+
+	// This object is responsible for providing custom images to mutable models (for image parameters)
+	// This object is called from the mutable thread, and it should only access data already build in
+	// the game thread and stored in FUnrealMutableImageProvider::ExternalImagesForCurrentInstance.
+	// Non-owned reference
+	class FUnrealMutableImageProvider* ImageProvider = nullptr;
+
+	// Cache of weak references to generated resources to see if they can be reused.
+	TArray<FMutableResourceCache> ModelResourcesCache;
+
+	// List of textures currently cached and valid for the current object that we are operating on.
+	// This array gets generated when the object cached resources are protected in SetResourceCacheProtected
+	// from the game thread, and it is read from the Mutable thread only while updating the instance.
+	TArray<mu::RESOURCE_ID> ProtectedObjectCachedImages;
+
+	FMutableQueue MutableOperationQueue;
+
+	// Queue of game-thread tasks that need to be executed for the current operation
+	TQueue<FMutableTask> PendingTasks;
+
+	static int32 EnableMutableProgressiveMipStreaming;
+
+	/** */
+	inline void AddGameThreadTask(const FMutableTask& Task)
+	{
+		PendingTasks.Enqueue(Task);
+	}
+
+	/** */
+#ifdef MUTABLE_USE_NEW_TASKGRAPH
+	template<typename TaskBodyType>
+	inline UE::Tasks::FTask AddMutableThreadTask(const TCHAR* DebugName, TaskBodyType&& TaskBody, UE::Tasks::ETaskPriority Priority)
+	{
+		// This could be called from the game thread or from a worker thread for the texture streaming system
+		//check(IsInGameThread());
+		FScopeLock Lock(&MutableTaskLock);
+
+		if (LastMutableTask.IsValid())
+		{
+			LastMutableTask = UE::Tasks::Launch(DebugName, MoveTemp(TaskBody), LastMutableTask, Priority);
+		}
+		else
+		{
+			LastMutableTask = UE::Tasks::Launch(DebugName, MoveTemp(TaskBody), Priority);
+		}
+		return LastMutableTask;
+	}
+#else
+	template<typename TaskBodyType>
+	inline FGraphEventRef AddMutableThreadTask(const TCHAR* DebugName, TaskBodyType&& TaskBody, UE::Tasks::ETaskPriority Priority)
+	{
+		// This could be called from the game thread or from a worker thread for the texture streaming system
+		//check(IsInGameThread());
+		FScopeLock Lock(&MutableTaskLock);
+
+		// Chain to pending tasks of the mutable thread if any. This ensures exclusive access to mutable.
+		FGraphEventArray Prerequisites;
+		if (LastMutableTask)
+		{
+			Prerequisites.Add(LastMutableTask);
+		}
+
+		LastMutableTask = FFunctionGraphTask::CreateAndDispatchWhenReady(
+			MoveTemp(TaskBody),
+			TStatId{},
+			&Prerequisites, 
+			ENamedThreads::AnyThread);
+		LastMutableTask->SetDebugName(DebugName);
+
+		return LastMutableTask;
+	}
+
+	/** This should only be used when shutting down. */
+	inline void WaitForMutableTasks()
+	{
+		FScopeLock Lock(&MutableTaskLock);
+
+		if (LastMutableTask.IsValid())
+		{
+#ifdef MUTABLE_USE_NEW_TASKGRAPH
+			LastMutableTask.Wait();
+#else
+			LastMutableTask->Wait();
+#endif
+			LastMutableTask = {};
+		}
+	}
+
+	inline void ClearMutableTaskIfDone()
+	{
+		FScopeLock Lock(&MutableTaskLock);
+
+#ifdef MUTABLE_USE_NEW_TASKGRAPH
+		if (LastMutableTask.IsValid() && LastMutableTask.IsCompleted())
+		{
+			LastMutableTask = {};
+		}
+#else
+		if (LastMutableTask.IsValid() && LastMutableTask->IsComplete())
+		{
+			LastMutableTask = nullptr;
+		}
+#endif
+	}
+
+
+	template<typename TaskBodyType>
+	inline void AddAnyThreadTask(const TCHAR* DebugName, TaskBodyType&& TaskBody, UE::Tasks::ETaskPriority Priority)
+	{
+		check(IsInGameThread());
+
+		FGraphEventRef Task = FFunctionGraphTask::CreateAndDispatchWhenReady(
+			MoveTemp(TaskBody),
+			TStatId{},
+			nullptr,
+			ENamedThreads::AnyThread);
+		Task->SetDebugName(TEXT("Mutable_Anythread"));
+	}
+#endif
+
+	//
+	//TSharedPtr<FMutableOperationData> CurrentOperationData;
+
+
+	/** FSerializableObject interface */
+	virtual void AddReferencedObjects(FReferenceCollector& Collector) override;
+	virtual FString GetReferencerName() const override
+	{
+		return TEXT("FCustomizableObjectSystemPrivate");
+	}
+
+
+	// Remove references to cached objects that have been deleted in the unreal
+	// side, and cannot be cached anyway.
+	// This should only happen in the mutable thread
+	void CleanupCache()
+	{
+		for (int ModelIndex=0; ModelIndex<ModelResourcesCache.Num();)
+		{
+			if (!ModelResourcesCache[ModelIndex].Object.IsValid(false,true))
+			{
+				// The whole object has been destroyed. Remove everything.
+				ModelResourcesCache.RemoveAtSwap(ModelIndex);
+			}
+			else
+			{
+				// \todo: Free invalid references to specific resources
+				++ModelIndex;
+			}
+		}
+	}
+
+
+	// This should only happen in the game thread
+	FMutableResourceCache& GetObjectCache(const UCustomizableObject* Object)
+	{
+		check(IsInGameThread());
+
+		// Not mandatory, but a good place for a cleanup
+		CleanupCache();
+
+		for (int ModelIndex = 0; ModelIndex < ModelResourcesCache.Num(); ++ModelIndex)
+		{
+			if (ModelResourcesCache[ModelIndex].Object==Object)
+			{
+				return ModelResourcesCache[ModelIndex];
+			}
+		}
+		
+		// Not found, create and add it.
+		ModelResourcesCache.Push(FMutableResourceCache());
+		ModelResourcesCache.Last().Object = Object;
+		return ModelResourcesCache.Last();
+	}
+
+
+	void AddTextureReference(uint32 TextureId)
+	{
+		uint32& CountRef = TextureReferenceCount.FindOrAdd(TextureId);
+
+		CountRef++;
+	}
+
+	
+	// Returns true if the texture's references become zero
+	bool RemoveTextureReference(uint32 TextureId)
+	{
+		uint32* CountPtr = TextureReferenceCount.Find(TextureId);
+
+		if (CountPtr && *CountPtr > 0)
+		{
+			(*CountPtr)--;
+
+			if (*CountPtr == 0)
+			{
+				TextureReferenceCount.Remove(TextureId);
+
+				return true;
+			}
+		}
+		else
+		{
+			ensure(false); // Mutable texture reference count is incorrect
+			TextureReferenceCount.Remove(TextureId);
+		}
+
+		return false;
+	}
+
+
+	bool TextureHasReferences(uint32 TextureId) const
+	{
+		const uint32* CountPtr = TextureReferenceCount.Find(TextureId);
+
+		if (CountPtr && *CountPtr > 0)
+		{
+			return true;
+		}
+
+		return false;
+	}
+
+
+	// Init the async Skeletal Mesh creation/update
+	void InitUpdateSkeletalMesh(UCustomizableObjectInstance* InCustomizableObjectInstance, FMutableQueueElem::EQueuePriorityType Priority);
+		
+	// Init an async and safe release of the UE4 and Mutable resources used by the instance without actually destroying the instance, for example if it's very far away
+	void InitDiscardResourcesSkeletalMesh(UCustomizableObjectInstance* InCustomizableObjectInstance);
+	
+	bool IsReplaceDiscardedWithReferenceMeshEnabled() const { return bReplaceDiscardedWithReferenceMesh; }
+	void SetReplaceDiscardedWithReferenceMeshEnabled(bool bIsEnabled) { bReplaceDiscardedWithReferenceMesh = bIsEnabled; }
+
+	int32 GetCountAllocatedSkeletalMesh() { return CountAllocatedSkeletalMesh; }
+
+	bool bCompactSerialization = true;
+
+	bool bReplaceDiscardedWithReferenceMesh = false;
+	bool bReleaseTexturesImmediately = false;
+
+	bool bSupport16BitBoneIndex = false;
+
+	// Statistics: Total time spent building instances
+	int32 TotalBuildMs = 0;
+
+	// Statistics: Total number of instances built or updated.
+	int32 TotalBuiltInstances = 0;
+
+	// Statistics: Number of instances alive (built)
+	int32 NumInstances = 0;
+
+	// Statistics: number of CustomizableObjectInstances waiting to be updated.
+	int32 NumPendingInstances = 0;
+
+	// Statistics: Total number of CustomizableObjectInstances, including not built.
+	int32 TotalInstances = 0;
+
+	// Statistics: total memory in bytes used for generated textures
+	int64_t TextureMemoryUsed = 0;
+
+	mutable uint32 CountAllocatedSkeletalMesh = 0;
+
+	static FCustomizableObjectCompilerBase* (*NewCompilerFunc)();
+
+	void CreatedTexture(UTexture2D* Texture);
+
+	// \TODO: Remove this array if we are not gathering stats!
+	TArray<TWeakObjectPtr<class UTexture2D>> TextureTrackerArray;
+	TMap<uint32, uint32> TextureReferenceCount; // Keeps a count of texture usage to decide if they have to be blocked from GC during an update
+
+	// This is protected from GC by AddReferencedObjects
+	UCustomizableObjectInstance* CurrentInstanceBeingUpdated = nullptr;
+
+	TSharedPtr<FMutableOperation> CurrentMutableOperation = nullptr;
+
+	// Handle to the registered TickDelegate.
+	FTSTicker::FDelegateHandle TickDelegateHandle;
+	FTickerDelegate TickDelegate;
+
+	//! Update the stats logged in unreal's stats system. 
+	void UpdateStats();
+
+	// Important!!! Never call when there's a Begin Update thread running!
+	void ReleasePendingMutableInstances();
+
+	// Check and update the streaming memory limit. Only safe from game thread and when the mutable thread is idle.
+	void UpdateStreamingLimit();
+
+private:
+
+	/** Access to this must be protected with this. */
+	mutable FCriticalSection MutableTaskLock;
+
+#ifdef MUTABLE_USE_NEW_TASKGRAPH
+	UE::Tasks::FTask LastMutableTask;
+#else
+	FGraphEventRef LastMutableTask = nullptr;
+#endif
+
+
+};
+
+
+void ConvertImage(UTexture2D* Texture, mu::ImagePtrConst MutableImage, const struct FMutableModelImageProperties& Props);
+
+
+
diff --git a/Engine/Plugins/Experimental/Mutable/Source/CustomizableObject/Private/CustomizableSkeletalComponent.cpp b/Engine/Plugins/Experimental/Mutable/Source/CustomizableObject/Private/CustomizableSkeletalComponent.cpp
new file mode 100644
--- /dev/null
+++ b/Engine/Plugins/Experimental/Mutable/Source/CustomizableObject/Private/CustomizableSkeletalComponent.cpp
@@ -0,0 +1,349 @@
+// Copyright Epic Games, Inc. All Rights Reserved.
+
+#include "CustomizableSkeletalComponent.h"
+
+#include "CustomizableObject.h"
+#include "CustomizableObjectInstance.h"
+#include "CustomizableInstancePrivateData.h"
+#include "CustomizableObjectSystem.h"
+
+#include "Engine/Texture.h"
+#include "Engine/SkeletalMesh.h"
+#include "Engine/StaticMesh.h"
+#include "Components/SkeletalMeshComponent.h"
+#include "Materials/MaterialInstanceDynamic.h"
+
+#include "AnimationRuntime.h"
+
+#include "UObject/ObjectSaveContext.h"
+
+
+UCustomizableSkeletalComponent::UCustomizableSkeletalComponent()
+{
+	PrimaryComponentTick.bCanEverTick = true;
+	//bTickInEditor = true;
+}
+
+
+USkeletalMesh* UCustomizableSkeletalComponent::GetSkeletalMesh() const
+{
+	return CustomizableObjectInstance ? CustomizableObjectInstance->GetSkeletalMesh(ComponentIndex) : nullptr;
+}
+
+void UCustomizableSkeletalComponent::SetSkeletalMesh(USkeletalMesh* SkeletalMesh, bool bReinitPose, bool bForceClothReset)
+{
+	USkeletalMeshComponent* Parent = Cast<USkeletalMeshComponent>(GetAttachParent());
+
+	// LINUX_PLATFORM needs a more aggressive workaround so morph and cloth glitches
+	// are not visible. 
+	// TODO: Try to find a better way of setting BP Morphs and clothing and remove
+	// the workaround.  
+	if (Parent)
+	{
+#if PLATFORM_LINUX
+		const bool bDisableClothSimulation = Parent->bDisableClothSimulation;
+		if (bForceClothReset)
+		{
+			Parent->bDisableClothSimulation = true;
+		}
+#endif
+
+		if(Parent->SkeletalMesh == SkeletalMesh)
+		{
+			Parent->RecreateRenderState_Concurrent();
+		}
+		
+		TMap<FName, float> MorphTargetCurves = Parent->GetMorphTargetCurves();
+		PreUpdateDelegate.ExecuteIfBound(this, SkeletalMesh);
+		Parent->SetSkeletalMesh(SkeletalMesh, bReinitPose);
+		
+		// USkeletalMeshCompoent MorphTargetCurves are reset when SetSkeletalMesh is called.
+		// Re-enable them if not bReinitPose.
+		if (!bReinitPose && MorphTargetCurves.Num() > 0)
+		{
+			for (const TPair<FName, float>& MorphTarget : MorphTargetCurves)
+			{
+				Parent->SetMorphTarget(MorphTarget.Key, MorphTarget.Value);
+			}
+
+#if PLATFORM_LINUX
+			FRenderStateRecreator RenderStateRecreator(Parent);
+			FAnimationRuntime::AppendActiveMorphTargets(SkeletalMesh, Parent->GetMorphTargetCurves(), Parent->ActiveMorphTargets, Parent->MorphTargetWeights);		
+#endif
+		}
+
+		if (bForceClothReset)
+		{
+			Parent->ForceClothNextUpdateTeleportAndReset();
+
+#if PLATFORM_LINUX
+			Parent->bDisableClothSimulation = bDisableClothSimulation;
+#endif
+		}
+
+		if (Parent->GetNumOverrideMaterials() > 0)
+		{
+			// For some reason the reference skeletal mesh materials are added as override materials, clear them if necessary
+			Parent->EmptyOverrideMaterials();
+		}	
+	}
+}
+
+void UCustomizableSkeletalComponent::SetPhysicsAsset(UPhysicsAsset* PhysicsAsset)
+{
+	USkeletalMeshComponent* Parent = Cast<USkeletalMeshComponent>(GetAttachParent());
+
+	if (Parent && GetWorld())
+	{
+		Parent->SetPhysicsAsset(PhysicsAsset, true);
+	}
+}
+
+
+USkeletalMesh* UCustomizableSkeletalComponent::GetAttachedSkeletalMesh() const
+{
+	USkeletalMeshComponent* Parent = Cast<USkeletalMeshComponent>(GetAttachParent());
+
+	if (Parent)
+	{
+		return Parent->SkeletalMesh;
+	}
+
+	return nullptr;
+}
+
+
+void UCustomizableSkeletalComponent::UpdateSkeletalMeshAsync(bool bNeverSkipUpdate)
+{
+	if (CustomizableObjectInstance)
+	{
+		CustomizableObjectInstance->UpdateSkeletalMeshAsync(false, false);
+	}
+}
+
+
+void UCustomizableSkeletalComponent::BeginDestroy()
+{
+	Super::BeginDestroy();
+}
+
+
+void UCustomizableSkeletalComponent::PreSave(FObjectPreSaveContext ObjectSaveContext)
+{
+	SetSkeletalMesh(nullptr); // To prevent "Graph is linked to private object(s) in an external package" errors when saving in editor after having PIE and generated customized skeletal meshes while in PIE
+
+	Super::PreSave(ObjectSaveContext);
+}
+
+#if WITH_EDITOR
+
+void UCustomizableSkeletalComponent::UpdateDistFromComponentToLevelEditorCamera(const FVector& CameraPosition)
+{
+	// We want instances in the editor to be generated
+	if (!GetWorld() || GetWorld()->WorldType != EWorldType::Editor)
+		return;
+
+	if (CustomizableObjectInstance)
+	{
+		AActor* ParentActor = GetAttachmentRootActor();
+		if (ParentActor && ParentActor->IsValidLowLevel())
+		{
+			// update distance to camera and set the instance as being used by a component
+			CustomizableObjectInstance->GetPrivate()->SetCOInstanceFlags(UsedByComponent);
+
+			float SquareDist = FVector::DistSquared(CameraPosition, ParentActor->GetActorLocation());
+			CustomizableObjectInstance->GetPrivate()->MinSquareDistFromComponentToPlayer = 
+				FMath::Min(SquareDist, CustomizableObjectInstance->GetPrivate()->MinSquareDistFromComponentToPlayer);
+		}
+
+		USkeletalMesh* AttachedSkeletalMesh = GetAttachedSkeletalMesh();
+		USkeletalMesh* GeneratedSKeletalMesh = CustomizableObjectInstance->GetSkeletalMesh(ComponentIndex);
+		if (!AttachedSkeletalMesh && !CustomizableObjectInstance->GetPrivate()->HasCOInstanceFlags(CreatingSkeletalMesh))
+		{
+			UCustomizableObject* Object = CustomizableObjectInstance->GetCustomizableObject(); 
+			USkeletalMesh* RefMesh = Object ? Object->GetRefSkeletalMesh(ComponentIndex) : nullptr;
+			SetSkeletalMesh(GeneratedSKeletalMesh ? GeneratedSKeletalMesh : RefMesh);
+		}
+		else if (GeneratedSKeletalMesh && AttachedSkeletalMesh != GeneratedSKeletalMesh)
+		{
+			SetSkeletalMesh(GeneratedSKeletalMesh);
+		}
+	}
+}
+
+
+void UCustomizableSkeletalComponent::EditorUpdateComponent()
+{
+	if (CustomizableObjectInstance)
+	{
+		CustomizableObjectInstance->GetPrivate()->SetCOInstanceFlags(UsedByComponent);
+
+		AActor* ParentActor = GetAttachmentRootActor();
+		if (ParentActor)
+		{
+			USkeletalMesh* AttachedSkeletalMesh = GetAttachedSkeletalMesh();
+			USkeletalMesh* GeneratedSKeletalMesh = CustomizableObjectInstance->GetSkeletalMesh(ComponentIndex);
+			if (!AttachedSkeletalMesh && !CustomizableObjectInstance->GetPrivate()->HasCOInstanceFlags(CreatingSkeletalMesh))
+			{
+				SetSkeletalMesh(GeneratedSKeletalMesh ? GeneratedSKeletalMesh : CustomizableObjectInstance->GetCustomizableObject()->GetRefSkeletalMesh(ComponentIndex));
+			}
+			else if (GeneratedSKeletalMesh && AttachedSkeletalMesh != GeneratedSKeletalMesh)
+			{
+				SetSkeletalMesh(GeneratedSKeletalMesh);
+			}
+		}
+	}
+}
+#endif
+
+
+void UCustomizableSkeletalComponent::SetVisibilityOfSkeletalMeshSectionWithMaterialName(bool bInVisible, const FString& MaterialName, int32 LOD)
+{
+	USkeletalMesh* SkeletalMesh = CustomizableObjectInstance->GetSkeletalMesh(ComponentIndex);
+
+	if (!SkeletalMesh)
+	{
+		return;
+	}
+
+	for (int32 i = 0; i < SkeletalMesh->GetMaterials().Num(); ++i)
+	{
+		UMaterialInstanceDynamic* MatInst = Cast<UMaterialInstanceDynamic>(SkeletalMesh->GetMaterials()[i].MaterialInterface);
+
+		if (MatInst && MatInst->Parent->GetName() == MaterialName)
+		{
+			if (FSkeletalMeshRenderData* SkelResource = SkeletalMesh->GetResourceForRendering() )
+			{
+				if ( SkelResource->LODRenderData.IsValidIndex(LOD))
+				{
+					FSkeletalMeshLODRenderData& LodData = SkelResource->LODRenderData[LOD];
+
+					for (int32 j = 0; j < LodData.RenderSections.Num(); ++j)
+					{
+						FSkelMeshRenderSection& Section = LodData.RenderSections[j];
+
+						if (Section.MaterialIndex == i)
+						{
+							Section.bDisabled = !bInVisible;
+						}
+					}
+				}
+			}
+		}
+	}
+}
+
+
+void UCustomizableSkeletalComponent::UpdateDistFromComponentToPlayer(const AActor* ViewCenter, const bool bForceEvenIfNotBegunPlay)
+{
+	if (CustomizableObjectInstance)
+	{
+		AActor* ParentActor = GetAttachmentRootActor();
+		CustomizableObjectInstance->SetIsPlayerOrNearIt(false);
+
+		if (ParentActor && ParentActor->IsValidLowLevel())
+		{
+			if (ParentActor->HasActorBegunPlay() || bForceEvenIfNotBegunPlay)
+			{
+				float SquareDist = FLT_MAX;
+
+				if (ViewCenter && ViewCenter->IsValidLowLevel())
+				{
+					APawn* pawn = Cast<APawn>(ParentActor);
+					bool bIsPawn = pawn ? pawn->IsLocallyControlled() : false;
+					CustomizableObjectInstance->SetIsPlayerOrNearIt(bIsPawn);
+
+					if (bIsPawn)
+					{
+						SquareDist = -0.01f; // Negative value to give the player character more priority than any other character
+					}
+					else
+					{
+						SquareDist = FVector::DistSquared(ViewCenter->GetActorLocation(), ParentActor->GetActorLocation());
+					}
+				}
+				else if (bForceEvenIfNotBegunPlay)
+				{
+					SquareDist = -0.01f; // This is a manual update before begin play and the creation of the pawn, so it should probably be high priority
+					CustomizableObjectInstance->GetPrivate()->LastMinSquareDistFromComponentToPlayer = FMath::Min(SquareDist, CustomizableObjectInstance->GetPrivate()->MinSquareDistFromComponentToPlayer);
+				}
+				else
+				{
+					SquareDist = 0.f; // This a mutable tick before begin play and the creation of the pawn, so it should have a definite and high priority but less than a manual update
+					CustomizableObjectInstance->GetPrivate()->LastMinSquareDistFromComponentToPlayer = FMath::Min(SquareDist, CustomizableObjectInstance->GetPrivate()->MinSquareDistFromComponentToPlayer);
+				}
+
+				CustomizableObjectInstance->GetPrivate()->MinSquareDistFromComponentToPlayer = FMath::Min(SquareDist, CustomizableObjectInstance->GetPrivate()->MinSquareDistFromComponentToPlayer);
+				CustomizableObjectInstance->SetIsBeingUsedByComponentInPlay(true);
+
+				if (CustomizableObjectInstance->GetPrivate()->MinSquareDistFromComponentToPlayer == SquareDist)
+				{
+					CustomizableObjectInstance->NearestToActor = this;
+					CustomizableObjectInstance->NearestToViewCenter = ViewCenter;
+				}
+			}
+		}
+
+		if (ParentActor && GetAttachedSkeletalMesh() == nullptr && CustomizableObjectInstance->GetSkeletalMesh(ComponentIndex) && !CustomizableObjectInstance->GetPrivate()->HasCOInstanceFlags(CreatingSkeletalMesh))
+		{
+			SetSkeletalMesh(CustomizableObjectInstance->GetSkeletalMesh(ComponentIndex));
+		}
+	}
+}
+
+
+void UCustomizableSkeletalComponent::TickComponent(float DeltaTime, ELevelTick TickType, FActorComponentTickFunction* ThisTickFunction)
+{
+	Super::TickComponent(DeltaTime, TickType, ThisTickFunction);
+
+	if (!bPendingSetSkeletalMesh)
+	{
+		USkeletalMeshComponent* Parent = Cast<USkeletalMeshComponent>(GetAttachParent());
+
+		if (Parent)
+		{
+			UCustomizableObject* CustomizableObject = CustomizableObjectInstance->GetCustomizableObject(); 
+
+			bPendingSetSkeletalMesh = true;
+
+			if (CustomizableObjectInstance && CustomizableObjectInstance->GetSkeletalMesh(ComponentIndex))
+			{
+				Parent->SetSkeletalMesh(CustomizableObjectInstance->GetSkeletalMesh(ComponentIndex));
+			}
+			else if (!bSkipSetReferenceSkeletalMesh && CustomizableObjectInstance && CustomizableObject)
+			{
+				Parent->SetSkeletalMesh(CustomizableObject->GetRefSkeletalMesh(ComponentIndex));
+			}
+
+			if (Parent->OverrideMaterials.Num() > 0)
+			{
+#if !(UE_BUILD_SHIPPING || UE_BUILD_TEST)
+				if (Parent->GetClass()->GetFName() != FName(TEXT("SkeletalMeshComponentBudgeted"))) // Reduce unnecessary logging
+				{
+					UE_LOG(LogMutable, Warning, TEXT("Attaching Customizable Skeletal Component to Skeletal Mesh Component with overriden materials! Deleting overrides."));
+				}
+#endif
+
+				Parent->EmptyOverrideMaterials();
+			}
+
+		}
+	}
+}
+
+
+void UCustomizableSkeletalComponent::OnAttachmentChanged()
+{
+	Super::OnAttachmentChanged();
+
+	USkeletalMeshComponent* Parent = Cast<USkeletalMeshComponent>(GetAttachParent());
+
+	if (Parent)
+	{
+		bPendingSetSkeletalMesh = false;
+	}
+	else if(!GetAttachParent())
+	{
+		DestroyComponent();
+	}
+}
diff --git a/Engine/Plugins/Experimental/Mutable/Source/CustomizableObject/Private/LogBenchmarkUtil.cpp b/Engine/Plugins/Experimental/Mutable/Source/CustomizableObject/Private/LogBenchmarkUtil.cpp
new file mode 100644
--- /dev/null
+++ b/Engine/Plugins/Experimental/Mutable/Source/CustomizableObject/Private/LogBenchmarkUtil.cpp
@@ -0,0 +1,123 @@
+// Copyright Epic Games, Inc. All Rights Reserved.
+
+#include "LogBenchmarkUtil.h"
+#include "Misc/FileHelper.h"
+#include "Containers/UnrealString.h"
+#include "HAL/PlatformFileManager.h"
+#include "Misc/DateTime.h"
+#include "UnrealPortabilityHelpers.h"
+
+bool LogBenchmarkUtil::bLoggingActive = false;
+
+LogBenchmarkUtil::~LogBenchmarkUtil() {}
+
+void LogBenchmarkUtil::startLogging()
+{
+	bLoggingActive = true;
+}
+
+void LogBenchmarkUtil::shutdownAndSaveResults()
+{
+	if (bLoggingActive)
+	{
+		FString SaveDirectory = Helper_GetSavedDir() + "/Logs";
+		FString FileName = FString("BenchmarkResult.txt");
+		FString AbsoluteFilePath = SaveDirectory + "/" + FileName;
+		IPlatformFile& PlatformFile = FPlatformFileManager::Get().GetPlatformFile();
+		if (PlatformFile.CreateDirectoryTree(*SaveDirectory))
+		{
+			FString currentTime = FDateTime::Now().ToString();
+			FString DateStamp = LINE_TERMINATOR + FString("[" + currentTime + "]LogBenchmark: (string) ") + FString("benchmark_results_location : ") + SaveDirectory;
+			FFileHelper::SaveStringToFile(DateStamp, *AbsoluteFilePath, FFileHelper::EEncodingOptions::ForceAnsi);
+
+			FString keyName;
+			int32 maxValue;
+			float averageValue;
+
+			// Iterate over the integer stats and log them to a file
+			for (TPair<FName, IntStat> ist : getInstance().IntStats)
+			{
+				currentTime = FDateTime::Now().ToString();
+
+				keyName = FString(ist.Key.ToString());
+				maxValue = ist.Value.max;
+				averageValue = (long double)ist.Value.total / (long double)ist.Value.totalCalls;
+
+				FString TextToSave = LINE_TERMINATOR;
+				TextToSave += FString("[" + currentTime + "]");
+				TextToSave += FString("LogBenchmark: (int) peak_");
+				TextToSave += FString::Printf(TEXT("%s : %lld"), *keyName, maxValue);
+				TextToSave += LINE_TERMINATOR;
+				TextToSave += FString("[" + currentTime + "]");
+				TextToSave += FString("LogBenchmark: (float) average_");
+				TextToSave += FString::Printf(TEXT("%s : %f"), *keyName, averageValue);
+
+				FFileHelper::SaveStringToFile(TextToSave, *AbsoluteFilePath, FFileHelper::EEncodingOptions::ForceAnsi, &IFileManager::Get(), FILEWRITE_Append);
+			}
+
+			// Iterate over the Float stats and log them to a file
+			for (TPair<FName, FloatStat> dst : getInstance().FloatStats)
+			{
+				currentTime = FDateTime::Now().ToString();
+
+				keyName = FString(dst.Key.ToString());
+				maxValue = dst.Value.max;
+				averageValue = dst.Value.total / dst.Value.totalCalls;
+
+				FString TextToSave = LINE_TERMINATOR;
+				TextToSave += FString("[" + currentTime + "]");
+				TextToSave += FString("LogBenchmark: (int) peak_");
+				TextToSave += FString::Printf(TEXT("%s : %lld"), *keyName, maxValue);
+				TextToSave += LINE_TERMINATOR;
+				TextToSave += FString("[" + currentTime + "]");
+				TextToSave += FString("LogBenchmark: (float) average_");
+				TextToSave += FString::Printf(TEXT("%s : %f"), *keyName, averageValue);
+
+				FFileHelper::SaveStringToFile(TextToSave, *AbsoluteFilePath, FFileHelper::EEncodingOptions::ForceAnsi, &IFileManager::Get(), FILEWRITE_Append);
+			}
+		}
+		bLoggingActive = false;
+	}
+}
+
+void LogBenchmarkUtil::updateStat(const FName& stat, int32 value)
+{
+	if (bLoggingActive)
+	{
+		if (getInstance().IntStats.Contains(stat))
+		{
+			getInstance().IntStats[stat] += value;
+		}
+		else
+		{
+			getInstance().IntStats.Add(stat) += value;
+		}
+	}
+}
+
+void LogBenchmarkUtil::updateStat(const FName & stat, double value)
+{
+	updateStat(stat, (long double)value);
+}
+
+void LogBenchmarkUtil::updateStat(const FName& stat, long double value)
+{
+	if (bLoggingActive)
+	{
+		if (getInstance().FloatStats.Contains(stat))
+		{
+			getInstance().FloatStats[stat] += value;
+		}
+		else
+		{
+			getInstance().FloatStats.Add(stat) += value;
+		}
+	}
+}
+
+bool LogBenchmarkUtil::isLoggingActive()
+{
+	return bLoggingActive;
+}
+
+LogBenchmarkUtil::LogBenchmarkUtil() {}
\ No newline at end of file
diff --git a/Engine/Plugins/Experimental/Mutable/Source/CustomizableObject/Private/LogBenchmarkUtil.h b/Engine/Plugins/Experimental/Mutable/Source/CustomizableObject/Private/LogBenchmarkUtil.h
new file mode 100644
--- /dev/null
+++ b/Engine/Plugins/Experimental/Mutable/Source/CustomizableObject/Private/LogBenchmarkUtil.h
@@ -0,0 +1,60 @@
+// Copyright Epic Games, Inc. All Rights Reserved.
+
+#pragma once
+
+#include "Containers/Map.h"
+#include "GenericPlatform/GenericPlatformMath.h"
+#include "UObject/NameTypes.h"
+
+class LogBenchmarkUtil
+{
+public:
+	~LogBenchmarkUtil();
+	static void startLogging();
+	static void shutdownAndSaveResults();
+	static void updateStat(const FName& stat, int32 value);
+	static void updateStat(const FName& stat, double value);
+	static void updateStat(const FName& stat, long double value);
+	static bool isLoggingActive();
+
+private:
+	struct IntStat {
+		int32 total;
+		int32 max;
+		int totalCalls;
+
+		IntStat& operator+=(const int32& rhs)
+		{
+			total += rhs;
+			max = FGenericPlatformMath::Max(max, rhs);
+			totalCalls++;
+			return *this;
+		}
+	};
+
+	struct FloatStat {
+		long double total;
+		long double max;
+		int totalCalls;
+
+		FloatStat& operator+=(const long double& rhs)
+		{
+			total += rhs;
+			max = FGenericPlatformMath::Max(max, rhs);
+			totalCalls++;
+			return *this;
+		}
+	};
+
+	static bool bLoggingActive;
+	TMap<FName, IntStat> IntStats;
+	TMap<FName, FloatStat> FloatStats;
+
+
+	static LogBenchmarkUtil &getInstance() {
+		static LogBenchmarkUtil instance;
+		return instance;
+	}
+
+	LogBenchmarkUtil();
+};
diff --git a/Engine/Plugins/Experimental/Mutable/Source/CustomizableObject/Private/LogInformationUtil.cpp b/Engine/Plugins/Experimental/Mutable/Source/CustomizableObject/Private/LogInformationUtil.cpp
new file mode 100644
--- /dev/null
+++ b/Engine/Plugins/Experimental/Mutable/Source/CustomizableObject/Private/LogInformationUtil.cpp
@@ -0,0 +1,445 @@
+// Copyright Epic Games, Inc. All Rights Reserved.
+
+#include "LogInformationUtil.h"
+#include "Engine/SkeletalMesh.h"
+#include "Engine/Texture2D.h"
+#include "CustomizableObject.h"
+#include "CustomizableInstancePrivateData.h"
+#include "GameFramework/PlayerController.h"
+
+int LogInformationUtil::CountLOD0 = 0;
+int LogInformationUtil::CountLOD1 = 0;
+int LogInformationUtil::CountLOD2 = 0;
+
+void LogInformationUtil::PrintGeneratedTextures(const TArray<struct FGeneratedTexture> GeneratedTextures, FString& Log, bool DoPrintInitialMessage)
+{
+	if (GeneratedTextures.Num() == 0)
+	{
+		return;
+	}
+
+	if (DoPrintInitialMessage)
+	{
+		Log += "Generated textures: ";
+	}
+
+	const int Max = GeneratedTextures.Num();
+	int i;
+
+	bool UseTabulation = (Max > 1);
+
+	if (!DoPrintInitialMessage)
+	{
+		Log += "\n\t\t";
+	}
+	
+
+	if (Max > 1)
+	{
+		UE_LOG(LogMutable, Warning, TEXT("Case"));
+	}
+
+	for (i = 0; i < Max; ++i)
+	{
+		Log += FString::Printf(TEXT("Texture name=%s, "), *GeneratedTextures[i].Name);
+		Log += FString::Printf(TEXT("id=%d, "), GeneratedTextures[i].Id);
+		Log += "\n\t\t\t";
+		PrintTextureData(GeneratedTextures[i].Texture, Log, false);
+		if (i < (Max - 1))
+		{
+			Log += "\n\t\t";
+		}
+	}
+}
+
+
+void LogInformationUtil::PrintGeneratedMaterial(const TArray<struct FGeneratedMaterial> GeneratedMaterials, FString& Log)
+{
+	if (GeneratedMaterials.Num() == 0)
+	{
+		return;
+	}
+
+	Log += "\tGenerated materials:\n";
+
+	FString MessageChunk = "\t\t";
+
+	const int Max = GeneratedMaterials.Num();
+	int i;
+	int entryLength = 25;
+	bool AnyHasTextures = false;
+
+	MessageChunk += "Materials with no textures: ";
+
+	// First only no texture materials
+	for (i = 0; i < Max; ++i)
+	{
+		if (GeneratedMaterials[i].Textures.Num() == 0)
+		{
+			MessageChunk += FString::Printf(TEXT("%d, "), i);
+		}
+		else
+		{
+			AnyHasTextures = true;
+		}
+	}
+
+	// Second texture materials
+	for (i = 0; i < Max; ++i)
+	{
+		if (GeneratedMaterials[i].Textures.Num() > 0)
+		{
+			MessageChunk += "\n\t\t";
+			MessageChunk += FString::Printf(TEXT("%d"), i);
+			MessageChunk += FString::Printf(TEXT(" has %d texture(s)"), GeneratedMaterials[i].Textures.Num());
+			PrintGeneratedTextures(GeneratedMaterials[i].Textures, MessageChunk, false);
+		}
+	}
+
+	Log += MessageChunk;
+}
+
+
+void LogInformationUtil::LogShowInstanceData(const UCustomizableObjectInstance* CustomizableObjectInstance, FString& LogData)
+{
+	LogData += FString::Printf(TEXT("Name=%s "), *CustomizableObjectInstance->GetName());
+	LogData += FString::Printf(TEXT("Priority=%f "), FMath::Sqrt(CustomizableObjectInstance->GetPrivate()->LastMinSquareDistFromComponentToPlayer));
+
+	const uint32 InstanceFlags = CustomizableObjectInstance->GetPrivate()->GetCOInstanceFlags();
+	LogData += FString::Printf(TEXT("bIsUpdating=%d "), InstanceFlags & ECOInstanceFlags::Updating);
+	LogData += FString::Printf(TEXT("bShouldUpdateLODs=%d "), InstanceFlags & ECOInstanceFlags::PendingLODsUpdate);
+
+	LogData += FString::Printf(TEXT("CurrentMinLOD=%d "), CustomizableObjectInstance->GetPrivate()->CurrentMinLOD);
+	LogData += FString::Printf(TEXT("CurrentMaxLOD=%d "), CustomizableObjectInstance->GetPrivate()->CurrentMaxLOD);
+	LogData += FString::Printf(TEXT("MinLODToLoad=%d "), CustomizableObjectInstance->GetPrivate()->MinLODToLoad);
+	LogData += FString::Printf(TEXT("MaxLODToLoad=%d\n"), CustomizableObjectInstance->GetPrivate()->MaxLODToLoad);
+}
+
+
+// TODO: [Acosin] put in a different file and include it
+void LogInformationUtil::LogShowInstanceDataFull(const UCustomizableObjectInstance* CustomizableObjectInstance, bool ShowMaterialInfo)
+{
+	FString LogData = "\n\n";
+
+	LogData += FString::Printf(TEXT("CustomizableObjectInstance %s, "), *CustomizableObjectInstance->GetName());
+	
+	PrintBoolParameters(CustomizableObjectInstance->GetBoolParameters(), LogData);
+	PrintIntParameters(CustomizableObjectInstance->GetIntParameters(), LogData);
+	PrintFloatParameters(CustomizableObjectInstance->GetFloatParameters(), LogData);
+	PrintVectorParameters(CustomizableObjectInstance->GetVectorParameters(), LogData);
+	PrintProjectorParameters(CustomizableObjectInstance->GetProjectorParameters(), LogData);
+
+	if (ShowMaterialInfo)
+	{
+		PrintGeneratedMaterial(CustomizableObjectInstance->GetPrivate()->GeneratedMaterials, LogData);
+
+		if (CustomizableObjectInstance->GetPrivate()->GeneratedTextures.Num() > 0)
+		{
+			LogData += "\n\t";
+		}
+
+		PrintGeneratedTextures(CustomizableObjectInstance->GetPrivate()->GeneratedTextures, LogData, true);
+	}
+
+	FString MessageChunk;
+
+	MessageChunk += "\n\t";
+	MessageChunk += FString::Printf(TEXT("        bBuildParameterDecorations = %d\n"), CustomizableObjectInstance->bBuildParameterDecorations);
+	MessageChunk += FString::Printf(TEXT("        bShowOnlyRuntimeParameters = %d\n"), CustomizableObjectInstance->bShowOnlyRuntimeParameters);
+	MessageChunk += FString::Printf(TEXT("        bShowOnlyRelevantParameters = %d\n"), CustomizableObjectInstance->bShowOnlyRelevantParameters);
+	MessageChunk += FString::Printf(TEXT("        MinSquareDistFromComponentToPlayer = %.2f\n"), CustomizableObjectInstance->GetPrivate()->MinSquareDistFromComponentToPlayer);
+	LogData += MessageChunk;
+
+	const uint32 InstanceFlags = CustomizableObjectInstance->GetPrivate()->GetCOInstanceFlags();
+	MessageChunk = FString::Printf(TEXT("        bIsBeingUsedByComponent = %d\n"), InstanceFlags & ECOInstanceFlags::UsedByComponent);
+	MessageChunk += FString::Printf(TEXT("        bIsBeingUsedByComponentInPlay = %d\n"), InstanceFlags & ECOInstanceFlags::UsedByComponentInPlay);
+	MessageChunk += FString::Printf(TEXT("        bIsPlayerOrNearIt = %d\n"), InstanceFlags & ECOInstanceFlags::UsedByPlayerOrNearIt);
+	MessageChunk += FString::Printf(TEXT("        bIsDiscardedByNumGeneratedInstancesLimit = %d\n"), InstanceFlags & ECOInstanceFlags::DiscardedByNumInstancesLimit);
+	MessageChunk += FString::Printf(TEXT("        LastMinSquareDistFromComponentToPlayer = %.2f\n"), CustomizableObjectInstance->GetPrivate()->LastMinSquareDistFromComponentToPlayer);
+	LogData += MessageChunk;
+
+	MessageChunk = FString::Printf(TEXT("        CurrentMinLOD = %d\n"), CustomizableObjectInstance->GetPrivate()->CurrentMinLOD);
+	MessageChunk += FString::Printf(TEXT("        CurrentMaxLOD = %d\n"), CustomizableObjectInstance->GetPrivate()->CurrentMaxLOD);
+	MessageChunk += FString::Printf(TEXT("        MinLODToLoad = %d\n"), CustomizableObjectInstance->GetPrivate()->MinLODToLoad);
+	MessageChunk += FString::Printf(TEXT("        MaxLODToLoad = %d\n"), CustomizableObjectInstance->GetPrivate()->MaxLODToLoad);
+	MessageChunk += FString::Printf(TEXT("        bShouldUpdateLODs = %d\n"), InstanceFlags & ECOInstanceFlags::PendingLODsUpdate);
+	LogData += MessageChunk;
+
+	MessageChunk = FString::Printf(TEXT("        bShouldUpdateLODsSecondStage = %d\n"), InstanceFlags & ECOInstanceFlags::PendingLODsUpdateSecondStage);
+	MessageChunk += FString::Printf(TEXT("        bIsDowngradeLODUpdate = %d\n"), InstanceFlags & ECOInstanceFlags::PendingLODsDowngrade);
+	MessageChunk += FString::Printf(TEXT("        bIsUpdating = %d\n"), InstanceFlags & ECOInstanceFlags::Updating);
+	MessageChunk += FString::Printf(TEXT("        bIsCreatingSkeletalMesh = %d\n"), InstanceFlags & ECOInstanceFlags::CreatingSkeletalMesh);
+	MessageChunk += FString::Printf(TEXT("        bIsGenerated = %d\n"), InstanceFlags & ECOInstanceFlags::Generated);
+	LogData += MessageChunk;
+
+	for (int32 ComponentIndex = 0; ComponentIndex < CustomizableObjectInstance->SkeletalMeshes.Num(); ++ComponentIndex)
+	{
+		if ((CustomizableObjectInstance->GetSkeletalMesh(ComponentIndex) != nullptr) && (CustomizableObjectInstance->GetSkeletalMesh(ComponentIndex)->GetResourceForRendering()))
+		{
+			if (CustomizableObjectInstance->GetPrivate()->GetMinLODToLoad() < 1)
+			{
+				CountLOD0++;
+			}
+			else if (CustomizableObjectInstance->GetPrivate()->GetMinLODToLoad() < 2)
+			{
+				CountLOD1++;
+			}
+			else
+			{
+				CountLOD2++;
+			}
+		}
+	}
+
+	UE_LOG(LogMutable, Warning, TEXT("%s"), *LogData);
+
+	UWorld* World = GWorld;
+
+	if (World)
+	{
+		APlayerController* PlayerController = World->GetFirstPlayerController();
+		if (PlayerController)
+		{
+			PlayerController->ClientMessage(LogData);
+		}
+	}
+}
+
+
+void LogInformationUtil::PrintBoolParameters(const TArray<struct FCustomizableObjectBoolParameterValue>& BoolParameters, FString& Log)
+{
+	if (BoolParameters.Num() == 0)
+	{
+		return;
+	}
+
+	Log += "Boolean parameters:\n";
+
+	const int Max = BoolParameters.Num();
+	int i;
+
+	for (i = 0; i < Max; ++i)
+	{
+		Log += FString::Printf(TEXT("        %s="), *BoolParameters[i].ParameterName);
+		Log += FString::Printf(TEXT("%d\n"), BoolParameters[i].ParameterValue);
+	}
+
+	Log += "\n";
+}
+
+
+void LogInformationUtil::PrintIntParameters(const TArray<struct FCustomizableObjectIntParameterValue>& IntParameters, FString& Log)
+{
+	if (IntParameters.Num() == 0)
+	{
+		return;
+	}
+
+	FString MessageChunk;
+
+	MessageChunk += "\tInt parameters:\n";
+
+	const int Max = IntParameters.Num();
+	int i;
+
+	for (i = 0; i < Max; ++i)
+	{
+		MessageChunk += FString::Printf(TEXT("        %s="), *IntParameters[i].ParameterName);
+		MessageChunk += FString::Printf(TEXT("%s\n"), *IntParameters[i].ParameterValueName);
+		FillToLength(MessageChunk,25);
+	}
+
+	MessageChunk += "\n";
+
+	Log += MessageChunk;
+}
+
+
+void LogInformationUtil::PrintFloatParameters(const TArray<struct FCustomizableObjectFloatParameterValue>& FloatParameters, FString& Log)
+{
+	if (FloatParameters.Num() == 0)
+	{
+		return;
+	}
+
+	Log += "Float parameters:\n";
+
+	const int Max = FloatParameters.Num();
+	int i;
+
+	for (i = 0; i < Max; ++i)
+	{
+		Log += FString::Printf(TEXT("        %s="), *FloatParameters[i].ParameterName);
+		Log += FString::Printf(TEXT("%.2f\n"), FloatParameters[i].ParameterValue);
+	}
+
+	Log += "\n";
+}
+
+
+void LogInformationUtil::PrintVectorParameters(const TArray<struct FCustomizableObjectVectorParameterValue>& VectorParameters, FString& Log)
+{
+	if (VectorParameters.Num() == 0)
+	{
+		return;
+	}
+
+	Log += "Vector parameters:\n";
+
+	const int Max = VectorParameters.Num();
+	int i;
+
+	for (i = 0; i < Max; ++i)
+	{
+		Log += FString::Printf(TEXT("        %s="), *VectorParameters[i].ParameterName);
+		Log += FString::Printf(TEXT("(%.2f,"), VectorParameters[i].ParameterValue.R);
+		Log += FString::Printf(TEXT("%.2f,"), VectorParameters[i].ParameterValue.G);
+		Log += FString::Printf(TEXT("%.2f,"), VectorParameters[i].ParameterValue.B);
+		Log += FString::Printf(TEXT("%.2f)\n"), VectorParameters[i].ParameterValue.A);
+	}
+
+	Log += "\n";
+}
+
+
+void LogInformationUtil::PrintProjectorParameters(const TArray<struct FCustomizableObjectProjectorParameterValue>& ProjectorParameters, FString& Log)
+{
+	if (ProjectorParameters.Num() == 0)
+	{
+		return;
+	}
+
+	Log += "Projector parameters:\n";
+
+	const int Max = ProjectorParameters.Num();
+	int i;
+
+	for (i = 0; i < Max; ++i)
+	{
+		Log += FString::Printf(TEXT("        %s has "), *ProjectorParameters[i].ParameterName);
+		Log += FString::Printf(TEXT(" position=(%.2f,"), ProjectorParameters[i].Value.Position.X);
+		Log += FString::Printf(TEXT("%.2f,"), ProjectorParameters[i].Value.Position.Y);
+		Log += FString::Printf(TEXT("%.2f)"), ProjectorParameters[i].Value.Position.Z);
+		Log += FString::Printf(TEXT(" direction=(%.2f,"), ProjectorParameters[i].Value.Direction.X);
+		Log += FString::Printf(TEXT("%.2f,"), ProjectorParameters[i].Value.Direction.Y);
+		Log += FString::Printf(TEXT("%.2f)"), ProjectorParameters[i].Value.Direction.Z);
+		Log += FString::Printf(TEXT(" up=(%.2f,"), ProjectorParameters[i].Value.Up.X);
+		Log += FString::Printf(TEXT("%.2f,"), ProjectorParameters[i].Value.Up.Y);
+		Log += FString::Printf(TEXT("%.2f)"), ProjectorParameters[i].Value.Up.Z);
+		Log += FString::Printf(TEXT(" scale=(%.2f,"), ProjectorParameters[i].Value.Scale.X);
+		Log += FString::Printf(TEXT("%.2f)\n"), ProjectorParameters[i].Value.Scale.Y);
+	}
+
+	Log += "\n";
+}
+
+
+void LogInformationUtil::PrintTextureData(const UTexture2D* Texture, FString& Log, bool UseTabulation)
+{
+	FString MessageChunk;
+	MessageChunk += FString::Printf(TEXT("        Name=%s\n"), *Texture->GetName());
+	if (UseTabulation)
+	{
+		FillToLength(MessageChunk, 37);
+	}
+	MessageChunk += FString::Printf(TEXT("                w=%d\n"), int32(Texture->GetSurfaceWidth()));
+	if (UseTabulation)
+	{
+		FillToLength(MessageChunk, 65);
+	}
+	MessageChunk += FString::Printf(TEXT("                h=%d\n"), int32(Texture->GetSurfaceHeight()));
+	if (UseTabulation)
+	{
+		FillToLength(MessageChunk, 95);
+	}
+	MessageChunk += FString::Printf(TEXT("                bForceMiplevelsToBeResident=%d\n"), Texture->bForceMiplevelsToBeResident);
+	if (UseTabulation)
+	{
+		MessageChunk += "\n\t\t\t";
+	}
+	MessageChunk += FString::Printf(TEXT("                bGlobalForceMipLevelsToBeResident=%d\n"), Texture->bGlobalForceMipLevelsToBeResident);
+	if (UseTabulation)
+	{
+		FillToLength(MessageChunk, 37);
+	}
+	MessageChunk += FString::Printf(TEXT("                bHasStreamingUpdatePending=%d\n"), Texture->bHasStreamingUpdatePending);
+	if (UseTabulation)
+	{
+		FillToLength(MessageChunk, 65);
+	}
+	MessageChunk += FString::Printf(TEXT("                bIgnoreStreamingMipBias=%d\n"), Texture->bIgnoreStreamingMipBias);
+	if (UseTabulation)
+	{
+		MessageChunk += "\n\t\t\t";
+	}
+	MessageChunk += FString::Printf(TEXT("                FirstResourceMemMip=%d\n"), Texture->FirstResourceMemMip);
+	if (UseTabulation)
+	{
+		FillToLength(MessageChunk, 37);
+	}
+	MessageChunk += FString::Printf(TEXT("                LevelIndex=%d\n"), Texture->LevelIndex);
+	if (UseTabulation)
+	{
+		FillToLength(MessageChunk, 65);
+	}
+	MessageChunk += FString::Printf(TEXT("                CompressionSettings=%s\n"), *UEnum::GetValueAsString(Texture->CompressionSettings));
+	if (UseTabulation)
+	{
+		FillToLength(MessageChunk, 95);
+	}
+	MessageChunk += FString::Printf(TEXT("                Filter=%s\n"), *UEnum::GetValueAsString(Texture->Filter));
+	if (UseTabulation)
+	{
+		FillToLength(MessageChunk, 128);
+	}
+	MessageChunk += FString::Printf(TEXT("                LODGroup=%s\n"), *UEnum::GetValueAsString(Texture->LODGroup));
+
+	Log += MessageChunk;
+}
+
+
+void LogInformationUtil::FillToLength(FString& Data, int32 NewLength)
+{
+	int32 i;
+	const int32 Max = NewLength - Data.Len();
+
+	for (i = 0; i < Max; ++i)
+	{
+		Data += " ";
+	}
+}
+
+
+void LogInformationUtil::ResetCounters()
+{
+	CountLOD0 = 0;
+	CountLOD1 = 0;
+	CountLOD2 = 0;
+}
+
+void LogInformationUtil::PrintImageToPlatformDataMap(const TMap<uint32, FTexturePlatformData*>& ImageToPlatformDataMap, FString& Log)
+{
+	if (ImageToPlatformDataMap.Num() == 0)
+	{
+		return;
+	}
+
+	Log += "\n\tImageToPlatformDataMap:";
+	for (const TPair<uint32, FTexturePlatformData*>& Elem : ImageToPlatformDataMap)
+	{
+		Log += "\n\t\t";
+		Log = FString::Printf(TEXT("key=%d, "), Elem.Key);
+		Log = FString::Printf(TEXT("value: SizeX=%d, "), Elem.Value->SizeX);
+		Log = FString::Printf(TEXT("SizeY=%d, "), Elem.Value->SizeY);
+		Log = FString::Printf(TEXT("PixelFormat=%d, "), *UEnum::GetValueAsString(Elem.Value->PixelFormat));
+		int i;
+		const int Max = Elem.Value->Mips.Num();
+		for (i = 0; i < Max; ++i)
+		{
+			Log += "\n\t\t\t";
+			Log = FString::Printf(TEXT("Mips[%d] "), i);
+			Log = FString::Printf(TEXT("SizeX=%d, "), Elem.Value->Mips[i].SizeX);
+			Log = FString::Printf(TEXT("SizeY=%d"), Elem.Value->Mips[i].SizeY);
+		}
+	}
+}
diff --git a/Engine/Plugins/Experimental/Mutable/Source/CustomizableObject/Private/LogInformationUtil.h b/Engine/Plugins/Experimental/Mutable/Source/CustomizableObject/Private/LogInformationUtil.h
new file mode 100644
--- /dev/null
+++ b/Engine/Plugins/Experimental/Mutable/Source/CustomizableObject/Private/LogInformationUtil.h
@@ -0,0 +1,71 @@
+// Copyright Epic Games, Inc. All Rights Reserved.
+
+#pragma once
+
+#include "Containers/Array.h"
+#include "Engine/Texture.h"
+
+class UCustomizableObjectInstance;
+
+#define STRINGIZE(A) STRINGIZE_NX(A)
+#define STRINGIZE_NX(A) #A
+#define ADD_INTEGER_LOG_VARIABLE(String, info, booleanValue) {String += FString::Printf(TEXT(STRINGIZE_NX(info %d)), #booleanValue);}
+#define ADD_FLOAT_LOG_VARIABLE(String, info, floatValue) {String += FString::Printf(TEXT(STRINGIZE_NX(info %d)), #floatValue);}
+#define ADD_STRING_LOG_VARIABLE(String, info, stringValue) {String += FString::Printf(TEXT(STRINGIZE_NX(info %s)), #stringValue);}
+#define ADD_STRING_LOG(String, stringValue) {String += FString::Printf(TEXT(STRINGIZE_NX(stringValues)), #stringValue);}
+
+class LogInformationUtil
+{
+public:
+	// Show full data about all UCustomizableObjectInstance existing elements
+	static void LogShowInstanceDataFull(const UCustomizableObjectInstance* CustomizableObjectInstance, bool ShowMaterialInfo);
+
+	// Show data about all UCustomizableObjectInstance existing elements
+	static void LogShowInstanceData(const UCustomizableObjectInstance* CustomizableObjectInstance, FString& LogData);
+
+	// Utility method just to log information, adds to the given FString containing the
+	// log message the log information of BoolParameters
+	static void PrintBoolParameters(const TArray<struct FCustomizableObjectBoolParameterValue>& BoolParameters, FString& Log);
+
+	// Utility method just to log information, adds to the given FString containing the
+	// log message the log information of IntParameters
+	static void PrintIntParameters(const TArray<struct FCustomizableObjectIntParameterValue>& IntParameters, FString& Log);
+
+	// Utility method just to log information, adds to the given FString containing the
+	// log message the log information of FloatParameters
+	static void PrintFloatParameters(const TArray<struct FCustomizableObjectFloatParameterValue>& FloatParameters, FString& Log);
+
+	// Utility method just to log information, adds to the given FString containing the
+	// log message the log information of VectorParameters
+	static void PrintVectorParameters(const TArray<struct FCustomizableObjectVectorParameterValue>& VectorParameters, FString& Log);
+
+	// Utility method just to log information, adds to the given FString containing the
+	// log message the log information of ProjectorParameters
+	static void PrintProjectorParameters(const TArray<struct FCustomizableObjectProjectorParameterValue>& ProjectorParameters, FString& Log);
+
+	// Utility method just to log information, adds to the given FString containing the
+	// log message the log information of GeneratedMaterials
+	static void PrintGeneratedMaterial(const TArray<struct FGeneratedMaterial> GeneratedMaterials, FString& Log);
+
+	// Utility method just to log information, adds to the given FString containing the
+	// log message the log information of GeneratedTextures
+	static void PrintGeneratedTextures(const TArray<struct FGeneratedTexture> GeneratedTextures, FString& Log, bool DoPrintInitialMessage);
+
+	// Utility method just to log information, adds to the given FString containing the
+	// log message the log information of Texture
+	static void PrintTextureData(const class UTexture2D* Texture, FString& Log, bool UseTabulation);
+
+	// Reset static lod counters
+	static void ResetCounters();
+
+	// Fills with empty spaces the Data string so it reaches a kength of NewLength (is lower than current length, nothing is done)
+	static void FillToLength(FString& Data, int32 NewLength);
+
+	// Utility method just to log information, adds to the given FString containing the
+	// log message the log information of ImageToPlatformDataMap
+	static void PrintImageToPlatformDataMap(const TMap<uint32, FTexturePlatformData*>& ImageToPlatformDataMap, FString& Log);
+
+	static int CountLOD0;
+	static int CountLOD1;
+	static int CountLOD2;
+};
diff --git a/Engine/Plugins/Experimental/Mutable/Source/CustomizableObject/Private/MultilayerProjector.cpp b/Engine/Plugins/Experimental/Mutable/Source/CustomizableObject/Private/MultilayerProjector.cpp
new file mode 100644
--- /dev/null
+++ b/Engine/Plugins/Experimental/Mutable/Source/CustomizableObject/Private/MultilayerProjector.cpp
@@ -0,0 +1,502 @@
+// Copyright Epic Games, Inc. All Rights Reserved.
+
+#include "MultilayerProjector.h"
+
+#include "CustomizableObjectInstance.h"
+
+
+void FMultilayerProjectorLayer::Read(const FMultilayerProjector& MultilayerProjector, const int32 Index)
+{
+	checkCode(MultilayerProjector.CheckInstanceParameters());
+	check(Index >= 0 && Index < MultilayerProjector.NumLayers()); // Layer out of range.
+	
+	const UCustomizableObjectInstance& Instance = *MultilayerProjector.Instance;
+	const FString ParamName = MultilayerProjector.ParamName.ToString();
+	
+	{
+		const int32 ProjectorParamIndex = Instance.FindProjectorParameterNameIndex(ParamName);
+		const FCustomizableObjectProjector& Projector = Instance.GetProjectorParameters()[ProjectorParamIndex].RangeValues[Index];
+		Position = static_cast<FVector3d>(Projector.Position);
+		Direction = static_cast<FVector3d>(Projector.Direction);
+		Up = static_cast<FVector3d>(Projector.Up);
+		Scale = static_cast<FVector3d>(Projector.Scale);
+		Angle = Projector.Angle;
+	}
+	
+	{
+		const int32 ImageParamIndex = Instance.FindIntParameterNameIndex(ParamName + FMultilayerProjector::IMAGE_PARAMETER_POSTFIX);
+		Image = Instance.GetIntParameters()[ImageParamIndex].ParameterRangeValueNames[Index];
+	}
+
+	{
+		const int32 OpacityParamIndex = Instance.FindFloatParameterNameIndex(ParamName + FMultilayerProjector::OPACITY_PARAMETER_POSTFIX);
+		Opacity = Instance.GetFloatParameters()[OpacityParamIndex].ParameterRangeValues[Index];
+	}
+}
+
+
+void FMultilayerProjectorLayer::Write(const FMultilayerProjector& MultilayerProjector, const int32 Index) const
+{
+	checkCode(MultilayerProjector.CheckInstanceParameters());
+	check(Index >= 0 && Index < MultilayerProjector.NumLayers()); // Layer out of range.
+
+	UCustomizableObjectInstance& Instance = *MultilayerProjector.Instance;
+	const FString ParamName = MultilayerProjector.ParamName.ToString();
+
+	{
+		const int32 ProjectorParamIndex = Instance.FindProjectorParameterNameIndex(ParamName);
+		FCustomizableObjectProjector& Projector = Instance.GetProjectorParameters()[ProjectorParamIndex].RangeValues[Index];
+		Projector.Position = static_cast<FVector3f>(Position);
+		Projector.Direction = static_cast<FVector3f>(Direction);
+		Projector.Up = static_cast<FVector3f>(Up);
+		Projector.Scale = static_cast<FVector3f>(Scale);
+		Projector.Angle = Angle;
+	}
+	
+	{
+		const int32 ImageParamIndex = Instance.FindIntParameterNameIndex(ParamName + FMultilayerProjector::IMAGE_PARAMETER_POSTFIX);
+		Instance.GetIntParameters()[ImageParamIndex].ParameterRangeValueNames[Index] = Image;
+	}
+
+	{
+		const int32 OpacityParamIndex = Instance.FindFloatParameterNameIndex(ParamName + FMultilayerProjector::OPACITY_PARAMETER_POSTFIX);
+		Instance.GetFloatParameters()[OpacityParamIndex].ParameterRangeValues[Index] = Opacity;
+	}
+}
+
+
+FMultilayerProjectorVirtualLayer::FMultilayerProjectorVirtualLayer(const FMultilayerProjectorLayer& Layer, const bool bEnabled, const int32 Order):
+	FMultilayerProjectorLayer(Layer),
+	bEnabled(bEnabled),
+	Order(Order)
+{
+}
+
+
+const FString FMultilayerProjector::NUM_LAYERS_PARAMETER_POSTFIX = FString("_NumLayers");
+
+const FString FMultilayerProjector::OPACITY_PARAMETER_POSTFIX = FString("_Opacity");
+
+const FString FMultilayerProjector::IMAGE_PARAMETER_POSTFIX = FString("_SelectedImages");
+
+const FString FMultilayerProjector::POSE_PARAMETER_POSTFIX = FString("_SelectedPoses");
+
+
+int32 FMultilayerProjector::NumLayers() const
+{
+	check(Instance); // Multilayer Projector created without using the Instance factory.
+
+	const FString NumLayersParamName = ParamName.ToString() + NUM_LAYERS_PARAMETER_POSTFIX;
+	
+	const int32 FloatParameterIndex = Instance->FindFloatParameterNameIndex(NumLayersParamName);
+	check(FloatParameterIndex != -1); // Parameter not found.
+
+	return Instance->GetFloatParameters()[FloatParameterIndex].ParameterValue;
+}
+
+
+void FMultilayerProjector::CreateLayer(const int32 Index) const
+{
+	check(Instance); // Multilayer Projector created without using the Instance factory.
+	checkCode(CheckInstanceParameters());
+	check(Index >= 0 && Index <= NumLayers()); // Layer is non-contiguous or out of range.
+
+	const UCustomizableObject* Object = Instance->GetCustomizableObject(); 
+	
+	// Num Layers.
+	{
+        const FString NumLayersParamName = ParamName.ToString() + NUM_LAYERS_PARAMETER_POSTFIX;
+		const int32 FloatParameterIndex = Instance->FindFloatParameterNameIndex(NumLayersParamName);
+
+        Instance->GetFloatParameters()[FloatParameterIndex].ParameterValue += 1;
+    }
+	// Projector Range.
+	{
+		const int32 ProjectorParameterIndex = Instance->FindProjectorParameterNameIndex(ParamName.ToString());
+		
+		FCustomizableObjectProjectorParameterValue& ProjectorParameter = Instance->GetProjectorParameters()[ProjectorParameterIndex];
+		const FCustomizableObjectProjector Projector = Instance->GetProjectorDefaultValue(Object->FindParameter(ParamName.ToString()));
+		ProjectorParameter.RangeValues.Insert(Projector, Index);
+	}
+	
+	// Selected Image Range.
+	{
+		const int32 IntParameterIndex = Instance->FindIntParameterNameIndex(ParamName.ToString() + FMultilayerProjector::IMAGE_PARAMETER_POSTFIX);
+
+		FCustomizableObjectIntParameterValue& IntParameter = Instance->GetIntParameters()[IntParameterIndex];
+		const int32 ParamIndexInObject = Object->FindParameter(IntParameter.ParameterName);
+
+		const FString DefaultValue = Object->GetIntParameterAvailableOption(ParamIndexInObject, 0); // TODO: Define the default option in the editor instead of taking the first available, like it's currently defined for GetProjectorDefaultValue()
+		IntParameter.ParameterRangeValueNames.Insert(DefaultValue, Index);
+	}
+	
+	// Opacity Range.
+	{
+		const int32 FloatParameterIndex = Instance->FindFloatParameterNameIndex(ParamName.ToString() + FMultilayerProjector::OPACITY_PARAMETER_POSTFIX);
+
+		FCustomizableObjectFloatParameterValue& FloatParameter = Instance->GetFloatParameters()[FloatParameterIndex];
+		FloatParameter.ParameterRangeValues.Insert(0.5, Index); // TODO: Define the default float in the editor instead of [0.5f], like it's currently defined for GetProjectorDefaultValue()
+	}
+}
+
+
+void FMultilayerProjector::RemoveLayerAt(const int32 Index) const
+{
+	check(Instance); // Multilayer Projector created without using the Instance factory.
+	checkCode(CheckInstanceParameters());
+	check(Index >= 0 && Index < NumLayers()); // Layer out of range.
+	
+	// Num Layers.
+	{
+		const FString NumLayersParamName = ParamName.ToString() + NUM_LAYERS_PARAMETER_POSTFIX;
+        const int32 FloatParameterIndex = Instance->FindFloatParameterNameIndex(NumLayersParamName);
+		
+    	Instance->GetFloatParameters()[FloatParameterIndex].ParameterValue -= 1;
+    }
+    
+	// Projector Range.
+	{
+		const int32 ProjectorParameterIndex = Instance->FindProjectorParameterNameIndex(ParamName.ToString());
+		
+		FCustomizableObjectProjectorParameterValue& ProjectorParameter = Instance->GetProjectorParameters()[ProjectorParameterIndex];
+		ProjectorParameter.RangeValues.RemoveAt(Index);
+	}
+	
+	// Selected Image Range.
+	{
+		const int32 IntParameterIndex = Instance->FindIntParameterNameIndex(ParamName.ToString() + IMAGE_PARAMETER_POSTFIX);
+		
+		FCustomizableObjectIntParameterValue& IntParameter = Instance->GetIntParameters()[IntParameterIndex];
+		IntParameter.ParameterRangeValueNames.RemoveAt(Index);
+	}
+	
+	// Opacity Range.
+	{
+		const int32 FloatParameterIndex = Instance->FindFloatParameterNameIndex(ParamName.ToString() + OPACITY_PARAMETER_POSTFIX);
+		
+		FCustomizableObjectFloatParameterValue& FloatParameter = Instance->GetFloatParameters()[FloatParameterIndex];
+		FloatParameter.ParameterRangeValues.RemoveAt(Index);
+	}
+}
+
+
+FMultilayerProjectorLayer FMultilayerProjector::GetLayer(const int32 Index) const
+{
+	FMultilayerProjectorLayer MultilayerProjectorLayer;
+	MultilayerProjectorLayer.Read(*this, Index);
+	return MultilayerProjectorLayer;
+}
+
+
+void FMultilayerProjector::UpdateLayer(const int32 Index, const FMultilayerProjectorLayer& Layer) const
+{
+	Layer.Write(*this, Index);
+}
+
+
+TArray<FName> FMultilayerProjector::GetVirtualLayers() const
+{
+	TArray<FName> VirtualLayers;
+	VirtualLayersMapping.GetKeys(VirtualLayers);
+	return VirtualLayers;
+}
+
+
+void FMultilayerProjector::CreateVirtualLayer(const FName& Id)
+{
+	if (!VirtualLayersMapping.Contains(Id))
+	{
+		const int32 Index = NumLayers();
+		
+		CreateLayer(Index);
+		VirtualLayersMapping.Add(Id, Index);
+		VirtualLayersOrder.Add(Id, NEW_VIRTUAL_LAYER_ORDER);
+	}
+}
+
+
+FMultilayerProjectorVirtualLayer FMultilayerProjector::FindOrCreateVirtualLayer(const FName& Id)
+{
+	FMultilayerProjectorLayer Layer;
+	bool bEnabled;
+	int32 Order;
+
+	if (const int32* Index = VirtualLayersMapping.Find(Id))
+	{
+		if (*Index == VIRTUAL_LAYER_DISABLED)
+		{
+			Layer = DisableVirtualLayers[Id];
+			bEnabled = false;
+		}
+		else
+		{
+			Layer = GetLayer(*Index);
+			bEnabled = true;
+		}
+
+		Order = VirtualLayersOrder[Id];
+	}
+	else
+	{
+		const int32 NewIndex = NumLayers();
+		constexpr int32 NewOrder = NEW_VIRTUAL_LAYER_ORDER;
+		
+		CreateLayer(NewIndex);
+		VirtualLayersMapping.Add(Id, NewIndex);
+		VirtualLayersOrder.Add(Id, NewOrder);
+
+		Layer = GetLayer(NewIndex);
+		bEnabled = true;
+		Order = NewOrder;
+	}
+
+	return FMultilayerProjectorVirtualLayer(Layer, bEnabled, Order);
+}
+
+
+void FMultilayerProjector::RemoveVirtualLayer(const FName& Id)
+{
+	const int32* Index = VirtualLayersMapping.Find(Id);
+	check(Index); // Virtual Layer not created.
+	
+	if (*Index == VIRTUAL_LAYER_DISABLED)
+	{
+		DisableVirtualLayers.Remove(Id);
+	}
+	else
+	{
+		RemoveLayerAt(*Index);
+		
+		for (TMap<FName, int32>::TIterator It =  VirtualLayersMapping.CreateIterator(); It; ++It)
+		{
+			if (It.Key() == Id)
+			{
+				It.RemoveCurrent();
+			}
+			else if (It.Value() > *Index) // Update following Layers.
+			{
+				--It.Value();
+			}
+		}
+	}
+
+	VirtualLayersOrder.Remove(Id);
+}
+
+
+FMultilayerProjectorVirtualLayer FMultilayerProjector::GetVirtualLayer(const FName& Id) const
+{
+	const int32* Index = VirtualLayersMapping.Find(Id);
+	check(Index); // Virtual Layer not created.
+
+	const FMultilayerProjectorLayer Layer = GetLayer(*Index);
+	const bool bEnabled = *Index != VIRTUAL_LAYER_DISABLED;
+	const int32 Order = VirtualLayersOrder[Id];
+	
+	return FMultilayerProjectorVirtualLayer(Layer, bEnabled, Order);
+}
+
+
+void FMultilayerProjector::UpdateVirtualLayer(const FName& Id, const FMultilayerProjectorVirtualLayer& Layer)
+{
+	const int32* Index = VirtualLayersMapping.Find(Id);
+	check(Index); // Virtual Layer not created.
+
+	const bool bEnabled = *Index != VIRTUAL_LAYER_DISABLED;
+	
+	if (!bEnabled)
+	{
+		DisableVirtualLayers[Id] = static_cast<FMultilayerProjectorLayer>(Layer); // Update disabled layer.
+		VirtualLayersOrder[Id] = Layer.Order;
+	}
+	else
+	{
+		int32* Order = VirtualLayersOrder.Find(Id);
+		if (*Order != Layer.Order) // Order changed, check if it needs to be moved.
+		{
+			const int32 OldIndex = *Index;
+			int32 NewIndex = CalculateVirtualLayerIndex(Id, Layer.Order);
+			if (OldIndex != NewIndex) // Move required. Could be optimized by moving only the in-between values.
+			{
+				RemoveLayerAt(OldIndex);
+				UpdateMappingVirtualLayerDisabled(Id, OldIndex);
+
+				if (OldIndex < NewIndex)
+				{
+					NewIndex -= 1;
+				}
+				
+				CreateLayer(NewIndex);
+				UpdateMappingVirtualLayerEnabled(Id, NewIndex);
+			}
+			
+			*Order = Layer.Order;
+		}
+
+		UpdateLayer(*Index, static_cast<FMultilayerProjectorLayer>(Layer)); // Update enabled layer.
+	}
+	
+	// Enable or disable virtual layer.
+	if (Layer.bEnabled && !bEnabled)
+	{
+		const int32 NewIndex = CalculateVirtualLayerIndex(Id, VirtualLayersOrder[Id]);
+
+		CreateLayer(NewIndex);
+		UpdateMappingVirtualLayerEnabled(Id, NewIndex);
+
+		UpdateLayer(NewIndex, static_cast<FMultilayerProjectorLayer>(Layer));
+		
+		DisableVirtualLayers.Remove(Id);
+	}
+	else if (!Layer.bEnabled && bEnabled)
+	{
+		RemoveLayerAt(*Index);
+		UpdateMappingVirtualLayerDisabled(Id, *Index);
+		
+		DisableVirtualLayers.Add(Id, Layer);
+	}
+}
+
+
+FMultilayerProjector::FMultilayerProjector(UCustomizableObjectInstance& Instance, const FName& ParamName) :
+	Instance(&Instance),
+	ParamName(ParamName)
+{
+}
+
+
+int32 FMultilayerProjector::CalculateVirtualLayerIndex(const FName& Id, const int32 InsertOrder) const
+{
+	int32 LayerBeforeIndex = -1;
+	int32 LayerBeforeOrder = -1;
+	
+	for (const TTuple<FName, int>& MappingTuple : VirtualLayersMapping) // Find closest smallest layer.
+	{
+		if (MappingTuple.Value != VIRTUAL_LAYER_DISABLED && MappingTuple.Key != Id)
+		{
+			const int32 LayerOrder = VirtualLayersOrder[MappingTuple.Key];
+			if (LayerOrder <= InsertOrder)
+			{
+				if ((LayerOrder > LayerBeforeOrder) ||
+					(LayerOrder == LayerBeforeOrder && MappingTuple.Value > LayerBeforeIndex))
+				{
+					LayerBeforeIndex = MappingTuple.Value;
+					LayerBeforeOrder = LayerOrder;
+				}
+			}
+		}
+	}
+	
+	return LayerBeforeIndex + 1;
+}
+
+
+void FMultilayerProjector::UpdateMappingVirtualLayerEnabled(const FName& Id, const int32 Index)
+{
+	for (TTuple<FName, int>& Tuple : VirtualLayersMapping)
+	{
+		if (Tuple.Key == Id)
+		{
+			Tuple.Value = Index;
+		}
+		else if (Tuple.Value >= Index) // Update following Layers.
+		{
+			++Tuple.Value;
+		}
+	}
+}
+
+
+void FMultilayerProjector::UpdateMappingVirtualLayerDisabled(const FName& Id, const int32 Index)
+{
+	for (TTuple<FName, int>& Tuple : VirtualLayersMapping)
+	{
+		if (Tuple.Key == Id)
+    	{
+			Tuple.Value = VIRTUAL_LAYER_DISABLED;
+    	}
+		else if (Tuple.Value > Index) // Update following Layers.
+		{
+			--Tuple.Value;
+		}
+	}
+}
+
+
+void FMultilayerProjector::CheckInstanceParameters() const
+{
+	check(Instance); // Multilayer Projector created without using the Instance factory.
+
+	const FString ParamNameString = ParamName.ToString();
+
+	// Num layers.
+	{
+		const FString NumLayersParamName = ParamNameString + NUM_LAYERS_PARAMETER_POSTFIX;
+		const int32 FloatParameterIndex = Instance->FindFloatParameterNameIndex(NumLayersParamName);
+		check(FloatParameterIndex >= 0); // Instance Parameter does not exist.
+	}
+    
+	// Projector.
+	{
+		const int32 ProjectorParameterIndex = Instance->FindProjectorParameterNameIndex(ParamNameString);
+		check(ProjectorParameterIndex >= 0) // Instance Parameter does not exist.
+	}
+	
+	// Selected Image.
+	{
+		const int32 IntParameterIndex = Instance->FindIntParameterNameIndex(ParamNameString + IMAGE_PARAMETER_POSTFIX);
+		check(IntParameterIndex >= 0) // Instance Parameter does not exist.
+	}
+	
+	// Opacity.
+	{
+		const int32 FloatParameterIndex = Instance->FindFloatParameterNameIndex(ParamNameString + OPACITY_PARAMETER_POSTFIX);
+		check(FloatParameterIndex >= 0) // Instance Parameter does not exist.
+	}
+}
+
+
+bool FMultilayerProjector::AreInstanceParametersValid(const UCustomizableObjectInstance& Instance, const FName& ParamName)
+{
+	const FString ParamNameString = ParamName.ToString();
+
+	// Num layers.
+	{
+		const FString NumLayersParamName = ParamNameString + NUM_LAYERS_PARAMETER_POSTFIX;
+		const int32 FloatParameterIndex = Instance.FindFloatParameterNameIndex(NumLayersParamName);
+		if (FloatParameterIndex < 0)
+		{
+			return false;
+		}
+	}
+    
+	// Projector.
+	{
+		const int32 ProjectorParameterIndex = Instance.FindProjectorParameterNameIndex(ParamNameString);
+		if (ProjectorParameterIndex < 0)
+		{
+			return false;
+		}
+	}
+	
+	// Selected Image.
+	{
+		const int32 IntParameterIndex = Instance.FindIntParameterNameIndex(ParamNameString + IMAGE_PARAMETER_POSTFIX);
+		if (IntParameterIndex < 0)
+		{
+			return false;
+		}
+	}
+	
+	// Opacity.
+	{
+		const int32 FloatParameterIndex = Instance.FindFloatParameterNameIndex(ParamNameString + OPACITY_PARAMETER_POSTFIX);
+		if (FloatParameterIndex < 0)
+		{
+			return false;
+		}
+	}
+
+	return true;
+}
+
diff --git a/Engine/Plugins/Experimental/Mutable/Source/CustomizableObject/Private/MutableMeshBufferUtils.cpp b/Engine/Plugins/Experimental/Mutable/Source/CustomizableObject/Private/MutableMeshBufferUtils.cpp
new file mode 100644
--- /dev/null
+++ b/Engine/Plugins/Experimental/Mutable/Source/CustomizableObject/Private/MutableMeshBufferUtils.cpp
@@ -0,0 +1,180 @@
+// Copyright Epic Games, Inc. All Rights Reserved.
+
+#include "MutableMeshBufferUtils.h"
+
+#include "Rendering/PositionVertexBuffer.h"
+#include "Rendering/StaticMeshVertexBuffer.h"
+
+
+void MutableMeshBufferUtils::SetupVertexPositionsBuffer(const int32& InCurrentVertexBuffer, mu::FMeshBufferSet& OutTargetVertexBuffers)
+{
+	using namespace mu;
+	const int32 ElementSize = sizeof(FPositionVertex);
+	constexpr int32 ChannelCount = 1;
+	const MESH_BUFFER_SEMANTIC Semantics[ChannelCount] = {MBS_POSITION};
+	const int32 SemanticIndices[ChannelCount] = {0};
+	const MESH_BUFFER_FORMAT Formats[ChannelCount] = {MBF_FLOAT32};
+	const int32 Components[ChannelCount] = {3};
+	const int32 Offsets[ChannelCount] =
+	{
+		STRUCT_OFFSET(FPositionVertex, Position)
+	};
+
+	OutTargetVertexBuffers.SetBuffer(InCurrentVertexBuffer, ElementSize, ChannelCount, Semantics, SemanticIndices,
+									  Formats, Components, Offsets);
+}
+
+
+void MutableMeshBufferUtils::SetupTangentBuffer(const int32& InCurrentVertexBuffer,
+                                                mu::FMeshBufferSet& OutTargetVertexBuffers)
+{
+	// \todo: support for high precision?
+	typedef TStaticMeshVertexTangentDatum<TStaticMeshVertexTangentTypeSelector<
+		EStaticMeshVertexTangentBasisType::Default>::TangentTypeT> TangentType;
+
+	using namespace mu;
+	const int32 ElementSize = sizeof(TangentType);
+	constexpr int32 ChannelCount = 2;
+	const MESH_BUFFER_SEMANTIC Semantics[ChannelCount] = {MBS_TANGENT, MBS_NORMAL};
+	const int32 SemanticIndices[ChannelCount] = {0, 0};
+	const MESH_BUFFER_FORMAT Formats[ChannelCount] = {MBF_PACKEDDIRS8, MBF_PACKEDDIRS8_W_TANGENTSIGN};
+	const int32 Components[ChannelCount] = {4, 4};
+	const int32 Offsets[ChannelCount] =
+	{
+		STRUCT_OFFSET(TangentType, TangentX),
+		STRUCT_OFFSET(TangentType, TangentZ)
+	};
+
+	OutTargetVertexBuffers.SetBuffer(InCurrentVertexBuffer, ElementSize, ChannelCount, Semantics, SemanticIndices,
+	                                  Formats, Components, Offsets);
+}
+
+
+void MutableMeshBufferUtils::SetupTexCoordinatesBuffer(const int32& InCurrentVertexBuffer, const int32& InChannelCount,
+                                                       mu::FMeshBufferSet& OutTargetVertexBuffers,const int32* InTextureSemanticIndicesOverride /* = nullptr */)
+{
+	// \todo: support for half precision?
+	typedef TStaticMeshVertexUVsDatum<TStaticMeshVertexUVsTypeSelector<
+		EStaticMeshVertexUVType::HighPrecision>::UVsTypeT> UVType;
+
+	using namespace mu;
+	const int32 ElementSize = sizeof(UVType) * InChannelCount;
+	constexpr int32 MaxChannelCount = MaxTexCordChannelCount;
+	const MESH_BUFFER_SEMANTIC Semantics[MaxChannelCount] = {
+		MBS_TEXCOORDS, MBS_TEXCOORDS, MBS_TEXCOORDS, MBS_TEXCOORDS
+	};
+	const MESH_BUFFER_FORMAT Formats[MaxChannelCount] = {MBF_FLOAT32, MBF_FLOAT32, MBF_FLOAT32, MBF_FLOAT32};
+	const int32 Components[MaxChannelCount] = {2, 2, 2, 2};
+	const int32 Offsets[MaxChannelCount] =
+	{
+		STRUCT_OFFSET(UVType, UVs) + 0 * sizeof(UVType),
+		STRUCT_OFFSET(UVType, UVs) + 1 * sizeof(UVType),
+		STRUCT_OFFSET(UVType, UVs) + 2 * sizeof(UVType),
+		STRUCT_OFFSET(UVType, UVs) + 3 * sizeof(UVType),
+	};
+
+	if (InTextureSemanticIndicesOverride)
+	{
+		OutTargetVertexBuffers.SetBuffer(InCurrentVertexBuffer, ElementSize, InChannelCount, Semantics,
+		                                  InTextureSemanticIndicesOverride, Formats, Components, Offsets);
+	}
+	else
+	{
+		const int32 SemanticIndices[MaxChannelCount] = {0, 1, 2, 3};
+		OutTargetVertexBuffers.SetBuffer(InCurrentVertexBuffer, ElementSize, InChannelCount, Semantics,
+		                                  SemanticIndices, Formats, Components, Offsets);
+	}
+}
+
+
+void MutableMeshBufferUtils::SetupSkinBuffer(const int32& InCurrentVertexBuffer, const int32& MaxBoneIndexTypeSizeBytes,
+                                             const int32& MaxNumBonesPerVertex,
+                                             mu::FMeshBufferSet& OutTargetVertexBuffers)
+{
+	using namespace mu;
+	const int32 meshWeightTypeSizeBytes = 1;
+	const int32 ElementSize = (meshWeightTypeSizeBytes + MaxBoneIndexTypeSizeBytes) * MaxNumBonesPerVertex;
+	constexpr int32 ChannelCount = 2;
+	const MESH_BUFFER_SEMANTIC Semantics[ChannelCount] = {MBS_BONEINDICES, MBS_BONEWEIGHTS};
+	const int32 SemanticIndices[ChannelCount] = {0, 0};
+
+	MESH_BUFFER_FORMAT Formats[ChannelCount] = {MBF_UINT8, MBF_NUINT8};
+	switch (MaxBoneIndexTypeSizeBytes)
+	{
+	case 1: Formats[0] = mu::MBF_UINT8;
+		break;
+	case 2: Formats[0] = mu::MBF_UINT16;
+		break;
+	case 4: Formats[0] = mu::MBF_UINT32;
+		break;
+	default:
+		// unsupported bone index type
+		check(false);
+		Formats[0] = mu::MBF_NONE;
+		break;
+	}
+	switch (meshWeightTypeSizeBytes)
+	{
+	case 1: Formats[1] = mu::MBF_NUINT8;
+		break;
+	case 2: Formats[1] = mu::MBF_NUINT16;
+		break;
+	case 4: Formats[1] = mu::MBF_NUINT32;
+		break;
+	default:
+		// unsupported bone weight type
+		check(false);
+		Formats[1] = mu::MBF_NONE;
+		break;
+	}
+
+	int32 Components[ChannelCount];
+	Components[0] = Components[1] = MaxNumBonesPerVertex;
+
+	int32 Offsets[ChannelCount];
+	Offsets[0] = 0;
+	Offsets[1] = MaxBoneIndexTypeSizeBytes * MaxNumBonesPerVertex;
+
+	OutTargetVertexBuffers.SetBuffer(InCurrentVertexBuffer, ElementSize, ChannelCount, Semantics, SemanticIndices,
+	                                  Formats, Components, Offsets);
+}
+
+
+void MutableMeshBufferUtils::SetupVertexColorBuffer(const int32& InCurrentVertexBuffer,
+                                                    mu::FMeshBufferSet& OutTargetVertexBuffers)
+{
+	using namespace mu;
+	const int32 ElementSize = sizeof(FColor);
+	constexpr int32 ChannelCount = 1;
+	const MESH_BUFFER_SEMANTIC Semantics[ChannelCount] = {MBS_COLOUR};
+	const int32 SemanticIndices[ChannelCount] = {0};
+	const MESH_BUFFER_FORMAT Formats[ChannelCount] = {MBF_NUINT8};
+	const int32 Components[ChannelCount] = {4};
+	const int32 Offsets[ChannelCount] = {0};
+	check(ElementSize == 4);
+
+	OutTargetVertexBuffers.SetBuffer(InCurrentVertexBuffer, ElementSize, ChannelCount, Semantics, SemanticIndices,
+	                                  Formats, Components, Offsets);
+}
+
+
+void MutableMeshBufferUtils::SetupIndexBuffer(mu::FMeshBufferSet& OutTargetIndexBuffers)
+{
+	OutTargetIndexBuffers.SetBufferCount(1);
+
+	using namespace mu;
+	const int32 ElementSize = sizeof(uint32);
+	//SkeletalMesh->GetImportedResource()->LODModels[LOD].MultiSizeIndexContainer.GetDataTypeSize();
+	constexpr int32 ChannelCount = 1;
+	const MESH_BUFFER_SEMANTIC Semantics[ChannelCount] = {MBS_VERTEXINDEX};
+	const int32 SemanticIndices[ChannelCount] = {0};
+	// We force 32 bit indices, since merging meshes may create vertex buffers bigger than the initial mesh
+	// and for now the mutable runtime doesn't handle it.
+	// \TODO: go back to 16-bit indices when possible.
+	MESH_BUFFER_FORMAT Formats[ChannelCount] = {MBF_UINT32};
+	const int32 Components[ChannelCount] = {1};
+	const int32 Offsets[ChannelCount] = {0};
+
+	OutTargetIndexBuffers.SetBuffer(0, ElementSize, ChannelCount, Semantics, SemanticIndices, Formats, Components,
+	                                 Offsets);
+}
\ No newline at end of file
diff --git a/Engine/Plugins/Experimental/Mutable/Source/CustomizableObject/Private/UnrealConversionUtils.cpp b/Engine/Plugins/Experimental/Mutable/Source/CustomizableObject/Private/UnrealConversionUtils.cpp
new file mode 100644
--- /dev/null
+++ b/Engine/Plugins/Experimental/Mutable/Source/CustomizableObject/Private/UnrealConversionUtils.cpp
@@ -0,0 +1,336 @@
+// Copyright Epic Games, Inc. All Rights Reserved.
+
+#include "UnrealConversionUtils.h"
+
+#include "Animation/Skeleton.h"
+#include "CustomizableObject/Private/CustomizableObjectSystemPrivate.h"
+#include "Materials/Material.h"
+#include "MutableRuntime/Private/MeshPrivate.h"
+#include "MutableRuntime/Private/OpMeshFormat.h"
+#include "MutableRuntime/Public/Mesh.h"
+#include "MutableRuntime/Public/MutableTrace.h"
+#include "MutableRuntime/Public/Parameters.h"
+#include "MutableRuntime/Public/Ptr.h"
+
+namespace UnrealConversionUtils
+{
+	// Hidden functions only used internally to aid other functions
+	namespace
+	{
+		/**
+		 * Initializes the static mesh vertex buffers object provided with the data found on the mutable buffers
+		 * @param OutVertexBuffers - The Unreal's vertex buffers container to be updated with the mutable data.
+		 * @param NumVertices - The amount of vertices on the buffer
+		 * @param NumTexCoords - The amount of texture coordinates
+		 * @param bUseFullPrecisionUVs - Determines if we want to use or not full precision UVs
+		 * @param bNeedCPUAccess - Determines if CPU access is required
+		 * @param InMutablePositionData - Mutable position data buffer
+		 * @param InMutableTangentData - Mutable tangent data buffer
+		 * @param InMutableTextureData - Mutable texture data buffer
+		 */
+		void FStaticMeshVertexBuffers_InitWithMutableData(
+		FStaticMeshVertexBuffers& OutVertexBuffers,
+		const int32 NumVertices,
+		const int32 NumTexCoords,
+		const bool bUseFullPrecisionUVs,
+		const bool bNeedCPUAccess,
+		const void* InMutablePositionData,
+		const void* InMutableTangentData,
+		const void* InMutableTextureData)
+{
+	// positions
+	OutVertexBuffers.PositionVertexBuffer.Init(NumVertices, bNeedCPUAccess);
+	FMemory::Memcpy(OutVertexBuffers.PositionVertexBuffer.GetVertexData(), InMutablePositionData, NumVertices * OutVertexBuffers.PositionVertexBuffer.GetStride());
+
+	// tangent and texture coords
+	OutVertexBuffers.StaticMeshVertexBuffer.SetUseFullPrecisionUVs(true);
+	OutVertexBuffers.StaticMeshVertexBuffer.SetUseHighPrecisionTangentBasis(false);
+	OutVertexBuffers.StaticMeshVertexBuffer.Init(NumVertices, NumTexCoords, bNeedCPUAccess);
+	FMemory::Memcpy(OutVertexBuffers.StaticMeshVertexBuffer.GetTangentData(), InMutableTangentData, OutVertexBuffers.StaticMeshVertexBuffer.GetTangentSize());
+	FMemory::Memcpy(OutVertexBuffers.StaticMeshVertexBuffer.GetTexCoordData(), InMutableTextureData, OutVertexBuffers.StaticMeshVertexBuffer.GetTexCoordSize());
+}
+
+
+		/**
+		 * Initializes the color vertex buffers object provided with the data found on the mutable buffers
+		 * @param OutVertexBuffers - The Unreal's vertex buffers container to be updated with the mutable data.
+		 * @param NumVertices - The amount of vertices on the buffer
+		 * @param InMutableColorData - Mutable color data buffer
+		 */
+		void FColorVertexBuffers_InitWithMutableData(
+	FStaticMeshVertexBuffers& OutVertexBuffers,
+	const int32 NumVertices,
+	const void* InMutableColorData
+)
+{
+	// positions
+	OutVertexBuffers.ColorVertexBuffer.Init(NumVertices);
+	FMemory::Memcpy(OutVertexBuffers.ColorVertexBuffer.GetVertexData(), InMutableColorData, NumVertices * OutVertexBuffers.ColorVertexBuffer.GetStride());
+}
+
+
+		/**
+		 * Initializes the skin vertex buffers object provided with the data found on the mutable buffers
+		 * @param OutVertexWeightBuffer - The Unreal's vertex buffers container to be updated with the mutable data.
+		   * @param NumVertices - The amount of vertices on the buffer
+		 * @param NumBones - The amount of bones to use to init the skin weights buffer
+		   * @param NumBoneInfluences - The amount of bone influences on the buffer
+		 * @param bNeedCPUAccess - Determines if CPU access is required
+		 * @param InMutableData - Mutable data buffer
+		 */
+		void FSkinWeightVertexBuffer_InitWithMutableData(
+	FSkinWeightVertexBuffer& OutVertexWeightBuffer,
+	const int32 NumVertices,
+	const int32 NumBones,
+	const int32 NumBoneInfluences,
+	const bool bNeedCPUAccess,
+	const void* InMutableData)
+{
+	// \todo Ugly cast.
+	const_cast<FSkinWeightDataVertexBuffer*>(OutVertexWeightBuffer.GetDataVertexBuffer())->SetMaxBoneInfluences(NumBoneInfluences);
+	const_cast<FSkinWeightDataVertexBuffer*>(OutVertexWeightBuffer.GetDataVertexBuffer())->Init(NumBones, NumVertices);
+
+	if (NumVertices)
+	{
+		OutVertexWeightBuffer.SetNeedsCPUAccess(bNeedCPUAccess);
+
+		FSkinWeightInfo* Data = OutVertexWeightBuffer.GetDataVertexBuffer()->GetWeightData();
+		FMemory::Memcpy(Data, InMutableData, OutVertexWeightBuffer.GetVertexDataSize());
+	}
+}
+	}
+
+
+	
+	void BuildRefSkeleton(FInstanceUpdateData::FSkeletonData* OutMutSkeletonData
+	,const FReferenceSkeleton& InSourceReferenceSkeleton,const TArray<bool>& InUsedBones,
+	FReferenceSkeleton& InRefSkeleton, const USkeleton* InSkeleton )
+{
+	const int32 SourceBoneCount = InSourceReferenceSkeleton.GetNum();
+	const TArray<FTransform>& SourceRawMeshBonePose = InSourceReferenceSkeleton.GetRawRefBonePose();
+	
+	MUTABLE_CPUPROFILER_SCOPE(BuildSkeletonData_BuildRefSkeleton);
+
+	// Build new RefSkeleton	
+	FReferenceSkeletonModifier RefSkeletonModifier(InRefSkeleton, InSkeleton);
+
+	const TArray<FMeshBoneInfo>& BoneInfo = InSourceReferenceSkeleton.GetRawRefBoneInfo();
+
+	TMap<FName, uint16> BoneToFinalBoneIndexMap;
+	BoneToFinalBoneIndexMap.Reserve(SourceBoneCount);
+	
+	uint32 FinalBoneCount = 0;
+	for (int32 BoneIndex = 0; BoneIndex < SourceBoneCount; ++BoneIndex)
+	{
+		if (!InUsedBones[BoneIndex])
+		{
+			continue;
+		}
+
+		FName BoneName = BoneInfo[BoneIndex].Name;
+
+		// Build a bone to index map so we can remap BoneMaps and ActiveBoneIndices later on
+		BoneToFinalBoneIndexMap.Add(BoneName, FinalBoneCount);
+		FinalBoneCount++;
+
+		// Find parent index
+		const int32 SourceParentIndex = BoneInfo[BoneIndex].ParentIndex;
+		const int32 ParentIndex = SourceParentIndex != INDEX_NONE ? BoneToFinalBoneIndexMap[BoneInfo[SourceParentIndex].Name] : INDEX_NONE;
+
+		RefSkeletonModifier.Add(FMeshBoneInfo(BoneName, BoneName.ToString(), ParentIndex), SourceRawMeshBonePose[BoneIndex]);
+	}
+}
+
+
+	void BuildSkeletalMeshElementDataAtLOD(const int32 MeshLODIndex,
+	mu::MeshPtrConst InMutableMesh, USkeletalMesh* OutSkeletalMesh)
+{
+	int32 SurfaceCount = 0;
+	if (InMutableMesh)
+	{
+		SurfaceCount = InMutableMesh->GetSurfaceCount();
+	}
+	
+	Helper_GetLODInfoArray(OutSkeletalMesh)[MeshLODIndex].LODMaterialMap.SetNum(1);
+	Helper_GetLODInfoArray(OutSkeletalMesh)[MeshLODIndex].LODMaterialMap[0] = 0;
+
+	const int32 MaterialCount = SurfaceCount;
+	int32 NewMaterialIndex = Helper_GetLODRenderSections(OutSkeletalMesh, MeshLODIndex).Num();
+	for (int32 MaterialIndex = 0; MaterialIndex < MaterialCount; MaterialIndex++, NewMaterialIndex++)
+	{
+		new(Helper_GetLODRenderSections(OutSkeletalMesh, MeshLODIndex)) Helper_SkelMeshRenderSection();
+		FSkelMeshRenderSection& Section = Helper_GetLODRenderSections(OutSkeletalMesh, MeshLODIndex)[NewMaterialIndex];
+		Section.MaterialIndex = 0;
+		Section.BaseIndex = 0;
+		Section.NumTriangles = 0;
+	}
+}
+
+
+	void SetupRenderSections(
+	const mu::MeshPtrConst InMutableMesh,
+	const USkeletalMesh* OutSkeletalMesh,
+	const int32 MeshLODIndex,
+	const int32 NumBoneInfluences,
+	const TArray<uint16>& InBoneMap)
+{
+	const int32 SurfaceCount = InMutableMesh->GetSurfaceCount();
+	for (int32 Surface = 0; Surface < SurfaceCount; Surface++)
+	{
+		MUTABLE_CPUPROFILER_SCOPE(UpdateSkeletalMesh_SurfaceLoop);
+
+		int32 FirstIndex;
+		int32 IndexCount;
+		int32 FirstVertex;
+		int32 VertexCount;
+		InMutableMesh->GetSurface(Surface, &FirstVertex, &VertexCount, &FirstIndex, &IndexCount);
+		FSkelMeshRenderSection& Section = Helper_GetLODRenderSections(OutSkeletalMesh, MeshLODIndex)[Surface];
+
+#if (ENGINE_MAJOR_VERSION == 4 && ENGINE_MINOR_VERSION >= 19) || ENGINE_MAJOR_VERSION >= 5
+		// \todo: investigate. Without this 4.19 rendering crashes. 
+		// It must be initialized even in disabled sections or else they will crash
+		Section.DuplicatedVerticesBuffer.Init(1, TMap<int, TArray<int32>>());
+#endif
+			
+		if (VertexCount == 0 || IndexCount == 0)
+		{
+			Section.bDisabled = true;
+			continue; // Unreal doesn't like empty meshes
+		}
+			
+		Section.BaseIndex = FirstIndex;
+		Section.NumTriangles = IndexCount / 3;
+		Section.BaseVertexIndex = FirstVertex;
+		Section.MaxBoneInfluences = NumBoneInfluences;
+		Section.NumVertices = VertexCount;
+
+		//Section.BoneMap.Append(OperationData->InstanceUpdateData.LODs[LODIndex].Components[ComponentIndex].BoneMap);
+		Section.BoneMap.Append(InBoneMap);
+	}
+}
+
+
+	void CopyMutableVertexBuffers(
+	USkeletalMesh* OutSkeletalMesh,
+	const int32 NumVerticesLODModel, const int32 NumBoneInfluences,
+	const int32 BoneIndexBuffer,
+	const mu::FMeshBufferSet& MutableMeshVertexBuffers,
+	const int32 MeshLODIndex,
+	const mu::MESH_BUFFER_FORMAT& InBoneIndexFormat
+	)
+{
+	FSkeletalMeshLODRenderData& LODModel = Helper_GetLODData(OutSkeletalMesh)[MeshLODIndex];
+	
+	MUTABLE_CPUPROFILER_SCOPE(UpdateSkeletalMesh_SurfaceLoop_MemCpy);
+
+	uint32 BuildFlags = OutSkeletalMesh->GetVertexBufferFlags();
+	const bool bUseFullPrecisionUVs = (BuildFlags & ESkeletalMeshVertexFlags::UseFullPrecisionUVs) != 0;
+	const bool bHasVertexColors = (BuildFlags & ESkeletalMeshVertexFlags::HasVertexColors) != 0;
+	const int NumTexCoords = MutableMeshVertexBuffers.GetBufferChannelCount(MUTABLE_VERTEXBUFFER_TEXCOORDS);
+
+	const bool bNeedsCPUAccess = Helper_GetLODInfoArray(OutSkeletalMesh)[MeshLODIndex].bAllowCPUAccess;
+	
+	FStaticMeshVertexBuffers_InitWithMutableData(
+		LODModel.StaticVertexBuffers,
+		NumVerticesLODModel, 
+		NumTexCoords, 
+		bUseFullPrecisionUVs, 
+		bNeedsCPUAccess,
+		MutableMeshVertexBuffers.GetBufferData(MUTABLE_VERTEXBUFFER_POSITION),
+		MutableMeshVertexBuffers.GetBufferData(MUTABLE_VERTEXBUFFER_TANGENT),
+		MutableMeshVertexBuffers.GetBufferData(MUTABLE_VERTEXBUFFER_TEXCOORDS)
+		);
+
+	// Init skin weight buffer
+	FSkinWeightVertexBuffer_InitWithMutableData(
+		LODModel.SkinWeightVertexBuffer,
+		NumVerticesLODModel,
+		NumBoneInfluences * NumVerticesLODModel,
+		NumBoneInfluences,
+		bNeedsCPUAccess,
+		MutableMeshVertexBuffers.GetBufferData(BoneIndexBuffer)
+	);
+
+	if (InBoneIndexFormat == mu::MBF_UINT16)
+	{
+		LODModel.SkinWeightVertexBuffer.SetUse16BitBoneIndex(true);
+	}
+
+	// Optional buffers
+	for ( int Buffer= MUTABLE_VERTEXBUFFER_TEXCOORDS + 1; Buffer<MutableMeshVertexBuffers.GetBufferCount(); ++Buffer)
+	{
+		if (MutableMeshVertexBuffers.GetBufferChannelCount(Buffer) > 0)
+		{
+			mu::MESH_BUFFER_SEMANTIC Semantic;
+			mu::MESH_BUFFER_FORMAT Format;
+			int32 SemanticIndex;
+			int32 ComponentCount;
+			int32 Offset;
+			MutableMeshVertexBuffers.GetChannel(Buffer, 0, &Semantic, &SemanticIndex, &Format, &ComponentCount, &Offset);
+			
+			if (Semantic == mu::MBS_BONEINDICES)
+			{
+				const int32 BonesPerVertex = ComponentCount;
+				const int32 NumBones = BonesPerVertex * NumVerticesLODModel;
+
+				check(LODModel.SkinWeightVertexBuffer.GetVariableBonesPerVertex() == false);
+				FSkinWeightVertexBuffer_InitWithMutableData(
+					LODModel.SkinWeightVertexBuffer,
+					NumVerticesLODModel,
+					NumBones, 
+					NumBoneInfluences,
+					bNeedsCPUAccess,
+					MutableMeshVertexBuffers.GetBufferData(Buffer)
+				);
+			}
+
+			// colour buffer?
+			else if (Semantic == mu::MBS_COLOUR)
+			{
+				OutSkeletalMesh->SetHasVertexColors(true);
+				const void* DataPtr = MutableMeshVertexBuffers.GetBufferData(Buffer);
+				FColorVertexBuffers_InitWithMutableData(LODModel.StaticVertexBuffers, NumVerticesLODModel, DataPtr);
+				check(LODModel.StaticVertexBuffers.ColorVertexBuffer.GetStride() == MutableMeshVertexBuffers.GetElementSize(Buffer));
+			}
+		}
+	}
+}
+
+
+	bool CopyMutableIndexBuffers(mu::MeshPtrConst InMutableMesh,FSkeletalMeshLODRenderData& LODModel)
+{
+	MUTABLE_CPUPROFILER_SCOPE(UpdateSkeletalMesh_BuildSkeletalMeshRenderData_IndexLoop);
+
+	const int32 IndexCount = InMutableMesh->GetIndexBuffers().GetElementCount();
+
+	if (IndexCount)
+	{
+		if (InMutableMesh->GetIndexBuffers().GetElementSize(0) == 2)
+		{
+			const uint16* IndexPtr = (const uint16*)InMutableMesh->GetIndexBuffers().GetBufferData(0);
+			LODModel.MultiSizeIndexContainer.CreateIndexBuffer(2);
+			LODModel.MultiSizeIndexContainer.GetIndexBuffer()->Insert(0, IndexCount);
+			FMemory::Memcpy(LODModel.MultiSizeIndexContainer.GetIndexBuffer()->GetPointerTo(0), IndexPtr, IndexCount * sizeof(uint16));
+		}
+		else
+		{
+			const uint32* IndexPtr = (const uint32*)InMutableMesh->GetIndexBuffers().GetBufferData(0);
+			LODModel.MultiSizeIndexContainer.CreateIndexBuffer(4);
+			LODModel.MultiSizeIndexContainer.GetIndexBuffer()->Insert(0, IndexCount);
+			FMemory::Memcpy(LODModel.MultiSizeIndexContainer.GetIndexBuffer()->GetPointerTo(0), IndexPtr, IndexCount * sizeof(uint32));
+		}
+
+	}
+	else
+	{
+		UE_LOG(LogMutable, Error, TEXT("UCustomizableInstancePrivateData::BuildSkeletalMeshRenderData is converting an empty mesh."));
+		return false;
+	}
+
+	return true;
+}
+
+		}
+
+	
+		
