[mutable] Move Mutable plugin from the restricted folder to the experimental plugins folder.

diff --git a/Engine/Plugins/Experimental/Mutable/Source/CustomizableObject/Private/UnrealMutableImageProvider.cpp b/Engine/Plugins/Experimental/Mutable/Source/CustomizableObject/Private/UnrealMutableImageProvider.cpp
new file mode 100644
--- /dev/null
+++ b/Engine/Plugins/Experimental/Mutable/Source/CustomizableObject/Private/UnrealMutableImageProvider.cpp
@@ -0,0 +1,144 @@
+// Copyright Epic Games, Inc. All Rights Reserved.
+
+#include "UnrealMutableImageProvider.h"
+
+#include "CustomizableObject.h"
+#include "CustomizableObjectSystem.h"
+
+
+//-------------------------------------------------------------------------------------------------
+namespace
+{
+
+	mu::ImagePtr ConvertTextureUnrealToMutable(UTexture2D* Texture)
+	{
+		mu::ImagePtr pResult;
+
+#if WITH_EDITOR
+		int LODs = 1;
+		int SizeX = Texture->Source.GetSizeX();
+		int SizeY = Texture->Source.GetSizeY();
+		ETextureSourceFormat Format = Texture->Source.GetFormat();
+		mu::EImageFormat MutableFormat = mu::EImageFormat::IF_NONE;
+
+		switch (Format)
+		{
+		case ETextureSourceFormat::TSF_BGRA8: MutableFormat = mu::EImageFormat::IF_BGRA_UBYTE; break;
+		// This format is deprecated ans using the enum fails to compile in some cases.
+		//case ETextureSourceFormat::TSF_RGBA8: MutableFormat = mu::EImageFormat::IF_RGBA_UBYTE; break;
+		case ETextureSourceFormat::TSF_G8: MutableFormat = mu::EImageFormat::IF_L_UBYTE; break;
+		default:
+			break;
+		}
+
+		pResult = new mu::Image(SizeX, SizeY, LODs, MutableFormat);
+
+		// If not locked ReadOnly the Texture Source's FGuid can change, invalidating the texture's caching/shaders
+		// making shader compile and cook times increase
+		const uint8* pSource = Texture->Source.LockMipReadOnly(0);
+		if (pSource)
+		{
+			FMemory::Memcpy(pResult->GetData(), pSource, pResult->GetDataSize());
+		}
+
+		Texture->Source.UnlockMip(0);
+#endif
+
+		return pResult;
+	}
+}
+
+
+//-------------------------------------------------------------------------------------------------
+
+mu::ImagePtr FUnrealMutableImageProvider::GetImage(mu::EXTERNAL_IMAGE_ID id)
+{
+	// Thread: Mutable worker
+
+	check(ExternalImagesForCurrentInstance.Contains(id));
+	return ExternalImagesForCurrentInstance[id];
+}
+
+
+void FUnrealMutableImageProvider::CacheImage(mu::EXTERNAL_IMAGE_ID id)
+{
+	check( IsInGameThread() );
+
+	mu::ImagePtr pResult;
+
+	// See if any providere provides this id.
+	for (int p = 0; !pResult && p<ImageProviders.Num(); ++p)
+	{
+		TWeakObjectPtr<UCustomizableSystemImageProvider> Provider = ImageProviders[p];
+
+		if (Provider.IsValid())
+		{
+			switch (Provider->HasTextureParameterValue(id))
+			{
+
+			case UCustomizableSystemImageProvider::ValueType::Raw:
+			{
+				FIntVector desc = Provider->GetTextureParameterValueSize(id);
+				pResult = new mu::Image(desc[0], desc[1], 1, mu::EImageFormat::IF_RGBA_UBYTE);
+				Provider->GetTextureParameterValueData(id, pResult->GetData());
+				break;
+			}
+
+			case UCustomizableSystemImageProvider::ValueType::Unreal:
+			{
+				UTexture2D* UnrealTexture = Provider->GetTextureParameterValue(id);
+				pResult = ConvertTextureUnrealToMutable(UnrealTexture);
+				break;
+			}
+
+			default:
+				break;
+			}
+		}
+	}
+
+	if (!pResult)
+	{
+		pResult = CreateDummy();
+	}
+
+	ExternalImagesForCurrentInstance.Add(id,pResult);
+}
+
+
+void FUnrealMutableImageProvider::ClearCache()
+{
+	check(IsInGameThread());
+
+	ExternalImagesForCurrentInstance.Empty();
+}
+
+
+mu::ImagePtr FUnrealMutableImageProvider::CreateDummy()
+{
+	// Create a dummy image
+	const int size = 32;
+	const int checkerSize = 4;
+	constexpr int checkerTileCount = 2;
+	uint8_t colours[checkerTileCount][4] = { { 255,255,0,255 },{ 0,0,255,255 } };
+
+	mu::ImagePtr pResult = new mu::Image(size, size, 1, mu::EImageFormat::IF_RGBA_UBYTE);
+
+	uint8_t* pData = pResult->GetData();
+	for (int x = 0; x < size; ++x)
+	{
+		for (int y = 0; y < size; ++y)
+		{
+			int checkerIndex = ((x / checkerSize) + (y / checkerSize)) % checkerTileCount;
+			pData[0] = colours[checkerIndex][0];
+			pData[1] = colours[checkerIndex][1];
+			pData[2] = colours[checkerIndex][2];
+			pData[3] = colours[checkerIndex][3];
+			pData += 4;
+		}
+	}
+
+	return pResult;
+}
+
+
diff --git a/Engine/Plugins/Experimental/Mutable/Source/CustomizableObject/Private/UnrealMutableImageProvider.h b/Engine/Plugins/Experimental/Mutable/Source/CustomizableObject/Private/UnrealMutableImageProvider.h
new file mode 100644
--- /dev/null
+++ b/Engine/Plugins/Experimental/Mutable/Source/CustomizableObject/Private/UnrealMutableImageProvider.h
@@ -0,0 +1,42 @@
+// Copyright Epic Games, Inc. All Rights Reserved.
+
+#pragma once
+
+#include "Containers/Array.h"
+#include "UObject/WeakObjectPtrTemplates.h"
+
+#include "MutableRuntime/Public/System.h"
+#include "MutableRuntime/Public/Image.h"
+
+
+/** Implementation of a mutable core provider for image parameters that are application-specific. */
+class FUnrealMutableImageProvider : public mu::ImageParameterGenerator
+{
+public:
+
+	// mu::ImageParameterGenerator interface
+	// Thread: Mutable
+	mu::ImagePtr GetImage(mu::EXTERNAL_IMAGE_ID id) override;
+
+	// Own interface
+	// Thread: Game
+	void CacheImage(mu::EXTERNAL_IMAGE_ID id);
+	void ClearCache();
+
+	/** List of actual image providers that have been registered to the CustomizableObjectSystem. */
+	TArray< TWeakObjectPtr<class UCustomizableSystemImageProvider> > ImageProviders;
+
+private:
+
+	/** This will be called if an image Id has been requested by Mutable core but it has not been provided by any provider. */
+	mu::ImagePtr CreateDummy();
+
+	/** List of external textures that may be required for the current instance under construction.
+	* This is only written from the game thread before scheduling a new instance update, and it
+	* is read from the mutable thread during the update.
+	*/
+	TMap<uint64, mu::ImagePtr> ExternalImagesForCurrentInstance;
+
+};
+
+
diff --git a/Engine/Plugins/Experimental/Mutable/Source/CustomizableObject/Private/UnrealMutableModelDiskStreamer.cpp b/Engine/Plugins/Experimental/Mutable/Source/CustomizableObject/Private/UnrealMutableModelDiskStreamer.cpp
new file mode 100644
--- /dev/null
+++ b/Engine/Plugins/Experimental/Mutable/Source/CustomizableObject/Private/UnrealMutableModelDiskStreamer.cpp
@@ -0,0 +1,363 @@
+// Copyright Epic Games, Inc. All Rights Reserved.
+
+#include "UnrealMutableModelDiskStreamer.h"
+
+#include "CustomizableObject.h"
+#include "CustomizableObjectSystem.h"
+
+#include "Serialization/MemoryReader.h"
+#include "Serialization/MemoryWriter.h"
+#include "Serialization/BulkData.h"
+
+
+#if WITH_EDITOR
+#include "Misc/Paths.h"
+#include "HAL/FileManager.h"
+#include "HAL/PlatformFileManager.h"
+#endif
+
+#include "MutableTrace.h"
+#include "Runtime/Launch/Resources/Version.h"
+
+
+//-------------------------------------------------------------------------------------------------
+#if WITH_EDITOR
+
+UnrealMutableOutputStream::UnrealMutableOutputStream(FArchive& ar)
+	: m_ar(ar)
+{
+}
+
+
+void UnrealMutableOutputStream::Write(const void* pData, uint64 size)
+{
+	// TODO: ugly const cast
+	m_ar.Serialize(const_cast<void*>(pData), size);
+}
+
+#endif
+
+//-------------------------------------------------------------------------------------------------
+UnrealMutableInputStream::UnrealMutableInputStream(FArchive& ar)
+	: m_ar(ar)
+{
+}
+
+
+void UnrealMutableInputStream::Read(void* pData, uint64 size)
+{
+	m_ar.Serialize(pData, size);
+}
+
+
+//-------------------------------------------------------------------------------------------------
+FUnrealMutableModelBulkStreamer::FUnrealMutableModelBulkStreamer(FArchive* InMainDataArchive, FArchive* InStreamedDataArchive)
+{
+#if WITH_EDITOR
+	MainDataArchive = InMainDataArchive;
+	StreamedDataArchive = InStreamedDataArchive;
+	CurrentWriteFile = nullptr;
+#endif
+}
+
+FUnrealMutableModelBulkStreamer::~FUnrealMutableModelBulkStreamer()
+{
+	EndStreaming();
+}
+
+
+void FUnrealMutableModelBulkStreamer::OpenWriteFile(const char* strModelName, uint64 key0)
+{
+#if WITH_EDITOR
+	if (key0 == 0) // Model
+	{
+		check(MainDataArchive);
+		CurrentWriteFile = MainDataArchive;
+	}
+	else
+	{
+		check(StreamedDataArchive);
+		CurrentWriteFile = StreamedDataArchive;
+	}
+#endif
+}
+
+
+void FUnrealMutableModelBulkStreamer::Write(const void* pBuffer, uint64 size)
+{
+#if WITH_EDITOR
+	check(CurrentWriteFile);
+	CurrentWriteFile->Serialize(const_cast<void*>(pBuffer), size);
+#endif
+}
+
+
+void FUnrealMutableModelBulkStreamer::CloseWriteFile()
+{
+#if WITH_EDITOR
+	CurrentWriteFile = nullptr;
+#endif
+}
+
+
+bool FUnrealMutableModelBulkStreamer::PrepareStreamingForObject(UCustomizableObject* CustomizableObject)
+{
+	// This happens in the game thread
+	check(IsInGameThread());
+
+	if (!CustomizableObject)
+	{
+		check(false);
+		return false;
+	}
+
+	// See if we can free previuously allocated resources
+	for (int32 ObjectIndex = 0; ObjectIndex < Objects.Num(); )
+	{
+		if (!Objects[ObjectIndex].Model.Pin())
+		{
+			// The CustomizableObject is gone, so we won't be streaming for it anymore.
+			for (TPair<OPERATION_ID, IAsyncReadRequest*>& it : Objects[ObjectIndex].CurrentReadRequests)
+			{
+				it.Value->WaitCompletion();
+				delete it.Value;
+			}
+
+			for (IAsyncReadFileHandle* ReadFileHandle : Objects[ObjectIndex].ReadFileHandles)
+			{
+				delete ReadFileHandle;
+			}
+
+			Objects.RemoveAtSwap(ObjectIndex);
+		}
+		else
+		{
+			++ObjectIndex;
+		}
+	}
+
+	// Is the object already prepared for streaming?
+	bool bAlreadyStreaming = Objects.FindByPredicate(
+		[CustomizableObject](const FObjectData& d) 
+		{ return d.Model.Pin() == CustomizableObject->GetModel(); })
+		!= 
+		nullptr;
+
+	if (!bAlreadyStreaming)
+	{
+		FObjectData NewData;
+		NewData.Model = mu::WeakPtr<mu::Model>(CustomizableObject->GetModel());
+
+#if WITH_EDITOR
+		FString FolderPath = CustomizableObject->GetCompiledDataFolderPath(true);
+		FString FullFileName = FolderPath + CustomizableObject->GetCompiledDataFileName(false);
+
+		IAsyncReadFileHandle* ReadFileHandle = FPlatformFileManager::Get().GetPlatformFile().OpenAsyncRead(*FullFileName);
+		if (!ReadFileHandle)
+		{
+			UE_LOG(LogMutable, Warning, TEXT("Streaming: Customizable Object %s is missing the Editor BulkData."), *CustomizableObject->GetName());
+			return false;
+		}
+
+		NewData.ReadFileHandles.Add(ReadFileHandle);
+#else
+		const UCustomizableObjectBulk* BulkData = CustomizableObject->GetStreamableBulkData();
+		if (!BulkData)
+		{
+			UE_LOG(LogMutable, Warning, TEXT("Streaming: Customizable Object %s is missing the BulkData export."), *CustomizableObject->GetName());
+			return false;
+		}
+
+		NewData.ReadFileHandles = BulkData->GetAsyncReadFileHandles();
+#endif
+
+
+		if (NewData.ReadFileHandles.IsEmpty())
+		{
+			UE_LOG(LogMutable, Warning, TEXT("Streaming: Customizable Object %s read file handles empty."), *CustomizableObject->GetName());
+
+			check(false);
+			return false;
+		}
+
+		NewData.StreamableBlocks = CustomizableObject->HashToStreamableBlock;
+		if (NewData.StreamableBlocks.IsEmpty())
+		{
+			UE_LOG(LogMutable, Warning, TEXT("Streaming: Customizable Object %s has no data to stream."), *CustomizableObject->GetName());
+
+#if !WITH_EDITOR
+			check(false);
+#endif
+			return false;
+		}
+
+		Objects.Add(MoveTemp(NewData));
+	}
+
+	return true;
+}
+
+
+#if WITH_EDITOR
+void FUnrealMutableModelBulkStreamer::CancelStreamingForObject(const UCustomizableObject* CustomizableObject)
+{
+	// This happens in the game thread
+	check(IsInGameThread());
+
+	if (!CustomizableObject)
+	{
+		check(false);
+	}
+
+	// See if we can free previuously allocated resources
+	for (int32 ObjectIndex = 0; ObjectIndex < Objects.Num(); ++ObjectIndex)
+	{
+		if (Objects[ObjectIndex].Model.Pin()==CustomizableObject->GetModel())
+		{
+			for (TPair<OPERATION_ID, IAsyncReadRequest*>& it : Objects[ObjectIndex].CurrentReadRequests)
+			{
+				it.Value->WaitCompletion();
+				delete it.Value;
+			}
+
+			for (IAsyncReadFileHandle* ReadFileHandle : Objects[ObjectIndex].ReadFileHandles)
+			{
+				delete ReadFileHandle;
+			}
+
+			Objects.RemoveAtSwap(ObjectIndex);
+			break;
+		}
+	}
+}
+
+#endif // WITH_EDITOR
+
+
+void FUnrealMutableModelBulkStreamer::EndStreaming()
+{
+	for (FObjectData& o : Objects)
+	{
+		for (TPair<OPERATION_ID, IAsyncReadRequest*>& it : o.CurrentReadRequests)
+		{
+			it.Value->WaitCompletion();
+			delete it.Value;
+		}
+
+		for (IAsyncReadFileHandle* ReadFileHandle : o.ReadFileHandles)
+		{
+			delete ReadFileHandle;
+		}
+	}
+	Objects.Empty();
+}
+
+
+static int32 StreamPriority = 4;
+FAutoConsoleVariableRef CVarStreamPriority(
+	TEXT("Mutable.StreamPriority"),
+	StreamPriority,
+	TEXT(""));
+
+
+mu::ModelStreamer::OPERATION_ID FUnrealMutableModelBulkStreamer::BeginReadBlock(const mu::Model* Model, uint64 Key, void* pBuffer, uint64 size)
+{
+	MUTABLE_CPUPROFILER_SCOPE(FUnrealMutableModelBulkStreamer::OpenReadFile);
+
+	UE_LOG(LogMutable, Verbose, TEXT("Streaming: reading data %08lx."), Key);
+
+	// Find the object we are streaming for
+	FObjectData* ObjectData = Objects.FindByPredicate(
+		[Model](const FObjectData& d) {return d.Model.Pin()==Model; });
+
+	if (!ObjectData)
+	{
+		// The object has been unloaded. Streaming is not possible. 
+		// This may happen in the editor if we are recompiling an object but we still have instances from the old
+		// object that have progressive mip generation.
+		return -1;
+	}
+
+	mu::ModelStreamer::OPERATION_ID Result = 0;
+
+	check(!ObjectData->ReadFileHandles.IsEmpty());
+
+	// this generally cannot fail because it is async
+	if (ObjectData->StreamableBlocks.Contains(Key))
+	{
+		Result = ++LastOperationID;
+
+		const FMutableStreamableBlock& Block = ObjectData->StreamableBlocks[Key];
+
+		int32 BulkDataOffsetInFile = 0;
+#if WITH_EDITOR
+		BulkDataOffsetInFile = sizeof(MutableCompiledDataStreamHeader);
+#endif
+
+		IAsyncReadRequest* ReadRequest = ObjectData->ReadFileHandles[Block.FileIndex]->ReadRequest(BulkDataOffsetInFile + Block.Offset, size, (EAsyncIOPriorityAndFlags)StreamPriority, nullptr, reinterpret_cast<uint8*>(pBuffer));
+		ObjectData->CurrentReadRequests.Add(Result, ReadRequest);
+	}
+	else
+	{
+		// File Handle not found! This shouldn't really happen.
+		UE_LOG(LogMutable, Error, TEXT("Streaming Block not found!"));
+		check(false);
+	}
+
+	INC_DWORD_STAT(STAT_MutableStreamingOps);
+
+	return Result;
+}
+
+
+bool FUnrealMutableModelBulkStreamer::IsReadCompleted(mu::ModelStreamer::OPERATION_ID OperationId)
+{
+	MUTABLE_CPUPROFILER_SCOPE(FUnrealMutableModelBulkStreamer::IsReadCompleted);
+
+	for (FObjectData& o : Objects)
+	{
+		if (IAsyncReadRequest** ReadRequest = o.CurrentReadRequests.Find(OperationId))
+		{
+			return (*ReadRequest)->PollCompletion();
+		}
+	}
+
+	UE_LOG(LogMutable, Error, TEXT("Operation not found in IsReadCompleted."));
+	check(false);
+	return true;
+}
+
+
+void FUnrealMutableModelBulkStreamer::EndRead(mu::ModelStreamer::OPERATION_ID OperationId)
+{
+	MUTABLE_CPUPROFILER_SCOPE(FUnrealMutableModelBulkStreamer::EndRead);
+
+	bool bFound = false;
+	for (FObjectData& o : Objects)
+	{
+		IAsyncReadRequest** ReadRequest = o.CurrentReadRequests.Find(OperationId);
+		if (ReadRequest)
+		{
+			if (*ReadRequest)
+			{
+				bool bCompleted = (*ReadRequest)->WaitCompletion();
+				if (!bCompleted)
+				{
+					UE_LOG(LogMutable, Error, TEXT("Operation failed to complete in EndRead."));
+					check(false);
+				}
+
+				delete* ReadRequest;
+			}
+			o.CurrentReadRequests.Remove(OperationId);
+			bFound = true;
+			break;
+		}
+	}
+	
+	if (!bFound)
+	{
+		UE_LOG(LogMutable, Error, TEXT("Operation not found in EndRead."));
+		check(false);
+	}
+}
diff --git a/Engine/Plugins/Experimental/Mutable/Source/CustomizableObject/Public/CustomizableObjectDGGUI.h b/Engine/Plugins/Experimental/Mutable/Source/CustomizableObject/Public/CustomizableObjectDGGUI.h
new file mode 100644
--- /dev/null
+++ b/Engine/Plugins/Experimental/Mutable/Source/CustomizableObject/Public/CustomizableObjectDGGUI.h
@@ -0,0 +1,27 @@
+// Copyright Epic Games, Inc. All Rights Reserved.
+
+#pragma once
+
+#include "Blueprint/UserWidget.h"
+#include "Animation/AnimClassInterface.h"
+#include "CustomizableObjectDGGUI.generated.h"
+
+
+UCLASS(Abstract, BlueprintType, Blueprintable)
+class CUSTOMIZABLEOBJECT_API UDGGUI : public UUserWidget
+{
+	GENERATED_BODY()
+public:
+	UFUNCTION(BlueprintImplementableEvent, Category = "DGGUI")
+	class UCustomizableSkeletalComponent* GetCustomizableSkeletalComponent();
+
+	UFUNCTION(BlueprintImplementableEvent, Category = "DGGUI")
+	void SetCustomizableSkeletalComponent(class UCustomizableSkeletalComponent* CustomizableSkeletalComponent);
+
+	static void OpenDGGUI(const int32 SlotID, UCustomizableSkeletalComponent* SelectedCustomizableSkeletalComponent, const UWorld* CurrentWorld, const int32 PlayerIndex = 0);
+	static bool CloseExistingDGGUI(const UWorld* CurrentWorld);
+private:
+	void CustomizableSkeletalMeshPreUpdate(class UCustomizableSkeletalComponent* Component, class USkeletalMesh* NextMesh);
+	void CustomizableSkeletalMeshUpdated();
+	TSubclassOf<UAnimInstance> LastAnimationClass;
+};
diff --git a/Engine/Plugins/Experimental/Mutable/Source/CustomizableObject/Public/CustomizableObjectUIData.h b/Engine/Plugins/Experimental/Mutable/Source/CustomizableObject/Public/CustomizableObjectUIData.h
new file mode 100644
--- /dev/null
+++ b/Engine/Plugins/Experimental/Mutable/Source/CustomizableObject/Public/CustomizableObjectUIData.h
@@ -0,0 +1,303 @@
+// Copyright Epic Games, Inc. All Rights Reserved.
+
+#pragma once
+
+#include "UnrealPortabilityHelpers.h"
+#include "UObject/ObjectMacros.h"
+#include "CustomizableObjectParameterTypeDefinitions.h"
+#include "CustomizableObjectCustomVersion.h"
+#include "CustomizableObjectUIData.generated.h"
+
+class UTexture2D;
+
+
+USTRUCT(BlueprintType)
+struct FMutableParamUIMetadata
+{
+	GENERATED_BODY()
+
+	FMutableParamUIMetadata()
+		: UIOrder(0)
+	{}
+
+	/** This is the name to be shown in UI */
+	UPROPERTY(BlueprintReadWrite, EditAnywhere, Category = UI)
+	FString ObjectFriendlyName;
+
+	/** This is the name of the section where the parameter will be placed in UI */
+	UPROPERTY(BlueprintReadWrite, EditAnywhere, Category = UI)
+	FString UISectionName;
+
+	/** This is the order of the parameter inside its section */
+	UPROPERTY(BlueprintReadWrite, EditAnywhere, Category = UI)
+	int32 UIOrder;
+
+	/** Thumnbail for UI */
+	UPROPERTY(BlueprintReadWrite, EditAnywhere, Category = UI)
+	TSoftObjectPtr<UTexture2D> UIThumbnail;
+
+	/** Extra information to be used in UI building, with semantics completely defined by the game/UI programmer, with a key to identify the semantic of its related value */
+	UPROPERTY(BlueprintReadWrite, EditAnywhere, Category = UI)
+	TMap<FString, FString> ExtraInformation;
+
+	/** Extra assets to be used in UI building */
+	UPROPERTY(BlueprintReadWrite, EditAnywhere, Category = UI)
+	TMap<FString, TSoftObjectPtr<UObject>> ExtraAssets;
+
+	UPROPERTY(BlueprintReadWrite, EditAnywhere, Category = UI)
+	float MinimumValue = 0.0f;
+
+	UPROPERTY(BlueprintReadWrite, EditAnywhere, Category = UI)
+	float MaximumValue = 1.0f;
+
+	bool operator ==(const FMutableParamUIMetadata& Other) const
+	{
+		if (ObjectFriendlyName != Other.ObjectFriendlyName || UISectionName != Other.UISectionName || UIOrder != Other.UIOrder || UIThumbnail != Other.UIThumbnail)
+		{
+			return false;
+		}
+
+		if (!ExtraInformation.OrderIndependentCompareEqual(Other.ExtraInformation))
+		{
+			return false;
+		}
+
+		if (!ExtraAssets.OrderIndependentCompareEqual(Other.ExtraAssets))
+		{
+			return false;
+		}
+
+		return true;
+	}
+
+	bool operator !=(const FMutableParamUIMetadata& Other) const
+	{
+		return !(*this == Other);
+	}
+
+	friend FArchive& operator <<(FArchive& Ar, FMutableParamUIMetadata& Metadata)
+	{
+		Ar.UsingCustomVersion(FCustomizableObjectCustomVersion::GUID);
+
+		Ar << Metadata.ObjectFriendlyName;
+		Ar << Metadata.UISectionName;
+		Ar << Metadata.UIOrder;
+
+		//Ar << Metadata.UIThumbnail;
+		if (Ar.IsLoading())
+		{
+			FString StringRef;
+			Ar << StringRef;
+			Metadata.UIThumbnail = TSoftObjectPtr<UTexture2D>(FSoftObjectPath(StringRef));
+		}
+		else
+		{
+			FString StringRef = Metadata.UIThumbnail.ToSoftObjectPath().ToString();
+			Ar << StringRef;
+		}
+
+		Ar << Metadata.ExtraInformation;
+
+		//Ar << Metadata.ExtraAssets;
+		if (Ar.IsLoading())
+		{
+			int32 NumReferencedAssets = 0;
+			Ar << NumReferencedAssets;
+			Metadata.ExtraAssets.Empty(NumReferencedAssets);
+
+			for (int32 i = 0; i < NumReferencedAssets; ++i)
+			{
+				FString Key, StringRef;
+				Ar << Key;
+				Ar << StringRef;
+
+				Metadata.ExtraAssets.Add(Key, TSoftObjectPtr<UObject>(FSoftObjectPath(StringRef)));
+			}
+		}
+		else
+		{
+			int32 NumReferencedAssets = Metadata.ExtraAssets.Num();
+			Ar << NumReferencedAssets;
+
+			for (TPair<FString, TSoftObjectPtr<UObject>>& AssetPair : Metadata.ExtraAssets)
+			{
+				FString StringRef = AssetPair.Value.ToSoftObjectPath().ToString();
+				Ar << AssetPair.Key;
+				Ar << StringRef;
+			}
+		}
+
+		// Update structure
+		if (Ar.CustomVer(FCustomizableObjectCustomVersion::GUID) >= FCustomizableObjectCustomVersion::BeforeCustomVersionWasAdded)
+		{
+			Ar << Metadata.MinimumValue;
+			Ar << Metadata.MaximumValue;
+		}
+
+		return Ar;
+	}
+
+#if WITH_EDITOR
+
+	/** Only called in the BeginCacheForCookPlatform to include new references in the final package */
+	void LoadResources()
+	{
+		UIThumbnail.LoadSynchronous();
+
+		for (const TPair<FString, TSoftObjectPtr<UObject>>& a : ExtraAssets)
+		{
+			a.Value.LoadSynchronous();
+		}
+	}
+
+#endif
+};
+
+USTRUCT(BlueprintType)
+struct FIntegerParameterUIData
+{
+	GENERATED_BODY()
+
+	FIntegerParameterUIData()
+	{
+
+	}
+
+	FIntegerParameterUIData(
+		const FString& ParamName,
+		const FMutableParamUIMetadata& InParamFriendlyName) :
+		  Name(ParamName)
+		, ParamUIMetadata(InParamFriendlyName)
+	{
+
+	}
+
+	/** Integer parameter option name */
+	UPROPERTY(BlueprintReadWrite, Category = UI)
+	FString Name;
+
+	UPROPERTY(BlueprintReadWrite, Category = UI, meta = (DisplayName = "Parameter UI Metadata"))
+	FMutableParamUIMetadata ParamUIMetadata;
+
+	bool operator ==(const FIntegerParameterUIData& Other) const
+	{
+		if (Name != Other.Name || ParamUIMetadata != Other.ParamUIMetadata)
+		{
+			return false;
+		}
+
+		return true;
+	}
+
+	bool operator !=(const FIntegerParameterUIData& Other) const
+	{
+		return !(*this == Other);
+	}
+
+	friend FArchive& operator <<(FArchive& Ar, FIntegerParameterUIData& UIData)
+	{
+		Ar << UIData.Name;
+		Ar << UIData.ParamUIMetadata;
+
+		return Ar;
+	}
+
+#if WITH_EDITOR
+
+	/** Only called in the BeginCacheForCookPlatform to include new references in the final package */
+	void LoadResources()
+	{
+		ParamUIMetadata.LoadResources();
+	}
+
+#endif
+};
+
+
+USTRUCT(BlueprintType)
+struct FParameterUIData
+{
+	GENERATED_BODY()
+
+	FParameterUIData() 
+		: Type(EMutableParameterType::None)
+		, IntegerParameterGroupType(ECustomizableObjectGroupType::COGT_ONE_OR_NONE)
+	{
+
+	}
+
+	FParameterUIData(
+		const FString& ParamName,
+		const FMutableParamUIMetadata& InParamUIMetadata,
+		EMutableParameterType ParamType) :
+		Name(ParamName),
+		ParamUIMetadata(InParamUIMetadata),
+		Type(ParamType), 
+		IntegerParameterGroupType(ECustomizableObjectGroupType::COGT_ONE_OR_NONE)
+	{
+
+	}
+
+	/** Parameter name */
+	UPROPERTY(BlueprintReadWrite, Category = UI)
+	FString Name;
+
+	UPROPERTY(BlueprintReadWrite, Category = UI, meta = (DisplayName = "Parameter UI Metadata"))
+	FMutableParamUIMetadata ParamUIMetadata;
+
+	/** Parameter type, using uint8 since the enum in declared in the class it is used */
+	UPROPERTY(BlueprintReadWrite, Category = UI)
+	EMutableParameterType Type;
+
+	/** In the case of an integer parameter, store here all options */
+	UPROPERTY(BlueprintReadWrite, Category = UI)
+	TArray<FIntegerParameterUIData> ArrayIntegerParameterOption;
+
+	/** In the case of an integer parameter, how are the different options selected (one, one or none, etc...) */
+	UPROPERTY(BlueprintReadWrite, Category = UI)
+	ECustomizableObjectGroupType IntegerParameterGroupType;
+
+	UPROPERTY(BlueprintReadWrite, Category = CustomizableObject)
+	bool bDontCompressRuntimeTextures = false; // Only useful for state metadata
+
+	UPROPERTY(BlueprintReadWrite, Category = CustomizableObject)
+	TMap<FString, FString> ForcedParameterValues;
+
+	bool operator ==(const FParameterUIData& Other) const
+	{
+		if (Name != Other.Name || ParamUIMetadata != Other.ParamUIMetadata || Type != Other.Type || ArrayIntegerParameterOption != Other.ArrayIntegerParameterOption || IntegerParameterGroupType != Other.IntegerParameterGroupType || !ForcedParameterValues.OrderIndependentCompareEqual(Other.ForcedParameterValues))
+		{
+			return false;
+		}
+
+		return true;
+	}
+
+	friend FArchive& operator <<(FArchive& Ar, FParameterUIData& UIData)
+	{
+		Ar << UIData.Name;
+		Ar << UIData.ParamUIMetadata;
+		Ar << UIData.Type;
+		Ar << UIData.ArrayIntegerParameterOption;
+		Ar << UIData.IntegerParameterGroupType;
+		Ar << UIData.bDontCompressRuntimeTextures;
+		Ar << UIData.ForcedParameterValues;
+
+		return Ar;
+	}
+
+#if WITH_EDITOR
+
+	/** Only called in the BeginCacheForCookPlatform to include new references in the final package */
+	void LoadResources()
+	{
+		ParamUIMetadata.LoadResources();
+
+		for (FIntegerParameterUIData& intParam : ArrayIntegerParameterOption)
+		{
+			intParam.LoadResources();
+		}
+	}
+
+#endif
+};
diff --git a/Engine/Plugins/Experimental/Mutable/Source/CustomizableObject/Public/MultilayerProjector.h b/Engine/Plugins/Experimental/Mutable/Source/CustomizableObject/Public/MultilayerProjector.h
new file mode 100644
--- /dev/null
+++ b/Engine/Plugins/Experimental/Mutable/Source/CustomizableObject/Public/MultilayerProjector.h
@@ -0,0 +1,197 @@
+// Copyright Epic Games, Inc. All Rights Reserved.
+
+#pragma once
+
+#include "MultilayerProjector.generated.h"
+
+class UCustomizableObjectInstance;
+struct FMultilayerProjector;
+
+
+/** Data structure representing a Multilayer Projector Layer.
+ *
+ * This struct is not actually saved, its values are obtained from the Instance Parameters. */
+USTRUCT(BlueprintType)
+struct FMultilayerProjectorLayer
+{
+	GENERATED_BODY()
+	
+	/** Read the Layer from the Instance Parameters.
+	 *
+	 * @param MultilayerProjector Multilayer Projector helper.
+	 * @param Index Index where to read the Layer in the Instance Parameters.
+	 */
+	void Read(const FMultilayerProjector& MultilayerProjector, int32 Index);
+
+	/** Write to Layer to the Instance Parameters.
+	 *
+	 * @param MultilayerProjector Multilayer Projector helper.
+	 * @param Index Index where to write the Layer in the Instance Parameters.
+	 */
+	void Write(const FMultilayerProjector& MultilayerProjector, int32 Index) const;
+	
+	/** Layer position. */
+	FVector Position;
+
+	/** Layer direction vector. */
+	FVector Direction;
+
+	/** Layer up direction vector. */
+	FVector Up;
+
+	/** Layer scale. */
+	FVector Scale;
+
+	/** Layer angle*/
+	float Angle;
+
+	/** Layer selected image. */
+	FString Image;	
+
+	/** Layer image opacity. */
+	float Opacity;
+};
+
+
+/** Data structure representing a Multilayer Projector Virtual Layer. */
+USTRUCT(BlueprintType)
+struct FMultilayerProjectorVirtualLayer : public FMultilayerProjectorLayer
+{
+	GENERATED_BODY()
+
+	FMultilayerProjectorVirtualLayer() = default;
+
+	FMultilayerProjectorVirtualLayer(const FMultilayerProjectorLayer& Layer, bool bEnabled, int32 Order);
+
+	/** True if the Virtual Layer is enabled. */
+	bool bEnabled;
+
+	/** Virtual Layer absolute order. Equal values may result in an undefined behaviour. */
+	int32 Order;
+};
+
+
+/** Multilayer Projector Helper. Eases the management of Layers and Virtual Layers.
+ *
+ * Layer: Management of Multilayer Projector Layers by index.
+ * All layers indices has to be consecutive.
+ * 
+ * Virtual Layer: Management of Multilayer Projector Layers by name.
+ * - Allows to enabled and disabled layers.
+ * - Allows to sort layers.
+ */
+USTRUCT()
+struct CUSTOMIZABLEOBJECT_API FMultilayerProjector
+{
+	GENERATED_BODY()
+
+	friend UCustomizableObjectInstance;
+	friend FMultilayerProjectorLayer;
+
+	// Parameters encoding
+	static const FString NUM_LAYERS_PARAMETER_POSTFIX;
+	static const FString OPACITY_PARAMETER_POSTFIX;
+	static const FString IMAGE_PARAMETER_POSTFIX;
+	static const FString POSE_PARAMETER_POSTFIX;
+
+	// Constructors
+	FMultilayerProjector() = default;
+
+	FMultilayerProjector(UCustomizableObjectInstance& Instance, const FName& ParamName);
+
+	// Layers
+	
+	/** Returns the number of layers. */
+	int32 NumLayers() const;
+
+	/** Insert the Layer at the given index moving all the following layers one position. */
+	void CreateLayer(int32 Index) const;
+
+	/** Remove the Layer at the given index moving all the following layers one position. */
+	void RemoveLayerAt(int32 Index) const;
+
+	/** Get the properties of the Layer at the given index. */
+	FMultilayerProjectorLayer GetLayer(int32 Index) const;
+
+	/** Update the Layer properties at the given index. */
+	void UpdateLayer(int32 Index, const FMultilayerProjectorLayer& Layer) const;
+
+	// Virtual layers
+
+	/** Get all created Virtual Layers. */
+	TArray<FName> GetVirtualLayers() const;
+	
+	/** Given a new identifier, create a new Virtual Layer (if non-existent).*/
+	void CreateVirtualLayer(const FName& Id);
+
+	/** Given a new identifier, find or create a new Virtual Layer.*/
+	FMultilayerProjectorVirtualLayer FindOrCreateVirtualLayer(const FName& Id);
+
+	/** Given its identifier, remove the Virtual Layer. */
+	void RemoveVirtualLayer(const FName& Id);
+
+	/** Given its identifier, get the Virtual Layer properties. */
+	FMultilayerProjectorVirtualLayer GetVirtualLayer(const FName& Id) const;
+
+	/** Given its identifier, update a Virtual Layer properties. */
+	void UpdateVirtualLayer(const FName& Id, const FMultilayerProjectorVirtualLayer& Layer);
+
+private:
+	/** Always valid. */
+	UPROPERTY()
+	TObjectPtr<UCustomizableObjectInstance> Instance = nullptr;
+
+	/** Multilayer Projector Parameter name. */
+	UPROPERTY()
+	FName ParamName;	
+	
+	/** Maps a Virtual Layer to a Layer.
+	 *
+ 	 * All enabled Virtual Layers must have an entry pointing to the Layer index.
+ 	 * Disabled layers should have VIRTUAL_LAYER_DISABLED as value.
+	 */
+	UPROPERTY()
+	TMap<FName, int32> VirtualLayersMapping;
+
+	/** Maps a Virtual Layer to a Layer.
+	 *
+	 * All created Virtual Layers (enabled or disabled) must have an entry to this map.
+	 */
+	UPROPERTY()
+	TMap<FName, int32> VirtualLayersOrder;
+	
+	/** Since disabling a Virtual Layers removes the Layer, its values are no longer on the Instance Parameters.
+	 * Instead, they are saved in this data structure to be able to reenable it later. */
+	UPROPERTY()
+	TMap<FName, FMultilayerProjectorLayer> DisableVirtualLayers;
+
+	/** Virtual Layer Disabled Index. */
+	static constexpr int32 VIRTUAL_LAYER_DISABLED = -1;
+
+	/** New Virtual Layer default order. */
+	static constexpr int32 NEW_VIRTUAL_LAYER_ORDER = 0;
+
+	/** Taking into consideration the Virtual Layer Order, Calculate where the Layer should be inserted. */
+	int32 CalculateVirtualLayerIndex(const FName& Id, int32 InsertOrder) const;
+	
+	/** Given an inserted Virtual Layer, update the mapping of the Virtual Layers.
+ 	 *
+ 	 * @param Id Virtual Layer.
+	 * @param Index Index which the layer has been enabled.
+	 */
+	void UpdateMappingVirtualLayerEnabled(const FName& Id, int32 Index);
+	
+	/** Given a removed Virtual Layer, update the mapping of the Virtual Layers.
+	 *
+ 	 * @param Id Virtual Layer.	
+	 * @param Index Index which the layer has been disabled.
+	 */
+	void UpdateMappingVirtualLayerDisabled(const FName& Id, int32 Index);
+
+	/** Rise an assert if the instance does not contain the necessary Instance Parameters. */
+	void CheckInstanceParameters() const;
+	
+	/** Return false if the Instance does not contain the necessary Instance Parameters. */
+	static bool AreInstanceParametersValid(const UCustomizableObjectInstance& Instance, const FName& ParamName);
+};
+
diff --git a/Engine/Plugins/Experimental/Mutable/Source/CustomizableObject/Public/UnrealBakeHelpers.h b/Engine/Plugins/Experimental/Mutable/Source/CustomizableObject/Public/UnrealBakeHelpers.h
new file mode 100644
--- /dev/null
+++ b/Engine/Plugins/Experimental/Mutable/Source/CustomizableObject/Public/UnrealBakeHelpers.h
@@ -0,0 +1,16 @@
+// Copyright Epic Games, Inc. All Rights Reserved.
+
+#pragma once
+
+#include "ICustomizableObjectModule.h"
+
+//-------------------------------------------------------------------------------------------------
+
+class CUSTOMIZABLEOBJECT_API FUnrealBakeHelpers
+{
+public:
+
+	static void BakeHelper_RegenerateImportedModel(class USkeletalMesh* SkeletalMesh);
+};
+
+//-------------------------------------------------------------------------------------------------
diff --git a/Engine/Plugins/Experimental/Mutable/Source/CustomizableObject/Public/UnrealConversionUtils.h b/Engine/Plugins/Experimental/Mutable/Source/CustomizableObject/Public/UnrealConversionUtils.h
new file mode 100644
--- /dev/null
+++ b/Engine/Plugins/Experimental/Mutable/Source/CustomizableObject/Public/UnrealConversionUtils.h
@@ -0,0 +1,90 @@
+// Copyright Epic Games, Inc. All Rights Reserved.
+
+#pragma once
+
+#include "CustomizableObject/Private/CustomizableObjectSystemPrivate.h"
+#include "MutableRuntime/Public/Mesh.h"
+
+
+namespace UnrealConversionUtils
+{
+	/*
+	 * Shared methods between mutable instance update and viewport mesh generation
+	 * These are the result of stripping up parts of the reference code to be able to share it on the two
+	 * current pipelines (instance update and USkeletal mesh generation for a mesh viewport)
+	 */
+	
+	/**
+	 * Assigns a new render section for each of the surfaces found on the mutable mesh and sets a default material for
+	 * each of them (same default material)
+	 * @param MeshLODIndex - The index of the LOD found on the OutSkeletalMesh to be targeting.
+	 * @param InMutableMesh - The Mutable mesh to get the surface count from. Determines the amount materials and also
+	 * render sections to be added to the OutSkeletalMesh.
+	 * @param OutSkeletalMesh - The skeletal mesh whose render data is being updated (adding render sections)
+	 */
+	CUSTOMIZABLEOBJECT_API void BuildSkeletalMeshElementDataAtLOD(const int32 MeshLODIndex, mu::MeshPtrConst InMutableMesh,
+										   USkeletalMesh* OutSkeletalMesh);
+	
+	
+	/**
+	 * Builds the reference skeleton by adding bones to the reference skeleton modifier provided. When exiting the scope
+	 * of this method those bones get generated on the reference skeleton. It also updates an array with the used bones.
+	 * @param OutMutSkeletonData - Object map used later when remaking bones.
+	 * @param InSourceReferenceSkeleton - The reference skeleton from the source object. Is used to get the skeleton bone info
+	 * @param InUsedBones - Array with the same count as the bones present on the skeleton and determines what bones are used and what not.
+	 * @param InRefSkeleton - The reference skeleton to be updated with new bone data
+	 * @param InSkeleton - Skeleton to be updated as consequence of the destruction of the FReferenceSkeletonModifier
+	 */
+	CUSTOMIZABLEOBJECT_API void BuildRefSkeleton(
+		FInstanceUpdateData::FSkeletonData* OutMutSkeletonData,
+		const FReferenceSkeleton& InSourceReferenceSkeleton,
+		const TArray<bool>& InUsedBones,
+		FReferenceSkeleton& InRefSkeleton,
+		const USkeleton* InSkeleton );
+
+	
+	/**
+	 * Prepares the render sections found on the InSkeletalMesh and sets them up accordingly what the InMutableMesh requires
+	 * @param InMutableMesh - Mutable mesh to be used as reference for the section data update on the Skeletal Mesh
+	 * @param OutSkeletalMesh - The Skeletal mesh whose sections are ought to be updated
+	 * @param MeshLODIndex - The Mesh lod to be working with
+	 * @param NumBoneInfluences - The amount of bones a pixel can be affected by
+	 * @param InBoneMap - The bones to be set as part of the sections.
+	 */
+	CUSTOMIZABLEOBJECT_API void SetupRenderSections(
+		const mu::MeshPtrConst InMutableMesh,
+		const USkeletalMesh* OutSkeletalMesh,
+		const int32 MeshLODIndex,
+		const int32 NumBoneInfluences,
+		const TArray<uint16>& InBoneMap);
+
+
+	/** Performs a copy of the data found on the vertex buffers on the mutable mesh to the buffers of the skeletal mesh
+	 * @param OutSkeletalMesh - Skeletal Mesh to be updated with new buffer data
+	 * @param NumVerticesLODModel - The amount of vertices on the LODModel.
+	 * @param NumBoneInfluences - The amount of bones that can influence a vertex
+	 * @param BoneIndexBuffer - The buffer containing the indices for the bones.
+	 * @param InMutableMeshVertexBuffers - The mutable buffers to be reading data from
+	 * @param MeshLODIndex - The LOD index we are working with.
+	 * @param InBoneIndexFormat - The format of the bone indices on the mutable mesh.
+	 */
+	CUSTOMIZABLEOBJECT_API void CopyMutableVertexBuffers(
+		USkeletalMesh* OutSkeletalMesh,
+		const int32 NumVerticesLODModel,
+		const int32 NumBoneInfluences,
+		const int32 BoneIndexBuffer,
+		const mu::FMeshBufferSet& InMutableMeshVertexBuffers,
+		const int32 MeshLODIndex,
+		const mu::MESH_BUFFER_FORMAT& InBoneIndexFormat);
+
+	
+	/**
+	 *Performs a copy of the data found on the index buffers on the mutable mesh to the buffers of the skeletal mesh
+	 * @param InMutableMesh - The mutable mesh whose index buffers you want to work with
+	 * @param OutLODModel - The LOD model to be updated.
+	 * @return True if the operation could be performed successfully, false if not.
+	 */
+	CUSTOMIZABLEOBJECT_API bool CopyMutableIndexBuffers(mu::MeshPtrConst InMutableMesh,
+	                                                    FSkeletalMeshLODRenderData& OutLODModel);
+	
+}
diff --git a/Engine/Plugins/Experimental/Mutable/Source/CustomizableObject/Public/UnrealMutableModelDiskStreamer.h b/Engine/Plugins/Experimental/Mutable/Source/CustomizableObject/Public/UnrealMutableModelDiskStreamer.h
new file mode 100644
--- /dev/null
+++ b/Engine/Plugins/Experimental/Mutable/Source/CustomizableObject/Public/UnrealMutableModelDiskStreamer.h
@@ -0,0 +1,105 @@
+// Copyright Epic Games, Inc. All Rights Reserved.
+
+#pragma once
+
+#include "Serialization/Archive.h"
+#include "Containers/UnrealString.h"
+#include "Async/AsyncFileHandle.h"
+
+#include "CustomizableObject.h"
+
+#include "MutableRuntime/Public/Model.h"
+#include "MutableRuntime/Public/Serialisation.h"
+
+
+// Implementation fo the basic mutable streams.
+#if WITH_EDITOR
+class UnrealMutableOutputStream : public mu::OutputStream
+{
+public:
+
+	UnrealMutableOutputStream(FArchive& ar);
+
+	// mu::OutputStream interface
+	void Write(const void* pData, uint64 size) override;
+
+private:
+	FArchive& m_ar;
+};
+#endif
+
+
+class UnrealMutableInputStream : public mu::InputStream
+{
+public:
+
+	UnrealMutableInputStream(FArchive& ar);
+
+	// mu::InputStream interface
+	void Read(void* pData, uint64 size) override;
+
+private:
+	FArchive& m_ar;
+};
+
+
+// Implementation of a mutable streamer using bulk storage.
+class CUSTOMIZABLEOBJECT_API FUnrealMutableModelBulkStreamer : public mu::ModelStreamer
+{
+public:
+	// 
+	FUnrealMutableModelBulkStreamer(FArchive* InMainDataArchive = nullptr, FArchive* InStreamedDataArchive = nullptr);
+	~FUnrealMutableModelBulkStreamer();
+	
+	// Own interface
+
+	/** Make sure that the provided object can stream data. */
+	bool PrepareStreamingForObject(UCustomizableObject* Object);
+
+#if WITH_EDITOR
+	/** Cancel any further streaming operations for the given object. This is necessary if the object compiled data is
+	 * going to be modified. This can only happen in the editor, when recompiling.
+	 * Any additional straming requests for this object will fail.
+	 */
+	void CancelStreamingForObject(const UCustomizableObject* CustomizableObject);
+#endif
+
+	/** Release all the pending resources. This disables treamings for all objects. */
+	void EndStreaming();
+	
+	// mu::ModelStreamer interface
+	void OpenWriteFile(const char* strModelName, uint64 key0) override;
+	void Write(const void* pBuffer, uint64 size) override;
+	void CloseWriteFile() override;
+	
+	OPERATION_ID BeginReadBlock(const mu::Model*, uint64 key0, void* pBuffer, uint64 size) override;
+	bool IsReadCompleted(OPERATION_ID) override;
+	void EndRead(OPERATION_ID) override;
+
+protected:
+
+#if WITH_EDITOR
+	// Write
+	// Non-owned pointer to an archive where we'll store the main model data (non-streamable)
+	FArchive* MainDataArchive = nullptr;
+
+	// Non-owned pointer to an archive where we'll store the resouces (streamable)
+	FArchive* StreamedDataArchive = nullptr;
+
+	FArchive* CurrentWriteFile = nullptr;
+#endif
+
+	/** Streaming data for one object. */
+	struct FObjectData
+	{
+		mu::WeakPtr<mu::Model> Model;
+		TArray<IAsyncReadFileHandle*> ReadFileHandles;
+		TMap<OPERATION_ID, IAsyncReadRequest*> CurrentReadRequests;
+		TMap<uint64, FMutableStreamableBlock> StreamableBlocks;
+	};
+
+	TArray<FObjectData> Objects;
+
+	/** This is used to generate unique ids for read requests. */
+	OPERATION_ID LastOperationID = 0;
+};
diff --git a/Engine/Plugins/Experimental/Mutable/Source/CustomizableObject/Public/UnrealPortabilityHelpers.h b/Engine/Plugins/Experimental/Mutable/Source/CustomizableObject/Public/UnrealPortabilityHelpers.h
new file mode 100644
--- /dev/null
+++ b/Engine/Plugins/Experimental/Mutable/Source/CustomizableObject/Public/UnrealPortabilityHelpers.h
@@ -0,0 +1,106 @@
+// Copyright Epic Games, Inc. All Rights Reserved.
+
+#pragma once
+
+#include "Runtime/Launch/Resources/Version.h"
+
+#include "Engine/SkeletalMesh.h"
+#include "Misc/Paths.h"
+#include "Rendering/SkeletalMeshRenderData.h"
+#include "Rendering/SkeletalMeshLODRenderData.h"
+#include "Rendering/SkeletalMeshModel.h"
+#include "Rendering/SkinWeightVertexBuffer.h"
+
+// Main engine branch should have 1 here. 
+// Projects with API modifications may set it 0.
+#define MUTABLE_CLEAN_ENGINE_BRANCH		1
+
+//-------------------------------------------------------------------------------------------------
+// Helpers to ease portability across unreal engine versions
+//-------------------------------------------------------------------------------------------------
+
+typedef FSkeletalMeshLODRenderData Helper_LODDataType;
+typedef FSkelMeshRenderSection Helper_SkelMeshRenderSection;
+
+inline FSkeletalMeshRenderData* Helper_GetResourceForRendering(const USkeletalMesh* Mesh)
+{
+	return Mesh->GetResourceForRendering();
+}
+
+inline TIndirectArray<FSkeletalMeshLODRenderData>* Helper_GetLODDataPtr(FSkeletalMeshRenderData* Resource)
+{
+	return &(Resource->LODRenderData);
+}
+
+inline TIndirectArray<FSkeletalMeshLODRenderData>& Helper_GetLODData(FSkeletalMeshRenderData* Resource)
+{
+	return Resource->LODRenderData;
+}
+
+inline const TIndirectArray<FSkeletalMeshLODRenderData>& Helper_GetLODData(const FSkeletalMeshRenderData* Resource)
+{
+	return Resource->LODRenderData;
+}
+
+inline TIndirectArray<FSkeletalMeshLODRenderData>& Helper_GetLODData(const USkeletalMesh* Mesh)
+{
+	return Mesh->GetResourceForRendering()->LODRenderData;
+}
+
+inline TArray<FSkelMeshRenderSection>& Helper_GetLODRenderSections(const USkeletalMesh* Mesh, int LODIndex)
+{
+	return  Mesh->GetResourceForRendering()->LODRenderData[LODIndex].RenderSections;
+}
+
+#if WITH_EDITORONLY_DATA
+inline FSkeletalMeshModel* Helper_GetImportedModel(const USkeletalMesh* Mesh)
+{
+	return Mesh->GetImportedModel();
+}
+#endif
+
+inline int32 Helper_LODGetTotalFaces(const FSkeletalMeshLODRenderData* Resource)
+{
+	return Resource->GetTotalFaces();
+}
+
+inline uint32 Helper_LODGetNumTexCoords(const FSkeletalMeshLODRenderData* Resource)
+{
+	return Resource->GetNumTexCoords();
+}
+
+inline const TArray<FSkelMeshRenderSection>* Helper_LODGetSections(const FSkeletalMeshLODRenderData* Resource)
+{
+	return &(Resource->RenderSections);
+}
+
+inline TArray<struct FSkeletalMeshLODInfo>& Helper_GetLODInfoArray(USkeletalMesh* Mesh)
+{
+	return Mesh->GetLODInfoArray();
+}
+
+inline const TArray<struct FSkeletalMeshLODInfo>& Helper_GetLODInfoArray(const USkeletalMesh* Mesh)
+{
+	return Mesh->GetLODInfoArray();
+}
+
+inline int Helper_GetLODNum(const USkeletalMesh* Mesh)
+{
+	return Mesh->GetLODNum();
+}
+
+inline const TArray<FBoneReference>& Helper_GetLODInfoBonesToRemove(const USkeletalMesh* Mesh, int LODIndex)
+{
+	return Mesh->GetLODInfo(LODIndex)->BonesToRemove;
+}
+
+inline TArray<FBoneReference>& Helper_GetLODInfoBonesToRemove(USkeletalMesh* Mesh, int LODIndex)
+{
+	return Mesh->GetLODInfo(LODIndex)->BonesToRemove;
+}
+
+inline FString Helper_GetSavedDir()
+{
+	return FPaths::ConvertRelativePathToFull(FPaths::ProjectSavedDir());
+}
+
diff --git a/Engine/Plugins/Experimental/Mutable/Source/CustomizableObjectEditor/Classes/CustomizableObjectFactory.h b/Engine/Plugins/Experimental/Mutable/Source/CustomizableObjectEditor/Classes/CustomizableObjectFactory.h
new file mode 100644
--- /dev/null
+++ b/Engine/Plugins/Experimental/Mutable/Source/CustomizableObjectEditor/Classes/CustomizableObjectFactory.h
@@ -0,0 +1,24 @@
+// Copyright Epic Games, Inc. All Rights Reserved.
+
+#pragma once
+
+#include "Factories/Factory.h"
+
+#include "CustomizableObjectFactory.generated.h"
+
+
+UCLASS(MinimalAPI)
+class UCustomizableObjectFactory : public UFactory
+{
+	GENERATED_BODY()
+
+	UCustomizableObjectFactory();
+
+	// Begin UObject Interface
+	virtual bool ConfigureProperties() override;
+
+	// Begin UFactory Interface
+	virtual bool DoesSupportClass(UClass * Class) override;
+	virtual UClass* ResolveSupportedClass() override;
+	virtual UObject* FactoryCreateNew(UClass* Class,UObject* InParent,FName Name,EObjectFlags Flags,UObject* Context,FFeedbackContext* Warn) override;
+};
diff --git a/Engine/Plugins/Experimental/Mutable/Source/CustomizableObjectEditor/Classes/CustomizableObjectGraph.h b/Engine/Plugins/Experimental/Mutable/Source/CustomizableObjectEditor/Classes/CustomizableObjectGraph.h
new file mode 100644
--- /dev/null
+++ b/Engine/Plugins/Experimental/Mutable/Source/CustomizableObjectEditor/Classes/CustomizableObjectGraph.h
@@ -0,0 +1,42 @@
+// Copyright Epic Games, Inc. All Rights Reserved.
+
+#pragma once
+
+#include "Logging/TokenizedMessage.h"
+#include "Runtime/CoreUObject/Public/Misc/UObjectToken.h"
+#include "EdGraph/EdGraph.h"
+
+#include "CustomizableObjectGraph.generated.h"
+
+
+class UCustomizableObjectNode;
+
+
+UCLASS()
+class CUSTOMIZABLEOBJECTEDITOR_API UCustomizableObjectGraph : public UEdGraph
+{
+public:
+	GENERATED_BODY()
+
+	UCustomizableObjectGraph();
+
+	// UObject Interface
+	virtual void PostLoad() override;
+	void PostRename(UObject * OldOuter, const FName OldName) override;
+
+	// Own Interface
+	void NotifyNodeIdChanged(const FGuid& OldGuid, const FGuid& NewGuid);
+
+	FGuid RequestNotificationForNodeIdChange(const FGuid& OldGuid, const FGuid& NodeToNotifyGuid);
+
+	void PostDuplicate(bool bDuplicateForPIE) override;
+
+private:
+
+	// Request Node Id Update Map
+	TMap<FGuid, TSet<FGuid>> NodesToNotifyMap;
+
+	// Guid map with the key beeing the old Guid and the Value the new one, filled after duplicating COs
+	TMap<FGuid, FGuid> NotifiedNodeIdsMap;
+};
+
diff --git a/Engine/Plugins/Experimental/Mutable/Source/CustomizableObjectEditor/Classes/CustomizableObjectInstanceFactory.h b/Engine/Plugins/Experimental/Mutable/Source/CustomizableObjectEditor/Classes/CustomizableObjectInstanceFactory.h
new file mode 100644
--- /dev/null
+++ b/Engine/Plugins/Experimental/Mutable/Source/CustomizableObjectEditor/Classes/CustomizableObjectInstanceFactory.h
@@ -0,0 +1,29 @@
+// Copyright Epic Games, Inc. All Rights Reserved.
+
+#pragma once
+
+#include "CoreMinimal.h"
+#include "UObject/ObjectMacros.h"
+#include "ActorFactories/ActorFactory.h"
+#include "CustomizableObjectInstanceFactory.generated.h"
+
+class AActor;
+struct FAssetData;
+class USkeletalMesh;
+
+UCLASS(MinimalAPI, config=Editor)
+class UCustomizableObjectInstanceFactory : public UActorFactory
+{
+	GENERATED_UCLASS_BODY()
+
+protected:
+	//~ Begin UActorFactory Interface
+	void PostSpawnActor(UObject* Asset, AActor* NewActor) override;
+	void PostCreateBlueprint(UObject* Asset, AActor* CDO) override;
+	FQuat AlignObjectToSurfaceNormal(const FVector& InSurfaceNormal, const FQuat& ActorRotation) const;
+	bool CanCreateActorFrom(const FAssetData& AssetData, FText& OutErrorMsg) override;
+	//~ End UActorFactory Interface
+
+	int32 GetNumberOfComponents(class UCustomizableObjectInstance* COInstance);
+	USkeletalMesh* GetSkeletalMeshFromAsset(UObject* Asset, int32 ComponentIndex = 0) const;
+};
diff --git a/Engine/Plugins/Experimental/Mutable/Source/CustomizableObjectEditor/Classes/CustomizableObjectLayout.h b/Engine/Plugins/Experimental/Mutable/Source/CustomizableObjectEditor/Classes/CustomizableObjectLayout.h
new file mode 100644
--- /dev/null
+++ b/Engine/Plugins/Experimental/Mutable/Source/CustomizableObjectEditor/Classes/CustomizableObjectLayout.h
@@ -0,0 +1,105 @@
+// Copyright Epic Games, Inc. All Rights Reserved.
+
+#pragma once
+
+#include "CustomizableObjectLayout.generated.h"
+
+UENUM()
+enum class ECustomizableObjectTextureLayoutPackingStrategy : uint8
+{
+	// The layout increases its size to fit all the blocks.
+	Resizable = 0 UMETA(DisplayName = "Resizable Layout"),
+	// The layout resizes the blocks to keep its size.
+	Fixed = 1 UMETA(DisplayName = "Fixed Layout")
+};
+
+USTRUCT()
+struct CUSTOMIZABLEOBJECTEDITOR_API FCustomizableObjectLayoutBlock
+{
+	GENERATED_USTRUCT_BODY()
+
+		FCustomizableObjectLayoutBlock()
+	{
+		Min = FIntPoint(0, 0);
+		Max = FIntPoint(1, 1);
+		Priority = 0;
+	}
+
+	UPROPERTY(EditAnywhere, Category = CustomizableObject)
+		FIntPoint Min;
+
+	UPROPERTY(EditAnywhere, Category = CustomizableObject)
+		FIntPoint Max;
+
+	UPROPERTY(EditAnywhere, Category = CustomizableObject)
+		uint32 Priority;
+
+	//! Unique unchangeable id used to reference this block from other nodes.
+	UPROPERTY()
+		FGuid Id;
+};
+
+UCLASS()
+class UCustomizableObjectLayout : public UObject
+{
+	GENERATED_BODY()
+
+public:
+
+	UCustomizableObjectLayout();
+
+	void SetLayout(UObject* InMesh, int32 LODIndex, int32 MatIndex, int32 UVIndex);
+	void SetPackingStrategy(ECustomizableObjectTextureLayoutPackingStrategy Strategy);
+	void SetGridSize(FIntPoint Size);
+	void SetMaxGridSize(FIntPoint Size);
+	void SetLayoutName(FString Name);
+
+	int32 GetLOD() const { return LOD; }
+	int32 GetMaterial() const { return Material; }
+	int32 GetUVChannel() const { return UVChannel; }
+	FString GetLayoutName() const { return LayoutName; }
+	UObject* GetMesh() const { return Mesh; }
+	FIntPoint GetGridSize() const { return GridSize; }
+	FIntPoint GetMaxGridSize() const { return MaxGridSize; }
+	ECustomizableObjectTextureLayoutPackingStrategy GetPackingStrategy() const { return PackingStrategy; }
+
+	void GetUVChannel(TArray<FVector2f>& UVs, int32 UVChannelIndex = 0) const;
+
+	// Get a block index in the array from its id. Return -1 if not found.
+	int32 FindBlock(const FGuid& InId) const;
+
+	void GenerateBlocksFromUVs();
+
+	UPROPERTY()
+	TArray<FCustomizableObjectLayoutBlock> Blocks;
+
+	TArray< TArray<FVector2f> > UnassignedUVs;
+
+private:
+
+	UPROPERTY()
+	TObjectPtr<UObject> Mesh = nullptr;
+
+	UPROPERTY()
+	int32 LOD;
+
+	UPROPERTY()
+	int32 Material;
+
+	UPROPERTY()
+	int32 UVChannel;
+
+	UPROPERTY()
+	FIntPoint GridSize;
+
+	/** Used with the fixed layout strategy. */
+	UPROPERTY()
+	FIntPoint MaxGridSize;
+
+	UPROPERTY()
+	ECustomizableObjectTextureLayoutPackingStrategy PackingStrategy = ECustomizableObjectTextureLayoutPackingStrategy::Resizable;
+
+	UPROPERTY()
+	FString LayoutName;
+
+};
diff --git a/Engine/Plugins/Experimental/Mutable/Source/CustomizableObjectEditor/Classes/EdGraphSchema_CustomizableObject.h b/Engine/Plugins/Experimental/Mutable/Source/CustomizableObjectEditor/Classes/EdGraphSchema_CustomizableObject.h
new file mode 100644
--- /dev/null
+++ b/Engine/Plugins/Experimental/Mutable/Source/CustomizableObjectEditor/Classes/EdGraphSchema_CustomizableObject.h
@@ -0,0 +1,168 @@
+// Copyright Epic Games, Inc. All Rights Reserved.
+
+#pragma once
+
+#include "AssetRegistry/AssetData.h"
+#include "EdGraphSchema_K2_Actions.h"
+#include "Runtime/Launch/Resources/Version.h"
+#include "MaterialTypes.h"
+
+#include "EdGraphSchema_CustomizableObject.generated.h"
+
+/** Action to add a node to the graph */
+USTRUCT()
+struct CUSTOMIZABLEOBJECTEDITOR_API FCustomizableObjectSchemaAction_NewNode : public FEdGraphSchemaAction
+{
+	GENERATED_USTRUCT_BODY();
+
+	/** Template of node we want to create */
+	UPROPERTY()
+	TObjectPtr<class UEdGraphNode> NodeTemplate;
+
+
+	FCustomizableObjectSchemaAction_NewNode() 
+		: FEdGraphSchemaAction()
+		, NodeTemplate(NULL)
+	{}
+
+	FCustomizableObjectSchemaAction_NewNode(const FString& InNodeCategory, const FText& InMenuDesc, const FText& InToolTip, const int32 InGrouping, const FText& InKeywords = FText(), int32 InSectionID = 0)
+		: FEdGraphSchemaAction(FText::FromString(InNodeCategory), InMenuDesc, InToolTip, InGrouping, InKeywords, InSectionID)
+		, NodeTemplate(NULL)
+	{}
+
+	// FEdGraphSchemaAction interface
+	virtual UEdGraphNode* PerformAction(class UEdGraph* ParentGraph, UEdGraphPin* FromPin, const FVector2D Location, bool bSelectNewNode = true) override;
+	virtual UEdGraphNode* PerformAction(class UEdGraph* ParentGraph, TArray<UEdGraphPin*>& FromPins, const FVector2D Location, bool bSelectNewNode = true) override;
+	virtual void AddReferencedObjects( FReferenceCollector& Collector ) override;
+
+	/** Reimplementaiton of EdGraphSchema::CreateNode(...). Performs the overlap calculation before calling AutowireNewNode(...). AutowireNewNode can induce a call to a ReconstrucNode() which removes pins required for the calculation. */
+	static UEdGraphNode* CreateNode(class UEdGraph* ParentGraph, UEdGraphPin* FromPin, const FVector2D Location, class UEdGraphNode* InNodeTemplate);
+	// End of FEdGraphSchemaAction interface
+
+	template <typename NodeType>
+	static NodeType* InstantSpawn(class UEdGraph* ParentGraph, NodeType* InTemplateNode, const FVector2D Location)
+	{
+		FEdGraphSchemaAction_K2NewNode Action;
+		Action.NodeTemplate = InTemplateNode;
+
+		return Cast<NodeType>(Action.PerformAction(ParentGraph, NULL, Location));
+	}
+};
+
+
+/** Action to paste clipboard contents into the graph */
+USTRUCT()
+struct CUSTOMIZABLEOBJECTEDITOR_API FCustomizableObjectSchemaAction_Paste : public FEdGraphSchemaAction
+{
+	GENERATED_USTRUCT_BODY();
+
+	// Simple type info
+	static FName StaticGetTypeId() { static FName Type("FCustomizableObjectSchemaAction_Paste"); return Type; }
+	virtual FName GetTypeId() const override { return StaticGetTypeId(); }
+
+	FCustomizableObjectSchemaAction_Paste()
+		: FEdGraphSchemaAction()
+	{}
+
+	FCustomizableObjectSchemaAction_Paste(const FText& InNodeCategory, const FText& InMenuDesc, const FString& InToolTip, const int32 InGrouping)
+		: FEdGraphSchemaAction(InNodeCategory, InMenuDesc, FText::FromString(InToolTip), InGrouping)
+	{}
+
+	// FEdGraphSchemaAction interface
+	virtual UEdGraphNode* PerformAction(class UEdGraph* ParentGraph, UEdGraphPin* FromPin, const FVector2D Location, bool bSelectNewNode = true) override;
+	// End of FEdGraphSchemaAction interface
+};
+
+
+UCLASS()
+class CUSTOMIZABLEOBJECTEDITOR_API UEdGraphSchema_CustomizableObject : public UEdGraphSchema
+{
+public:
+	GENERATED_BODY()
+
+	UEdGraphSchema_CustomizableObject();
+
+	// Allowable PinType.PinCategory values
+	static const FName PC_Object;
+	static const FName PC_Material;
+	static const FName PC_Mesh;
+	static const FName PC_Layout;
+	static const FName PC_Image;
+	static const FName PC_Projector;
+	static const FName PC_GroupProjector;
+	static const FName PC_Color;
+	static const FName PC_Float;
+	static const FName PC_Bool;
+	static const FName PC_Enum;
+	static const FName PC_Stack;
+	static const FName PC_MaterialAsset;
+	
+	// Begin EdGraphSchema interface
+	virtual void GetAssetsGraphHoverMessage(const TArray<FAssetData>& Assets, const UEdGraph* HoverGraph, FString& OutTooltipText, bool& OutOkIcon) const override;
+	void GetGraphContextActions(FGraphContextMenuBuilder& ContextMenuBuilder) const override;
+	const FPinConnectionResponse CanCreateConnection(const UEdGraphPin* A, const UEdGraphPin* B) const override;
+	FLinearColor GetPinTypeColor(const FEdGraphPinType& PinType) const override;
+	bool ShouldHidePinDefaultValue(UEdGraphPin* Pin) const override;
+	virtual void DroppedAssetsOnGraph(const TArray<FAssetData>& Assets, const FVector2D& GraphPosition, UEdGraph* Graph) const override;
+
+	FLinearColor GetPinTypeColor(const FName& PinType) const;
+
+#if (ENGINE_MAJOR_VERSION == 4 && ENGINE_MINOR_VERSION>=18) || ENGINE_MAJOR_VERSION >= 5
+	void GetContextMenuActions(class UToolMenu* Menu, class UGraphNodeContextMenuContext* Context) const override;
+#else
+	void GetContextMenuActions(const UEdGraph* CurrentGraph, const UEdGraphNode* InGraphNode, const UEdGraphPin* InGraphPin, class FMenuBuilder* MenuBuilder, bool bIsDebugging) const override;
+#endif
+
+	virtual void BreakPinLinks(UEdGraphPin& TargetPin, bool bSendsNodeNotification) const override;
+	virtual void BreakSinglePinLink(UEdGraphPin* SourcePin, UEdGraphPin* TargetPin) const override;
+	// End EdGraphSchema interface
+
+	/** Utility to create a node creating action */
+	static TSharedPtr<FCustomizableObjectSchemaAction_NewNode> AddNewNodeAction(FGraphActionListBuilderBase& ContextMenuBuilder, const FString& Category, const FText& MenuDesc, const FText& Tooltip, const int32 Grouping = 0, const FString& Keywords = FString());
+
+	/** */
+	void GetBreakLinkToSubMenuActions( class FMenuBuilder& MenuBuilder, UEdGraphPin* InGraphPin ) const;
+
+private: 
+/** Enum containing all the object types that we are able to convert onto a node when dragging
+ * and dropping and asset of that type onto the CO graph.
+ * Each value of this enumeration will, in practice, have a CO node to be represented by.
+ */
+// TODO: Replace the usage of this enum class with something similar to typeID (not casting)
+enum class ESpawnableObjectType : int32
+{
+	None = -1,		// Invalid value
+	
+	UTexture2D,
+	USkeletalMesh,
+	UStaticMesh,
+	UMaterialInterface,
+};
+
+public:
+
+	/**
+	 * Provided an asset object check if mutable does have a relative node to be able to host this asset's data.
+	 *  @param InAsset - The asset object we want to check if can be used to generate a new mutable node containing it.
+	 *  @param OutObjectType - The type of object of the asset.
+	 *  @return True if the asset can be used to spawn a mutable node, false if not.
+	 *  @note All possible object types that we can convert onto mutable nodes must be defined here.
+	 */
+	bool IsSpawnableAsset(const FAssetData& InAsset, ESpawnableObjectType& OutObjectType) const;
+
+	static TSharedPtr<class ICustomizableObjectEditor> GetCustomizableObjectEditor(const class UEdGraph* ParentGraph);
+	UEdGraphNode* AddComment(class UEdGraph* ParentGraph, UEdGraphPin* FromPin, const FVector2D Location, bool bSelectNewNode = true) const;
+
+	/**
+	 * Break non allowed connections once the connection has been created.
+	 * Allows the nodes to have greater connectivity granularity than simply forbidding to create a connection with UEdGraphSchema_CustomizableObject::CanCreateConnection.
+	 */ 
+	virtual bool TryCreateConnection(UEdGraphPin* PinA, UEdGraphPin* PinB) const override;
+
+	/** Given a pin category name, get its friendly name (user readable). */
+	static FText GetPinCategoryName(const FName& PinCategory);
+
+	/** Given a Material parameter type, returns the Mutable pin category. */
+	static const FName& GetPinCategory(EMaterialParameterType Type);
+};
+
diff --git a/Engine/Plugins/Experimental/Mutable/Source/CustomizableObjectEditor/Classes/Nodes/CustomizableObjectNode.h b/Engine/Plugins/Experimental/Mutable/Source/CustomizableObjectEditor/Classes/Nodes/CustomizableObjectNode.h
new file mode 100644
--- /dev/null
+++ b/Engine/Plugins/Experimental/Mutable/Source/CustomizableObjectEditor/Classes/Nodes/CustomizableObjectNode.h
@@ -0,0 +1,320 @@
+// Copyright Epic Games, Inc. All Rights Reserved.
+
+#pragma once
+
+#include "EdGraph/EdGraphNode.h"
+#include "CustomizableObject.h"
+#include "RemapPins/CustomizableObjectNodeRemapPins.h"
+#include "RemapPins/CustomizableObjectNodeRemapPinsByName.h"
+#include "RemapPins/CustomizableObjectNodeRemapPinsByPosition.h"
+
+#include "CustomizableObjectNode.generated.h"
+
+// Cyclical references
+class UCustomizableObjectNodeObject;
+
+
+/** Abstract base class for all pin data. */
+UCLASS(Abstract)
+class UCustomizableObjectNodePinData : public UObject
+{
+	GENERATED_BODY()
+
+public:
+	UCustomizableObjectNodePinData();
+};
+
+/** Encapsulation of parameters for the FPostEditChangePropertyDelegate delegate function.
+ *
+ * The delegate requires the function being a UFUNCTION which requires all parameters to be UCLASS, USTRUC or UENUM.
+ */
+USTRUCT()
+struct FPostEditChangePropertyDelegateParameters
+{
+	GENERATED_BODY()
+
+	UObject* Node;
+	FPropertyChangedEvent* FPropertyChangedEvent;
+};
+
+DECLARE_DYNAMIC_MULTICAST_DELEGATE_OneParam(FPostEditChangePropertyDelegate, FPostEditChangePropertyDelegateParameters&, Parameters); // Deprecated, use non-dynamic version.
+DECLARE_MULTICAST_DELEGATE_TwoParams(FPostEditChangePropertyRegularDelegate, UObject*, FPropertyChangedEvent&);
+
+DECLARE_MULTICAST_DELEGATE(FPostReconstructNodeDelegate);
+
+DECLARE_MULTICAST_DELEGATE(FNodeConnectionListChangedDelegate);
+
+DECLARE_MULTICAST_DELEGATE(FDestroyNodeDelegate);
+
+using FRemapPinsDelegateParameter = TMap<UEdGraphPin*, UEdGraphPin*>; // Required for the delegate macro.
+DECLARE_MULTICAST_DELEGATE_OneParam(FRemapPinsDelegate, const FRemapPinsDelegateParameter&);
+
+/** Base class of all Customizable Object nodes.
+ * 
+ * The Customizable Object node system is build on top of the following premises. To avoid breaking this system, these premise must be hold!
+ * 
+ * PREMISES:
+ * 1. A node can only be modified when a node is reconstructed (i.e., inside the ReconstructNode, AllocateDefaultPins functions). Exceptionally, a pin can also be created inside the BeginConstruct.
+ *    We consider that a node has been modified if one of the following things happen:
+ *        - Creation, destruction and modification of pins.
+ *        - Modification of any data that would modify any pin.
+ */
+UCLASS(abstract)
+class CUSTOMIZABLEOBJECTEDITOR_API UCustomizableObjectNode : public UEdGraphNode
+{
+public:
+	GENERATED_BODY()
+
+	// UObject interface
+	void Serialize(FArchive& Ar) override;
+	virtual void PostLoad() override; // Do work at PostBackwardsCompatibleFixup.
+
+	// UEdGraphNode interface
+	void AutowireNewNode(UEdGraphPin* FromPin) override;
+	void NodeConnectionListChanged() override;
+	void DestroyNode() override;
+
+	/** Allocates the default pins using the empty remap pins action. Usually called from CreateNode. */
+	void AllocateDefaultPins() override final; // Final. Override AllocateDefaultPins(UCustomizableObjectNodeRemapPins* RemapPins) instead
+	
+	/** Reconstructs the node using its default remap pins action. */
+	void ReconstructNode() override final; // Final. Override ReconstructNode(UCustomizableObjectNodeRemapPins* RemapPins) instead
+	void PostInitProperties() override;
+
+	// Own interface
+	/** Called at the beginning of the node lifecycle.*/
+	virtual void BeginConstruct() {};
+
+	/** Add backwards compatible code here.
+	 * When called, it is guaranteed that all nodes in this graph will have executed the PostLoad function. */
+	virtual void BackwardsCompatibleFixup();
+
+	/** Add post load work here.
+	 * When called, it is guaranteed that all nodes in this graph will have executed the BackwardsCompatibleFixup function. */
+	virtual void PostBackwardsCompatibleFixup() {};
+	
+	/** Virtual implementation of RemovePin. Allows to do work before removing a pin.
+	 * Use this function instead of RemovePin. RemovePin does not removes possible attached PinData. */
+	virtual bool CustomRemovePin(UEdGraphPin& Pin);
+	
+	void GetInputPins(TArray<class UEdGraphPin*>& OutInputPins) const;
+	void GetOutputPins(TArray<class UEdGraphPin*>& OutOutputPins) const;
+	UEdGraphPin* GetOutputPin(int32 OutputIndex) const;
+	
+	virtual bool IsNodeOutDatedAndNeedsRefresh() { return false; }
+	virtual FString GetRefreshMessage() const { return "Refresh Node."; }
+	void SetRefreshNodeWarning();
+	void RemoveWarnings();
+
+	// Used to replace Ids of referenced nodes by their new ids after duplicating the CustomizableObject.
+	virtual void UpdateReferencedNodeId(const FGuid& NewGuid) {};
+
+	TSharedPtr<class ICustomizableObjectEditor> GetGraphEditor() const;
+
+	virtual bool ProvidesCustomPinRelevancyTest() const { return false; } // Override and return true if IsPinRelevant is overridden.
+	virtual bool IsPinRelevant(const UEdGraphPin*) const { return true; } // Override if default pins allocated with AllocateDefaultPins do not provide enough information.
+
+	struct FAttachedErrorDataView
+	{
+		TArrayView<const float> UnassignedUVs;
+	};
+	
+	/** 
+	 * Check if two pins can be connected. 
+	 * 
+	 * @param	InOwnedInputPin	Input pin which belongs to this node.
+	 * @param	InOutputPin		Output pin which belongs to another node. If the node connects to iself, it could belong to this node.
+	 * @param	bOutIsOtherNodeBlocklisted		Is the other node of a type we are not allowed to connect?
+	 * @param	bOutArePinsCompatible		Does InOutputPin pin share the same type as the InOwnedInputPin?
+	 * @return	True if the pin types match and if the other node is not one of the Blocklisted types.
+	 */
+	virtual bool CanConnect( const UEdGraphPin* InOwnedInputPin, const UEdGraphPin* InOutputPin, bool& bOutIsOtherNodeBlocklisted, bool& bOutArePinsCompatible) const;
+
+	// Used during compilation process to cache the node for all LODs, or generate it specifically for each of them.
+	virtual bool IsAffectedByLOD() const { return true; }
+
+	// Get the CustomizableObject graph that owns this node
+	class UCustomizableObjectGraph* GetCustomizableObjectGraph() const;
+
+	/** If returns true, only a single link can go out of every output pin. False by default. */
+	virtual bool IsSingleOutputNode() const;
+
+	// Wrapper of the pin creation method, to cope with multiple Unreal Engine versions
+	UEdGraphPin* CustomCreatePin( EEdGraphPinDirection Direction, const FName& Type, const FName& Name, bool bIsArray=false );
+
+	/** Create pin with attached data.
+	 *
+     * @param PinData Pin data to be saved. */
+	UEdGraphPin* CustomCreatePin(EEdGraphPinDirection Direction, const FName& Type, const FName& Name, UCustomizableObjectNodePinData* PinData);
+
+	
+	virtual void AddAttachedErrorData(const FAttachedErrorDataView& ) {};
+	virtual void ResetAttachedErrorData() {};
+
+	/** 
+	 * When creating a new connection, break all previous connections. This method can be overridden.
+	 * 
+	 * @return false by default.
+	 */
+	virtual bool ShouldBreakExistingConnections(const UEdGraphPin* InputPin, const UEdGraphPin* OutputPin) const;
+
+	/** Called when a new connection is perform. Used to break an specific existing connectoin. Empty default implementation. This method can be overridden.
+	 *
+	 * @param	InputPin	New connection input pin.
+	 * @param	OutputPin	New connection output pin.
+	 */
+	virtual void BreakExistingConnectionsPostConnection(UEdGraphPin* InputPin, UEdGraphPin* OutputPin) {};
+
+	void PostEditChangeProperty(FPropertyChangedEvent& PropertyChangedEvent) override;
+
+
+	// Use this function to reconstruct the graph in the next tick
+	void MarkForReconstruct();
+
+	/** Custom post duplicate function. Called at the beginning of duplication, before the nodes have their Guid updated. */
+	virtual void BeginPostDuplicate(bool bDuplicateForPIE) {};
+	
+	/** Gets all non orphan pins this node owns. */
+	TArray<UEdGraphPin*> GetAllNonOrphanPins() const;
+
+	/** Gets all orphan pins this node owns. */
+	TArray<UEdGraphPin*> GetAllOrphanPins() const;
+
+	/**
+	 * Specialization of ReconstructNode UEdGraphNode function.
+	 *
+	 * @param RemapPins pointer to a remap pins action. Remap pins action which will be used once the node has been reconstructed. Can be null.
+	 */
+	virtual void ReconstructNode(UCustomizableObjectNodeRemapPins* RemapPinsMode);
+
+	/** Return the LOD which this node belongs to. Return -1 if it is not connected (directly or indirectly) to a LOD pin. */
+	int32 GetLOD() const;
+
+	/** Get all parent Object nodes of this node which are parents of this node.
+	 * 
+	 * If a Object node is connected through a LOD pin, it has to match the given LOD.
+	 */
+	TArray<UCustomizableObjectNodeObject*> GetParentObjectNodes(int LOD) const;
+	
+	/** Set a pin to be hidden or not. */
+	void SetPinHidden(UEdGraphPin& Pin, bool bHidden);
+
+	/** Set an array of pins to be hidden or not. 
+	 *
+	 * @param Pins Input pins. All pins must be non null.
+	 */
+	void SetPinHidden(const TArray<UEdGraphPin*>& Pins, bool bHidden);
+
+	/** Return true if a pin can be hidden. Override it to have a custom behaivour. */
+	virtual bool CanPinBeHidden(const UEdGraphPin& Pin) const;
+
+	/** Returns pin custom details. Override if required. */
+	virtual TSharedPtr<SWidget> CustomizePinDetails(UEdGraphPin& Pin);
+	
+	/** See GetPinData(const FEdGraphPinReference&). */
+	UCustomizableObjectNodePinData* GetPinData(const UEdGraphPin &Pin) const;
+	
+	/** Given a pin which has attached data, get the typed pin data.
+	 * @param Pin Pin which has attached data.
+	 * @return Pin data already typed. */
+	template<typename T>
+	T& GetPinData(const UEdGraphPin& Pin) const
+	{
+		return *CastChecked<T>(PinsDataId[Pin.PinId]);
+	}
+	
+	UPROPERTY()
+	FPostEditChangePropertyDelegate PostEditChangePropertyDelegate;
+
+	FPostEditChangePropertyRegularDelegate PostEditChangePropertyRegularDelegate;
+
+	FPostReconstructNodeDelegate PostReconstructNodeDelegate;
+
+	FNodeConnectionListChangedDelegate NodeConnectionListChangedDelegate;
+
+	FDestroyNodeDelegate DestroyNodeDelegate;
+
+	FRemapPinsDelegate RemapPinsDelegate;
+
+protected:
+	/**
+	 * Specialization of AllocateDefaultPins UEdGraphNode function. Override.
+	 * 
+	 * Given the node context (this), create all the required pins.
+	 *
+	 * @param RemapPins remap pins action. Use this object to save any information required to later perfrom the remap pins action. Can be null.
+	 */
+	virtual void AllocateDefaultPins(UCustomizableObjectNodeRemapPins* RemapPins) {};
+
+	/**
+	 * Creates and returns the node default node remap pins action.
+	 * 
+	 * By default all nodes remap by name. Override if a node requires a diferent default remap pins action.
+	 */
+	virtual UCustomizableObjectNodeRemapPins* CreateRemapPinsDefault() const;
+
+	/**
+	 * Creates and returns a remap pins by name action. 
+	 * 
+	 * Override a if a node requires an specialized remap pins by name.
+	 */
+	virtual UCustomizableObjectNodeRemapPinsByName* CreateRemapPinsByName() const;
+
+	/**
+	 * Creates and returns a remap pins by position action.
+	 *
+	 * Override a if a node requires an specialized remap pins by position.
+	 */
+	virtual UCustomizableObjectNodeRemapPinsByPosition* CreateRemapPinsByPosition() const;
+	
+	/** Allows to perform work when remapping a pin. */
+	virtual void RemapPins(const TMap<UEdGraphPin*, UEdGraphPin*>& PinsToRemap);
+
+	/*** Allows to perform work when remapping the pin data. */
+	virtual void RemapPinsData(const TMap<UEdGraphPin*, UEdGraphPin*>& PinsToRemap);
+	
+	/** Attach data to a given pin. */
+	void AddPinData(const UEdGraphPin& Pin, UCustomizableObjectNodePinData& PinData);
+
+private:
+	/** Data attached to a given pin. Not all pins contain data.
+	 * Data must be eliminated once the pins is removed. This is done automatically on CustomRemovePin.
+	 * Do not use RemovePin since it will not remove the data. */
+	UPROPERTY()
+	TMap<FGuid, TObjectPtr<UCustomizableObjectNodePinData>> PinsDataId;
+
+	// Deprectated properties
+	UPROPERTY()
+	TMap<FEdGraphPinReference, TObjectPtr<UCustomizableObjectNodePinData>> PinsData_DEPRECATED;
+};
+
+
+/** Return the specified node located in another Customizable Object graph.
+ *
+ * @param Customizable Object object to look for the node. If nulltpr, it returns nullptr.
+ * @param NodeGuid node guid. If invalid, it returns nullptr.
+ * 
+ * Return nullptr if not found.
+ */
+template<class NODE_TYPE>
+NODE_TYPE* GetCustomizableObjectExternalNode(UCustomizableObject* Object, const FGuid& NodeGuid)
+{
+	NODE_TYPE* Result = nullptr;
+
+	if (Object && Object->Source && NodeGuid.IsValid())
+	{
+		TArray<NODE_TYPE*> GroupNodes;
+		Object->Source->GetNodesOfClass<NODE_TYPE>(GroupNodes);
+
+		for (NODE_TYPE* GroupNode : GroupNodes)
+		{
+			if (NodeGuid == GroupNode->NodeGuid)
+			{
+				Result = GroupNode;
+				break;
+			}
+		}
+	}
+
+	return Result;
+}
\ No newline at end of file
diff --git a/Engine/Plugins/Experimental/Mutable/Source/CustomizableObjectEditor/Classes/Nodes/CustomizableObjectNodeAnimationPose.h b/Engine/Plugins/Experimental/Mutable/Source/CustomizableObjectEditor/Classes/Nodes/CustomizableObjectNodeAnimationPose.h
new file mode 100644
--- /dev/null
+++ b/Engine/Plugins/Experimental/Mutable/Source/CustomizableObjectEditor/Classes/Nodes/CustomizableObjectNodeAnimationPose.h
@@ -0,0 +1,36 @@
+// Copyright Epic Games, Inc. All Rights Reserved.
+
+#pragma once
+
+#include "Nodes/CustomizableObjectNode.h"
+#include "Animation/PoseAsset.h"
+#include "CustomizableObjectNodeAnimationPose.generated.h"
+
+
+UCLASS()
+class CUSTOMIZABLEOBJECTEDITOR_API UCustomizableObjectNodeAnimationPose : public UCustomizableObjectNode
+{
+public:
+	GENERATED_BODY()
+
+	UCustomizableObjectNodeAnimationPose();
+
+	UPROPERTY(EditAnywhere, Category=CustomizableObject)
+	TObjectPtr<UPoseAsset> PoseAsset;
+
+	// UObject interface.
+	virtual void PostEditChangeProperty(FPropertyChangedEvent& PropertyChangedEvent) override;
+
+	// Begin EdGraphNode interface
+	UEdGraphPin* GetInputMeshPin() const;
+	virtual FText GetNodeTitle(ENodeTitleType::Type TitleType) const override;
+	virtual FLinearColor GetNodeTitleColor() const override;
+	// End EdGraphNode interface
+
+	// UCustomizableObjectNode interface
+	void AllocateDefaultPins(UCustomizableObjectNodeRemapPins* RemapPins) override;
+
+	// Retrieve the pose information from the PoseAsset
+	static void StaticRetrievePoseInformation(UPoseAsset* PoseAsset, class USkeletalMesh* RefSkeletalMesh, TArray<FString>& OutArrayBoneName, TArray<FTransform>& OutArrayTransform);
+};
+
diff --git a/Engine/Plugins/Experimental/Mutable/Source/CustomizableObjectEditor/Classes/Nodes/CustomizableObjectNodeColorArithmeticOp.h b/Engine/Plugins/Experimental/Mutable/Source/CustomizableObjectEditor/Classes/Nodes/CustomizableObjectNodeColorArithmeticOp.h
new file mode 100644
--- /dev/null
+++ b/Engine/Plugins/Experimental/Mutable/Source/CustomizableObjectEditor/Classes/Nodes/CustomizableObjectNodeColorArithmeticOp.h
@@ -0,0 +1,58 @@
+// Copyright Epic Games, Inc. All Rights Reserved.
+
+#pragma once
+
+#include "Nodes/CustomizableObjectNode.h"
+
+#include "CustomizableObjectNodeColorArithmeticOp.generated.h"
+
+
+UENUM(BlueprintType)
+enum class EColorArithmeticOperation : uint8
+{
+	E_Add 	UMETA(DisplayName = "+"),
+	E_Sub 	UMETA(DisplayName = "-"),
+	E_Mul	UMETA(DisplayName = "x"),
+	E_Div	UMETA(DisplayName = "/")
+};
+
+UCLASS()
+class CUSTOMIZABLEOBJECTEDITOR_API UCustomizableObjectNodeColorArithmeticOp : public UCustomizableObjectNode
+{
+public:
+	GENERATED_BODY()
+
+	UCustomizableObjectNodeColorArithmeticOp();
+
+	// UObject interface.
+	void PostEditChangeProperty(FPropertyChangedEvent& PropertyChangedEvent) override;
+
+	// Begin EdGraphNode interface
+	FText GetNodeTitle(ENodeTitleType::Type TitleType) const override;
+	FLinearColor GetNodeTitleColor() const override;
+	FText GetTooltipText() const override;
+	bool ShouldOverridePinNames() const override { return true; }
+
+	// UCustomizableObjectNode interface
+	void AllocateDefaultPins(UCustomizableObjectNodeRemapPins* RemapPins) override;
+
+	// First Operand
+	UEdGraphPin* XPin() const
+	{
+		return FindPin(TEXT("A"));
+	}
+
+	// Second Operand
+	UEdGraphPin* YPin() const
+	{
+		return FindPin(TEXT("B"));
+	}
+
+	UPROPERTY(EditAnywhere, Category = ColorArithmeticOperation)
+	EColorArithmeticOperation Operation;
+
+	// UCustomizableObjectNode interface
+	virtual bool IsAffectedByLOD() const { return false; }
+
+};
+
diff --git a/Engine/Plugins/Experimental/Mutable/Source/CustomizableObjectEditor/Classes/Nodes/CustomizableObjectNodeColorConstant.h b/Engine/Plugins/Experimental/Mutable/Source/CustomizableObjectEditor/Classes/Nodes/CustomizableObjectNodeColorConstant.h
new file mode 100644
--- /dev/null
+++ b/Engine/Plugins/Experimental/Mutable/Source/CustomizableObjectEditor/Classes/Nodes/CustomizableObjectNodeColorConstant.h
@@ -0,0 +1,100 @@
+// Copyright Epic Games, Inc. All Rights Reserved.
+
+#pragma once
+
+#include "Nodes/CustomizableObjectNode.h"
+#include "SGraphNode.h"
+
+#include "CustomizableObjectNodeColorConstant.generated.h"
+
+
+enum class ColorChannel
+{
+	RED,
+	GREEN,
+	BLUE,
+	ALPHA
+};
+
+// Class create color inputs in the NodeColorConstant
+class SGraphNodeColorConstant : public SGraphNode
+{
+public:
+
+	SLATE_BEGIN_ARGS(SGraphNodeColorConstant) {}
+	SLATE_END_ARGS();
+
+	SGraphNodeColorConstant() : SGraphNode() {};
+
+	// Builds the SGraphNodeFloatConstant when needed
+	void Construct(const FArguments& InArgs, UEdGraphNode* InGraphNode);
+
+	// Calls the needed functions to build the SGraphNode widgets
+	void UpdateGraphNode();
+
+	// Overriden functions to build the SGraphNode widgets
+	virtual void SetDefaultTitleAreaWidget(TSharedRef<SOverlay> DefaultTitleAreaWidget) override;
+	virtual void CreateBelowPinControls(TSharedPtr<SVerticalBox> MainBox) override;
+	virtual bool ShouldAllowCulling() const override { return false; }
+
+	// Callbacks for the widget
+	void OnExpressionPreviewChanged(const ECheckBoxState NewCheckedState);
+	ECheckBoxState IsExpressionPreviewChecked() const;
+	const FSlateBrush* GetExpressionPreviewArrow() const;
+
+private:
+
+	// Callback for the OnValueChanged of the SpinBox
+	void OnSpinBoxValueChanged(float Value, ColorChannel channel);
+
+	// Callback for the OnValueCommited of the SpinBox
+	void OnSpinBoxValueCommitted(float Value, ETextCommit::Type, ColorChannel channel);
+
+	// Callback for the OnClicked of the ColorBox
+	FReply OnColorPreviewClicked(const FGeometry& MyGeometry, const FPointerEvent& MouseEvent);
+
+private:
+
+	// Pointer to the NodeFloatconstant that owns this SGraphNode
+	UCustomizableObjectNodeColorConstant* NodeColorConstant;
+
+	// Style for the SpinBox
+	FSpinBoxStyle WidgetStyle;
+
+	// Pointer to store the color editors to manage the visibility
+	TSharedPtr<SVerticalBox> ColorEditor;
+};
+
+UCLASS()
+class CUSTOMIZABLEOBJECTEDITOR_API UCustomizableObjectNodeColorConstant : public UCustomizableObjectNode
+{
+public:
+	GENERATED_BODY()
+
+	UCustomizableObjectNodeColorConstant();
+
+	/**  */
+	UPROPERTY(EditAnywhere, Category=CustomizableObject, meta = (DontUpdateWhileEditing))
+	FLinearColor Value;
+
+	// UObject interface.
+	void PostEditChangeProperty(FPropertyChangedEvent& PropertyChangedEvent) override;
+
+	// Begin EdGraphNode interface
+	FText GetNodeTitle(ENodeTitleType::Type TitleType) const override;
+	FLinearColor GetNodeTitleColor() const override;
+	FText GetTooltipText() const override;
+	// End EdGraphNode interface
+
+	// UCustomizableObjectNode interface
+	void AllocateDefaultPins(UCustomizableObjectNodeRemapPins* RemapPins) override;
+	bool IsAffectedByLOD() const override { return false; }
+
+	// Creates the SGraph Node widget for the Color Editor
+	TSharedPtr<SGraphNode> CreateVisualWidget() override;
+
+	// Determines if the Node is collapsed or not
+	bool bCollapsed;
+
+};
+
diff --git a/Engine/Plugins/Experimental/Mutable/Source/CustomizableObjectEditor/Classes/Nodes/CustomizableObjectNodeColorFromFloats.h b/Engine/Plugins/Experimental/Mutable/Source/CustomizableObjectEditor/Classes/Nodes/CustomizableObjectNodeColorFromFloats.h
new file mode 100644
--- /dev/null
+++ b/Engine/Plugins/Experimental/Mutable/Source/CustomizableObjectEditor/Classes/Nodes/CustomizableObjectNodeColorFromFloats.h
@@ -0,0 +1,52 @@
+// Copyright Epic Games, Inc. All Rights Reserved.
+
+#pragma once
+
+#include "Nodes/CustomizableObjectNode.h"
+
+#include "CustomizableObjectNodeColorFromFloats.generated.h"
+
+
+UCLASS()
+class CUSTOMIZABLEOBJECTEDITOR_API UCustomizableObjectNodeColorFromFloats : public UCustomizableObjectNode
+{
+public:
+	GENERATED_BODY()
+
+	UCustomizableObjectNodeColorFromFloats();
+
+	// UObject interface.
+	void PostEditChangeProperty(FPropertyChangedEvent& PropertyChangedEvent) override;
+
+	// Begin EdGraphNode interface
+	FText GetNodeTitle(ENodeTitleType::Type TitleType) const override;
+	FLinearColor GetNodeTitleColor() const override;
+	FText GetTooltipText() const override;
+
+	// UCustomizableObjectNode interface
+	void AllocateDefaultPins(UCustomizableObjectNodeRemapPins* RemapPins) override;
+
+	UEdGraphPin* RPin() const
+	{
+		return FindPin(TEXT("R"));
+	}
+
+	UEdGraphPin* GPin() const
+	{
+		return FindPin(TEXT("G"));
+	}
+
+	UEdGraphPin* BPin() const
+	{
+		return FindPin(TEXT("B"));
+	}
+
+	UEdGraphPin* APin() const
+	{
+		return FindPin(TEXT("A"));
+	}
+
+	// UCustomizableObjectNode interface
+	virtual bool IsAffectedByLOD() const override { return false; }
+};
+
diff --git a/Engine/Plugins/Experimental/Mutable/Source/CustomizableObjectEditor/Classes/Nodes/CustomizableObjectNodeColorParameter.h b/Engine/Plugins/Experimental/Mutable/Source/CustomizableObjectEditor/Classes/Nodes/CustomizableObjectNodeColorParameter.h
new file mode 100644
--- /dev/null
+++ b/Engine/Plugins/Experimental/Mutable/Source/CustomizableObjectEditor/Classes/Nodes/CustomizableObjectNodeColorParameter.h
@@ -0,0 +1,34 @@
+// Copyright Epic Games, Inc. All Rights Reserved.
+
+#pragma once
+
+#include "Nodes/CustomizableObjectNode.h"
+
+#include "CustomizableObjectNodeColorParameter.generated.h"
+
+
+UCLASS()
+class CUSTOMIZABLEOBJECTEDITOR_API UCustomizableObjectNodeColorParameter : public UCustomizableObjectNode
+{
+public:
+	GENERATED_BODY()
+
+	UPROPERTY(EditAnywhere, Category=CustomizableObject, meta = (DontUpdateWhileEditing))
+	FLinearColor DefaultValue = FLinearColor(1, 1, 1, 1);
+
+	UPROPERTY(EditAnywhere, Category=CustomizableObject)
+	FString ParameterName = "Default Name";
+
+	UPROPERTY(EditAnywhere, Category = UI, meta = (DisplayName = "Parameter UI Metadata"))
+	FMutableParamUIMetadata ParamUIMetadata;
+
+	// Begin EdGraphNode interface
+	FText GetNodeTitle(ENodeTitleType::Type TitleType) const override;
+	FLinearColor GetNodeTitleColor() const override;
+	FText GetTooltipText() const override;
+
+	// UCustomizableObjectNode interface
+	void AllocateDefaultPins(UCustomizableObjectNodeRemapPins* RemapPins) override;
+	bool IsAffectedByLOD() const override { return false; }
+};
+
diff --git a/Engine/Plugins/Experimental/Mutable/Source/CustomizableObjectEditor/Classes/Nodes/CustomizableObjectNodeColorSwitch.h b/Engine/Plugins/Experimental/Mutable/Source/CustomizableObjectEditor/Classes/Nodes/CustomizableObjectNodeColorSwitch.h
new file mode 100644
--- /dev/null
+++ b/Engine/Plugins/Experimental/Mutable/Source/CustomizableObjectEditor/Classes/Nodes/CustomizableObjectNodeColorSwitch.h
@@ -0,0 +1,29 @@
+// Copyright Epic Games, Inc. All Rights Reserved.
+
+#pragma once
+
+#include "Nodes/CustomizableObjectNodeSwitchBase.h"
+
+#include "EdGraph/EdGraphPin.h"
+
+#include "CustomizableObjectNodeColorSwitch.generated.h"
+
+
+UCLASS()
+class CUSTOMIZABLEOBJECTEDITOR_API UCustomizableObjectNodeColorSwitch : public UCustomizableObjectNodeSwitchBase
+{
+public:
+	GENERATED_BODY()
+
+	// UObject interface
+	virtual void Serialize(FArchive& Ar) override;
+	
+	// UCustomizableObjectNodeSwitchBase interface
+	FString GetOutputPinName() const override;
+
+	FName GetCategory() const override;
+
+	// Own interface
+	FString GetPinPrefix() const;
+};
+
diff --git a/Engine/Plugins/Experimental/Mutable/Source/CustomizableObjectEditor/Classes/Nodes/CustomizableObjectNodeColorVariation.h b/Engine/Plugins/Experimental/Mutable/Source/CustomizableObjectEditor/Classes/Nodes/CustomizableObjectNodeColorVariation.h
new file mode 100644
--- /dev/null
+++ b/Engine/Plugins/Experimental/Mutable/Source/CustomizableObjectEditor/Classes/Nodes/CustomizableObjectNodeColorVariation.h
@@ -0,0 +1,65 @@
+// Copyright Epic Games, Inc. All Rights Reserved.
+
+#pragma once
+
+#include "Nodes/CustomizableObjectNode.h"
+
+#include "CustomizableObjectNodeColorVariation.generated.h"
+
+
+USTRUCT()
+struct CUSTOMIZABLEOBJECTEDITOR_API FCustomizableObjectColorVariation
+{
+	GENERATED_USTRUCT_BODY()
+
+	UPROPERTY(EditAnywhere, Category = CustomizableObject)
+	FString Tag;
+};
+
+
+UCLASS()
+class CUSTOMIZABLEOBJECTEDITOR_API UCustomizableObjectNodeColorVariation : public UCustomizableObjectNode
+{
+public:
+	GENERATED_BODY()
+
+	UCustomizableObjectNodeColorVariation();
+
+	UPROPERTY(EditAnywhere, Category = CustomizableObject)
+	TArray<FCustomizableObjectColorVariation> Variations;
+
+	// UObject interface.
+	void PostEditChangeProperty(FPropertyChangedEvent& PropertyChangedEvent) override;
+
+	// EdGraphNode interface
+	FText GetNodeTitle(ENodeTitleType::Type TitleType) const override;
+	FLinearColor GetNodeTitleColor() const override;
+	FText GetTooltipText() const override;
+
+	// UCustomizableObjectNode interface
+	void AllocateDefaultPins(UCustomizableObjectNodeRemapPins* RemapPins) override;
+
+	UEdGraphPin* OutputPin() const
+	{
+		return FindPin(TEXT("Color"));
+	}
+
+	UEdGraphPin* DefaultPin() const
+	{
+		return FindPin(TEXT("Default"));
+	}
+
+	int32 GetNumVariations() const
+	{
+		return Variations.Num();
+	}
+
+	UEdGraphPin* VariationPin(int Index) const
+	{
+		return FindPin(FString::Printf(TEXT("Variation %d"),Index));
+	}
+
+private:
+	TMap<UEdGraphPin*, UEdGraphPin*> RemapPinsMap;
+};
+
diff --git a/Engine/Plugins/Experimental/Mutable/Source/CustomizableObjectEditor/Classes/Nodes/CustomizableObjectNodeCopyMaterial.h b/Engine/Plugins/Experimental/Mutable/Source/CustomizableObjectEditor/Classes/Nodes/CustomizableObjectNodeCopyMaterial.h
new file mode 100644
--- /dev/null
+++ b/Engine/Plugins/Experimental/Mutable/Source/CustomizableObjectEditor/Classes/Nodes/CustomizableObjectNodeCopyMaterial.h
@@ -0,0 +1,67 @@
+// Copyright Epic Games, Inc. All Rights Reserved.
+
+#pragma once
+
+#include "Nodes/CustomizableObjectNodeMaterial.h"
+#include "Nodes/CustomizableObjectNodeSkeletalMesh.h"
+
+#include "CustomizableObjectNodeCopyMaterial.generated.h"
+
+/**
+ * Copy Material node. Duplicates a Material Node. Duplicates all Material node input pins and properties except for the Mesh input pin.
+ * A new Mesh has to be defined through the new Mesh input pin.
+ *
+ * Input pins:
+ * - Mesh: New mesh.
+ * - Material: Material to duplicate.
+ * 
+ * Output pins:
+ * - Material: New material.
+ * 
+ * Properties:
+ *   [NONE]
+ */
+UCLASS()
+class CUSTOMIZABLEOBJECTEDITOR_API UCustomizableObjectNodeCopyMaterial : public UCustomizableObjectNodeMaterial
+{
+public:
+	GENERATED_BODY()
+
+	FText GetNodeTitle(ENodeTitleType::Type TitleType) const override;
+
+	// UCustomizableObjectNode interface
+	void AllocateDefaultPins(UCustomizableObjectNodeRemapPins* RemapPins) override;
+
+	/** Get Mesh input pin. */
+	UEdGraphPin* GetMeshPin() const override;
+
+	/** Get Material input pin. */
+	UEdGraphPin* GetMaterialPin() const;
+
+	/** Get Mesh node. */
+	UCustomizableObjectNodeSkeletalMesh* GetMeshNode() const;
+
+	/** Get the Material node. */
+	UCustomizableObjectNodeMaterial* GetMaterialNode() const;
+
+	/** Input material pin can only be connect to a NodeMaterial output pin or a NodeExternalPin. A connection to a NodeMaterialCopy output pin is not allowed. */
+	virtual bool CanConnect(const UEdGraphPin* InOwnedInputPin, const UEdGraphPin* InOutputPin, bool& bOutIsOtherNodeBlocklisted, bool& bOutArePinsCompatible) const override;
+	
+	bool ShouldBreakExistingConnections(const UEdGraphPin* InputPin, const UEdGraphPin* OutputPin) const override;
+
+	/** Override the NodeMaterial method. Does not break any connection. */
+	void BreakExistingConnectionsPostConnection(UEdGraphPin* InputPin, UEdGraphPin* OutputPin) override {};
+
+	virtual bool IsNodeOutDatedAndNeedsRefresh() override;
+
+	bool ProvidesCustomPinRelevancyTest() const override;
+
+	/** 
+	 * Only are relevant pins:
+	 * - From: NodeMaterial, excluding NodeCopyMaterial 
+	 * - To: NodeBaseObject
+	 */
+	bool IsPinRelevant(const UEdGraphPin* Pin) const override;
+
+	FText GetTooltipText() const override;
+};
\ No newline at end of file
diff --git a/Engine/Plugins/Experimental/Mutable/Source/CustomizableObjectEditor/Classes/Nodes/CustomizableObjectNodeCurve.h b/Engine/Plugins/Experimental/Mutable/Source/CustomizableObjectEditor/Classes/Nodes/CustomizableObjectNodeCurve.h
new file mode 100644
--- /dev/null
+++ b/Engine/Plugins/Experimental/Mutable/Source/CustomizableObjectEditor/Classes/Nodes/CustomizableObjectNodeCurve.h
@@ -0,0 +1,47 @@
+// Copyright Epic Games, Inc. All Rights Reserved.
+
+#pragma once
+
+#include "Nodes/CustomizableObjectNode.h"
+#include "Curves/CurveBase.h"
+
+#include "CustomizableObjectNodeCurve.generated.h"
+
+
+UCLASS()
+class CUSTOMIZABLEOBJECTEDITOR_API UCustomizableObjectNodeCurve : public UCustomizableObjectNode
+{
+public:
+	GENERATED_BODY()
+
+	UCustomizableObjectNodeCurve();
+
+	// UObject interface.
+	void PostEditChangeProperty(FPropertyChangedEvent& PropertyChangedEvent) override;
+	virtual void Serialize(FArchive& Ar) override;
+
+	// Begin EdGraphNode interface
+	FText GetNodeTitle(ENodeTitleType::Type TitleType) const override;
+	FLinearColor GetNodeTitleColor() const override;
+	FText GetTooltipText() const override;
+
+	// UCustomizableObjectNode interface
+	void AllocateDefaultPins(UCustomizableObjectNodeRemapPins* RemapPins) override;
+
+	UEdGraphPin* InputPin() const;
+	UEdGraphPin* CurvePins(int32 Index) const;
+	int32 GetNumCurvePins() const;
+
+	UPROPERTY(EditAnywhere, Category = Curve)
+	TObjectPtr<UCurveBase> CurveAsset;
+
+	// UCustomizableObjectNode interface
+	virtual bool IsAffectedByLOD() const override { return false; }
+
+	bool ProvidesCustomPinRelevancyTest() const override { return true; }
+	bool IsPinRelevant(const UEdGraphPin* Pin) const override;
+
+	/** Override the default behaivour of remap pins. Use remap pins by position by default. */
+	UCustomizableObjectNodeRemapPins* CreateRemapPinsDefault() const override;
+};
+
diff --git a/Engine/Plugins/Experimental/Mutable/Source/CustomizableObjectEditor/Classes/Nodes/CustomizableObjectNodeEditLayoutBlocks.h b/Engine/Plugins/Experimental/Mutable/Source/CustomizableObjectEditor/Classes/Nodes/CustomizableObjectNodeEditLayoutBlocks.h
new file mode 100644
--- /dev/null
+++ b/Engine/Plugins/Experimental/Mutable/Source/CustomizableObjectEditor/Classes/Nodes/CustomizableObjectNodeEditLayoutBlocks.h
@@ -0,0 +1,21 @@
+// Copyright Epic Games, Inc. All Rights Reserved.
+
+#pragma once
+
+#include "Nodes/CustomizableObjectNodeEditMaterialBase.h"
+
+#include "CustomizableObjectNodeEditLayoutBlocks.generated.h"
+
+
+UCLASS(Abstract)
+class CUSTOMIZABLEOBJECTEDITOR_API UCustomizableObjectNodeEditLayoutBlocks : public UCustomizableObjectNodeEditMaterialBase
+{
+public:
+	GENERATED_BODY()
+
+	// Selected blocks
+	UPROPERTY()
+	TArray<FGuid> BlockIds;
+
+};
+
diff --git a/Engine/Plugins/Experimental/Mutable/Source/CustomizableObjectEditor/Classes/Nodes/CustomizableObjectNodeEditMaterial.h b/Engine/Plugins/Experimental/Mutable/Source/CustomizableObjectEditor/Classes/Nodes/CustomizableObjectNodeEditMaterial.h
new file mode 100644
--- /dev/null
+++ b/Engine/Plugins/Experimental/Mutable/Source/CustomizableObjectEditor/Classes/Nodes/CustomizableObjectNodeEditMaterial.h
@@ -0,0 +1,84 @@
+// Copyright Epic Games, Inc. All Rights Reserved.
+
+#pragma once
+
+#include "Nodes/CustomizableObjectNodeEditLayoutBlocks.h"
+#include "Nodes/CustomizableObjectNodeUseMaterial.h"
+#include "CustomizableObjectEditor_Deprecated.h"
+
+#include "CustomizableObjectNodeEditMaterial.generated.h"
+
+
+/** Additional data for the Parent Texture Parameter to edit pin. */
+UCLASS()
+class CUSTOMIZABLEOBJECTEDITOR_API UCustomizableObjectNodeEditMaterialPinEditImageData : public UCustomizableObjectNodePinData
+{
+	GENERATED_BODY()
+
+public:
+	UPROPERTY()
+	FGuid ImageId;
+
+	UPROPERTY();
+	FEdGraphPinReference PinMask;
+};
+
+
+UCLASS()
+class CUSTOMIZABLEOBJECTEDITOR_API UCustomizableObjectNodeEditMaterial :
+	public UCustomizableObjectNodeEditLayoutBlocks,
+	public FCustomizableObjectNodeUseMaterial
+{
+public:
+	GENERATED_BODY()
+
+	// UObject interface
+	virtual void Serialize(FArchive& Ar) override;
+	
+	// EdGraphNode interface
+	virtual FText GetNodeTitle(ENodeTitleType::Type TitleType) const override;
+	virtual FLinearColor GetNodeTitleColor() const override;
+	virtual FText GetTooltipText() const override;
+	virtual void PinConnectionListChanged(UEdGraphPin* Pin) override;
+
+	// UCustomizableObjectNode interface
+	virtual void BackwardsCompatibleFixup() override;
+	virtual void PostBackwardsCompatibleFixup() override;
+	virtual void AllocateDefaultPins(UCustomizableObjectNodeRemapPins* RemapPins) override;
+	virtual void RemapPins(const TMap<UEdGraphPin*, UEdGraphPin*>& PinsToRemap) override;
+	virtual bool IsNodeOutDatedAndNeedsRefresh() override;
+	virtual FString GetRefreshMessage() const override;
+	virtual bool IsSingleOutputNode() const override;
+	virtual bool CustomRemovePin(UEdGraphPin& Pin) override;
+
+	// UCustomizableObjectNodeEditMaterialBase interface
+	virtual UEdGraphPin* OutputPin() const override;
+	
+	// ICustomizableObjectNodeParentedMaterial interface
+	virtual void SetParentNode(UCustomizableObject* Object, FGuid NodeId) override;
+
+	// FCustomizableObjectNodeUseMaterial interface
+	virtual UCustomizableObjectNode& GetNode() override;
+	virtual FCustomizableObjectNodeParentedMaterial& GetNodeParentedMaterial() override;
+	virtual TMap<FGuid, FEdGraphPinReference>& GetPinsParameter() override;
+	
+	
+	/** Returns the Image mask pin of the given Image that will be edited.
+	 *
+	 * @returns Always returns a valid pin if EditsImage(const FGuid&) returns true. */
+	const UEdGraphPin* GetUsedImageMaskPin(const FGuid& ImageId) const;
+
+private:
+	/** Relates a Parameter id to a Pin. Only used to improve performance. */
+	UPROPERTY()
+	TMap<FGuid, FEdGraphPinReference> PinsParameter;
+	
+	// Deprecated properties
+	UPROPERTY()
+	TArray<FCustomizableObjectNodeEditMaterialImage> Images_DEPRECATED;
+	
+	// Old layout blocks to patch. Now in parent class.
+	UPROPERTY()
+	TArray<int32> Blocks_DEPRECATED;
+};
+
diff --git a/Engine/Plugins/Experimental/Mutable/Source/CustomizableObjectEditor/Classes/Nodes/CustomizableObjectNodeEditMaterialBase.h b/Engine/Plugins/Experimental/Mutable/Source/CustomizableObjectEditor/Classes/Nodes/CustomizableObjectNodeEditMaterialBase.h
new file mode 100644
--- /dev/null
+++ b/Engine/Plugins/Experimental/Mutable/Source/CustomizableObjectEditor/Classes/Nodes/CustomizableObjectNodeEditMaterialBase.h
@@ -0,0 +1,45 @@
+// Copyright Epic Games, Inc. All Rights Reserved.
+
+#pragma once
+
+#include "Nodes/CustomizableObjectNode.h"
+#include "CustomizableObjectNodeParentedMaterial.h"
+
+#include "CustomizableObjectNodeEditMaterialBase.generated.h"
+
+
+UCLASS(Abstract)
+class CUSTOMIZABLEOBJECTEDITOR_API UCustomizableObjectNodeEditMaterialBase : public UCustomizableObjectNode, public FCustomizableObjectNodeParentedMaterial
+{
+	GENERATED_BODY()
+	
+public:
+	// FCustomizableObjectNodeParentedMaterial interface
+	virtual void SetParentNode(UCustomizableObject* Object, FGuid NodeId) override;
+	virtual void SaveParentNode(UCustomizableObject* Object, FGuid NodeId) override;
+	virtual UCustomizableObjectNode& GetNode() override;
+	virtual FGuid GetParentNodeId() const override;
+	virtual UCustomizableObject* GetParentObject() const override;
+
+	// Own interface
+	/** Index of the layout to use to patch blocks. */
+	UPROPERTY(EditAnywhere, Category = CustomizableObject)
+	int ParentLayoutIndex = 0;
+
+	void BeginPostDuplicate(bool bDuplicateForPIE) override;
+
+	void UpdateReferencedNodeId(const FGuid& NewGuid) override;
+
+	virtual UEdGraphPin* OutputPin() const
+	{
+		return FindPin(TEXT("Material"));
+	}
+
+private:
+	UPROPERTY(EditAnywhere, Category = CustomizableObject)
+	TObjectPtr<UCustomizableObject> ParentMaterialObject = nullptr;
+
+	UPROPERTY()
+	FGuid ParentMaterialNodeId;
+};
+
diff --git a/Engine/Plugins/Experimental/Mutable/Source/CustomizableObjectEditor/Classes/Nodes/CustomizableObjectNodeEnumParameter.h b/Engine/Plugins/Experimental/Mutable/Source/CustomizableObjectEditor/Classes/Nodes/CustomizableObjectNodeEnumParameter.h
new file mode 100644
--- /dev/null
+++ b/Engine/Plugins/Experimental/Mutable/Source/CustomizableObjectEditor/Classes/Nodes/CustomizableObjectNodeEnumParameter.h
@@ -0,0 +1,53 @@
+// Copyright Epic Games, Inc. All Rights Reserved.
+
+#pragma once
+
+#include "Nodes/CustomizableObjectNode.h"
+
+#include "CustomizableObjectNodeEnumParameter.generated.h"
+
+
+USTRUCT()
+struct CUSTOMIZABLEOBJECTEDITOR_API FCustomizableObjectNodeEnumValue
+{
+	GENERATED_USTRUCT_BODY()
+
+	UPROPERTY(EditAnywhere, Category=CustomizableObject)
+	FString Name;
+
+	UPROPERTY(EditAnywhere, Category = UI, meta = (DisplayName = "Parameter UI Metadata"))
+	FMutableParamUIMetadata ParamUIMetadata;
+};
+
+
+UCLASS()
+class CUSTOMIZABLEOBJECTEDITOR_API UCustomizableObjectNodeEnumParameter : public UCustomizableObjectNode
+{
+public:
+	GENERATED_BODY()
+
+	UPROPERTY(EditAnywhere, Category=CustomizableObject)
+	int32 DefaultIndex = 0;
+
+	UPROPERTY(EditAnywhere, Category=CustomizableObject)
+	FString ParameterName = "Default Name";
+
+	UPROPERTY(EditAnywhere, Category = UI, meta = (DisplayName = "Parameter UI Metadata"))
+	FMutableParamUIMetadata ParamUIMetadata;
+
+	UPROPERTY(EditAnywhere, Category=CustomizableObject)
+	TArray<FCustomizableObjectNodeEnumValue> Values;
+
+	// UObject interface.
+	void PostEditChangeProperty(FPropertyChangedEvent& PropertyChangedEvent) override;
+
+	// Begin EdGraphNode interface
+	FText GetNodeTitle(ENodeTitleType::Type TitleType) const override;
+	FLinearColor GetNodeTitleColor() const override;
+	FText GetTooltipText() const override;
+
+	// UCustomizableObjectNode interface
+	void AllocateDefaultPins(UCustomizableObjectNodeRemapPins* RemapPins) override;
+	virtual bool IsAffectedByLOD() const override { return false; }
+};
+
