[mutable] Move Mutable plugin from the restricted folder to the experimental plugins folder.

diff --git a/Engine/Plugins/Experimental/Mutable/Source/CustomizableObjectEditor/Classes/Nodes/CustomizableObjectNodeExposePin.h b/Engine/Plugins/Experimental/Mutable/Source/CustomizableObjectEditor/Classes/Nodes/CustomizableObjectNodeExposePin.h
new file mode 100644
--- /dev/null
+++ b/Engine/Plugins/Experimental/Mutable/Source/CustomizableObjectEditor/Classes/Nodes/CustomizableObjectNodeExposePin.h
@@ -0,0 +1,50 @@
+// Copyright Epic Games, Inc. All Rights Reserved.
+
+#pragma once
+
+#include "Nodes/CustomizableObjectNode.h"
+
+#include "CustomizableObjectNodeExposePin.generated.h"
+
+
+DECLARE_MULTICAST_DELEGATE(FOnNameChangedDelegate);
+
+UCLASS()
+class CUSTOMIZABLEOBJECTEDITOR_API UCustomizableObjectNodeExposePin : public UCustomizableObjectNode
+{
+public:
+	GENERATED_BODY()
+
+	// UObject interface
+	void PostEditChangeProperty(FPropertyChangedEvent& PropertyChangedEvent) override;
+
+	// EdGraphNode interface
+	FText GetNodeTitle(ENodeTitleType::Type TitleType) const override;
+	FLinearColor GetNodeTitleColor() const override;
+	FText GetTooltipText() const override;
+
+	virtual bool CanConnect(const UEdGraphPin* InOwnedInputPin, const UEdGraphPin* InOutputPin, bool& bOutIsOtherNodeBlocklisted, bool& bOutArePinsCompatible) const override;
+
+	// UCustomizableObjectNode interface
+	void AllocateDefaultPins(UCustomizableObjectNodeRemapPins* RemapPins) override;
+
+	/** Return the Expose Pin Node expose pin name. */
+	FString GetNodeName() const;
+	
+	// This is actually PinCategory
+	UPROPERTY()
+	FName PinType;
+	
+	UEdGraphPin* InputPin() const
+	{
+		return FindPin(TEXT("Object"));
+	}
+
+	/** Boradcasted when the UPROPERTY Name changes. */
+	FOnNameChangedDelegate OnNameChangedDelegate;
+	
+private:
+	UPROPERTY(Category = CustomizableObject, EditAnywhere)
+	FString Name = "Default Name";
+};
+
diff --git a/Engine/Plugins/Experimental/Mutable/Source/CustomizableObjectEditor/Classes/Nodes/CustomizableObjectNodeExtendMaterial.h b/Engine/Plugins/Experimental/Mutable/Source/CustomizableObjectEditor/Classes/Nodes/CustomizableObjectNodeExtendMaterial.h
new file mode 100644
--- /dev/null
+++ b/Engine/Plugins/Experimental/Mutable/Source/CustomizableObjectEditor/Classes/Nodes/CustomizableObjectNodeExtendMaterial.h
@@ -0,0 +1,77 @@
+// Copyright Epic Games, Inc. All Rights Reserved.
+
+#pragma once
+
+#include "Nodes/CustomizableObjectNodeMaterialBase.h"
+#include "Nodes/CustomizableObjectNodeParentedMaterial.h"
+#include "Nodes/CustomizableObjectNodeUseMaterial.h"
+#include "CustomizableObjectEditor_Deprecated.h"
+
+#include "CustomizableObjectNodeExtendMaterial.generated.h"
+
+
+UCLASS()
+class CUSTOMIZABLEOBJECTEDITOR_API UCustomizableObjectNodeExtendMaterial :
+	public UCustomizableObjectNodeMaterialBase,
+	public FCustomizableObjectNodeParentedMaterial,
+	public FCustomizableObjectNodeUseMaterial
+{
+	GENERATED_BODY()
+
+public:
+	UPROPERTY(EditAnywhere, BlueprintReadWrite, Category = CustomizableObject)
+	TArray<FString> Tags;
+
+	// UObject interface.
+	virtual void Serialize(FArchive& Ar) override;
+
+	// UCustomizableObjectNode
+	virtual void PostBackwardsCompatibleFixup() override;
+	virtual void BackwardsCompatibleFixup() override;
+	
+	// Begin EdGraphNode interface
+	virtual FText GetNodeTitle(ENodeTitleType::Type TitleType) const override;
+	virtual FLinearColor GetNodeTitleColor() const override;
+	virtual FText GetTooltipText() const override;
+
+	// UCustomizableObjectNode interface
+	virtual void AllocateDefaultPins(UCustomizableObjectNodeRemapPins* RemapPins) override;
+	virtual void PinConnectionListChanged(UEdGraphPin * Pin) override;
+	virtual bool IsNodeOutDatedAndNeedsRefresh() override;
+	virtual FString GetRefreshMessage() const override;
+	virtual bool IsSingleOutputNode() const override;
+	virtual bool CustomRemovePin(UEdGraphPin& Pin) override;
+
+	// UCustomizableObjectNodeMaterialBase interface
+	virtual UEdGraphPin* OutputPin() const override;
+	virtual void SetParentNode(UCustomizableObject* Object, FGuid NodeId) override;
+	
+	// FCustomizableObjectNodeParentMaterial interface
+	virtual void SaveParentNode(UCustomizableObject* Object, FGuid NodeId) override;
+	virtual UCustomizableObjectNode& GetNode() override;
+	virtual FGuid GetParentNodeId() const override;
+	virtual UCustomizableObject* GetParentObject() const override;
+	
+	// FCustomizableObjectNodeUseMaterial interface
+	virtual FCustomizableObjectNodeParentedMaterial& GetNodeParentedMaterial() override;
+	virtual TMap<FGuid, FEdGraphPinReference>& GetPinsParameter() override;
+	
+	// Own interface
+	UEdGraphPin* AddMeshPin() const;
+	
+private:
+	UPROPERTY(EditAnywhere, Category = CustomizableObject)
+	TObjectPtr<UCustomizableObject> ParentMaterialObject = nullptr;
+
+	UPROPERTY()
+	FGuid ParentMaterialNodeId;
+	
+	/** Relates a Parameter id to a Pin. Only used to improve performance. */
+   	UPROPERTY()
+   	TMap<FGuid, FEdGraphPinReference> PinsParameter;
+
+	// Deprecated properties
+	UPROPERTY()
+	TArray<FCustomizableObjectNodeExtendMaterialImage> Images_DEPRECATED;
+};
+
diff --git a/Engine/Plugins/Experimental/Mutable/Source/CustomizableObjectEditor/Classes/Nodes/CustomizableObjectNodeExternalPin.h b/Engine/Plugins/Experimental/Mutable/Source/CustomizableObjectEditor/Classes/Nodes/CustomizableObjectNodeExternalPin.h
new file mode 100644
--- /dev/null
+++ b/Engine/Plugins/Experimental/Mutable/Source/CustomizableObjectEditor/Classes/Nodes/CustomizableObjectNodeExternalPin.h
@@ -0,0 +1,85 @@
+// Copyright Epic Games, Inc. All Rights Reserved.
+
+#pragma once
+
+#include "Nodes/CustomizableObjectNode.h"
+#include "Nodes/CustomizableObjectNodeExposePin.h"
+
+#include "CustomizableObjectNodeExternalPin.generated.h"
+
+
+/**
+ * Custom remap pins by name action.
+ *
+ * If the node can not find the exposed pin, set the external pin as deprecated.
+ */
+UCLASS()
+class UCustomizableObjectNodeRemapPinsCustomExternalPin : public UCustomizableObjectNodeRemapPinsByName
+{
+	GENERATED_BODY()
+
+public:
+	class UCustomizableObjectNodeExternalPin* Node = nullptr;
+
+	virtual void RemapPins(const TArray<UEdGraphPin*>& OldPins, const TArray<UEdGraphPin*>& NewPins, TMap<UEdGraphPin*, UEdGraphPin*>& PinsToRemap, TArray<UEdGraphPin*>& PinsToOrphan) override;
+};
+
+UCLASS()
+class CUSTOMIZABLEOBJECTEDITOR_API UCustomizableObjectNodeExternalPin : public UCustomizableObjectNode
+{
+public:
+	GENERATED_BODY()
+
+	// UObject interface.
+	virtual void Serialize(FArchive& Ar) override;
+	
+	// EdGraphNode interface 
+	FText GetNodeTitle(ENodeTitleType::Type TitleType) const override;
+	FLinearColor GetNodeTitleColor() const override;
+	FText GetTooltipText() const override;
+	
+	// UCustomizableObjectNode interface
+	virtual void BeginConstruct() override;
+	virtual void BackwardsCompatibleFixup() override;
+	virtual void PostBackwardsCompatibleFixup() override;
+	void AllocateDefaultPins(UCustomizableObjectNodeRemapPins* RemapPins) override;
+	UCustomizableObjectNodeRemapPins* CreateRemapPinsDefault() const override;
+	void BeginPostDuplicate(bool bDuplicateForPIE) override;
+	virtual bool CanConnect(const UEdGraphPin* InOwnedInputPin, const UEdGraphPin* InOutputPin, bool& bOutIsOtherNodeBlocklisted, bool& bOutArePinsCompatible) const override;
+
+	void UpdateReferencedNodeId(const FGuid& NewGuid) override;
+
+	// Own interface
+	
+	/** Create the external pin. 
+	 *
+	 * @param NodeExposePin Used to obtain expose pin friendly name. If it is nullptr, the name will be set to a default.
+	 */
+	UEdGraphPin* CreateExternalPin(const UCustomizableObjectNodeExposePin* NodeExposePin);
+
+	/** Set the linked Node Expose Pin node guid. */
+	void SetExternalObjectNodeId(FGuid Guid);
+
+	/** Return the external pin. Can return nullptr. */
+	UEdGraphPin* GetExternalPin() const;
+
+	/** Return the linked Expose Pin node. Return nullptr if not set. */
+	UCustomizableObjectNodeExposePin* GetNodeExposePin() const;
+
+	// This is actually PinCategory
+	UPROPERTY()
+	FName PinType;
+
+	/** External Customizable Object which the linked Node Expose Pin belong to. */
+	UPROPERTY()
+	TObjectPtr<UCustomizableObject> ExternalObject;
+
+private:
+	/** Linked Node Expose Pin node guid. */
+	UPROPERTY()
+	FGuid ExternalObjectNodeId;
+	
+	FDelegateHandle OnNameChangedDelegateHandle;
+	FDelegateHandle DestroyNodeDelegateHandle;
+};
+
diff --git a/Engine/Plugins/Experimental/Mutable/Source/CustomizableObjectEditor/Classes/Nodes/CustomizableObjectNodeFloatConstant.h b/Engine/Plugins/Experimental/Mutable/Source/CustomizableObjectEditor/Classes/Nodes/CustomizableObjectNodeFloatConstant.h
new file mode 100644
--- /dev/null
+++ b/Engine/Plugins/Experimental/Mutable/Source/CustomizableObjectEditor/Classes/Nodes/CustomizableObjectNodeFloatConstant.h
@@ -0,0 +1,86 @@
+// Copyright Epic Games, Inc. All Rights Reserved.
+
+#pragma once
+
+#include "Nodes/CustomizableObjectNode.h"
+#include "SGraphNode.h"
+
+#include "CustomizableObjectNodeFloatConstant.generated.h"
+
+// Class create an input float in the NodeFloatConstant
+class SGraphNodeFloatConstant : public SGraphNode
+{
+public:
+
+	SLATE_BEGIN_ARGS(SGraphNodeFloatConstant) {}
+
+	SLATE_END_ARGS();
+
+	SGraphNodeFloatConstant() : SGraphNode() {};
+
+	// Builds the SGraphNodeFloatConstant when needed
+	void Construct(const FArguments& InArgs, UEdGraphNode* InGraphNode);
+
+	// Calls the needed functions to build the SGraphNode widgets
+	void UpdateGraphNode();
+
+	// Overriden functions to build the SGraphNode widgets
+	virtual void SetDefaultTitleAreaWidget(TSharedRef<SOverlay> DefaultTitleAreaWidget) override;
+	virtual void CreateBelowPinControls(TSharedPtr<SVerticalBox> MainBox) override;
+	virtual bool ShouldAllowCulling() const override { return false; }
+
+	// Callbacks for the widget
+	void OnExpressionPreviewChanged(const ECheckBoxState NewCheckedState);
+	ECheckBoxState IsExpressionPreviewChecked() const;
+	const FSlateBrush* GetExpressionPreviewArrow() const;
+
+private:
+
+	// Callback for the OnValueChanged of the SpinBox
+	void OnSpinBoxValueChanged(float Value);
+
+	// Callback for the OnValueCommited of the SpinBox
+	void OnSpinBoxValueCommitted(float Value, ETextCommit::Type);
+
+private:
+
+	// Pointer to the NodeFloatconstant that owns this SGraphNode
+	UCustomizableObjectNodeFloatConstant* NodeFloatConstant;
+
+	// Style for the SpinBox
+	FSpinBoxStyle WidgetStyle;
+
+};
+
+
+UCLASS()
+class CUSTOMIZABLEOBJECTEDITOR_API UCustomizableObjectNodeFloatConstant : public UCustomizableObjectNode
+{
+public:
+	GENERATED_BODY()
+
+	UCustomizableObjectNodeFloatConstant();
+
+	/**  */
+	UPROPERTY(EditAnywhere, Category = CustomizableObject, meta = (DontUpdateWhileEditing))
+	float Value;
+
+	// UObject interface.
+	void PostEditChangeProperty(FPropertyChangedEvent& PropertyChangedEvent) override;
+
+	// Begin EdGraphNode interface
+	FText GetNodeTitle(ENodeTitleType::Type TitleType) const override;
+	FLinearColor GetNodeTitleColor() const override;
+	FText GetTooltipText() const override;
+
+	// UCustomizableObjectNode interface
+	void AllocateDefaultPins(UCustomizableObjectNodeRemapPins* RemapPins) override;
+	bool IsAffectedByLOD() const override { return false; }
+
+	// Creates the SGraph Node widget for the thumbnail
+	TSharedPtr<SGraphNode> CreateVisualWidget() override;
+
+	// Determines if the Node is collapsed or not
+	bool bCollapsed;
+};
+
diff --git a/Engine/Plugins/Experimental/Mutable/Source/CustomizableObjectEditor/Classes/Nodes/CustomizableObjectNodeFloatParameter.h b/Engine/Plugins/Experimental/Mutable/Source/CustomizableObjectEditor/Classes/Nodes/CustomizableObjectNodeFloatParameter.h
new file mode 100644
--- /dev/null
+++ b/Engine/Plugins/Experimental/Mutable/Source/CustomizableObjectEditor/Classes/Nodes/CustomizableObjectNodeFloatParameter.h
@@ -0,0 +1,60 @@
+// Copyright Epic Games, Inc. All Rights Reserved.
+
+#pragma once
+
+#include "Nodes/CustomizableObjectNode.h"
+#include "CustomizableObjectNodeFloatParameter.generated.h"
+
+
+USTRUCT()
+struct CUSTOMIZABLEOBJECTEDITOR_API FCustomizableObjectNodeFloatDescription
+{
+	GENERATED_USTRUCT_BODY()
+
+	UPROPERTY(EditAnywhere, Category=CustomizableObject)
+	FString Name;
+};
+
+
+UCLASS()
+class CUSTOMIZABLEOBJECTEDITOR_API UCustomizableObjectNodeFloatParameter : public UCustomizableObjectNode
+{
+public:
+	GENERATED_BODY()
+
+	UPROPERTY(EditAnywhere, Category=CustomizableObject)
+	float DefaultValue = 1.0f;
+
+	UPROPERTY(EditAnywhere, Category=CustomizableObject)
+	FString ParameterName = "Default Name";
+
+	UPROPERTY(EditAnywhere, Category = UI, meta = (DisplayName = "Parameter UI Metadata"))
+	FMutableParamUIMetadata ParamUIMetadata;
+	
+	// UObject interface.
+	void PostEditChangeProperty(FPropertyChangedEvent& PropertyChangedEvent) override;
+	virtual void Serialize(FArchive& Ar) override;
+	
+	// Begin EdGraphNode interface
+	FText GetNodeTitle(ENodeTitleType::Type TitleType) const override;
+	FLinearColor GetNodeTitleColor() const override;
+	FText GetTooltipText() const override;
+	void PostPasteNode() override;
+
+	// UCustomizableObjectNode interface
+	void AllocateDefaultPins(UCustomizableObjectNodeRemapPins* RemapPins) override;
+	virtual bool IsAffectedByLOD() const override;
+
+	// Own interface
+	int32 GetNumDescriptionImage() const;
+	
+	const UEdGraphPin* GetDescriptionImagePin(int32 Index) const;
+
+private:
+	UPROPERTY()
+	TArray<FEdGraphPinReference> DescriptionImagePinsReferences;
+	
+	UPROPERTY(EditAnywhere, Category=CustomizableObject)
+	TArray<FCustomizableObjectNodeFloatDescription> DescriptionImage;
+};
+
diff --git a/Engine/Plugins/Experimental/Mutable/Source/CustomizableObjectEditor/Classes/Nodes/CustomizableObjectNodeFloatSwitch.h b/Engine/Plugins/Experimental/Mutable/Source/CustomizableObjectEditor/Classes/Nodes/CustomizableObjectNodeFloatSwitch.h
new file mode 100644
--- /dev/null
+++ b/Engine/Plugins/Experimental/Mutable/Source/CustomizableObjectEditor/Classes/Nodes/CustomizableObjectNodeFloatSwitch.h
@@ -0,0 +1,28 @@
+// Copyright Epic Games, Inc. All Rights Reserved.
+
+#pragma once
+
+#include "Nodes/CustomizableObjectNodeSwitchBase.h"
+
+#include "EdGraph/EdGraphPin.h"
+
+#include "CustomizableObjectNodeFloatSwitch.generated.h"
+
+
+UCLASS()
+class CUSTOMIZABLEOBJECTEDITOR_API UCustomizableObjectNodeFloatSwitch : public UCustomizableObjectNodeSwitchBase
+{
+public:
+	GENERATED_BODY()
+
+	// UObject interface
+	virtual void Serialize(FArchive& Ar) override;
+	
+	// UCustomizableObjectNodeSwitchBase interface
+	FString GetOutputPinName() const override;
+
+	FName GetCategory() const override;
+
+	// Own interface
+	FString GetPinPrefix() const;
+};
diff --git a/Engine/Plugins/Experimental/Mutable/Source/CustomizableObjectEditor/Classes/Nodes/CustomizableObjectNodeFloatVariation.h b/Engine/Plugins/Experimental/Mutable/Source/CustomizableObjectEditor/Classes/Nodes/CustomizableObjectNodeFloatVariation.h
new file mode 100644
--- /dev/null
+++ b/Engine/Plugins/Experimental/Mutable/Source/CustomizableObjectEditor/Classes/Nodes/CustomizableObjectNodeFloatVariation.h
@@ -0,0 +1,62 @@
+// Copyright Epic Games, Inc. All Rights Reserved.
+
+#pragma once
+
+#include "Nodes/CustomizableObjectNode.h"
+
+#include "CustomizableObjectNodeFloatVariation.generated.h"
+
+
+USTRUCT()
+struct CUSTOMIZABLEOBJECTEDITOR_API FCustomizableObjectFloatVariation
+{
+	GENERATED_USTRUCT_BODY()
+
+	UPROPERTY(EditAnywhere, Category = CustomizableObject)
+	FString Tag;
+};
+
+
+UCLASS()
+class CUSTOMIZABLEOBJECTEDITOR_API UCustomizableObjectNodeFloatVariation : public UCustomizableObjectNode
+{
+public:
+	GENERATED_BODY()
+
+	UCustomizableObjectNodeFloatVariation();
+
+	UPROPERTY(EditAnywhere, Category = CustomizableObject)
+	TArray<FCustomizableObjectFloatVariation> Variations;
+
+	// UObject interface.
+	void PostEditChangeProperty(FPropertyChangedEvent& PropertyChangedEvent) override;
+
+	// EdGraphNode interface
+	FText GetNodeTitle(ENodeTitleType::Type TitleType) const override;
+	FLinearColor GetNodeTitleColor() const override;
+	FText GetTooltipText() const override;
+
+	// UCustomizableObjectNode interface
+	void AllocateDefaultPins(UCustomizableObjectNodeRemapPins* RemapPins) override;
+
+	UEdGraphPin* OutputPin() const
+	{
+		return FindPin(TEXT("Float"));
+	}
+
+	UEdGraphPin* DefaultPin() const
+	{
+		return FindPin(TEXT("Default"));
+	}
+
+	int32 GetNumVariations() const
+	{
+		return Variations.Num();
+	}
+
+	UEdGraphPin* VariationPin(int Index) const
+	{
+		return FindPin(FString::Printf(TEXT("Variation %d"),Index));
+	}
+};
+
diff --git a/Engine/Plugins/Experimental/Mutable/Source/CustomizableObjectEditor/Classes/Nodes/CustomizableObjectNodeGroupProjectorParameter.h b/Engine/Plugins/Experimental/Mutable/Source/CustomizableObjectEditor/Classes/Nodes/CustomizableObjectNodeGroupProjectorParameter.h
new file mode 100644
--- /dev/null
+++ b/Engine/Plugins/Experimental/Mutable/Source/CustomizableObjectEditor/Classes/Nodes/CustomizableObjectNodeGroupProjectorParameter.h
@@ -0,0 +1,110 @@
+// Copyright Epic Games, Inc. All Rights Reserved.
+
+#pragma once
+
+#include "CustomizableObjectInstance.h"
+#include "Nodes/CustomizableObjectNodeProjectorParameter.h"
+#include "Engine/DataTable.h"
+#include "Animation/PoseAsset.h"
+
+#include "CustomizableObjectNodeGroupProjectorParameter.generated.h"
+
+USTRUCT()
+struct FGroupProjectorParameterImage
+{
+	GENERATED_BODY()
+
+	UPROPERTY(EditAnywhere, Category = ProjectorGroup)
+	FString OptionName;
+
+	UPROPERTY(EditAnywhere, Category = ProjectorGroup)
+	TObjectPtr<UTexture2D> OptionImage = nullptr;
+};
+
+
+USTRUCT()
+struct FGroupProjectorParameterPose
+{
+	GENERATED_BODY()
+
+	UPROPERTY(EditAnywhere, Category = ProjectorGroup)
+	FString PoseName;
+
+	UPROPERTY(EditAnywhere, Category = ProjectorGroup)
+	TObjectPtr<UPoseAsset> OptionPose = nullptr;
+};
+
+
+UCLASS(hideCategories = (CustomizableObjectHide))
+class CUSTOMIZABLEOBJECTEDITOR_API UCustomizableObjectNodeGroupProjectorParameter : public UCustomizableObjectNodeProjectorParameter
+{
+public:
+	GENERATED_BODY()
+
+	UCustomizableObjectNodeGroupProjectorParameter();
+
+	/** Return array with the sticker name and UTexture2D for projection */
+	TArray<FGroupProjectorParameterImage> GetOptionImagesFromTable() const;
+
+	/** Returns the final option images without repeated elements in the option names,
+	* the data table has preference over elements in the OptionImages array */
+	TArray<FGroupProjectorParameterImage> GetFinalOptionImagesNoRepeat() const;
+
+	UPROPERTY(EditAnywhere, Category = ProjectorGroup, Meta = (ToolTip = "Only used when connected to a Group node. Specifies which material channel in the Group node's child material nodes will be connected to the projection."))
+	FString MaterialChannelNameToConnect;
+
+	UPROPERTY(EditAnywhere, Category = ProjectorGroup, Meta = (ToolTip = "Only used when connected to a Group node. Specifies which material channel will be used to mask out the projection."))
+	FString MaskedOutAreaMaterialChannelName;
+
+	/** Reference Texture used to decide the texture properties of the mutable-generated textures
+	* connected to this material. If null, it will try to be guessed at compile time from
+	* the graph. */
+	UPROPERTY(EditAnywhere, Category = ProjectorGroup)
+	TObjectPtr<UTexture2D> ReferenceTexture = nullptr;
+
+	UPROPERTY(EditAnywhere, Category = ProjectorGroup)
+	int32 ProjectionTextureSize = 512;
+
+	UPROPERTY(EditAnywhere, Category = ProjectorGroup, Meta = (ToolTip = "Specifies at which LOD level the projection texture will not be used and possibly save memory. A negative value means they will never be dropped."))
+	int32 DropProjectionTextureAtLOD = -1;
+
+	UPROPERTY(EditAnywhere, Category = ProjectorGroup, Meta = (ToolTip = "If true, projection textures will be shared between LODs of the same object, and will save memory. Only use if all the LODs share the same UV layout."))
+	bool bShareProjectionTexturesBetweenLODs = false;
+
+	UPROPERTY(EditAnywhere, Category = ProjectorGroup)
+	TArray<FGroupProjectorParameterImage> OptionImages;
+
+	UPROPERTY(EditAnywhere, Category = ProjectorGroup)
+	TArray<FGroupProjectorParameterPose> OptionPoses;
+
+	/** Name of the column in the Option Images Data Table with the additional option images, (UTexture2D assets). */
+	UPROPERTY(EditAnywhere, Category = ProjectorGroup)
+	FName DataTableTextureColumnName;
+
+	/** Table where additional option images besides Option Images are read. The elements in this table have priority
+	* over elements from Option Images in case of duplicity. Use the "Data Table Texture Column Name" property to specify
+	* the name of the column where textures are read in the table. */
+	UPROPERTY(EditAnywhere, Category = ProjectorGroup)
+	TObjectPtr<UDataTable> OptionImagesDataTable = nullptr;
+
+	UPROPERTY(EditAnywhere, Category = ProjectorGroup)
+	FString AlternateProjectionResolutionStateName;
+
+	UPROPERTY(EditAnywhere, Category = ProjectorGroup)
+	float AlternateProjectionResolutionFactor = 0.0f;
+
+	UPROPERTY(EditAnywhere, Category = ProjectorGroup, Meta = (ClampMin = "0"))
+	int32 UVLayout = 0;
+
+	// UObject interface
+	void Serialize(FArchive& Ar) override;
+
+	// Begin EdGraphNode interface
+	virtual FText GetNodeTitle(ENodeTitleType::Type TitleType) const override;
+	FLinearColor GetNodeTitleColor() const override;
+	FText GetTooltipText() const override;
+
+	// UCustomizableObjectNode interface.
+	void AllocateDefaultPins(UCustomizableObjectNodeRemapPins* RemapPins);
+};
+
diff --git a/Engine/Plugins/Experimental/Mutable/Source/CustomizableObjectEditor/Classes/Nodes/CustomizableObjectNodeLayoutBlocks.h b/Engine/Plugins/Experimental/Mutable/Source/CustomizableObjectEditor/Classes/Nodes/CustomizableObjectNodeLayoutBlocks.h
new file mode 100644
--- /dev/null
+++ b/Engine/Plugins/Experimental/Mutable/Source/CustomizableObjectEditor/Classes/Nodes/CustomizableObjectNodeLayoutBlocks.h
@@ -0,0 +1,70 @@
+// Copyright Epic Games, Inc. All Rights Reserved.
+
+#pragma once
+
+#include "Nodes/CustomizableObjectNode.h"
+#include "CustomizableObjectLayout.h"
+
+#include "CustomizableObjectNodeLayoutBlocks.generated.h"
+
+
+UCLASS()
+class CUSTOMIZABLEOBJECTEDITOR_API UCustomizableObjectNodeLayoutBlocks : public UCustomizableObjectNode
+{
+public:
+	GENERATED_BODY()
+
+	UCustomizableObjectNodeLayoutBlocks();
+
+	UPROPERTY()
+	FIntPoint GridSize_DEPRECATED;
+
+	/** Used with the fixed layout strategy. */
+	UPROPERTY()
+	FIntPoint MaxGridSize_DEPRECATED;
+
+	UPROPERTY()
+	TArray<FCustomizableObjectLayoutBlock> Blocks_DEPRECATED;
+
+	UPROPERTY()
+	ECustomizableObjectTextureLayoutPackingStrategy PackingStrategy_DEPRECATED;
+
+	UPROPERTY()
+	TObjectPtr<UCustomizableObjectLayout> Layout = nullptr;
+
+	// EdGraphNode interface
+	FText GetNodeTitle(ENodeTitleType::Type TitleType) const override;
+	FLinearColor GetNodeTitleColor() const override;
+	FText GetTooltipText() const override;
+	void Serialize(FArchive& Ar) override;
+	void PinConnectionListChanged(UEdGraphPin* Pin) override;
+	void PostPasteNode() override;
+
+
+	// UCustomizableObjectNode interface
+	virtual void PostBackwardsCompatibleFixup() override;
+	void AllocateDefaultPins(UCustomizableObjectNodeRemapPins* RemapPins) override;
+
+	UEdGraphPin* OutputPin() const
+	{
+		return FindPin(TEXT("Layout"));
+	}
+
+	void AddAttachedErrorData(const FAttachedErrorDataView& AttachedErrorData) override;
+	void ResetAttachedErrorData() override;
+
+	bool IsSingleOutputNode() const override;
+
+private:
+
+	/** Last static or skeletal mesh connected. Used to remove the callback once disconnected. */
+	TWeakObjectPtr<UCustomizableObjectNode> LastMeshNodeConnected;
+
+	/** Connected NodeStaticMesh or NodeSkeletalMesh Mesh UPROPERTY changed callback function. Sets the layout mesh. */
+	void MeshPostEditChangeProperty(UObject* Node, FPropertyChangedEvent& FPropertyChangedEvent);
+
+	void LinkPostEditChangePropertyDelegate(UEdGraphPin* Pin);
+
+	void SetLayoutSkeletalMesh();
+};
+
diff --git a/Engine/Plugins/Experimental/Mutable/Source/CustomizableObjectEditor/Classes/Nodes/CustomizableObjectNodeMaterial.h b/Engine/Plugins/Experimental/Mutable/Source/CustomizableObjectEditor/Classes/Nodes/CustomizableObjectNodeMaterial.h
new file mode 100644
--- /dev/null
+++ b/Engine/Plugins/Experimental/Mutable/Source/CustomizableObjectEditor/Classes/Nodes/CustomizableObjectNodeMaterial.h
@@ -0,0 +1,333 @@
+// Copyright Epic Games, Inc. All Rights Reserved.
+
+#pragma once
+
+#include "Engine/Texture2D.h"
+#include "Nodes/CustomizableObjectNodeMaterialBase.h"
+#include "EdGraphSchema_CustomizableObject.h"
+#include "CustomizableObjectEditor_Deprecated.h"
+
+#include "CustomizableObjectNodeMaterial.generated.h"
+
+
+DECLARE_MULTICAST_DELEGATE(FPostImagePinModeChangedDelegate)
+
+
+/** Custom remap pins by name action.
+ *
+ * Remap pins by Texture Parameter Id. */
+UCLASS()
+class UCustomizableObjectNodeMaterialRemapPinsByName : public UCustomizableObjectNodeRemapPinsByName
+{
+	GENERATED_BODY()
+public:
+	class UCustomizableObjectNodeMaterial* Node = nullptr;
+
+	virtual bool Equal(const UEdGraphPin& OldPin, const UEdGraphPin& NewPin) const override;
+
+	virtual void RemapPins(const TArray<UEdGraphPin*>& OldPins, const TArray<UEdGraphPin*>& NewPins, TMap<UEdGraphPin*, UEdGraphPin*>& PinsToRemap, TArray<UEdGraphPin*>& PinsToOrphan) override;
+
+	bool HasSavedPinData(const UEdGraphPin &Pin) const;
+};
+
+
+/** Base class for all Material Parameters. */
+UCLASS()
+class CUSTOMIZABLEOBJECTEDITOR_API UCustomizableObjectNodeMaterialPinDataParameter : public UCustomizableObjectNodePinData
+{
+	GENERATED_BODY()
+
+public:
+	/** Texture Parameter Id. */
+	UPROPERTY()
+	FGuid ParameterId;
+
+	/** Returns true if all properties are in its default state. */
+	virtual bool IsDefault() const;
+};
+
+
+/** Node pin mode. All pins set to EPinMode::Default will use this this mode. */
+UENUM()
+enum class ENodePinMode
+{
+	Mutable UMETA(ToolTip = "All Material Texture Parameters go through Mutable."),
+	Passthrough UMETA(ToolTip = "All Material Texture Parameters are not modified by Mutable.")
+};
+
+
+/** Image pin, pin mode. */
+UENUM()
+enum class EPinMode
+{
+	/** Use the defined node pin mode. Does not override it. */
+	Default,
+	/** Override the node pin mode. Set it to Mutable mode. */
+	Mutable,
+	/** Override the node pin mode. Set it to Pass-through mode. */
+	Passthrough
+};
+
+/** Enum to FText. */
+FText EPinModeToText(EPinMode PinMode);
+
+UCLASS()
+class CUSTOMIZABLEOBJECTEDITOR_API UCustomizableObjectNodeMaterial : public UCustomizableObjectNodeMaterialBase
+{
+public:
+	GENERATED_BODY()
+
+	UPROPERTY(EditAnywhere, Category=CustomizableObject)
+	TObjectPtr<UMaterialInterface> Material = nullptr;
+
+	UPROPERTY(EditAnywhere, Category=CustomizableObject, Meta = (ToolTip = "Set all Mateiral Texture Parameters to the specified mode. Each Texture Parameter Pin can override this mode."))
+	ENodePinMode TextureParametersMode = ENodePinMode::Passthrough;
+
+	UPROPERTY(EditAnywhere, BlueprintReadWrite, Category = CustomizableObject)
+	TArray<FString> Tags;
+
+	/** Selects which Mesh component of the Instance this material belongs to */
+	UPROPERTY(EditAnywhere, BlueprintReadWrite, Category = CustomizableObject, meta = (ClampMin = "0"))
+	int32 MeshComponentIndex = 0;
+	
+	// UObject interface.
+	virtual void Serialize(FArchive& Ar) override;
+	virtual void PostLoad() override;
+	void PostEditChangeProperty(FPropertyChangedEvent& PropertyChangedEvent) override;
+
+	// EdGraphNode interface
+	FText GetNodeTitle(ENodeTitleType::Type TitleType) const override;
+	FLinearColor GetNodeTitleColor() const override;
+	FText GetTooltipText() const override;
+	TSharedPtr<SGraphNode> CreateVisualWidget() override;
+	virtual void PinConnectionListChanged(UEdGraphPin* Pin) override;
+	virtual void PostPasteNode() override;
+
+	// UCustomizableObjectNode interface
+	virtual void BackwardsCompatibleFixup() override;
+	virtual void PostBackwardsCompatibleFixup() override;
+	void AllocateDefaultPins(UCustomizableObjectNodeRemapPins* RemapPins) override;
+	bool CanPinBeHidden(const UEdGraphPin& Pin) const override;
+	virtual UCustomizableObjectNodeRemapPinsByName* CreateRemapPinsDefault() const override;
+	bool ProvidesCustomPinRelevancyTest() const override { return true; }
+	bool IsPinRelevant(const UEdGraphPin* Pin) const override;
+	virtual bool CustomRemovePin(UEdGraphPin& Pin) override;
+	virtual void ReconstructNode(UCustomizableObjectNodeRemapPins* RemapPinsMode) override;
+
+	// Own Interface
+	virtual UEdGraphPin* GetMeshPin() const;
+	UEdGraphPin* GetMaterialAssetPin() const;
+
+	/** Allow only one connection to a single BaseObjec node. There can be multiple conections to NodeCopyMaterial nodes. */
+	void BreakExistingConnectionsPostConnection(UEdGraphPin* InputPin, UEdGraphPin* OutputPin) override;
+
+	bool IsNodeOutDatedAndNeedsRefresh() override;
+	FString GetRefreshMessage() const override;
+	virtual TSharedPtr<SWidget> CustomizePinDetails(UEdGraphPin& Pin) override;
+
+	
+	// UCustomizableObjectNodeMaterialBase interface
+	TArray<class UCustomizableObjectLayout*> GetLayouts() override;
+	UEdGraphPin* OutputPin() const override;
+	
+	/** Return true if a Material Parameter has changed on which we had a pin connected or data saved. */
+	bool RealMaterialDataHasChanged() const;
+	
+	// --------------------
+	// ALL PARAMETERS
+	// --------------------
+
+	/** Returns the number of Material Parameters. */
+	int32 GetNumParameters(EMaterialParameterType Type) const;
+
+	/** Returns the Material Parameter id.
+	 *
+	 * @param ParameterIndex Have to be valid. */
+	FGuid GetParameterId(EMaterialParameterType Type, int32 ParameterIndex) const;
+	
+	/** Returns the Material Parameter name.
+	 *
+	 * @param ParameterIndex Have to be valid. */
+	FName GetParameterName(EMaterialParameterType Type, int32 ParameterIndex) const;
+	
+	/** Get the Material Parameter layer index.
+	 *
+	 * @param ParameterIndex Have to be valid.
+	 * @returns INDEX_NONE for global parameters. */
+	int32 GetParameterLayerIndex(EMaterialParameterType Type, int32 ParameterIndex) const;
+	
+	/** Get the Material Parameter layer name.
+	 *
+	 * @param ParameterIndex Have to be valid. */
+	FText GetParameterLayerName(EMaterialParameterType Type, int32 ParameterIndex) const;
+
+	/** Returns true if the Material contains the given Material Parameter. */
+	bool HasParameter(const FGuid& ParameterId) const;
+	
+	/** Get the Vector pin for the given Material Vector Parameter.
+	 * Not all parameters have pin.
+	 *
+	 * @param ParameterIndex Has to be valid.
+	 * @return Can return nullptr. */
+	const UEdGraphPin* GetParameterPin(EMaterialParameterType Type, int32 ParameterIndex) const;
+	
+	// --------------------
+	// IMAGES PARAMETERS
+	// --------------------
+
+	/** Returns true if the Material Texture Parameter goes through Mutable.
+	 *
+	 * @param ImageIndex Have to be valid. */
+	bool IsImageMutableMode(int32 ImageIndex) const;
+
+	/** Given an Image pin, returns true if the Material Texture Parameter goes through Mutable. */
+	bool IsImageMutableMode(const UEdGraphPin& Pin) const;
+
+	/** Update a Material Texture Parameter Mode. */
+	void UpdateImagePinMode(const FGuid ParameterId);
+	
+	/** Update a Material Texture Parameter Mode. */
+	void UpdateImagePinMode(const UEdGraphPin& Pin);
+
+	/** Update all Material Texture Parameters Mode. */
+	void UpdateAllImagesPinMode();
+	
+	/** Returns the reference texture assigned to a Material Texture Parameter.
+	 *
+	 * @param ImageIndex Have to be valid.
+	 * @return nullptr if it does not have one assigned. */
+	UTexture2D* GetImageReferenceTexture(int32 ImageIndex) const;
+
+	/** Returns the Texture set in the Material Texture Parameter.
+	 *
+	 * @param ImageIndex Have to be valid. */
+	UTexture2D* GetImageValue(int32 ImageIndex) const;
+	
+	/** Get the Material Texture Parameter UV index.
+	 *
+	 * @param ImageIndex Have to be valid. */
+	int32 GetImageUVLayout(int32 ImageIndex) const;
+
+	/** Delegate called when a Texture Parameter Pin Mode changes. */
+	FPostImagePinModeChangedDelegate PostImagePinModeChangedDelegate;
+	
+private:
+	/** Last static or skeletal mesh connected. Used to remove the callback once disconnected. */
+	TWeakObjectPtr<UCustomizableObjectNode> LastMeshNodeConnected;
+
+	static const TArray<EMaterialParameterType> ParameterTypes;
+
+	/** Relates a Parameter id (key) to a Pin (value). Only used to improve performance. */
+	UPROPERTY()
+	TMap<FGuid, FEdGraphPinReference> PinsParameter;
+
+	/** Relates an Image pin (key) to its Image Pin Mode (value). */
+	UPROPERTY()
+	TMap<FGuid, EPinMode> PinsImagePinMode;
+	
+	/** Create the pin data of the given parameter type. */
+	UCustomizableObjectNodeMaterialPinDataParameter* CreatePinData(EMaterialParameterType Type, int32 ParameterIndex);
+
+	/** Allocate a pin for each parameter of the given type. */
+	void AllocateDefaultParameterPins(EMaterialParameterType Type);
+
+	/** Set the default Material from the connected static or skeletal mesh. */
+	void SetDefaultMaterial();
+
+	/** Connected NodeStaticMesh or NodeSkeletalMesh Mesh UPROPERTY changed callback function. Sets the default material. */
+	UFUNCTION()
+	void MeshPostEditChangeProperty(FPostEditChangePropertyDelegateParameters& Parameters);
+
+	/** Format pin name. */
+	FName GetPinName(EMaterialParameterType Type, int32 ParameterIndex) const;
+
+	/** Returns the texture coordinate of the given Material Expression. Returns -1 if not found. */
+	static int32 GetExpressionTextureCoordinate(UMaterial* Material, const FGuid &ImageId);
+
+	/** Converts node NodePinMode (does not include Default mode) to PinMode. */
+	static EPinMode NodePinModeToImagePinMode(ENodePinMode NodePinMode);
+
+	/** Returns the Image Pin Mode the pin should be at. It does not update its mode, to update it call UpdateImagePinMode. */
+	EPinMode GetImagePinMode(const UEdGraphPin& Pin) const;
+	
+	// Deprecated properties
+	/** Set all pins to Mutable mode. Even so, each pin can override its behaviour. */
+	UPROPERTY()
+	bool bDefaultPinModeMutable_DEPRECATED = false;
+	
+	UPROPERTY()
+	TArray<FCustomizableObjectNodeMaterialImage> Images_DEPRECATED;
+	
+	UPROPERTY()
+	TArray<FCustomizableObjectNodeMaterialVector> VectorParams_DEPRECATED;
+
+	UPROPERTY()
+	TArray<FCustomizableObjectNodeMaterialScalar> ScalarParams_DEPRECATED;
+};
+
+
+/** Additional data for a Material Texture Parameter pin. */
+UCLASS()
+class CUSTOMIZABLEOBJECTEDITOR_API UCustomizableObjectNodeMaterialPinDataImage : public UCustomizableObjectNodeMaterialPinDataParameter
+{
+	GENERATED_BODY()
+
+	friend void UCustomizableObjectNodeMaterial::BackwardsCompatibleFixup();
+	
+public:
+	// UObject interface
+	virtual void PostEditChangeProperty(FPropertyChangedEvent& PropertyChangedEvent) override;
+	
+	// UCustomizableObjectNodeMaterialPinParameter interface
+	virtual bool IsDefault() const override;
+
+	// Own interface
+	/** Constructor parameters. Should always be called after a NewObject. */
+	void Init(UCustomizableObjectNodeMaterial& InNodeMaterial);
+
+	EPinMode GetPinMode() const;
+
+	void SetPinMode(EPinMode InPinMode);
+	
+private:
+	/** Image pin mode. If is not default, overrides the defined node behaviour. */
+	UPROPERTY()
+	EPinMode PinMode = EPinMode::Default;
+
+public:
+	/* UVLayout Mode. Indicates that the texture should not be transformed by any layout. Theses textures will not be reduced automatically for LODs. */
+	constexpr static int32 UV_LAYOUT_IGNORE = -1;
+	
+	/** UVLayout Mode. Does not override the Material Texture Parameter. */
+	constexpr static int32 UV_LAYOUT_DEFAULT = -2;
+	
+	/** Index of the UV channel that will be used with this image.It is necessary to apply the proper layout transformations to it. */
+	UPROPERTY()
+	int32 UVLayout = UV_LAYOUT_DEFAULT;
+
+	/** Reference Texture used to decide the texture properties of the mutable-generated textures
+	* connected to this material. If null, it will try to be guessed at compile time from
+	* the graph. */
+	UPROPERTY(EditAnywhere, Category=CustomizableObject) // Required to be EditAnywhere for the selector to work.
+	TObjectPtr<UTexture2D> ReferenceTexture = nullptr;
+
+private:
+	UPROPERTY()
+	TObjectPtr<UCustomizableObjectNodeMaterial> NodeMaterial = nullptr;
+};
+
+
+/** Additional data for a Material Vector Parameter pin. */
+UCLASS()
+class CUSTOMIZABLEOBJECTEDITOR_API UCustomizableObjectNodeMaterialPinDataVector : public UCustomizableObjectNodeMaterialPinDataParameter
+{
+	GENERATED_BODY()
+};
+
+
+/** Additional data for a Material Float Parameter pin. */
+UCLASS()
+class CUSTOMIZABLEOBJECTEDITOR_API UCustomizableObjectNodeMaterialPinDataScalar : public UCustomizableObjectNodeMaterialPinDataParameter
+{
+	GENERATED_BODY()
+};
\ No newline at end of file
diff --git a/Engine/Plugins/Experimental/Mutable/Source/CustomizableObjectEditor/Classes/Nodes/CustomizableObjectNodeMaterialBase.h b/Engine/Plugins/Experimental/Mutable/Source/CustomizableObjectEditor/Classes/Nodes/CustomizableObjectNodeMaterialBase.h
new file mode 100644
--- /dev/null
+++ b/Engine/Plugins/Experimental/Mutable/Source/CustomizableObjectEditor/Classes/Nodes/CustomizableObjectNodeMaterialBase.h
@@ -0,0 +1,29 @@
+// Copyright Epic Games, Inc. All Rights Reserved.
+
+#pragma once
+
+#include "Nodes/CustomizableObjectNode.h"
+#include "CustomizableObjectNodeMaterialBase.generated.h"
+
+
+
+UCLASS(abstract)
+class CUSTOMIZABLEOBJECTEDITOR_API UCustomizableObjectNodeMaterialBase : public UCustomizableObjectNode
+{
+public:
+	GENERATED_BODY()
+
+	// Own interface
+	virtual TArray<class UCustomizableObjectLayout*> GetLayouts() { return TArray<class UCustomizableObjectLayout*>(); }
+
+	// This method should be overidden in all derived classes
+	virtual UEdGraphPin* OutputPin() const
+	{
+		check(false); 
+		return nullptr;
+	}
+
+	/** Allow multiple connections to NodeCopyMaterial but only one connection to a NodeObject. */
+	bool ShouldBreakExistingConnections(const UEdGraphPin* InputPin, const UEdGraphPin* OutputPin) const override;
+};
+
diff --git a/Engine/Plugins/Experimental/Mutable/Source/CustomizableObjectEditor/Classes/Nodes/CustomizableObjectNodeMaterialVariation.h b/Engine/Plugins/Experimental/Mutable/Source/CustomizableObjectEditor/Classes/Nodes/CustomizableObjectNodeMaterialVariation.h
new file mode 100644
--- /dev/null
+++ b/Engine/Plugins/Experimental/Mutable/Source/CustomizableObjectEditor/Classes/Nodes/CustomizableObjectNodeMaterialVariation.h
@@ -0,0 +1,73 @@
+// Copyright Epic Games, Inc. All Rights Reserved.
+
+#pragma once
+
+#include "Nodes/CustomizableObjectNodeMaterialBase.h"
+
+#include "CustomizableObjectNodeMaterialVariation.generated.h"
+
+
+USTRUCT()
+struct CUSTOMIZABLEOBJECTEDITOR_API FCustomizableObjectMaterialVariation
+{
+	GENERATED_USTRUCT_BODY()
+
+	UPROPERTY(EditAnywhere, Category = CustomizableObject)
+	FString Tag;
+};
+
+
+UENUM(BlueprintType)
+enum class ECustomizableObjectNodeMaterialVariationType : uint8
+{
+	Tag 		UMETA(DisplayName = "Tag"),
+	State 		UMETA(DisplayName = "State"),
+};
+
+
+UCLASS()
+class CUSTOMIZABLEOBJECTEDITOR_API UCustomizableObjectNodeMaterialVariation : public UCustomizableObjectNodeMaterialBase
+{
+public:
+	GENERATED_BODY()
+
+	UCustomizableObjectNodeMaterialVariation();
+	
+	UPROPERTY(EditAnywhere, Category = CustomizableObject)
+	ECustomizableObjectNodeMaterialVariationType Type = ECustomizableObjectNodeMaterialVariationType::Tag;
+
+	UPROPERTY(EditAnywhere, Category = CustomizableObject)
+	TArray<FCustomizableObjectMaterialVariation> Variations;
+
+	// UObject interface.
+	void PostEditChangeProperty(FPropertyChangedEvent& PropertyChangedEvent) override;
+
+	// Begin EdGraphNode interface
+	FText GetNodeTitle(ENodeTitleType::Type TitleType) const override;
+	FLinearColor GetNodeTitleColor() const override;
+	FText GetTooltipText() const override;
+
+	// UCustomizableObjectNode interface
+	void AllocateDefaultPins(UCustomizableObjectNodeRemapPins* RemapPins) override;
+
+	UEdGraphPin* OutputPin() const
+	{
+		return FindPin(TEXT("Material"));
+	}
+
+	UEdGraphPin* DefaultPin() const
+	{
+		return FindPin(TEXT("Default"));
+	}
+
+	int32 GetNumVariations() const
+	{
+		return Variations.Num();
+	}
+
+	UEdGraphPin* VariationPin(int Index) const
+	{
+		return FindPin(FString::Printf(TEXT("Variation %d"),Index));
+	}
+};
+
diff --git a/Engine/Plugins/Experimental/Mutable/Source/CustomizableObjectEditor/Classes/Nodes/CustomizableObjectNodeMesh.h b/Engine/Plugins/Experimental/Mutable/Source/CustomizableObjectEditor/Classes/Nodes/CustomizableObjectNodeMesh.h
new file mode 100644
--- /dev/null
+++ b/Engine/Plugins/Experimental/Mutable/Source/CustomizableObjectEditor/Classes/Nodes/CustomizableObjectNodeMesh.h
@@ -0,0 +1,42 @@
+// Copyright Epic Games, Inc. All Rights Reserved.
+
+#pragma once
+
+#include "CustomizableObjectLayout.h"
+#include "Nodes/CustomizableObjectNode.h"
+
+#include "CustomizableObjectNodeMesh.generated.h"
+
+
+UCLASS(abstract)
+class CUSTOMIZABLEOBJECTEDITOR_API UCustomizableObjectNodeMesh : public UCustomizableObjectNode
+{
+public:
+	GENERATED_BODY()
+
+	// Own interface
+	/** Returns the multiple Layouts of a given Mesh pin. A pin can have multiple Layouts since it can have multiple UVs. Override. */
+	virtual TArray<UCustomizableObjectLayout*> GetLayouts(const UEdGraphPin* OutPin) const PURE_VIRTUAL(UCustomizableObjectNodeMesh::GetLayouts, return {};);
+
+	virtual UTexture2D* FindTextureForPin(const UEdGraphPin* Pin) const PURE_VIRTUAL(UCustomizableObjectNodeMesh::FindTextureForPin, return {};);
+
+	virtual void GetUVChannelForPin(const UEdGraphPin* Pin, TArray<FVector2f>& OutSegments, int32 UVIndex) const PURE_VIRTUAL(UCustomizableObjectNodeMesh::GetUVChannelForPin, );
+
+	/** Returns the Unreal mesh (e.g., USkeletalMesh, UStaticMesh...). */
+	virtual UObject* GetMesh() const PURE_VIRTUAL(UCustomizableObjectNodeMesh::GetMesh, return {};);;
+
+	/** Returns the out mesh pin associated to the given LOD and Material. Override. */
+	virtual UEdGraphPin* GetMeshPin(int32 LOD, int MaterialIndex) const PURE_VIRTUAL(UCustomizableObjectNodeMesh::GetMeshPin, return {};);
+
+	/** Given a pin, return the Section and Layout index.
+	 *
+	 * Will always return a valid result. The result can be valid but out of sync with respect the Unreal mesh asset.
+	 * In other words, the pin still represents a LOD 3 but the asset no longer has this third LOD.
+	 * 
+	 * @param Pin In. Node's owning pin.
+	 * @param LODIndex Out.
+	 * @param SectionIndex Out.
+	 * @param LayoutIndex Out. */
+	virtual void GetPinSection(const UEdGraphPin& Pin, int32& OutLODIndex, int32& OutSectionIndex, int32& OutLayoutIndex) const PURE_VIRTUAL(UCustomizableObjectNodeMesh::GetPinSection, );
+};
+
diff --git a/Engine/Plugins/Experimental/Mutable/Source/CustomizableObjectEditor/Classes/Nodes/CustomizableObjectNodeMeshClipDeform.h b/Engine/Plugins/Experimental/Mutable/Source/CustomizableObjectEditor/Classes/Nodes/CustomizableObjectNodeMeshClipDeform.h
new file mode 100644
--- /dev/null
+++ b/Engine/Plugins/Experimental/Mutable/Source/CustomizableObjectEditor/Classes/Nodes/CustomizableObjectNodeMeshClipDeform.h
@@ -0,0 +1,49 @@
+// Copyright Epic Games, Inc. All Rights Reserved.
+
+#pragma once
+
+#include "CustomizableObjectNodeModifierBase.h"
+
+#include "CustomizableObjectNodeMeshClipDeform.generated.h"
+
+UENUM()
+enum class EShapeBindingMethod : uint32
+{
+	ClosestProject = 0,
+	ClosestToSurface = 1,
+	NormalProject = 2
+};
+
+UCLASS()
+class CUSTOMIZABLEOBJECTEDITOR_API UCustomizableObjectNodeMeshClipDeform : public UCustomizableObjectNodeModifierBase
+{
+public:
+	GENERATED_BODY()
+
+	UCustomizableObjectNodeMeshClipDeform();
+
+	UPROPERTY(EditAnywhere, BlueprintReadWrite, Category = MeshClipDeform)
+	TArray<FString> Tags;
+	
+	UPROPERTY(EditAnywhere, Category = MeshClipDeform)
+	EShapeBindingMethod BindingMethod;
+	
+	// Begin EdGraphNode interface
+	FText GetNodeTitle(ENodeTitleType::Type TitleType) const override;
+	FLinearColor GetNodeTitleColor() const override;
+	FText GetTooltipText() const override;
+
+	// UCustomizableObjectNode interface
+	void AllocateDefaultPins(UCustomizableObjectNodeRemapPins* RemapPins) override;
+	
+	inline UEdGraphPin* ClipShapePin() const
+	{
+		return FindPin(TEXT("Clip Shape"), EGPD_Input);
+	}
+
+	UEdGraphPin* OutputPin() const override
+	{
+		return FindPin(TEXT("Material"));
+	}
+};
+
diff --git a/Engine/Plugins/Experimental/Mutable/Source/CustomizableObjectEditor/Classes/Nodes/CustomizableObjectNodeMeshClipMorph.h b/Engine/Plugins/Experimental/Mutable/Source/CustomizableObjectEditor/Classes/Nodes/CustomizableObjectNodeMeshClipMorph.h
new file mode 100644
--- /dev/null
+++ b/Engine/Plugins/Experimental/Mutable/Source/CustomizableObjectEditor/Classes/Nodes/CustomizableObjectNodeMeshClipMorph.h
@@ -0,0 +1,86 @@
+// Copyright Epic Games, Inc. All Rights Reserved.
+
+#pragma once
+
+#include "CustomizableObjectNodeModifierBase.h"
+
+#include "CustomizableObjectNodeMeshClipMorph.generated.h"
+
+UCLASS()
+class CUSTOMIZABLEOBJECTEDITOR_API UCustomizableObjectNodeMeshClipMorph : public UCustomizableObjectNodeModifierBase
+{
+public:
+	GENERATED_BODY()
+
+	UCustomizableObjectNodeMeshClipMorph();
+
+	UPROPERTY(EditAnywhere, Category = MeshToClipAndMorph)
+	FName BoneName;
+
+	UPROPERTY(EditAnywhere, BlueprintReadWrite, Category = MeshToClipAndMorph)
+	TArray<FString> Tags;
+
+	UPROPERTY(EditAnywhere, Category = MeshToClipAndMorph)
+	uint32 ReferenceSkeletonIndex = 0;
+
+	// UObject interface.
+	void PostEditChangeProperty(FPropertyChangedEvent& PropertyChangedEvent) override;
+	virtual void Serialize(FArchive& Ar) override;
+
+	// Begin EdGraphNode interface
+	FText GetNodeTitle(ENodeTitleType::Type TitleType) const override;
+	FLinearColor GetNodeTitleColor() const override;
+	FText GetTooltipText() const override;
+
+	// UCustomizableObjectNode interface
+	void AllocateDefaultPins(UCustomizableObjectNodeRemapPins* RemapPins) override;
+	void PinConnectionListChanged(UEdGraphPin* Pin) override;
+
+	UPROPERTY(EditAnywhere, Category = MeshClipParameters, meta = (DisplayName="Morph Start Offset", ToolTip="Offset from the origin of the selected bone to the actual start of the morph."))
+	FVector StartOffset;
+
+	UPROPERTY(EditAnywhere, Category = MeshClipParameters, meta = (DisplayName = "Local Start Offset", ToolTip = "Toggles between a local or global start offset."))
+	bool bLocalStartOffset;
+
+	UPROPERTY(EditAnywhere, Category = MeshClipParameters, meta = (DisplayName="Morph length", ToolTip="The length from the morph start to the clip plane."))
+	float B;
+
+	UPROPERTY(EditAnywhere, Category = MeshClipParameters, meta = (DisplayName="Ellipse Radius 1", ToolTip="First radius of the ellipse that the mesh is morphed into."))
+	float Radius;
+
+	UPROPERTY(EditAnywhere, Category = MeshClipParameters, meta = (DisplayName = "Ellipse Radius 2", ToolTip = "Second radius of the ellipse that the mesh is morphed into."))
+	float Radius2;
+
+	UPROPERTY(EditAnywhere, Category = MeshClipParameters, meta = (DisplayName = "Ellipse Rotation", ToolTip = "Ellipse Rotation in degrees around the bone axis."))
+	float RotationAngle;
+
+	UPROPERTY(EditAnywhere, Category = MeshClipParameters, meta = (DisplayName = "Morph Curve Control", ToolTip = "Controls the morph curve shape. A value of 1 is linear, less than 1 is concave and greater than 1 convex."))
+	float Exponent;
+
+	UPROPERTY(EditAnywhere, Category = MeshClipParameters, meta = (DisplayName = "Invert normal direction", ToolTip = "Flag to invert the normal direction"))
+	bool bInvertNormal;
+
+	UPROPERTY()
+	FVector Origin;
+
+	UPROPERTY()
+	FVector Normal;
+	
+	UPROPERTY(EditAnywhere, Category = MeshClipParameters, meta = (DisplayName = "Max Effect Radius", ToolTip = "The maximum distance from the origin of the widget where vertices will be affected. If negative, there will be no limit."))
+	float MaxEffectRadius;
+
+	bool bUpdateViewportWidget;
+
+	FVector GetOriginWithOffset() const;
+
+	void FindLocalAxes(FVector& Right, FVector& Up, FVector& Forward) const;
+
+	// Change StartOffset from World to Local of the other way arround
+	void ChangeStartOffsetTransform();
+
+	virtual UEdGraphPin* OutputPin() const override
+	{
+		return FindPin(TEXT("Material"));
+	}
+};
+
diff --git a/Engine/Plugins/Experimental/Mutable/Source/CustomizableObjectEditor/Classes/Nodes/CustomizableObjectNodeMeshClipWithMesh.h b/Engine/Plugins/Experimental/Mutable/Source/CustomizableObjectEditor/Classes/Nodes/CustomizableObjectNodeMeshClipWithMesh.h
new file mode 100644
--- /dev/null
+++ b/Engine/Plugins/Experimental/Mutable/Source/CustomizableObjectEditor/Classes/Nodes/CustomizableObjectNodeMeshClipWithMesh.h
@@ -0,0 +1,65 @@
+// Copyright Epic Games, Inc. All Rights Reserved.
+
+#pragma once
+
+#include "CustomizableObjectNodeModifierBase.h"
+#include "CustomizableObjectNodeMaterial.h"
+
+#include "CustomizableObjectNodeMeshClipWithMesh.generated.h"
+
+UCLASS()
+class CUSTOMIZABLEOBJECTEDITOR_API UCustomizableObjectNodeMeshClipWithMesh : public UCustomizableObjectNodeModifierBase
+{
+public:
+	GENERATED_BODY()
+
+	UCustomizableObjectNodeMeshClipWithMesh();
+
+	UPROPERTY(EditAnywhere, BlueprintReadWrite, Category = MeshToClip)
+	TArray<FString> Tags;
+
+	//!< If assigned, then a material inside this CO will be clipped by this node.
+    //!< If several materials with the same name, all are considered (to cover all LOD levels)
+    UPROPERTY(EditAnywhere, Category = CustomizableObjectToClip)
+	TObjectPtr<UCustomizableObject> CustomizableObjectToClipWith;
+
+    //!< Array with the Guids of the nodes with the same material inside the CustomizableObjectToClipWith CO (if any is assigned)
+    UPROPERTY(EditAnywhere, Category = CustomizableObjectToClip)
+    TArray<FGuid> ArrayMaterialNodeToClipWithID;
+
+	/** Mesh Transform*/
+	UPROPERTY(EditAnywhere, Category = CustomizableObject)
+	FTransform Transform;
+
+	// Begin EdGraphNode interface
+	FText GetNodeTitle(ENodeTitleType::Type TitleType) const override;
+	FLinearColor GetNodeTitleColor() const override;
+	FText GetTooltipText() const override;
+
+	// UCustomizableObjectNode interface
+	void AllocateDefaultPins(UCustomizableObjectNodeRemapPins* RemapPins) override;
+	void PinConnectionListChanged(UEdGraphPin* Pin) override;
+	
+	virtual UEdGraphPin* OutputPin() const override
+	{
+		return FindPin(TEXT("Material"));
+	}
+
+	UEdGraphPin* ClipMeshPin() const
+	{
+		return FindPin(TEXT("Clip Mesh"));
+	}
+
+	void BeginPostDuplicate(bool bDuplicateForPIE) override;
+
+	// UCustomizableObjectNode interface
+	void UpdateReferencedNodeId(const FGuid& NewGuid) override;
+
+	// Details view variables
+	// The clipping node uses tags
+	bool bUseTags;
+
+	// The clipping node uses a material name
+	bool bUseMaterials;
+};
+
diff --git a/Engine/Plugins/Experimental/Mutable/Source/CustomizableObjectEditor/Classes/Nodes/CustomizableObjectNodeMeshGeometryOperation.h b/Engine/Plugins/Experimental/Mutable/Source/CustomizableObjectEditor/Classes/Nodes/CustomizableObjectNodeMeshGeometryOperation.h
new file mode 100644
--- /dev/null
+++ b/Engine/Plugins/Experimental/Mutable/Source/CustomizableObjectEditor/Classes/Nodes/CustomizableObjectNodeMeshGeometryOperation.h
@@ -0,0 +1,47 @@
+// Copyright Epic Games, Inc. All Rights Reserved.
+
+#pragma once
+
+#include "Nodes/CustomizableObjectNode.h"
+
+#include "CustomizableObjectNodeMeshGeometryOperation.generated.h"
+
+
+UCLASS()
+class CUSTOMIZABLEOBJECTEDITOR_API UCustomizableObjectNodeMeshGeometryOperation : public UCustomizableObjectNode
+{
+public:
+	GENERATED_BODY()
+
+	UCustomizableObjectNodeMeshGeometryOperation();
+
+	// Begin EdGraphNode interface
+	FText GetNodeTitle(ENodeTitleType::Type TitleType) const override;
+	FLinearColor GetNodeTitleColor() const override;
+	FText GetTooltipText() const override;
+
+	// UCustomizableObjectNode interface
+	void AllocateDefaultPins(UCustomizableObjectNodeRemapPins* RemapPins) override;
+
+	inline UEdGraphPin* MeshAPin() const
+	{
+		return FindPin(TEXT("Mesh A"), EGPD_Input);
+	}
+
+	inline UEdGraphPin* MeshBPin() const
+	{
+		return FindPin(TEXT("Mesh B"), EGPD_Input);
+	}
+
+	inline UEdGraphPin* ScalarAPin() const
+	{
+		return FindPin(TEXT("Scalar A"));
+	}
+
+	inline UEdGraphPin* ScalarBPin() const
+	{
+		return FindPin(TEXT("Scalar B"));
+	}
+	
+};
+
diff --git a/Engine/Plugins/Experimental/Mutable/Source/CustomizableObjectEditor/Classes/Nodes/CustomizableObjectNodeMeshMorph.h b/Engine/Plugins/Experimental/Mutable/Source/CustomizableObjectEditor/Classes/Nodes/CustomizableObjectNodeMeshMorph.h
new file mode 100644
--- /dev/null
+++ b/Engine/Plugins/Experimental/Mutable/Source/CustomizableObjectEditor/Classes/Nodes/CustomizableObjectNodeMeshMorph.h
@@ -0,0 +1,69 @@
+// Copyright Epic Games, Inc. All Rights Reserved.
+
+#pragma once
+
+#include "Nodes/CustomizableObjectNode.h"
+#include "Nodes/CustomizableObjectNodeMeshReshapeCommon.h"
+
+#include "CustomizableObjectNodeMeshMorph.generated.h"
+
+UCLASS()
+class CUSTOMIZABLEOBJECTEDITOR_API UCustomizableObjectNodeMeshMorph : public UCustomizableObjectNode
+{
+public:
+	GENERATED_BODY()
+
+	UCustomizableObjectNodeMeshMorph();
+
+	// Begin EdGraphNode interface
+	FText GetNodeTitle(ENodeTitleType::Type TitleType) const override;
+	FLinearColor GetNodeTitleColor() const override;
+	FText GetTooltipText() const override;
+
+	// UCustomizableObjectNode interface
+	void AllocateDefaultPins(UCustomizableObjectNodeRemapPins* RemapPins) override;
+	bool IsNodeOutDatedAndNeedsRefresh() override;
+	FString GetRefreshMessage() const override;
+	void PinConnectionListChanged(UEdGraphPin* Pin) override;
+
+
+	// Own interface
+	class UCustomizableObjectNodeSkeletalMesh* GetSourceSkeletalMesh() const;
+
+	UEdGraphPin* MeshPin() const
+	{
+		return FindPin(TEXT("Mesh"), EGPD_Input);
+	}
+
+	UEdGraphPin* FactorPin() const
+	{
+		return FindPin(TEXT("Factor"));
+	}
+
+	UPROPERTY(Category=CustomizableObject, EditAnywhere)
+	FString MorphTargetName;
+
+	/** Enable the deformation of the skeleton of the base mesh. */
+	UPROPERTY(EditAnywhere, Category = CustomizableObject)
+	bool bReshapeSkeleton = false;
+
+	/** Enable the deformation of physics volumes of the base mesh */
+    UPROPERTY(EditAnywhere, Category = CustomizableObject)
+    bool bReshapePhysicsVolumes = false;
+	
+	/** Enables the deformation of all bones of the skeleton */
+	UPROPERTY(EditAnywhere, Category = ReshapeBones, meta=(EditCondition="bReshapeSkeleton"))
+	bool bDeformAllBones = false;
+
+	/** Array with bones that will be deformed */
+	UPROPERTY(EditAnywhere, Category = ReshapeBones, meta=(EditCondition="bReshapeSkeleton && !bDeformAllBones"))
+	TArray<FMeshReshapeBoneReference> BonesToDeform;
+
+	UPROPERTY(EditAnywhere, Category = ReshapePhysics, meta = (EditCondition="bReshapePhysicsVolumes"))
+	bool bDeformAllPhysicsBodies = false;
+
+	UPROPERTY(EditAnywhere, Category = ReshapePhysics, meta=(EditCondition="bReshapePhysicsVolumes && !bDeformAllPhysicsBodies"))
+	TArray<FMeshReshapeBoneReference> PhysicsBodiesToDeform;
+
+};
+
diff --git a/Engine/Plugins/Experimental/Mutable/Source/CustomizableObjectEditor/Classes/Nodes/CustomizableObjectNodeMeshMorphStackApplication.h b/Engine/Plugins/Experimental/Mutable/Source/CustomizableObjectEditor/Classes/Nodes/CustomizableObjectNodeMeshMorphStackApplication.h
new file mode 100644
--- /dev/null
+++ b/Engine/Plugins/Experimental/Mutable/Source/CustomizableObjectEditor/Classes/Nodes/CustomizableObjectNodeMeshMorphStackApplication.h
@@ -0,0 +1,48 @@
+// Copyright Epic Games, Inc. All Rights Reserved.
+
+#pragma once
+
+#include "CustomizableObjectNode.h"
+#include "CustomizableObjectNodeMeshMorphStackApplication.generated.h"
+
+
+UCLASS()
+class CUSTOMIZABLEOBJECTEDITOR_API UCustomizableObjectNodeMeshMorphStackApplication : public UCustomizableObjectNode
+{
+public:
+	GENERATED_BODY()
+
+	UCustomizableObjectNodeMeshMorphStackApplication();
+
+	// Begin EdGraphNode interface
+	FText GetNodeTitle(ENodeTitleType::Type TittleType)const override;
+	FLinearColor GetNodeTitleColor() const override;
+	FText GetTooltipText() const override;
+
+	// UCustomizableObjectNode interface
+	void AllocateDefaultPins(UCustomizableObjectNodeRemapPins* RemapPins) override;
+	void PinConnectionListChanged(UEdGraphPin* Pin) override;
+	void ReconstructNode(UCustomizableObjectNodeRemapPins* RemapPins) override;
+	bool IsNodeOutDatedAndNeedsRefresh() override;
+
+	// Fills the list with all the morphs
+	void UpdateMorphList();
+
+	// Returns the mesh pin
+	UEdGraphPin* GetMeshPin() const
+	{
+		return FindPin(TEXT("InMesh"), EGPD_Input);
+	}
+	
+	// Returns the stack pin
+	UEdGraphPin* GetStackPin() const
+	{
+		return FindPin(TEXT("Stack"), EGPD_Input);
+	}
+
+public:
+
+	// List with all the morphs of the linked skeletal mesh
+	UPROPERTY()
+	TArray<FString> MorphNames;
+};
diff --git a/Engine/Plugins/Experimental/Mutable/Source/CustomizableObjectEditor/Classes/Nodes/CustomizableObjectNodeMeshMorphStackDefinition.h b/Engine/Plugins/Experimental/Mutable/Source/CustomizableObjectEditor/Classes/Nodes/CustomizableObjectNodeMeshMorphStackDefinition.h
new file mode 100644
--- /dev/null
+++ b/Engine/Plugins/Experimental/Mutable/Source/CustomizableObjectEditor/Classes/Nodes/CustomizableObjectNodeMeshMorphStackDefinition.h
@@ -0,0 +1,55 @@
+// Copyright Epic Games, Inc. All Rights Reserved.
+
+#pragma once
+
+#include "CustomizableObjectNode.h"
+#include "CustomizableObjectNodeMeshMorphStackDefinition.generated.h"
+
+
+UCLASS()
+class CUSTOMIZABLEOBJECTEDITOR_API UCustomizableObjectNodeMeshMorphStackDefinition : public UCustomizableObjectNode
+{
+public:
+	GENERATED_BODY()
+
+	UCustomizableObjectNodeMeshMorphStackDefinition();
+
+	// Begin EdGraphNode interface
+	FText GetNodeTitle(ENodeTitleType::Type TittleType)const override;
+	FLinearColor GetNodeTitleColor() const override;
+	FText GetTooltipText() const override;
+	void PinConnectionListChanged(UEdGraphPin* Pin) override;
+	bool IsNodeOutDatedAndNeedsRefresh() override;
+
+	// UCustomizableObjectNode interface
+	void AllocateDefaultPins(UCustomizableObjectNodeRemapPins* RemapPins) override;
+	void ReconstructNode(UCustomizableObjectNodeRemapPins* RemapPins) override;
+
+	// Fills the list with all the morphs
+	void UpdateMorphList();
+
+	// Returns the mesh pin of the node
+	UEdGraphPin* GetMeshPin() const
+	{
+		return FindPin(TEXT("Mesh"), EGPD_Input);
+	}
+
+	// Returns the stack pin of the node
+	UEdGraphPin* GetStackPin() const
+	{
+		return FindPin(TEXT("Stack"), EGPD_Output);
+	}
+
+	// Returns the morph pin at Index
+	UEdGraphPin* GetMorphPin(int32 Index) const;
+	
+	// Returns the index of the next connected stack node. Returns -1 if there is none.
+	int32 NextConnectedPin(int32 Index, TArray<FString> AvailableMorphs)const;
+
+public:
+
+	// List with all the morphs of the linked skeletal mesh
+	UPROPERTY()
+	TArray<FString> MorphNames;
+
+};
diff --git a/Engine/Plugins/Experimental/Mutable/Source/CustomizableObjectEditor/Classes/Nodes/CustomizableObjectNodeMeshReshape.h b/Engine/Plugins/Experimental/Mutable/Source/CustomizableObjectEditor/Classes/Nodes/CustomizableObjectNodeMeshReshape.h
new file mode 100644
--- /dev/null
+++ b/Engine/Plugins/Experimental/Mutable/Source/CustomizableObjectEditor/Classes/Nodes/CustomizableObjectNodeMeshReshape.h
@@ -0,0 +1,70 @@
+// Copyright Epic Games, Inc. All Rights Reserved.
+
+#pragma once
+
+#include "Nodes/CustomizableObjectNode.h"
+#include "Nodes/CustomizableObjectNodeMeshReshapeCommon.h"
+
+#include "CustomizableObjectNodeMeshReshape.generated.h"
+
+UCLASS()
+class CUSTOMIZABLEOBJECTEDITOR_API UCustomizableObjectNodeMeshReshape : public UCustomizableObjectNode
+{
+public:
+	GENERATED_BODY()
+
+	UCustomizableObjectNodeMeshReshape();
+
+	// Begin EdGraphNode interface
+	FText GetNodeTitle(ENodeTitleType::Type TitleType) const override;
+	FLinearColor GetNodeTitleColor() const override;
+	FText GetTooltipText() const override;
+
+	// UCustomizableObjectNode interface
+	void AllocateDefaultPins(UCustomizableObjectNodeRemapPins* RemapPins) override;
+	FString GetRefreshMessage() const override;
+
+	inline UEdGraphPin* BaseMeshPin() const
+	{
+		return FindPin(TEXT("Base Mesh"), EGPD_Input);
+	}
+
+	inline UEdGraphPin* BaseShapePin() const
+	{
+		return FindPin(TEXT("Base Shape"), EGPD_Input);
+	}
+
+	inline UEdGraphPin* TargetShapePin() const
+	{
+		return FindPin(TEXT("Target Shape"), EGPD_Input);
+	}
+
+	/** Enable the deformation of the skeleton of the base mesh. */
+	UPROPERTY(EditAnywhere, Category = CustomizableObject)
+	bool bReshapeSkeleton = false;	
+
+	/** Enable the deformation of physics volumes of the base mesh */
+    UPROPERTY(EditAnywhere, Category = CustomizableObject)
+    bool bReshapePhysicsVolumes = false;
+	
+	/** Enable rigid parts base mesh. */
+	UPROPERTY(EditAnywhere, Category = CustomizableObject)
+	bool bEnableRigidParts = false;
+
+	/** Enables the deformation of all bones of the skeleton */
+	UPROPERTY(EditAnywhere, Category = ReshapeBones, meta=(EditCondition="bReshapeSkeleton"))
+	bool bDeformAllBones = false;
+
+	/** Array with bones that will be deformed */
+	UPROPERTY(EditAnywhere, Category = ReshapeBones, meta=(EditCondition="bReshapeSkeleton && !bDeformAllBones"))
+	TArray<FMeshReshapeBoneReference> BonesToDeform;
+
+	/** Enables the deformation of all physic bodies*/
+	UPROPERTY(EditAnywhere, Category = ReshapePhysics, meta = (EditCondition = "bReshapePhysicsVolumes"))
+	bool bDeformAllPhysicsBodies = false;
+
+	/** Array with bones with physics bodies that will be deformed */
+	UPROPERTY(EditAnywhere, Category = ReshapePhysics, meta=(EditCondition="bReshapePhysicsVolumes && !bDeformAllPhysicsBodies"))
+	TArray<FMeshReshapeBoneReference> PhysicsBodiesToDeform;
+};
+
diff --git a/Engine/Plugins/Experimental/Mutable/Source/CustomizableObjectEditor/Classes/Nodes/CustomizableObjectNodeMeshSwitch.h b/Engine/Plugins/Experimental/Mutable/Source/CustomizableObjectEditor/Classes/Nodes/CustomizableObjectNodeMeshSwitch.h
new file mode 100644
--- /dev/null
+++ b/Engine/Plugins/Experimental/Mutable/Source/CustomizableObjectEditor/Classes/Nodes/CustomizableObjectNodeMeshSwitch.h
@@ -0,0 +1,28 @@
+// Copyright Epic Games, Inc. All Rights Reserved.
+
+#pragma once
+
+#include "Nodes/CustomizableObjectNodeSwitchBase.h"
+
+#include "EdGraph/EdGraphPin.h"
+
+#include "CustomizableObjectNodeMeshSwitch.generated.h"
+
+
+UCLASS()
+class CUSTOMIZABLEOBJECTEDITOR_API UCustomizableObjectNodeMeshSwitch : public UCustomizableObjectNodeSwitchBase
+{
+public:
+	GENERATED_BODY()
+
+	// UObject interface
+	virtual void Serialize(FArchive& Ar) override;
+	
+	// UCustomizableObjectNodeSwitchBase interface
+	FString GetOutputPinName() const override;
+
+	FName GetCategory() const override;
+
+	// Own interface
+	FString GetPinPrefix() const;
+};
diff --git a/Engine/Plugins/Experimental/Mutable/Source/CustomizableObjectEditor/Classes/Nodes/CustomizableObjectNodeMeshVariation.h b/Engine/Plugins/Experimental/Mutable/Source/CustomizableObjectEditor/Classes/Nodes/CustomizableObjectNodeMeshVariation.h
new file mode 100644
--- /dev/null
+++ b/Engine/Plugins/Experimental/Mutable/Source/CustomizableObjectEditor/Classes/Nodes/CustomizableObjectNodeMeshVariation.h
@@ -0,0 +1,62 @@
+// Copyright Epic Games, Inc. All Rights Reserved.
+
+#pragma once
+
+#include "Nodes/CustomizableObjectNode.h"
+
+#include "CustomizableObjectNodeMeshVariation.generated.h"
+
+
+USTRUCT()
+struct CUSTOMIZABLEOBJECTEDITOR_API FCustomizableObjectMeshVariation
+{
+	GENERATED_USTRUCT_BODY()
+
+	UPROPERTY(EditAnywhere, Category = CustomizableObject)
+	FString Tag;
+};
+
+
+UCLASS()
+class CUSTOMIZABLEOBJECTEDITOR_API UCustomizableObjectNodeMeshVariation : public UCustomizableObjectNode
+{
+public:
+	GENERATED_BODY()
+
+	UCustomizableObjectNodeMeshVariation();
+
+	UPROPERTY(EditAnywhere, Category = CustomizableObject)
+	TArray<FCustomizableObjectMeshVariation> Variations;
+
+	// UObject interface.
+	void PostEditChangeProperty(FPropertyChangedEvent& PropertyChangedEvent) override;
+
+	// EdGraphNode interface
+	FText GetNodeTitle(ENodeTitleType::Type TitleType) const override;
+	FLinearColor GetNodeTitleColor() const override;
+	FText GetTooltipText() const override;
+
+	// UCustomizableObjectNode interface
+	void AllocateDefaultPins(UCustomizableObjectNodeRemapPins* RemapPins) override;
+
+	UEdGraphPin* OutputPin() const
+	{
+		return FindPin(TEXT("Mesh"));
+	}
+
+	UEdGraphPin* DefaultPin() const
+	{
+		return FindPin(TEXT("Default"));
+	}
+
+	int32 GetNumVariations() const
+	{
+		return Variations.Num();
+	}
+
+	UEdGraphPin* VariationPin(int Index) const
+	{
+		return FindPin(FString::Printf(TEXT("Variation %d"),Index));
+	}
+};
+
diff --git a/Engine/Plugins/Experimental/Mutable/Source/CustomizableObjectEditor/Classes/Nodes/CustomizableObjectNodeModifierBase.h b/Engine/Plugins/Experimental/Mutable/Source/CustomizableObjectEditor/Classes/Nodes/CustomizableObjectNodeModifierBase.h
new file mode 100644
--- /dev/null
+++ b/Engine/Plugins/Experimental/Mutable/Source/CustomizableObjectEditor/Classes/Nodes/CustomizableObjectNodeModifierBase.h
@@ -0,0 +1,19 @@
+// Copyright Epic Games, Inc. All Rights Reserved.
+
+#pragma once
+
+#include "Nodes/CustomizableObjectNodeMaterialBase.h"
+
+#include "CustomizableObjectNodeModifierBase.generated.h"
+
+
+UCLASS(abstract)
+class CUSTOMIZABLEOBJECTEDITOR_API UCustomizableObjectNodeModifierBase : public UCustomizableObjectNodeMaterialBase
+{
+public:
+	GENERATED_BODY()
+
+	UCustomizableObjectNodeModifierBase();
+
+};
+
diff --git a/Engine/Plugins/Experimental/Mutable/Source/CustomizableObjectEditor/Classes/Nodes/CustomizableObjectNodeMorphMaterial.h b/Engine/Plugins/Experimental/Mutable/Source/CustomizableObjectEditor/Classes/Nodes/CustomizableObjectNodeMorphMaterial.h
new file mode 100644
--- /dev/null
+++ b/Engine/Plugins/Experimental/Mutable/Source/CustomizableObjectEditor/Classes/Nodes/CustomizableObjectNodeMorphMaterial.h
@@ -0,0 +1,41 @@
+// Copyright Epic Games, Inc. All Rights Reserved.
+
+#pragma once
+
+#include "Nodes/CustomizableObjectNodeEditMaterialBase.h"
+
+#include "CustomizableObjectNodeMorphMaterial.generated.h"
+
+
+UCLASS()
+class CUSTOMIZABLEOBJECTEDITOR_API UCustomizableObjectNodeMorphMaterial : public UCustomizableObjectNodeEditMaterialBase
+{
+public:
+
+	GENERATED_BODY()
+	
+	// Begin EdGraphNode interface
+	FText GetNodeTitle(ENodeTitleType::Type TitleType) const override;
+	FLinearColor GetNodeTitleColor() const override;
+	FText GetTooltipText() const override;
+
+	// UCustomizableObjectNode interface
+	virtual void BackwardsCompatibleFixup() override;
+	void AllocateDefaultPins(UCustomizableObjectNodeRemapPins* RemapPins) override;
+	void PinConnectionListChanged(UEdGraphPin* Pin) override;
+	void PostEditChangeProperty(FPropertyChangedEvent& PropertyChangedEvent) override;
+
+	bool IsNodeOutDatedAndNeedsRefresh() override;
+	FString GetRefreshMessage() const override;
+
+	UPROPERTY(Category=CustomizableObject, EditAnywhere)
+	FString MorphTargetName;
+
+	bool IsSingleOutputNode() const override;
+
+	UEdGraphPin* FactorPin() const
+	{
+		return FindPin(TEXT("Factor"));
+	}
+};
+
diff --git a/Engine/Plugins/Experimental/Mutable/Source/CustomizableObjectEditor/Classes/Nodes/CustomizableObjectNodeObject.h b/Engine/Plugins/Experimental/Mutable/Source/CustomizableObjectEditor/Classes/Nodes/CustomizableObjectNodeObject.h
new file mode 100644
--- /dev/null
+++ b/Engine/Plugins/Experimental/Mutable/Source/CustomizableObjectEditor/Classes/Nodes/CustomizableObjectNodeObject.h
@@ -0,0 +1,204 @@
+// Copyright Epic Games, Inc. All Rights Reserved.
+
+#pragma once
+
+#include "Nodes/CustomizableObjectNodeObjectGroup.h"
+#include "Nodes/CustomizableObjectNodeMaterial.h"
+
+#include "CustomizableObjectNodeObject.generated.h"
+
+
+USTRUCT()
+struct CUSTOMIZABLEOBJECTEDITOR_API FCustomizableObjectState
+{
+	GENERATED_USTRUCT_BODY()
+
+	UPROPERTY(EditAnywhere, Category=CustomizableObject)
+	FString Name;
+
+	UPROPERTY(EditAnywhere, Category=CustomizableObject)
+	TArray<FString> RuntimeParameters;
+
+	UPROPERTY(EditAnywhere, Category = CustomizableObject)
+	bool bDontCompressRuntimeTextures = false;
+
+	UPROPERTY(EditAnywhere, Category = CustomizableObject)
+	bool bBuildOnlyFirstLOD = false;
+
+	UPROPERTY(EditAnywhere, Category = CustomizableObject)
+	TMap<FString, FString> ForcedParameterValues;
+
+	UPROPERTY(EditAnywhere, Category = UI, meta = (DisplayName = "State UI Metadata"))
+	FMutableParamUIMetadata StateUIMetadata;
+};
+
+
+UENUM()
+enum class ECustomizableObjectAutomaticLODStrategy : uint8
+{
+	// Use the same strategy than the parent object. If root, then use "Manual".
+	Inherited = 0 UMETA(DisplayName = "Inherit from parent object"),
+	// Don't try to generate LODs automatically for the child nodes. Only the ones tha explicitely define them will be used.
+	Manual = 1 UMETA(DisplayName = "Only manually created LODs"),
+	// Try to generate the same material structure than LOd 0 if the source meshes have LODs.
+	AutomaticFromMesh = 2 UMETA(DisplayName = "Automatic from mesh")
+};
+
+UENUM()
+enum class ECustomizableObjectSelectionOverride : uint8
+{
+    NoOverride = 0 UMETA(DisplayName = "No Override"),
+    Disable    = 1 UMETA(DisplayName = "Disable"    ),
+    Enable     = 2 UMETA(DisplayName = "Enable"     )
+};
+
+USTRUCT()
+struct FRealTimeMorphSelectionOverride
+{
+	GENERATED_BODY()
+
+    FRealTimeMorphSelectionOverride()
+    {
+    }
+
+    FRealTimeMorphSelectionOverride(const FName& InMorphName)
+        : MorphName(InMorphName)
+    {
+    }
+
+    UPROPERTY()
+    FName MorphName;
+
+    UPROPERTY()
+    ECustomizableObjectSelectionOverride SelectionOverride = 
+        ECustomizableObjectSelectionOverride::NoOverride;
+
+    UPROPERTY()
+    TArray<FName> SkeletalMeshesNames;
+
+    UPROPERTY()
+    TArray<ECustomizableObjectSelectionOverride> Override;
+};
+
+UCLASS()
+class CUSTOMIZABLEOBJECTEDITOR_API UCustomizableObjectNodeObject : public UCustomizableObjectNode
+{
+public:
+	GENERATED_BODY()
+
+	UCustomizableObjectNodeObject();
+
+	UPROPERTY(EditAnywhere, Category=CustomizableObject)
+	FString ObjectName;
+
+	UPROPERTY(EditAnywhere, Category = UI, meta = (DisplayName = "Parameter UI Metadata"))
+	FMutableParamUIMetadata ParamUIMetadata;
+
+	UPROPERTY(EditAnywhere, Category = CustomizableObject)
+	int32 NumLODs;
+
+	UPROPERTY(EditAnywhere, Category = CustomizableObject)
+	ECustomizableObjectAutomaticLODStrategy AutoLODStrategy;
+
+	UPROPERTY(EditAnywhere, Category = CustomizableObject, meta = (ClampMin = "1"))
+	int32 NumMeshComponents = 1;
+
+	UPROPERTY(EditAnywhere, Category=CustomizableObject)
+	TArray<FCustomizableObjectState> States;
+
+	UPROPERTY(EditAnywhere, Category = AttachedToExternalObject)
+	TObjectPtr<UCustomizableObject> ParentObject;
+
+	UPROPERTY(EditAnywhere, Category = AttachedToExternalObject)
+	FGuid ParentObjectGroupId;
+
+	UPROPERTY()
+	FGuid Identifier;
+
+    // Soft references SkeletalMeshes found in the provoius compilation.
+    // Only populated if the node is the root.
+    UPROPERTY()
+    TArray<FSoftObjectPath> ReferencedSkeletalMeshes;
+
+    // Information about the realtime morph targets usage. It indexes to ReferncedSkeletakMeshes array
+    // so it is need to keep them syncronized. 
+    // This overrides the per skeletal mesh node selection 
+    UPROPERTY()
+    TArray<FRealTimeMorphSelectionOverride> RealTimeMorphSelectionOverrides; 
+	
+    // To avoid any no properly saved GUIDs
+	FGuid IdentifierVerification;
+
+	// UObject interface.
+	void PostEditChangeProperty(FPropertyChangedEvent& PropertyChangedEvent) override;
+
+	// EdGraphNode interface
+	FText GetNodeTitle(ENodeTitleType::Type TitleType) const override;
+	FLinearColor GetNodeTitleColor() const override;
+	FText GetTooltipText() const override;
+	void PrepareForCopying() override;
+	void PostPasteNode() override;
+	void PostDuplicate(bool bDuplicateForPIE) override;
+
+	// UCustomizableObjectNode interface
+	virtual void PostBackwardsCompatibleFixup() override;
+	void AllocateDefaultPins(UCustomizableObjectNodeRemapPins* RemapPins) override;
+
+	// Own Interface
+	void SetParentObject(UCustomizableObject* CustomizableParentObject);
+
+	// Own interface
+	UPROPERTY()
+	bool bIsBase;
+
+	UEdGraphPin* LODPin(int32 LODIndex) const
+	{
+		FString LODName = FString::Printf(TEXT("LOD %d "), LODIndex);
+		return FindPin(LODName);
+	}
+
+	int32 GetNumLODPins() const
+	{
+		int32 Count = 0;
+
+		for (UEdGraphPin* Pin : GetAllNonOrphanPins())
+		{
+			if (Pin->GetName().StartsWith(TEXT("LOD ")))
+			{
+				Count++;
+			}
+		}
+
+		return Count;
+	}
+
+	UEdGraphPin* ChildrenPin() const
+	{
+		return FindPin(TEXT("Children"));
+	}
+
+	UEdGraphPin* OutputPin() const
+	{
+		return FindPin(TEXT("Object"));
+	}
+
+	/** Return the LOD which a LOD pin references to. Retrun -1 if a pin does not belong to any LOD. */
+	int32 GetLOD(UEdGraphPin* Pin) const;
+
+	virtual bool CanUserDeleteNode() const override;
+	virtual bool CanDuplicateNode() const override;
+
+	/** Get all Material Nodes int his Customizable Object which belong to the given LOD.
+	 *
+	 * @param LOD LOD which materials have to belong to.
+	 */
+	TArray<UCustomizableObjectNodeMaterial*> GetMaterialNodes(int LOD) const;
+
+	bool IsSingleOutputNode() const override;
+
+	void SetMeshComponentNumFromParent(int32 MeshComponentNum)
+	{
+		NumMeshComponents = MeshComponentNum;
+	}
+};
+
diff --git a/Engine/Plugins/Experimental/Mutable/Source/CustomizableObjectEditor/Classes/Nodes/CustomizableObjectNodeObjectChild.h b/Engine/Plugins/Experimental/Mutable/Source/CustomizableObjectEditor/Classes/Nodes/CustomizableObjectNodeObjectChild.h
new file mode 100644
--- /dev/null
+++ b/Engine/Plugins/Experimental/Mutable/Source/CustomizableObjectEditor/Classes/Nodes/CustomizableObjectNodeObjectChild.h
@@ -0,0 +1,26 @@
+// Copyright Epic Games, Inc. All Rights Reserved.
+
+#pragma once
+
+#include "Nodes/CustomizableObjectNodeObject.h"
+
+#include "CustomizableObjectNodeObjectChild.generated.h"
+
+
+UCLASS()
+class CUSTOMIZABLEOBJECTEDITOR_API UCustomizableObjectNodeObjectChild : public UCustomizableObjectNodeObject
+{
+public:
+	GENERATED_BODY()
+
+	UCustomizableObjectNodeObjectChild();
+
+public:
+	// UEdGraphNode interface
+	FText GetNodeTitle(ENodeTitleType::Type TitleType) const override;
+	void PrepareForCopying() override;
+	FText GetTooltipText() const override;
+	bool CanUserDeleteNode() const override;
+	bool CanDuplicateNode() const override;
+};
+
diff --git a/Engine/Plugins/Experimental/Mutable/Source/CustomizableObjectEditor/Classes/Nodes/CustomizableObjectNodeObjectGroup.h b/Engine/Plugins/Experimental/Mutable/Source/CustomizableObjectEditor/Classes/Nodes/CustomizableObjectNodeObjectGroup.h
new file mode 100644
--- /dev/null
+++ b/Engine/Plugins/Experimental/Mutable/Source/CustomizableObjectEditor/Classes/Nodes/CustomizableObjectNodeObjectGroup.h
@@ -0,0 +1,49 @@
+// Copyright Epic Games, Inc. All Rights Reserved.
+
+#pragma once
+
+#include "CustomizableObjectNode.h"
+#include "CustomizableObjectNodeObjectGroup.generated.h"
+
+
+UCLASS()
+class CUSTOMIZABLEOBJECTEDITOR_API UCustomizableObjectNodeObjectGroup : public UCustomizableObjectNode
+{
+public:
+	GENERATED_BODY()
+
+	UCustomizableObjectNodeObjectGroup();
+
+	UPROPERTY(EditAnywhere, Category=CustomizableObject)
+	FString GroupName;
+
+	UPROPERTY(EditAnywhere, Category=CustomizableObject)
+	ECustomizableObjectGroupType GroupType;
+
+	UPROPERTY(EditAnywhere, Category = UI, meta = (DisplayName = "Parameter UI Metadata"))
+	FMutableParamUIMetadata ParamUIMetadata;
+
+	// UObject interface.
+	void Serialize(FArchive& Ar) override;
+	void PostEditChangeProperty(FPropertyChangedEvent& PropertyChangedEvent) override;
+
+	// EdGraphNode interface
+	FText GetNodeTitle(ENodeTitleType::Type TitleType) const override;
+	FLinearColor GetNodeTitleColor() const override;
+	FText GetTooltipText() const override;
+
+	// UCustomizableObjectNode interface
+	void AllocateDefaultPins(UCustomizableObjectNodeRemapPins* RemapPins) override;
+
+	// Own interface
+	UEdGraphPin* ObjectsPin() const
+	{
+		return FindPin(TEXT("Objects"));
+	}
+
+	UEdGraphPin* GroupProjectorsPin() const
+	{
+		return FindPin(TEXT("Projectors"));
+	}
+};
+
diff --git a/Engine/Plugins/Experimental/Mutable/Source/CustomizableObjectEditor/Classes/Nodes/CustomizableObjectNodeParentedMaterial.h b/Engine/Plugins/Experimental/Mutable/Source/CustomizableObjectEditor/Classes/Nodes/CustomizableObjectNodeParentedMaterial.h
new file mode 100644
--- /dev/null
+++ b/Engine/Plugins/Experimental/Mutable/Source/CustomizableObjectEditor/Classes/Nodes/CustomizableObjectNodeParentedMaterial.h
@@ -0,0 +1,29 @@
+// Copyright Epic Games, Inc. All Rights Reserved.
+
+#pragma once
+
+#include "CustomizableObjectNodeParentedNode.h"
+#include "Nodes/CustomizableObjectNode.h"
+#include "Nodes/CustomizableObjectNodeMaterial.h"
+
+
+/** Node Material specialization of ICustomizableObjectNodeParentedNode. */
+class FCustomizableObjectNodeParentedMaterial : public ICustomizableObjectNodeParentedNode
+{
+public:
+	// Own interface
+	/** Return the parent material node. */
+	UCustomizableObjectNodeMaterial* GetParentMaterialNode() const;
+	
+	/** Return all possible parent material nodes of the node. */
+	TArray<UCustomizableObjectNodeMaterial*> GetPossibleParentMaterialNodes() const;
+
+	/** Returns the parent material node if there exist a path to it. */
+	UCustomizableObjectNodeMaterial* GetParentMaterialNodeIfPath() const;
+	
+protected:
+	/** Return the node which this interface belongs to. */
+	virtual UCustomizableObjectNode& GetNode() = 0;
+	const UCustomizableObjectNode& GetNode() const;
+};
+
diff --git a/Engine/Plugins/Experimental/Mutable/Source/CustomizableObjectEditor/Classes/Nodes/CustomizableObjectNodeParentedNode.h b/Engine/Plugins/Experimental/Mutable/Source/CustomizableObjectEditor/Classes/Nodes/CustomizableObjectNodeParentedNode.h
new file mode 100644
--- /dev/null
+++ b/Engine/Plugins/Experimental/Mutable/Source/CustomizableObjectEditor/Classes/Nodes/CustomizableObjectNodeParentedNode.h
@@ -0,0 +1,28 @@
+// Copyright Epic Games, Inc. All Rights Reserved.
+
+#pragma once
+
+#include "CustomizableObject.h"
+#include "Nodes/CustomizableObjectNode.h"
+
+
+/** Interface for a node that saves a parent node. */
+class ICustomizableObjectNodeParentedNode
+{
+public:
+	virtual ~ICustomizableObjectNodeParentedNode() {};
+
+	/** Save the parent node reference. */
+	virtual void SetParentNode(UCustomizableObject* Object, FGuid NodeId);
+	
+	/** Return the saved parent material node. Can return nullptr. */
+	UCustomizableObjectNode* GetParentNode() const;
+
+protected:
+	/** Save the parent node reference. Do not call directly, use SetParentNode instead. */
+	virtual void SaveParentNode(UCustomizableObject* Object, FGuid NodeId) = 0;
+
+	virtual FGuid GetParentNodeId() const = 0;
+
+	virtual UCustomizableObject* GetParentObject() const = 0;
+};
\ No newline at end of file
diff --git a/Engine/Plugins/Experimental/Mutable/Source/CustomizableObjectEditor/Classes/Nodes/CustomizableObjectNodeProjectorConstant.h b/Engine/Plugins/Experimental/Mutable/Source/CustomizableObjectEditor/Classes/Nodes/CustomizableObjectNodeProjectorConstant.h
new file mode 100644
--- /dev/null
+++ b/Engine/Plugins/Experimental/Mutable/Source/CustomizableObjectEditor/Classes/Nodes/CustomizableObjectNodeProjectorConstant.h
@@ -0,0 +1,66 @@
+// Copyright Epic Games, Inc. All Rights Reserved.
+
+#pragma once
+
+#include "CustomizableObjectInstance.h"
+#include "Nodes/CustomizableObjectNode.h"
+
+#include "CustomizableObjectNodeProjectorConstant.generated.h"
+
+
+UCLASS(hideCategories = (CustomizableObjectHide))
+class CUSTOMIZABLEOBJECTEDITOR_API UCustomizableObjectNodeProjectorConstant : public UCustomizableObjectNode
+{
+public:
+	GENERATED_BODY()
+
+	UCustomizableObjectNodeProjectorConstant();
+
+	/**  */
+	UPROPERTY(EditAnywhere, Category= CustomizableObjectHide)
+	FCustomizableObjectProjector Value;
+
+	UPROPERTY(EditAnywhere, Category = UI, meta = (DisplayName = "Parameter UI Metadata"))
+	FMutableParamUIMetadata ParamUIMetadata;
+
+	UPROPERTY(EditAnywhere, Category = CustomizableObject)
+	ECustomizableObjectProjectorType ProjectionType;
+
+	UPROPERTY(EditAnywhere, Category = CustomizableObject, Meta = (DisplayName = "Projection Angle (degrees)"))
+	float ProjectionAngle;
+
+	UPROPERTY(EditAnywhere, Category = ProjectorSnapToBone)
+	uint32 ReferenceSkeletonIndex = 0;
+
+	UPROPERTY(EditAnywhere, Category = ProjectorSnapToBone)
+	FName ProjectorBone;
+
+	/** Temporary variable where to put the location information for bone combo box selection changes (in FCustomizableObjectNodeProjectorParameterDetails)*/
+	FVector BoneComboBoxLocation;
+
+	/** Temporary variable where to put the direction information for bone combo box selection changes (in FCustomizableObjectNodeProjectorParameterDetails)*/
+	FVector BoneComboBoxForwardDirection;
+
+	/** Temporary variable where to put the direction information for bone combo box selection changes (in FCustomizableObjectNodeProjectorParameterDetails)*/
+	FVector BoneComboBoxUpDirection;
+
+	/** Flag to know which parameters have been modified of the node in the details tab.
+	* 0 Means projection type
+	* 1 projection angle
+	* 2 projector transform */
+	int32 ParameterSetModified;
+
+	// UObject interface.
+	void PostEditChangeProperty(FPropertyChangedEvent& PropertyChangedEvent) override;
+
+	// Begin EdGraphNode interface
+	FText GetNodeTitle(ENodeTitleType::Type TitleType) const override;
+	FLinearColor GetNodeTitleColor() const override;
+	FText GetTooltipText() const override;
+
+	// UCustomizableObjectNode interface
+	void AllocateDefaultPins(UCustomizableObjectNodeRemapPins* RemapPins) override;
+
+	UEdGraphPin* ValuePin();
+};
+
diff --git a/Engine/Plugins/Experimental/Mutable/Source/CustomizableObjectEditor/Classes/Nodes/CustomizableObjectNodeProjectorParameter.h b/Engine/Plugins/Experimental/Mutable/Source/CustomizableObjectEditor/Classes/Nodes/CustomizableObjectNodeProjectorParameter.h
new file mode 100644
--- /dev/null
+++ b/Engine/Plugins/Experimental/Mutable/Source/CustomizableObjectEditor/Classes/Nodes/CustomizableObjectNodeProjectorParameter.h
@@ -0,0 +1,70 @@
+// Copyright Epic Games, Inc. All Rights Reserved.
+
+#pragma once
+
+#include "CustomizableObjectInstance.h"
+#include "Nodes/CustomizableObjectNode.h"
+
+#include "CustomizableObjectNodeProjectorParameter.generated.h"
+
+
+UCLASS(hideCategories = (CustomizableObjectHide))
+class CUSTOMIZABLEOBJECTEDITOR_API UCustomizableObjectNodeProjectorParameter : public UCustomizableObjectNode
+{
+public:
+	GENERATED_BODY()
+
+	UPROPERTY(EditAnywhere, Category=CustomizableObjectHide)
+	FCustomizableObjectProjector DefaultValue;
+
+	UPROPERTY(EditAnywhere, Category=CustomizableObject)
+	FString ParameterName = "Default Name";
+
+	UPROPERTY(EditAnywhere, Category = UI, meta = (DisplayName = "Parameter UI Metadata"))
+	FMutableParamUIMetadata ParamUIMetadata;
+
+	UPROPERTY(EditAnywhere, Category = CustomizableObject)
+	ECustomizableObjectProjectorType ProjectionType = ECustomizableObjectProjectorType::Planar;
+
+	UPROPERTY(EditAnywhere, Category = CustomizableObject, Meta = (DisplayName = "Projection Angle (degrees)"))
+	float ProjectionAngle = 360.0f;
+
+	UPROPERTY(EditAnywhere, Category = ProjectorSnapToBone)
+	uint32 ReferenceSkeletonIndex = 0;
+
+	UPROPERTY(EditAnywhere, Category = ProjectorSnapToBone)
+	FName ProjectorBone;
+
+	/** Temporary variable where to put the location information for bone combo box selection changes (in FCustomizableObjectNodeProjectorParameterDetails)*/
+	FVector BoneComboBoxLocation = FVector::ZeroVector;
+
+	/** Temporary variable where to put the direction information for bone combo box selection changes (in FCustomizableObjectNodeProjectorParameterDetails)*/
+	FVector BoneComboBoxForwardDirection = FVector::ZeroVector;
+
+	/** Temporary variable where to put the direction information for bone combo box selection changes (in FCustomizableObjectNodeProjectorParameterDetails)*/
+	FVector BoneComboBoxUpDirection = FVector::ZeroVector;
+
+	/** Flag to know which parameters have been modified of the node in the details tab.
+	* 0 Means projection type
+	* 1 projection angle
+	* 2 projector transform */
+	int32 ParameterSetModified = -1;
+
+	// UObject interface.
+	void PostEditChangeProperty(FPropertyChangedEvent& PropertyChangedEvent) override;
+
+	// Begin EdGraphNode interface
+	FText GetNodeTitle(ENodeTitleType::Type TitleType) const override;
+	FLinearColor GetNodeTitleColor() const override;
+	FText GetTooltipText() const override;
+
+	// UCustomizableObjectNode interface
+	void AllocateDefaultPins(UCustomizableObjectNodeRemapPins* RemapPins) override;
+	virtual bool IsAffectedByLOD() const override { return false; }
+
+	UEdGraphPin* ProjectorPin() const
+	{
+		return FindPin(TEXT("Value"));
+	}
+};
+
diff --git a/Engine/Plugins/Experimental/Mutable/Source/CustomizableObjectEditor/Classes/Nodes/CustomizableObjectNodeRemoveMesh.h b/Engine/Plugins/Experimental/Mutable/Source/CustomizableObjectEditor/Classes/Nodes/CustomizableObjectNodeRemoveMesh.h
new file mode 100644
--- /dev/null
+++ b/Engine/Plugins/Experimental/Mutable/Source/CustomizableObjectEditor/Classes/Nodes/CustomizableObjectNodeRemoveMesh.h
@@ -0,0 +1,31 @@
+// Copyright Epic Games, Inc. All Rights Reserved.
+
+#pragma once
+
+#include "Nodes/CustomizableObjectNodeEditMaterialBase.h"
+
+#include "CustomizableObjectNodeRemoveMesh.generated.h"
+
+UCLASS()
+class CUSTOMIZABLEOBJECTEDITOR_API UCustomizableObjectNodeRemoveMesh : public UCustomizableObjectNodeEditMaterialBase
+{
+public:
+	GENERATED_BODY()
+
+	// Begin EdGraphNode interface
+	FText GetNodeTitle(ENodeTitleType::Type TitleType) const override;
+	FLinearColor GetNodeTitleColor() const override;
+	FText GetTooltipText() const override;
+
+	// UCustomizableObjectNode interface
+	void AllocateDefaultPins(UCustomizableObjectNodeRemapPins* RemapPins) override;
+	void PinConnectionListChanged(UEdGraphPin * Pin);
+
+	UEdGraphPin* RemoveMeshPin() const
+	{
+		return FindPin(TEXT("Remove Mesh"));
+	}
+
+	bool IsSingleOutputNode() const override;
+};
+
diff --git a/Engine/Plugins/Experimental/Mutable/Source/CustomizableObjectEditor/Classes/Nodes/CustomizableObjectNodeRemoveMeshBlocks.h b/Engine/Plugins/Experimental/Mutable/Source/CustomizableObjectEditor/Classes/Nodes/CustomizableObjectNodeRemoveMeshBlocks.h
new file mode 100644
--- /dev/null
+++ b/Engine/Plugins/Experimental/Mutable/Source/CustomizableObjectEditor/Classes/Nodes/CustomizableObjectNodeRemoveMeshBlocks.h
@@ -0,0 +1,35 @@
+// Copyright Epic Games, Inc. All Rights Reserved.
+
+#pragma once
+
+#include "Nodes/CustomizableObjectNodeEditLayoutBlocks.h"
+
+#include "CustomizableObjectNodeRemoveMeshBlocks.generated.h"
+
+UCLASS()
+class CUSTOMIZABLEOBJECTEDITOR_API UCustomizableObjectNodeRemoveMeshBlocks : public UCustomizableObjectNodeEditLayoutBlocks
+{
+public:
+	GENERATED_BODY()
+
+	UPROPERTY()
+	TArray<int32> Blocks_DEPRECATED;
+
+	// UObject interface.
+	void PostEditChangeProperty(FPropertyChangedEvent& PropertyChangedEvent) override;
+	virtual void Serialize(FArchive& Ar) override;
+
+	// Begin EdGraphNode interface
+	FText GetNodeTitle(ENodeTitleType::Type TitleType) const override;
+	FLinearColor GetNodeTitleColor() const override;
+	FText GetTooltipText() const override;
+
+	// UCustomizableObjectNode interface
+	virtual void BackwardsCompatibleFixup() override;
+	void AllocateDefaultPins(UCustomizableObjectNodeRemapPins* RemapPins) override;
+	bool IsNodeOutDatedAndNeedsRefresh() override;
+	FString GetRefreshMessage() const override;
+	void PinConnectionListChanged(UEdGraphPin * Pin);
+	bool IsSingleOutputNode() const override;
+};
+
diff --git a/Engine/Plugins/Experimental/Mutable/Source/CustomizableObjectEditor/Classes/Nodes/CustomizableObjectNodeSkeletalMesh.h b/Engine/Plugins/Experimental/Mutable/Source/CustomizableObjectEditor/Classes/Nodes/CustomizableObjectNodeSkeletalMesh.h
new file mode 100644
--- /dev/null
+++ b/Engine/Plugins/Experimental/Mutable/Source/CustomizableObjectEditor/Classes/Nodes/CustomizableObjectNodeSkeletalMesh.h
@@ -0,0 +1,187 @@
+// Copyright Epic Games, Inc. All Rights Reserved.
+
+#pragma once
+
+#include "AssetThumbnail.h"
+#include "EdGraph/EdGraphPin.h"
+#include "GameplayTagContainer.h"
+#include "Nodes/CustomizableObjectNodeMesh.h"
+#include "SGraphNode.h"
+
+#include "CustomizableObjectNodeSkeletalMesh.generated.h"
+
+
+// Class to render the Skeletal Mesh thumbnail of a CustomizableObjectNodeSkeletalMesh
+class SGraphNodeSkeletalMesh : public SGraphNode
+{
+public:
+
+	SLATE_BEGIN_ARGS(SGraphNodeSkeletalMesh) {}
+	SLATE_END_ARGS();
+
+	SGraphNodeSkeletalMesh() : SGraphNode() {};
+
+	// Builds the SGraphNodeSkeletalMesh when needed
+	void Construct(const FArguments& InArgs, UEdGraphNode* InGraphNode);
+
+	// Calls the needed functions to build the SGraphNode widgets
+	void UpdateGraphNode();
+
+	// Overriden functions to build the SGraphNode widgets
+	virtual void SetDefaultTitleAreaWidget(TSharedRef<SOverlay> DefaultTitleAreaWidget) override;
+	virtual void CreateBelowPinControls(TSharedPtr<SVerticalBox> MainBox) override;
+	virtual bool ShouldAllowCulling() const override { return false; }
+
+	// Callbacks for the widget
+	void OnExpressionPreviewChanged(const ECheckBoxState NewCheckedState);
+	ECheckBoxState IsExpressionPreviewChecked() const;
+	const FSlateBrush* GetExpressionPreviewArrow() const;
+	EVisibility ExpressionPreviewVisibility() const;
+
+public:
+
+	// Single property that only draws the combo box widget of the skeletal mesh
+	TSharedPtr<ISinglePropertyView> SkeletalMeshSelector;
+
+	// Pointer to the NodeSkeletalMesh that owns this SGraphNode
+	class UCustomizableObjectNodeSkeletalMesh* NodeSkeletalMesh;
+
+private:
+	
+	// Classes needed to get and render the thumbnail of the Skeletal Mesh
+	TSharedPtr<FAssetThumbnailPool> AssetThumbnailPool;
+	TSharedPtr<FAssetThumbnail> AssetThumbnail;
+
+	//This parameter defines the size of the thumbnail widget inside the Node
+	float WidgetSize;
+
+	// This parameter defines the resolution of the thumbnail
+	uint32 ThumbnailSize;
+
+};
+
+
+
+USTRUCT()
+struct CUSTOMIZABLEOBJECTEDITOR_API FCustomizableObjectNodeSkeletalMeshMaterial
+{
+	GENERATED_USTRUCT_BODY()
+
+	UPROPERTY()
+	FString Name;
+
+	UPROPERTY()
+	TObjectPtr<UEdGraphPin_Deprecated> MeshPin = nullptr;
+
+	UPROPERTY()
+	TArray< TObjectPtr<UEdGraphPin_Deprecated> > LayoutPins;
+
+	UPROPERTY()
+	TArray< TObjectPtr<UEdGraphPin_Deprecated> > ImagePins;
+
+	UPROPERTY()
+	FEdGraphPinReference MeshPinRef;
+
+	UPROPERTY()
+	TArray<FEdGraphPinReference> LayoutPinsRef;
+
+	UPROPERTY()
+	TArray<FEdGraphPinReference> ImagePinsRef;
+};
+
+
+USTRUCT()
+struct CUSTOMIZABLEOBJECTEDITOR_API FCustomizableObjectNodeSkeletalMeshLOD
+{
+	GENERATED_USTRUCT_BODY()
+
+	UPROPERTY()
+	TArray<FCustomizableObjectNodeSkeletalMeshMaterial> Materials;
+};
+
+
+UCLASS()
+class CUSTOMIZABLEOBJECTEDITOR_API UCustomizableObjectNodeSkeletalMesh : public UCustomizableObjectNodeMesh
+{
+public:
+	GENERATED_BODY()
+	
+	UPROPERTY(EditAnywhere, Category = CustomizableObject)
+	TObjectPtr<USkeletalMesh> SkeletalMesh;
+
+	/** Images */
+	UPROPERTY()
+	TArray<FCustomizableObjectNodeSkeletalMeshLOD> LODs;
+
+	/** Default pin when there is no mesh. */
+	UPROPERTY()
+	FEdGraphPinReference DefaultPin;
+	
+	/** Morphs */
+	UPROPERTY(EditAnywhere, Category = CustomizableObject)
+	TArray<FString> UsedRealTimeMorphTargetNames;
+
+	UPROPERTY(EditAnywhere, Category = CustomizableObject)
+    bool bUseAllRealTimeMorphs = false; 
+	
+	/** The anim instance that will be gathered by a Generated instance if it contains this skeletal mesh part, 
+		it will be grouped by component and AnimBlueprintSlot (the next UProperty). */
+	UPROPERTY(EditAnywhere, Category = CustomizableObject)
+	TSoftClassPtr<UAnimInstance> AnimInstance;
+
+	/** The anim slot associated with the AnimInstance */
+	UPROPERTY(EditAnywhere, Category = CustomizableObject)
+	int32 AnimBlueprintSlot = 0;
+
+	/** Animation tags that will be gathered by a Generated instance if it contains this skeletal mesh part,
+		it will not be grouped by component or AnimBlueprintSlot */
+	UPROPERTY(EditAnywhere, Category = CustomizableObject)
+	FGameplayTagContainer AnimationGameplayTags;
+
+	// UObject interface.
+	virtual void PostEditChangeProperty(FPropertyChangedEvent& PropertyChangedEvent) override;
+	virtual void Serialize(FArchive& Ar) override;
+
+	// UEdGraphNode interface
+	virtual FText GetNodeTitle(ENodeTitleType::Type TitleType) const override;
+	virtual FLinearColor GetNodeTitleColor() const override;
+	virtual FText GetTooltipText() const override;
+
+	// UCustomizableObjectNode interface
+	virtual void BackwardsCompatibleFixup() override;
+	virtual void AllocateDefaultPins(UCustomizableObjectNodeRemapPins* RemapPins) override;
+	virtual UCustomizableObjectNodeRemapPinsByName* CreateRemapPinsByName() const override;
+
+	// UCustomizableObjectNodeMesh interface
+	virtual UTexture2D* FindTextureForPin(const UEdGraphPin* Pin) const override;
+	virtual void GetUVChannelForPin(const UEdGraphPin* Pin, TArray<FVector2f>& OutSegments, int32 UVIndex) const override;
+	virtual TArray<UCustomizableObjectLayout*> GetLayouts(const UEdGraphPin* OutPin) const override;
+	virtual UObject* GetMesh() const override;
+	virtual UEdGraphPin* GetMeshPin(int32 LOD, int MaterialIndex) const override;
+	virtual void GetPinSection(const UEdGraphPin& Pin, int32& OutLODIndex, int32& OutSectionIndex, int32& OutLayoutIndex) const override;
+	
+	/** Returns the material associated to the given output pin. */
+	UMaterialInterface* GetMaterialFor(const UEdGraphPin* Pin) const;
+	FSkeletalMaterial* GetSkeletalMaterialFor(const UEdGraphPin* Pin) const;
+
+	virtual bool ProvidesCustomPinRelevancyTest() const override { return true; }
+	virtual bool IsPinRelevant(const UEdGraphPin* Pin) const override;
+
+	virtual bool IsNodeOutDatedAndNeedsRefresh() override;
+	virtual FString GetRefreshMessage() const override;
+
+	// Creates the SGraph Node widget for the thumbnail
+	virtual TSharedPtr<SGraphNode> CreateVisualWidget() override;
+
+	// Check if previous UV channels pins of the same LOD and material are also linked to a layout node
+	bool CheckIsValidLayout(const UEdGraphPin* Pin, int32& LayoutIndex, FString& MaterialName);
+
+	// Determines if the Node is collapsed or not
+	bool bCollapsed = true;
+
+	// Pointer to the SGraphNode Skeletal Mesh
+	TWeakPtr< SGraphNodeSkeletalMesh > GraphNodeSkeletalMesh;
+private:
+	UMaterialInterface* GetMaterialInterfaceFor(const int LODIndex, const int MaterialIndex, const FSkeletalMeshModel* ImportedModel = nullptr) const;
+	FSkeletalMaterial* GetSkeletalMaterialFor(const int LODIndex, const int MaterialIndex, const FSkeletalMeshModel* ImportedModel = nullptr) const;
+};
