[mutable] Move Mutable plugin from the restricted folder to the experimental plugins folder.

diff --git a/Engine/Plugins/Experimental/Mutable/Source/CustomizableObjectEditor/Classes/Nodes/CustomizableObjectNodeStaticMesh.h b/Engine/Plugins/Experimental/Mutable/Source/CustomizableObjectEditor/Classes/Nodes/CustomizableObjectNodeStaticMesh.h
new file mode 100644
--- /dev/null
+++ b/Engine/Plugins/Experimental/Mutable/Source/CustomizableObjectEditor/Classes/Nodes/CustomizableObjectNodeStaticMesh.h
@@ -0,0 +1,156 @@
+// Copyright Epic Games, Inc. All Rights Reserved.
+
+#pragma once
+
+#include "Nodes/CustomizableObjectNodeMesh.h"
+#include "SGraphNode.h"
+#include "AssetThumbnail.h"
+
+#include "CustomizableObjectNodeStaticMesh.generated.h"
+
+
+// Class to render the Static Mesh thumbnail of a CustomizableObjectNodeStaticMesh
+class SGraphNodeStaticMesh : public SGraphNode
+{
+public:
+
+	SLATE_BEGIN_ARGS(SGraphNodeStaticMesh) {}
+	SLATE_END_ARGS();
+
+	SGraphNodeStaticMesh() : SGraphNode() {};
+
+	// Builds the SGraphNodeStaticMesh when needed
+	void Construct(const FArguments& InArgs, UEdGraphNode* InGraphNode);
+
+	// Calls the needed functions to build the SGraphNode widgets
+	void UpdateGraphNode();
+
+	// Overriden functions to build the SGraphNode widgets
+	virtual void SetDefaultTitleAreaWidget(TSharedRef<SOverlay> DefaultTitleAreaWidget) override;
+	virtual void CreateBelowPinControls(TSharedPtr<SVerticalBox> MainBox) override;
+	virtual bool ShouldAllowCulling() const override { return false; }
+
+	// Callbacks for the widget
+	void OnExpressionPreviewChanged(const ECheckBoxState NewCheckedState);
+	ECheckBoxState IsExpressionPreviewChecked() const;
+	const FSlateBrush* GetExpressionPreviewArrow() const;
+	EVisibility ExpressionPreviewVisibility() const;
+
+public:
+
+	// Single property that only draws the combo box widget of the static mesh
+	TSharedPtr<class ISinglePropertyView> StaticMeshSelector;
+
+	// Pointer to the NodeStaticMesh that owns this SGraphNode
+	class UCustomizableObjectNodeStaticMesh* NodeStaticMesh;
+
+private:
+
+	// Classes needed to get and render the thumbnail of the Static Mesh
+	TSharedPtr<FAssetThumbnailPool> AssetThumbnailPool;
+	TSharedPtr<FAssetThumbnail> AssetThumbnail;
+
+	//This parameter defines the size of the thumbnail widget inside the Node
+	float WidgetSize;
+
+	// This parameter defines the resolution of the thumbnail
+	uint32 ThumbnailSize;
+
+};
+
+
+USTRUCT()
+struct CUSTOMIZABLEOBJECTEDITOR_API FCustomizableObjectNodeStaticMeshMaterial
+{
+	GENERATED_USTRUCT_BODY()
+
+	UPROPERTY()
+	FString Name;
+
+	UPROPERTY()
+	TObjectPtr<UEdGraphPin_Deprecated> MeshPin = nullptr;
+
+	UPROPERTY()
+	TObjectPtr<UEdGraphPin_Deprecated> LayoutPin = nullptr;
+
+	UPROPERTY()
+	TArray< TObjectPtr<UEdGraphPin_Deprecated>> ImagePins;
+
+	UPROPERTY()
+	FEdGraphPinReference MeshPinRef;
+
+	UPROPERTY()
+	FEdGraphPinReference LayoutPinRef;
+
+	UPROPERTY()
+	TArray<FEdGraphPinReference> ImagePinsRef;
+};
+
+
+USTRUCT()
+struct CUSTOMIZABLEOBJECTEDITOR_API FCustomizableObjectNodeStaticMeshLOD
+{
+	GENERATED_USTRUCT_BODY()
+
+	UPROPERTY()
+	TArray<FCustomizableObjectNodeStaticMeshMaterial> Materials;
+};
+
+
+UCLASS()
+class CUSTOMIZABLEOBJECTEDITOR_API UCustomizableObjectNodeStaticMesh : public UCustomizableObjectNodeMesh
+{ 
+public:
+	GENERATED_BODY()
+
+	UPROPERTY(EditAnywhere, Category = CustomizableObject)
+	TObjectPtr<UStaticMesh> StaticMesh;
+
+	/** Images */
+	UPROPERTY()
+	TArray<FCustomizableObjectNodeStaticMeshLOD> LODs;
+
+	/** Default pin when there is no mesh. */
+	UPROPERTY()
+	FEdGraphPinReference DefaultPin;
+	
+	// UObject interface.
+	void PostEditChangeProperty(FPropertyChangedEvent& PropertyChangedEvent) override;
+	virtual void Serialize(FArchive& Ar) override;
+
+	// UEdGraphNode interface
+	FText GetNodeTitle(ENodeTitleType::Type TitleType) const override;
+	FLinearColor GetNodeTitleColor() const override;
+	FText GetTooltipText() const override;
+
+	// UCustomizableObjectNode interface
+	void AllocateDefaultPins(UCustomizableObjectNodeRemapPins* RemapPins) override;
+	bool ProvidesCustomPinRelevancyTest() const override { return true; }
+	bool IsPinRelevant(const UEdGraphPin* Pin) const override;
+	UCustomizableObjectNodeRemapPinsByName* CreateRemapPinsByName() const override;
+
+	bool IsNodeOutDatedAndNeedsRefresh() override;
+	FString GetRefreshMessage() const override;
+
+	// UCustomizableObjectNodeMesh interface
+	virtual UTexture2D* FindTextureForPin(const UEdGraphPin* Pin) const override;
+	virtual void GetUVChannelForPin(const UEdGraphPin* Pin, TArray<FVector2f>& OutSegments, int32 UVIndex) const override;
+	virtual TArray<UCustomizableObjectLayout*> GetLayouts(const UEdGraphPin* OutPin) const override;
+	virtual UObject* GetMesh() const override;
+	virtual UEdGraphPin* GetMeshPin(int32 LOD, int MaterialIndex) const override;
+	virtual void GetPinSection(const UEdGraphPin& Pin, int32& OutLODIndex, int32& OutSectionIndex, int32& OutLayoutIndex) const override;
+
+	/** Returns the material assossiated to the given output pin. */
+	UMaterialInterface* GetMaterialFor(const UEdGraphPin* Pin) const;
+
+	// Creates the SGraph Node widget for the thumbnail
+	TSharedPtr<SGraphNode> CreateVisualWidget() override;
+
+	// Determines if the Node is collapsed or not
+	bool bCollapsed = true;
+
+	// Pointer to the SGraphNodeStaticMesh
+	TWeakPtr< SGraphNodeStaticMesh > GraphNodeStaticMesh;
+
+};
+
diff --git a/Engine/Plugins/Experimental/Mutable/Source/CustomizableObjectEditor/Classes/Nodes/CustomizableObjectNodeSwitchBase.h b/Engine/Plugins/Experimental/Mutable/Source/CustomizableObjectEditor/Classes/Nodes/CustomizableObjectNodeSwitchBase.h
new file mode 100644
--- /dev/null
+++ b/Engine/Plugins/Experimental/Mutable/Source/CustomizableObjectEditor/Classes/Nodes/CustomizableObjectNodeSwitchBase.h
@@ -0,0 +1,85 @@
+// Copyright Epic Games, Inc. All Rights Reserved.
+
+#pragma once
+
+#include "Nodes/CustomizableObjectNode.h"
+
+#include "EdGraph/EdGraphPin.h"
+
+#include "CustomizableObjectNodeSwitchBase.generated.h"
+
+
+UCLASS()
+class CUSTOMIZABLEOBJECTEDITOR_API UCustomizableObjectNodeSwitchBase : public UCustomizableObjectNode
+{
+public:
+	GENERATED_BODY()
+
+
+	// UObject interface
+	virtual void Serialize(FArchive& Ar) override;
+	
+	// Begin EdGraphNode interface
+	FText GetNodeTitle(ENodeTitleType::Type TitleType) const;
+	FLinearColor GetNodeTitleColor() const override;
+	FText GetTooltipText() const override;
+	void PinConnectionListChanged(UEdGraphPin* Pin) override;
+	virtual void PostPasteNode() override;
+
+
+	// UCustomizableObjectNode interface
+	virtual void PostBackwardsCompatibleFixup() override;
+	void AllocateDefaultPins(UCustomizableObjectNodeRemapPins* RemapPins) override;
+	void ReconstructNode(UCustomizableObjectNodeRemapPins* RemapPins) override;
+
+	/** Get the output pin catergory. Override. */
+	virtual FName GetCategory() const { return FName(); };
+
+	UEdGraphPin* OutputPin() const;
+
+	UEdGraphPin* SwitchParameter() const;
+
+	UEdGraphPin* GetElementPin(int32 Index) const
+	{
+		return FindPin(GetPinPrefix(Index));
+	}
+
+	int32 GetNumElements() const;
+
+	/** Links the PostEditChangeProperty delegate */
+	void LinkPostEditChangePropertyDelegate(const UEdGraphPin& Pin);
+
+protected:
+	/** Get the pin prefix. Used for retrocompatibility. Override. */
+	virtual FString GetPinPrefix() const { return FString(); };
+
+	/** Get the ouput pin name. Override. */
+	virtual FString GetOutputPinName() const { return FString(); };
+
+	UPROPERTY()
+	FEdGraphPinReference OutputPinReference;
+
+private:
+	/** Get the pin prefix with index. Used for retrocompatibility.*/
+	FString GetPinPrefix(int32 Index) const;
+
+	void ReloadEnumParam();
+
+	/** Last NodeEnumParameter connected. Used to remove the callback once desconnected. */
+	TWeakObjectPtr<UCustomizableObjectNode> LastNodeEnumParameterConnected;
+
+	/** NodeEnumParameter property changed callback function. Reconstructs the node. */
+	UFUNCTION()
+	void EnumParameterPostEditChangeProperty(FPostEditChangePropertyDelegateParameters& Parameters);
+
+	/** The node has to be reconstructed. PinConnectionListChanged(...) can not reconstruct the node, flag used to reconstruct the node on NodeConnectionListChanged(). */
+	bool bMarkReconstruct = false;
+
+	UPROPERTY()
+	TArray<FString> ReloadingElementsNames;
+
+private:
+	UPROPERTY()
+	FEdGraphPinReference SwitchParameterPinReference;
+};
+
diff --git a/Engine/Plugins/Experimental/Mutable/Source/CustomizableObjectEditor/Classes/Nodes/CustomizableObjectNodeTable.h b/Engine/Plugins/Experimental/Mutable/Source/CustomizableObjectEditor/Classes/Nodes/CustomizableObjectNodeTable.h
new file mode 100644
--- /dev/null
+++ b/Engine/Plugins/Experimental/Mutable/Source/CustomizableObjectEditor/Classes/Nodes/CustomizableObjectNodeTable.h
@@ -0,0 +1,222 @@
+// Copyright Epic Games, Inc. All Rights Reserved.
+
+#pragma once
+
+#include "CustomizableObjectNodeMaterial.h"
+#include "Engine/DataTable.h"
+
+#include "Nodes/CustomizableObjectNode.h"
+#include "Nodes/CustomizableObjectNodeLayoutBlocks.h"
+
+#include "CustomizableObjectNodeTable.generated.h"
+
+
+/** Base class for all Table Pins. */
+UCLASS()
+class CUSTOMIZABLEOBJECTEDITOR_API UCustomizableObjectNodeTableObjectPinData : public UCustomizableObjectNodePinData
+{
+	GENERATED_BODY()
+
+public:
+
+	/** Name of the data table column related to the pin */
+	UPROPERTY()
+	FString ColumnName = "";
+
+};
+
+
+/** Additional data for a Mesh pins. */
+UCLASS()
+class CUSTOMIZABLEOBJECTEDITOR_API UCustomizableObjectNodeTableMeshPinData : public UCustomizableObjectNodeTableObjectPinData
+{
+	GENERATED_BODY()
+
+public:
+	
+	/** Anim Instance Column name related to this Mesh pin */
+	UPROPERTY()
+	FString AnimInstanceColumnName = "";
+
+	/** Anim Slot Column name related to this Mesh pin */
+	UPROPERTY()
+	FString AnimSlotColumnName = "";
+
+	/** Anim Tag Column name related to this Mesh pin */
+	UPROPERTY()
+	FString AnimTagColumnName = "";
+
+	UPROPERTY()
+	FString MutableColumnName = "";
+
+	/** LOD of the mesh related to this Mesh pin */
+	UPROPERTY()
+	int32 LOD = 0;
+
+	/** Material Index (Surface Index) of the mesh related to this Mesh pin */
+	UPROPERTY()
+	int32 Material = 0;
+
+	/** Layouts related to this Mesh pin */
+	UPROPERTY()
+	TArray< TObjectPtr<UCustomizableObjectLayout> > Layouts;
+};
+
+
+UCLASS()
+class CUSTOMIZABLEOBJECTEDITOR_API UCustomizableObjectNodeTable : public UCustomizableObjectNode
+{
+public:
+	
+	GENERATED_BODY()
+
+	// Name of the property parameter
+	UPROPERTY(EditAnywhere, Category = CustomizableObject)
+	FString ParameterName = "Default Name";
+
+	UPROPERTY(EditAnywhere, Category = CustomizableObject)
+	bool bAddNoneOption;
+
+	// Pointer to the Data Table Asset represented in this node
+	UPROPERTY(EditAnywhere, Category = CustomizableObject, meta = (DontUpdateWhileEditing))
+	TObjectPtr<UDataTable> Table = nullptr;
+	
+	UPROPERTY(EditAnywhere, Category = UI, meta = (DisplayName = "Parameter UI Metadata"))
+	FMutableParamUIMetadata ParamUIMetadata;
+
+	// UObject interface
+	virtual void PreEditChange(FProperty* PropertyAboutToChange) override;
+	virtual void PostEditChangeChainProperty(FPropertyChangedChainEvent& PropertyChangedEvent) override;
+	virtual void PostEditChangeProperty(FPropertyChangedEvent& PropertyChangedEvent) override;
+	
+	// EdGraphNode interface
+	FText GetNodeTitle(ENodeTitleType::Type TitleType) const override;
+	FLinearColor GetNodeTitleColor() const override;
+	FText GetTooltipText() const override;
+	
+	// UCustomizableObjectNode interface
+	virtual void PostBackwardsCompatibleFixup() override;
+	void AllocateDefaultPins(UCustomizableObjectNodeRemapPins* RemapPins) override;
+	bool IsNodeOutDatedAndNeedsRefresh() override;
+	FString GetRefreshMessage() const override;
+	virtual void RemapPins(const TMap<UEdGraphPin*, UEdGraphPin*>& PinsToRemap) override;
+	virtual bool ProvidesCustomPinRelevancyTest() const override;
+	virtual bool IsPinRelevant(const UEdGraphPin* Pin) const override;
+	
+	// Own interface
+	/** Given an Image pin, return the PinMode. */
+	bool ForceImageMutableMode(const UEdGraphPin* Pin, FGuid ParameterId) const;
+	
+	// Returns the reference Texture parameter from a Material
+	UTexture2D* FindReferenceTextureParameter(const UEdGraphPin* Pin, FString ParameterImageName) const;
+
+	// Methods to get the UVs of the skeletal mesh 
+	void GetUVChannel(const UCustomizableObjectLayout* CurrentLayout, TArray<FVector2f>& OutSegments) const;
+	void GetUVChannelForPin(const UEdGraphPin* Pin, TArray<FVector2f>& OutSegments, int32 UVChannel = 0) const;
+
+	// Methods to provide the Layouts to the Layout block editors
+	void SetLayoutInLayoutEditor(UCustomizableObjectLayout* CurrentLayout);
+	TArray<UCustomizableObjectLayout*> GetLayouts(const UEdGraphPin* Pin) const;
+
+	// Returns the name of the table column related to a pin
+	FString GetColumnNameByPin(const UEdGraphPin* Pin) const;
+
+	// Returns the name of a Mesh Mutable Column from its LOD and pin (Only for AutomaticFromMesh LOD Strategy)
+	FString GetMutableColumnName(const UEdGraphPin* Pin, const int32& LOD) const;
+
+	// Returns the LOD of the mesh associated to the input pin
+	void GetPinLODAndMaterial(const UEdGraphPin* Pin, int32& LOD, int32& Material) const;
+
+	// Get the anim blueprint and anim slot columns related to a mesh
+	void GetAnimationColumns(const FString& ColumnName, FString& AnimBPColumnName, FString& AnimSlotColumnName, FString& AnimTagColumnName) const;
+
+	/** Callback called when the Table property or its contents has changed. */
+	void OnTableChanged();
+	
+	template<class T> 
+	T* GetColumnDefaultAssetByType(FString ColumnName) const
+	{
+		T* ObjectToReturn = nullptr;
+
+		if (Table)
+		{
+			// Getting Struct Pointer
+			const UScriptStruct* TableStruct = Table->GetRowStruct();
+
+			if (TableStruct)
+			{
+				// Getting Default Struct Values
+				uint8* DefaultRowData = (uint8*)FMemory::Malloc(TableStruct->GetStructureSize());
+				TableStruct->InitializeStruct(DefaultRowData);
+
+				if (DefaultRowData)
+				{
+					FProperty* Property = Table->FindTableProperty(FName(*ColumnName));
+
+					if (Property)
+					{
+						if (const FSoftObjectProperty* SoftObjectProperty = CastField<FSoftObjectProperty>(Property))
+						{
+							UObject* Object = nullptr;
+							
+							// Getting default UObject
+							uint8* CellData = SoftObjectProperty->ContainerPtrToValuePtr<uint8>(DefaultRowData);
+
+							if (CellData)
+							{
+								Object = SoftObjectProperty->GetPropertyValue(CellData).LoadSynchronous();
+							}
+
+							if (Object)
+							{
+								if (Object->IsA(T::StaticClass()))
+								{
+									ObjectToReturn = Cast <T>(Object);
+								}
+							}
+						}
+					}
+
+					// Cleaning Default Structure Pointer
+					TableStruct->DestroyStruct(DefaultRowData);
+					FMemory::Free(DefaultRowData);
+				}
+			}
+		}
+
+		return ObjectToReturn;
+	}
+
+	template<class T>
+	T* GetColumnDefaultAssetByType(const UEdGraphPin* Pin) const
+	{
+		if (Pin)
+		{
+			FString ColumnName = GetColumnNameByPin(Pin);
+
+			if (ColumnName != FString())
+			{
+				return GetColumnDefaultAssetByType<T>(ColumnName);
+			}
+		}
+
+		return nullptr;
+	}
+
+	// We should do this in a template!
+	USkeletalMesh* GetSkeletalMeshAt(const UEdGraphPin* Pin, int32 RowIndex) const;
+
+private:
+	
+	/** Number of properties to know when the node needs an update */
+	UPROPERTY()
+	int32 NumProperties;
+
+	FDelegateHandle OnTableChangedDelegateHandle;
+	
+	// Generates a mesh pin for each LOD and Material Surface of the reference SkeletalMesh
+	void GenerateMeshPins(UObject* Mesh, FString Name);
+
+	// Checks if a pin already exists and if it has the same type as before the node refresh
+	bool CheckPinUpdated(const FString& PinName, const FName& PinType) const;
+};
diff --git a/Engine/Plugins/Experimental/Mutable/Source/CustomizableObjectEditor/Classes/Nodes/CustomizableObjectNodeTexture.h b/Engine/Plugins/Experimental/Mutable/Source/CustomizableObjectEditor/Classes/Nodes/CustomizableObjectNodeTexture.h
new file mode 100644
--- /dev/null
+++ b/Engine/Plugins/Experimental/Mutable/Source/CustomizableObjectEditor/Classes/Nodes/CustomizableObjectNodeTexture.h
@@ -0,0 +1,74 @@
+// Copyright Epic Games, Inc. All Rights Reserved.
+
+#pragma once
+
+#include "Nodes/CustomizableObjectNode.h"
+#include "SGraphNode.h"
+
+#include "CustomizableObjectNodeTexture.generated.h"
+
+// Class to render the Texture thumbnail of a CustomizableObjectNodeTexture
+class SGraphNodeTexture : public SGraphNode
+{
+public:
+
+	SLATE_BEGIN_ARGS(SGraphNodeTexture) {}
+	SLATE_END_ARGS();
+
+	SGraphNodeTexture() : SGraphNode() {};
+
+	// Builds the SGraphNodeTexture when needed
+	void Construct(const FArguments& InArgs, UEdGraphNode* InGraphNode);
+
+	// Calls the needed functions to build the SGraphNode widgets
+	void UpdateGraphNode();
+
+	// Overriden functions to build the SGraphNode widgets
+	virtual void SetDefaultTitleAreaWidget(TSharedRef<SOverlay> DefaultTitleAreaWidget) override;
+	virtual void CreateBelowPinControls(TSharedPtr<SVerticalBox> MainBox) override;
+	virtual void Tick(const FGeometry& AllottedGeometry, const double InCurrentTime, const float InDeltaTime) override;
+	virtual bool ShouldAllowCulling() const override { return false; }
+
+	// Callbacks for the widget
+	void OnExpressionPreviewChanged(const ECheckBoxState NewCheckedState);
+	ECheckBoxState IsExpressionPreviewChecked() const;
+	const FSlateBrush* GetExpressionPreviewArrow() const;
+	EVisibility ExpressionPreviewVisibility() const;
+
+public:
+
+	// Pointer to the NodeTexture that owns this SGraphNode
+	UCustomizableObjectNodeTexture* NodeTexture;
+
+	// Single property that only draws the combo box widget of the Texture
+	TSharedPtr<class ISinglePropertyView> TextureSelector;
+
+	// Brush to draw the texture to the widget
+	FSlateBrush TextureBrush;
+
+};
+
+
+UCLASS()
+class CUSTOMIZABLEOBJECTEDITOR_API UCustomizableObjectNodeTexture : public UCustomizableObjectNode
+{
+public:
+	GENERATED_BODY()
+
+	UPROPERTY(EditAnywhere, Category=CustomizableObject)
+	TObjectPtr<UTexture2D> Texture = nullptr;
+
+	// Begin EdGraphNode interface
+	FText GetNodeTitle(ENodeTitleType::Type TitleType) const override;
+	FLinearColor GetNodeTitleColor() const override;
+	FText GetTooltipText() const override;
+
+	// UCustomizableObjectNode interface
+	void AllocateDefaultPins(UCustomizableObjectNodeRemapPins* RemapPins) override;
+
+	// Creates the SGraph Node widget for the thumbnail
+	TSharedPtr<SGraphNode> CreateVisualWidget() override;
+	
+	// Determines if the Node is collapsed or not
+	bool bCollapsed = true;
+};
\ No newline at end of file
diff --git a/Engine/Plugins/Experimental/Mutable/Source/CustomizableObjectEditor/Classes/Nodes/CustomizableObjectNodeTextureBinarise.h b/Engine/Plugins/Experimental/Mutable/Source/CustomizableObjectEditor/Classes/Nodes/CustomizableObjectNodeTextureBinarise.h
new file mode 100644
--- /dev/null
+++ b/Engine/Plugins/Experimental/Mutable/Source/CustomizableObjectEditor/Classes/Nodes/CustomizableObjectNodeTextureBinarise.h
@@ -0,0 +1,33 @@
+// Copyright Epic Games, Inc. All Rights Reserved.
+#pragma once
+
+#include "Nodes/CustomizableObjectNode.h"
+
+#include "CustomizableObjectNodeTextureBinarise.generated.h"
+
+UCLASS()
+class CUSTOMIZABLEOBJECTEDITOR_API UCustomizableObjectNodeTextureBinarise : public UCustomizableObjectNode
+{
+public:
+	GENERATED_BODY()
+
+	// Begin EdGraphNode interface
+	FText GetNodeTitle(ENodeTitleType::Type TittleType) const override;
+	FLinearColor GetNodeTitleColor() const override;
+	FText GetTooltipText() const override;
+	virtual void Serialize(FArchive& Ar) override;
+	
+	// UCustomizableObjectNode interface
+	void AllocateDefaultPins(UCustomizableObjectNodeRemapPins* RemapPins) override;
+
+	UEdGraphPin* GetBaseImagePin() const;
+
+	UEdGraphPin* GetThresholdPin() const
+	{
+		return FindPin(TEXT("Threshold"));
+	}
+
+private:
+	UPROPERTY()
+	FEdGraphPinReference BaseImagePinReference;
+};
\ No newline at end of file
diff --git a/Engine/Plugins/Experimental/Mutable/Source/CustomizableObjectEditor/Classes/Nodes/CustomizableObjectNodeTextureColourMap.h b/Engine/Plugins/Experimental/Mutable/Source/CustomizableObjectEditor/Classes/Nodes/CustomizableObjectNodeTextureColourMap.h
new file mode 100644
--- /dev/null
+++ b/Engine/Plugins/Experimental/Mutable/Source/CustomizableObjectEditor/Classes/Nodes/CustomizableObjectNodeTextureColourMap.h
@@ -0,0 +1,39 @@
+// Copyright Epic Games, Inc. All Rights Reserved.
+
+#pragma once
+
+#include "Nodes/CustomizableObjectNode.h"
+
+#include "CustomizableObjectNodeTextureColourMap.generated.h"
+
+
+UCLASS()
+class CUSTOMIZABLEOBJECTEDITOR_API UCustomizableObjectNodeTextureColourMap : public UCustomizableObjectNode
+{
+public:
+	GENERATED_BODY()
+
+	// Begin EdGraphNode interface
+	FText GetNodeTitle(ENodeTitleType::Type TitleType) const override;
+	FLinearColor GetNodeTitleColor() const override;
+	FText GetTooltipText() const override;
+
+	// UCustomizableObjectNode interface
+	void AllocateDefaultPins(UCustomizableObjectNodeRemapPins* RemapPins) override;
+
+	UEdGraphPin* GetMapPin() const
+	{
+		return FindPin(TEXT("Map"));
+	}
+
+	UEdGraphPin* GetBasePin() const
+	{
+		return FindPin(TEXT("Base"));
+	}
+
+	UEdGraphPin* GetMaskPin() const
+	{
+		return FindPin(TEXT("Mask"));
+	}
+
+};
diff --git a/Engine/Plugins/Experimental/Mutable/Source/CustomizableObjectEditor/Classes/Nodes/CustomizableObjectNodeTextureFromChannels.h b/Engine/Plugins/Experimental/Mutable/Source/CustomizableObjectEditor/Classes/Nodes/CustomizableObjectNodeTextureFromChannels.h
new file mode 100644
--- /dev/null
+++ b/Engine/Plugins/Experimental/Mutable/Source/CustomizableObjectEditor/Classes/Nodes/CustomizableObjectNodeTextureFromChannels.h
@@ -0,0 +1,44 @@
+// Copyright Epic Games, Inc. All Rights Reserved.
+
+#pragma once
+
+#include "Nodes/CustomizableObjectNode.h"
+
+#include "CustomizableObjectNodeTextureFromChannels.generated.h"
+
+
+UCLASS()
+class CUSTOMIZABLEOBJECTEDITOR_API UCustomizableObjectNodeTextureFromChannels : public UCustomizableObjectNode
+{
+public:
+	GENERATED_BODY()
+
+	// Begin EdGraphNode interface
+	FText GetNodeTitle(ENodeTitleType::Type TitleType) const override;
+	FLinearColor GetNodeTitleColor() const override;
+	FText GetTooltipText() const override;
+
+	// UCustomizableObjectNode interface
+	void AllocateDefaultPins(UCustomizableObjectNodeRemapPins* RemapPins) override;
+
+	UEdGraphPin* RPin() const
+	{
+		return FindPin(TEXT("R"));
+	}
+
+	UEdGraphPin* GPin() const
+	{
+		return FindPin(TEXT("G"));
+	}
+
+	UEdGraphPin* BPin() const
+	{
+		return FindPin(TEXT("B"));
+	}
+
+	UEdGraphPin* APin() const
+	{
+		return FindPin(TEXT("A"));
+	}
+};
+
diff --git a/Engine/Plugins/Experimental/Mutable/Source/CustomizableObjectEditor/Classes/Nodes/CustomizableObjectNodeTextureFromColor.h b/Engine/Plugins/Experimental/Mutable/Source/CustomizableObjectEditor/Classes/Nodes/CustomizableObjectNodeTextureFromColor.h
new file mode 100644
--- /dev/null
+++ b/Engine/Plugins/Experimental/Mutable/Source/CustomizableObjectEditor/Classes/Nodes/CustomizableObjectNodeTextureFromColor.h
@@ -0,0 +1,29 @@
+// Copyright Epic Games, Inc. All Rights Reserved.
+
+#pragma once
+
+#include "Nodes/CustomizableObjectNode.h"
+
+#include "CustomizableObjectNodeTextureFromColor.generated.h"
+
+
+UCLASS()
+class CUSTOMIZABLEOBJECTEDITOR_API UCustomizableObjectNodeTextureFromColor : public UCustomizableObjectNode
+{
+public:
+	GENERATED_BODY()
+
+	// Begin EdGraphNode interface
+	FText GetNodeTitle(ENodeTitleType::Type TitleType) const override;
+	FLinearColor GetNodeTitleColor() const override;
+	FText GetTooltipText() const override;
+
+	// UCustomizableObjectNode interface
+	void AllocateDefaultPins(UCustomizableObjectNodeRemapPins* RemapPins) override;
+
+	UEdGraphPin* ColorPin() const
+	{
+		return FindPin(TEXT("Color"));
+	}
+};
+
diff --git a/Engine/Plugins/Experimental/Mutable/Source/CustomizableObjectEditor/Classes/Nodes/CustomizableObjectNodeTextureInterpolate.h b/Engine/Plugins/Experimental/Mutable/Source/CustomizableObjectEditor/Classes/Nodes/CustomizableObjectNodeTextureInterpolate.h
new file mode 100644
--- /dev/null
+++ b/Engine/Plugins/Experimental/Mutable/Source/CustomizableObjectEditor/Classes/Nodes/CustomizableObjectNodeTextureInterpolate.h
@@ -0,0 +1,57 @@
+// Copyright Epic Games, Inc. All Rights Reserved.
+
+#pragma once
+
+#include "Nodes/CustomizableObjectNode.h"
+
+#include "CustomizableObjectNodeTextureInterpolate.generated.h"
+
+
+UCLASS()
+class CUSTOMIZABLEOBJECTEDITOR_API UCustomizableObjectNodeTextureInterpolate : public UCustomizableObjectNode
+{
+public:
+	GENERATED_BODY()
+
+	UCustomizableObjectNodeTextureInterpolate();
+
+	// UObject interface.
+	void PostEditChangeProperty(FPropertyChangedEvent& PropertyChangedEvent) override;
+
+	// Begin EdGraphNode interface
+	FText GetNodeTitle(ENodeTitleType::Type TitleType) const override;
+	FLinearColor GetNodeTitleColor() const override;
+	FText GetTooltipText() const override;
+
+	// UCustomizableObjectNode interface
+	void AllocateDefaultPins(UCustomizableObjectNodeRemapPins* RemapPins) override;
+
+	UEdGraphPin* FactorPin() const
+	{
+		return FindPin(TEXT("Factor"));
+	}
+
+	UEdGraphPin* Targets(int32 Index) const
+	{
+		return FindPin(FString::Printf(TEXT("Target %d"), Index));
+	}
+
+	int32 GetNumTargets() const
+	{
+		int32 Count = 0;
+
+		for (UEdGraphPin* Pin : GetAllNonOrphanPins())
+		{
+			if (Pin->GetName().StartsWith(TEXT("Target ")))
+			{
+				Count++;
+			}
+		}
+
+		return Count;
+	}
+
+	UPROPERTY(EditAnywhere, Category=CustomizableObject)
+	int32 NumTargets;
+};
+
diff --git a/Engine/Plugins/Experimental/Mutable/Source/CustomizableObjectEditor/Classes/Nodes/CustomizableObjectNodeTextureInvert.h b/Engine/Plugins/Experimental/Mutable/Source/CustomizableObjectEditor/Classes/Nodes/CustomizableObjectNodeTextureInvert.h
new file mode 100644
--- /dev/null
+++ b/Engine/Plugins/Experimental/Mutable/Source/CustomizableObjectEditor/Classes/Nodes/CustomizableObjectNodeTextureInvert.h
@@ -0,0 +1,28 @@
+// Copyright Epic Games, Inc. All Rights Reserved.
+#pragma once
+
+#include "Nodes/CustomizableObjectNode.h"
+
+#include "CustomizableObjectNodeTextureInvert.generated.h"
+
+UCLASS()
+class CUSTOMIZABLEOBJECTEDITOR_API UCustomizableObjectNodeTextureInvert : public UCustomizableObjectNode
+{
+public:
+	GENERATED_BODY()
+	
+	// Begin EdGraphNode interface
+	FText GetNodeTitle(ENodeTitleType::Type TittleType)const override;
+	FLinearColor GetNodeTitleColor() const override;
+	FText GetTooltipText() const override;
+	virtual void Serialize(FArchive& Ar) override;
+	
+	// UCustomizableObjectNode interface
+	void AllocateDefaultPins(UCustomizableObjectNodeRemapPins* RemapPins) override;
+
+	UEdGraphPin* GetBaseImagePin() const;
+
+private:
+	UPROPERTY()
+	FEdGraphPinReference BaseImagePinReference;
+};
diff --git a/Engine/Plugins/Experimental/Mutable/Source/CustomizableObjectEditor/Classes/Nodes/CustomizableObjectNodeTextureLayer.h b/Engine/Plugins/Experimental/Mutable/Source/CustomizableObjectEditor/Classes/Nodes/CustomizableObjectNodeTextureLayer.h
new file mode 100644
--- /dev/null
+++ b/Engine/Plugins/Experimental/Mutable/Source/CustomizableObjectEditor/Classes/Nodes/CustomizableObjectNodeTextureLayer.h
@@ -0,0 +1,116 @@
+// Copyright Epic Games, Inc. All Rights Reserved.
+
+#pragma once
+
+#include "Nodes/CustomizableObjectNode.h"
+#include "CustomizableObjectNodeTextureLayer.generated.h"
+
+
+UENUM()
+enum ECustomizableObjectTextureLayerEffect
+{
+	COTLE_MODULATE			= 0 UMETA(DisplayName = "MODULATE"), 
+	COTLE_MULTIPLY				UMETA(DisplayName = "MULTIPLY"), 
+	COTLE_SOFTLIGHT				UMETA(DisplayName = "SOFTLIGHT"),
+	COTLE_HARDLIGHT				UMETA(DisplayName = "HARDLIGHT"),
+	COTLE_DODGE					UMETA(DisplayName = "DODGE"),
+	COTLE_BURN					UMETA(DisplayName = "BURN"),
+	COTLE_SCREEN				UMETA(DisplayName = "SCREEN"),
+	COTLE_OVERLAY				UMETA(DisplayName = "OVERLAY"),
+	COTLE_ALPHA_OVERLAY			UMETA(DisplayName = "ALPHA_OVERLAY"),
+	COTLE_NORMAL_COMBINE		UMETA(DisplayName = "COMBINE")
+};
+
+
+USTRUCT()
+struct CUSTOMIZABLEOBJECTEDITOR_API FCustomizableObjectTextureLayer
+{
+	GENERATED_USTRUCT_BODY()
+
+	FCustomizableObjectTextureLayer()
+	{
+		Effect = COTLE_SOFTLIGHT;
+	}
+
+	UPROPERTY(EditAnywhere, Category=CustomizableObject)
+	TEnumAsByte<enum ECustomizableObjectTextureLayerEffect> Effect;
+};
+
+
+UCLASS()
+class CUSTOMIZABLEOBJECTEDITOR_API UCustomizableObjectNodeTextureLayer : public UCustomizableObjectNode
+{
+public:
+	GENERATED_BODY()
+
+	UCustomizableObjectNodeTextureLayer();
+
+	/**  */
+	UPROPERTY(EditAnywhere, Category=CustomizableObject)
+	TArray<FCustomizableObjectTextureLayer> Layers;
+
+	// UObject interface.
+	void PostEditChangeProperty(FPropertyChangedEvent& PropertyChangedEvent) override;
+	virtual void Serialize(FArchive& Ar) override;
+	
+	// Begin EdGraphNode interface
+	FText GetNodeTitle(ENodeTitleType::Type TitleType) const override;
+	FLinearColor GetNodeTitleColor() const override;
+	FText GetTooltipText() const override;
+	void PrepareForCopying() override;
+	void PostPasteNode() override;
+
+	// UCustomizableObjectNode interface
+	void AllocateDefaultPins(UCustomizableObjectNodeRemapPins* RemapPins) override;
+	virtual bool CanConnect(const UEdGraphPin* InOwnedInputPin, const UEdGraphPin* InOutputPin, bool& bOutIsOtherNodeBlocklisted, bool& bOutArePinsCompatible) const override;
+		
+	UEdGraphPin* BasePin() const
+	{
+		return FindPin(TEXT("Base"));
+	}
+
+	UEdGraphPin* OutputPin() const;
+
+	UEdGraphPin* LayerPin(int32 Index) const
+	{
+		FString PinName = FString::Printf(TEXT("Layer %d "), Index);
+		return FindPin(PinName);
+	}
+
+	UEdGraphPin* MaskPin(int32 Index) const
+	{
+		FString PinName = FString::Printf(TEXT("Mask %d "), Index);
+		return FindPin(PinName);
+	}
+
+	int32 GetNumLayers() const
+	{
+		int32 Count = 0;
+
+		for (UEdGraphPin* Pin : GetAllNonOrphanPins())
+		{
+			if (Pin->GetName().StartsWith(TEXT("Layer ")))
+			{
+				Count++;
+			}
+		}
+
+		return Count;
+	}
+protected:
+	// This struct will store the information of each textue layer node when its destroyed, to be able to recover the UCustomizableObjectNodeTextureLayer::Layers
+	// data in a post paste operation. Otherwise, this information is lost.
+	struct TextureLayerLayersData
+	{
+		FString Name; // Texture layer node name
+		TArray<FCustomizableObjectTextureLayer> Layers;
+		UEdGraph* Graph; // Graph that contained this node
+	};
+	static int32 GetIndexInArrayDestroyedNodes(UCustomizableObjectNodeTextureLayer* Node);
+	static TArray<TextureLayerLayersData> ArrayDestroyedNodes;
+
+private:
+	UPROPERTY()
+	FEdGraphPinReference OutputPinReference;
+};
+
diff --git a/Engine/Plugins/Experimental/Mutable/Source/CustomizableObjectEditor/Classes/Nodes/CustomizableObjectNodeTextureParameter.h b/Engine/Plugins/Experimental/Mutable/Source/CustomizableObjectEditor/Classes/Nodes/CustomizableObjectNodeTextureParameter.h
new file mode 100644
--- /dev/null
+++ b/Engine/Plugins/Experimental/Mutable/Source/CustomizableObjectEditor/Classes/Nodes/CustomizableObjectNodeTextureParameter.h
@@ -0,0 +1,36 @@
+// Copyright Epic Games, Inc. All Rights Reserved.
+
+#pragma once
+
+#include "Nodes/CustomizableObjectNode.h"
+
+#include "Engine/Texture2D.h"
+
+#include "CustomizableObjectNodeTextureParameter.generated.h"
+
+
+UCLASS()
+class CUSTOMIZABLEOBJECTEDITOR_API UCustomizableObjectNodeTextureParameter : public UCustomizableObjectNode
+{
+public:
+	GENERATED_BODY()
+
+	UPROPERTY(EditAnywhere, Category=CustomizableObject)
+	TObjectPtr<UTexture2D> DefaultValue;
+
+	UPROPERTY(EditAnywhere, Category=CustomizableObject)
+	FString ParameterName = "Default Name";
+
+	UPROPERTY(EditAnywhere, Category = UI, meta = (DisplayName = "Parameter UI Metadata"))
+	FMutableParamUIMetadata ParamUIMetadata;
+
+	// Begin EdGraphNode interface
+	FText GetNodeTitle(ENodeTitleType::Type TitleType) const override;
+	FLinearColor GetNodeTitleColor() const override;
+	FText GetTooltipText() const override;
+
+	// UCustomizableObjectNode interface
+	void AllocateDefaultPins(UCustomizableObjectNodeRemapPins* RemapPins) override;
+	virtual bool IsAffectedByLOD() const override { return false; }
+};
+
diff --git a/Engine/Plugins/Experimental/Mutable/Source/CustomizableObjectEditor/Classes/Nodes/CustomizableObjectNodeTextureProject.h b/Engine/Plugins/Experimental/Mutable/Source/CustomizableObjectEditor/Classes/Nodes/CustomizableObjectNodeTextureProject.h
new file mode 100644
--- /dev/null
+++ b/Engine/Plugins/Experimental/Mutable/Source/CustomizableObjectEditor/Classes/Nodes/CustomizableObjectNodeTextureProject.h
@@ -0,0 +1,83 @@
+// Copyright Epic Games, Inc. All Rights Reserved.
+
+#pragma once
+
+#include "Nodes/CustomizableObjectNode.h"
+
+#include "CustomizableObjectNodeTextureProject.generated.h"
+
+
+UCLASS()
+class CUSTOMIZABLEOBJECTEDITOR_API UCustomizableObjectNodeTextureProject : public UCustomizableObjectNode
+{
+public:
+	GENERATED_BODY()
+
+	UCustomizableObjectNodeTextureProject();
+
+	// Layout to use for the generated images.
+	UPROPERTY(EditAnywhere, Category = CustomizableObject)
+	int32 Layout = 0;
+
+	UPROPERTY(EditAnywhere, Category = CustomizableObject)
+	int32 TextureSizeX = 0; 
+
+	UPROPERTY(EditAnywhere, Category = CustomizableObject)
+	int32 TextureSizeY = 0;
+
+	UPROPERTY(EditAnywhere, Category = CustomizableObject)
+	int32 Textures = 0;
+
+	// UObject interface.
+	void PostEditChangeProperty(FPropertyChangedEvent& PropertyChangedEvent) override;
+	virtual void Serialize(FArchive& Ar) override;
+	
+	// Begin EdGraphNode interface
+	FText GetNodeTitle(ENodeTitleType::Type TitleType) const override;
+	FLinearColor GetNodeTitleColor() const override;
+	FText GetTooltipText() const override;
+
+	// UCustomizableObjectNode interface
+	void AllocateDefaultPins(UCustomizableObjectNodeRemapPins* RemapPins) override;
+
+	UEdGraphPin* MeshPin() const
+	{
+		return FindPin(TEXT("Mesh"));
+	}
+
+	UEdGraphPin* MeshMaskPin() const
+	{
+		return FindPin(TEXT("Mesh Mask"));
+	}
+
+	UEdGraphPin* AngleFadeStartPin() const
+	{
+		return FindPin(TEXT("Fade Start Angle"));
+	}
+
+	UEdGraphPin* AngleFadeEndPin() const
+	{
+		return FindPin(TEXT("Fade End Angle"));
+	}
+
+	UEdGraphPin* ProjectorPin() const
+	{
+		return FindPin(TEXT("Projector"));
+	}
+
+	UEdGraphPin* TexturePins(int32 Index) const;
+
+	UEdGraphPin* OutputPins(int32 Index) const;
+
+	int32 GetNumTextures() const;
+
+	int32 GetNumOutputs() const;
+
+private:
+	UPROPERTY()
+	TArray<FEdGraphPinReference> TexturePinsReferences;
+
+	UPROPERTY()
+	TArray<FEdGraphPinReference> OutputPinsReferences;
+};
+
diff --git a/Engine/Plugins/Experimental/Mutable/Source/CustomizableObjectEditor/Classes/Nodes/CustomizableObjectNodeTextureSample.h b/Engine/Plugins/Experimental/Mutable/Source/CustomizableObjectEditor/Classes/Nodes/CustomizableObjectNodeTextureSample.h
new file mode 100644
--- /dev/null
+++ b/Engine/Plugins/Experimental/Mutable/Source/CustomizableObjectEditor/Classes/Nodes/CustomizableObjectNodeTextureSample.h
@@ -0,0 +1,44 @@
+// Copyright Epic Games, Inc. All Rights Reserved.
+
+#pragma once
+
+#include "Nodes/CustomizableObjectNode.h"
+
+#include "CustomizableObjectNodeTextureSample.generated.h"
+
+
+UCLASS()
+class CUSTOMIZABLEOBJECTEDITOR_API UCustomizableObjectNodeTextureSample : public UCustomizableObjectNode
+{
+public:
+	GENERATED_BODY()
+
+	UCustomizableObjectNodeTextureSample();
+
+	// UObject interface.
+	void PostEditChangeProperty(FPropertyChangedEvent& PropertyChangedEvent) override;
+
+	// Begin EdGraphNode interface
+	FText GetNodeTitle(ENodeTitleType::Type TitleType) const override;
+	FLinearColor GetNodeTitleColor() const override;
+	FText GetTooltipText() const override;
+
+	// UCustomizableObjectNode interface
+	void AllocateDefaultPins(UCustomizableObjectNodeRemapPins* RemapPins) override;
+
+	const UEdGraphPin* TexturePin() const
+	{
+		return FindPin(TEXT("Texture"));
+	}
+
+	const UEdGraphPin* XPin() const
+	{
+		return FindPin(TEXT("X"));
+	}
+
+	const UEdGraphPin* YPin() const
+	{
+		return FindPin(TEXT("Y"));
+	}
+};
+
diff --git a/Engine/Plugins/Experimental/Mutable/Source/CustomizableObjectEditor/Classes/Nodes/CustomizableObjectNodeTextureSwitch.h b/Engine/Plugins/Experimental/Mutable/Source/CustomizableObjectEditor/Classes/Nodes/CustomizableObjectNodeTextureSwitch.h
new file mode 100644
--- /dev/null
+++ b/Engine/Plugins/Experimental/Mutable/Source/CustomizableObjectEditor/Classes/Nodes/CustomizableObjectNodeTextureSwitch.h
@@ -0,0 +1,27 @@
+// Copyright Epic Games, Inc. All Rights Reserved.
+
+#pragma once
+
+#include "Nodes/CustomizableObjectNodeSwitchBase.h"
+
+#include "EdGraph/EdGraphPin.h"
+
+#include "CustomizableObjectNodeTextureSwitch.generated.h"
+
+
+UCLASS()
+class CUSTOMIZABLEOBJECTEDITOR_API UCustomizableObjectNodeTextureSwitch : public UCustomizableObjectNodeSwitchBase
+{
+public:
+	GENERATED_BODY()
+
+	// UObject interface
+	virtual void Serialize(FArchive& Ar) override;
+
+	// UCustomizableObjectNodeSwitchBase interface
+	FString GetOutputPinName() const override;
+
+	FName GetCategory() const override;
+
+	FString GetPinPrefix() const;
+};
diff --git a/Engine/Plugins/Experimental/Mutable/Source/CustomizableObjectEditor/Classes/Nodes/CustomizableObjectNodeTextureToChannels.h b/Engine/Plugins/Experimental/Mutable/Source/CustomizableObjectEditor/Classes/Nodes/CustomizableObjectNodeTextureToChannels.h
new file mode 100644
--- /dev/null
+++ b/Engine/Plugins/Experimental/Mutable/Source/CustomizableObjectEditor/Classes/Nodes/CustomizableObjectNodeTextureToChannels.h
@@ -0,0 +1,54 @@
+// Copyright Epic Games, Inc. All Rights Reserved.
+
+#pragma once
+
+#include "Nodes/CustomizableObjectNode.h"
+
+#include "CustomizableObjectNodeTextureToChannels.generated.h"
+
+
+UCLASS()
+class CUSTOMIZABLEOBJECTEDITOR_API UCustomizableObjectNodeTextureToChannels : public UCustomizableObjectNode
+{
+public:
+	GENERATED_BODY()
+	
+	// UObject interface.
+	void PostEditChangeProperty(FPropertyChangedEvent& PropertyChangedEvent) override;
+	virtual void Serialize(FArchive& Ar) override;
+	
+	// Begin EdGraphNode interface
+	FText GetNodeTitle(ENodeTitleType::Type TitleType) const override;
+	FLinearColor GetNodeTitleColor() const override;
+	FText GetTooltipText() const override;
+
+	// UCustomizableObjectNode interface
+	void AllocateDefaultPins(UCustomizableObjectNodeRemapPins* RemapPins) override;
+
+	UEdGraphPin* InputPin() const;
+
+	UEdGraphPin* RPin() const
+	{
+		return FindPin(TEXT("R"));
+	}
+
+	UEdGraphPin* GPin() const
+	{
+		return FindPin(TEXT("G"));
+	}
+
+	UEdGraphPin* BPin() const
+	{
+		return FindPin(TEXT("B"));
+	}
+
+	UEdGraphPin* APin() const
+	{
+		return FindPin(TEXT("A"));
+	}
+
+private:
+	UPROPERTY()
+	FEdGraphPinReference InputPinReference;
+};
+
diff --git a/Engine/Plugins/Experimental/Mutable/Source/CustomizableObjectEditor/Classes/Nodes/CustomizableObjectNodeTextureTransform.h b/Engine/Plugins/Experimental/Mutable/Source/CustomizableObjectEditor/Classes/Nodes/CustomizableObjectNodeTextureTransform.h
new file mode 100644
--- /dev/null
+++ b/Engine/Plugins/Experimental/Mutable/Source/CustomizableObjectEditor/Classes/Nodes/CustomizableObjectNodeTextureTransform.h
@@ -0,0 +1,53 @@
+// Copyright Epic Games, Inc. All Rights Reserved.
+#pragma once
+
+#include "Nodes/CustomizableObjectNode.h"
+
+#include "CustomizableObjectNodeTextureTransform.generated.h"
+
+UCLASS()
+class CUSTOMIZABLEOBJECTEDITOR_API UCustomizableObjectNodeTextureTransform : public UCustomizableObjectNode
+{
+public:
+	GENERATED_BODY()
+
+	// Begin EdGraphNode interface
+	FText GetNodeTitle(ENodeTitleType::Type TittleType) const override;
+	FLinearColor GetNodeTitleColor() const override;
+	FText GetTooltipText() const override;
+	virtual void Serialize(FArchive& Ar) override;
+	
+	// UCustomizableObjectNode interface
+	void AllocateDefaultPins(UCustomizableObjectNodeRemapPins* RemapPins) override;
+
+	UEdGraphPin* GetBaseImagePin() const;
+
+	UEdGraphPin* GetOffsetXPin() const
+	{
+		return FindPin(TEXT("Offset X"));
+	}
+
+	UEdGraphPin* GetOffsetYPin() const
+	{
+		return FindPin(TEXT("Offset Y"));
+	}
+
+	UEdGraphPin* GetScaleXPin() const
+	{
+		return FindPin(TEXT("Scale X"));
+	}
+
+	UEdGraphPin* GetScaleYPin() const
+	{
+		return FindPin(TEXT("Scale Y"));
+	}
+
+	UEdGraphPin* GetRotationPin() const
+	{
+		return FindPin(TEXT("Rotation"));
+	}
+
+private:
+	UPROPERTY()
+	FEdGraphPinReference BaseImagePinReference;
+};
diff --git a/Engine/Plugins/Experimental/Mutable/Source/CustomizableObjectEditor/Classes/Nodes/CustomizableObjectNodeTextureVariation.h b/Engine/Plugins/Experimental/Mutable/Source/CustomizableObjectEditor/Classes/Nodes/CustomizableObjectNodeTextureVariation.h
new file mode 100644
--- /dev/null
+++ b/Engine/Plugins/Experimental/Mutable/Source/CustomizableObjectEditor/Classes/Nodes/CustomizableObjectNodeTextureVariation.h
@@ -0,0 +1,62 @@
+// Copyright Epic Games, Inc. All Rights Reserved.
+
+#pragma once
+
+#include "Nodes/CustomizableObjectNode.h"
+
+#include "CustomizableObjectNodeTextureVariation.generated.h"
+
+
+USTRUCT()
+struct CUSTOMIZABLEOBJECTEDITOR_API FCustomizableObjectTextureVariation
+{
+	GENERATED_USTRUCT_BODY()
+
+	UPROPERTY(EditAnywhere, Category = CustomizableObject)
+	FString Tag;
+};
+
+
+UCLASS()
+class CUSTOMIZABLEOBJECTEDITOR_API UCustomizableObjectNodeTextureVariation : public UCustomizableObjectNode
+{
+public:
+	GENERATED_BODY()
+
+	UCustomizableObjectNodeTextureVariation();
+
+	UPROPERTY(EditAnywhere, Category = CustomizableObject)
+	TArray<FCustomizableObjectTextureVariation> Variations;
+
+	// UObject interface.
+	void PostEditChangeProperty(FPropertyChangedEvent& PropertyChangedEvent) override;
+
+	// EdGraphNode interface
+	FText GetNodeTitle(ENodeTitleType::Type TitleType) const override;
+	FLinearColor GetNodeTitleColor() const override;
+	FText GetTooltipText() const override;
+	
+	// UCustomizableObjectNode interface
+	void AllocateDefaultPins(UCustomizableObjectNodeRemapPins* RemapPins) override;
+
+	UEdGraphPin* OutputPin() const
+	{
+		return FindPin(TEXT("Texture"));
+	}
+
+	UEdGraphPin* DefaultPin() const
+	{
+		return FindPin(TEXT("Default"));
+	}
+
+	int32 GetNumVariations() const
+	{
+		return Variations.Num();
+	}
+
+	UEdGraphPin* VariationPin(int Index) const
+	{
+		return FindPin(FString::Printf(TEXT("Variation %d"),Index));
+	}
+};
+
diff --git a/Engine/Plugins/Experimental/Mutable/Source/CustomizableObjectEditor/Classes/Nodes/CustomizableObjectNodeUseMaterial.h b/Engine/Plugins/Experimental/Mutable/Source/CustomizableObjectEditor/Classes/Nodes/CustomizableObjectNodeUseMaterial.h
new file mode 100644
--- /dev/null
+++ b/Engine/Plugins/Experimental/Mutable/Source/CustomizableObjectEditor/Classes/Nodes/CustomizableObjectNodeUseMaterial.h
@@ -0,0 +1,60 @@
+// Copyright Epic Games, Inc. All Rights Reserved.
+
+#pragma once
+
+#include "CustomizableObjectNodeParentedMaterial.h"
+#include "Nodes/CustomizableObjectNode.h"
+
+
+/** Common functionalities from EditMaterial and ExtendMaterial.
+ *
+ * See must call functions section! These functions should be called in the corresponding inheritor functions. */
+class FCustomizableObjectNodeUseMaterial
+{
+public:
+	virtual ~FCustomizableObjectNodeUseMaterial() = default;
+	
+	/** Returns true if this node uses the given Image. */
+	bool UsesImage(const FGuid& ImageId) const;
+
+	/** Given an Image id, returns its Image pin. Can nullptr if the Image is not used. */
+	const UEdGraphPin* GetUsedImagePin(const FGuid& ImageId) const;
+
+protected:
+	// Begin must call functions
+	void PostBackwardsCompatibleFixupWork();
+	
+	bool IsNodeOutDatedAndNeedsRefreshWork();
+
+	/** Call before SetParentNode. */
+	void PreSetParentNodeWork(UCustomizableObject* Object, const FGuid NodeId);
+
+	/** Call after SetParentNode. */
+	void PostSetParentNodeWork(UCustomizableObject* Object, const FGuid NodeId);
+
+	void PinConnectionListChangedWork(UEdGraphPin* Pin);
+
+	void CustomRemovePinWork(UEdGraphPin& Pin);
+	// End must call functions
+
+	/** Return the node which this interface belongs to. */
+	virtual UCustomizableObjectNode& GetNode() = 0;
+	virtual const UCustomizableObjectNode& GetNode() const;
+
+	/** Returns the Parented Material interface this node belongs to. */
+	virtual FCustomizableObjectNodeParentedMaterial& GetNodeParentedMaterial() = 0;
+	
+	/** Return the PinsParameter map which this interface belongs to. */
+	virtual TMap<FGuid, FEdGraphPinReference>& GetPinsParameter() = 0;
+	virtual const TMap<FGuid, FEdGraphPinReference>& GetPinsParameter() const;
+
+	/** Get the output material pin. */
+	virtual UEdGraphPin* OutputPin() const = 0;
+
+private:
+	/** Handler to the reconstruct node delegate. Used to remove the previously attached delegate. */
+	FDelegateHandle PostReconstructNodeDelegateHandler;
+	
+	/** Handler to the mutable mode changed delegate. Used to remove the previously attached delegate. */
+	FDelegateHandle PostTextureParameterModeChangedDelegateHandle;
+};
diff --git a/Engine/Plugins/Experimental/Mutable/Source/CustomizableObjectEditor/Classes/Thumbnails/CustomizableObjectInstanceThumbnailRenderer.h b/Engine/Plugins/Experimental/Mutable/Source/CustomizableObjectEditor/Classes/Thumbnails/CustomizableObjectInstanceThumbnailRenderer.h
new file mode 100644
--- /dev/null
+++ b/Engine/Plugins/Experimental/Mutable/Source/CustomizableObjectEditor/Classes/Thumbnails/CustomizableObjectInstanceThumbnailRenderer.h
@@ -0,0 +1,36 @@
+// Copyright Epic Games, Inc. All Rights Reserved.
+
+/**
+ *
+ * This thumbnail renderer displays a given Customizable Object Instance
+ */
+
+#pragma once
+#include "ThumbnailHelpers.h"
+#include "ThumbnailRendering/DefaultSizedThumbnailRenderer.h"
+#include "CustomizableObjectInstanceThumbnailRenderer.generated.h"
+
+
+UCLASS(config=Editor)
+class CUSTOMIZABLEOBJECTEDITOR_API UCustomizableObjectInstanceThumbnailRenderer : public UDefaultSizedThumbnailRenderer
+{
+public:
+	GENERATED_BODY()
+
+	UCustomizableObjectInstanceThumbnailRenderer();
+
+	// Begin UThumbnailRenderer Object
+	void Draw(UObject* Object, int32 X, int32 Y, uint32 Width, uint32 Height, FRenderTarget*, FCanvas* Canvas, bool bAdditionalViewFamily) override;
+	// End UThumbnailRenderer Object
+
+	// UObject implementation
+	void BeginDestroy() override;
+
+private:
+	class FSkeletalMeshThumbnailScene* ThumbnailScene;
+
+	UPROPERTY()
+	TObjectPtr<class UTexture2D> NoImage;
+
+};
+
diff --git a/Engine/Plugins/Experimental/Mutable/Source/CustomizableObjectEditor/Classes/Thumbnails/CustomizableObjectThumbnailRenderer.h b/Engine/Plugins/Experimental/Mutable/Source/CustomizableObjectEditor/Classes/Thumbnails/CustomizableObjectThumbnailRenderer.h
new file mode 100644
--- /dev/null
+++ b/Engine/Plugins/Experimental/Mutable/Source/CustomizableObjectEditor/Classes/Thumbnails/CustomizableObjectThumbnailRenderer.h
@@ -0,0 +1,39 @@
+// Copyright Epic Games, Inc. All Rights Reserved.
+
+/**
+ *
+ * This thumbnail renderer displays a given Customizable Object
+ */
+
+#pragma once
+#include "ThumbnailHelpers.h"
+#include "ThumbnailRendering/DefaultSizedThumbnailRenderer.h"
+
+#include "CustomizableObjectThumbnailRenderer.generated.h"
+
+
+UCLASS(config=Editor)
+class CUSTOMIZABLEOBJECTEDITOR_API UCustomizableObjectThumbnailRenderer : public UDefaultSizedThumbnailRenderer
+{
+public:
+	GENERATED_BODY()
+
+	UCustomizableObjectThumbnailRenderer();
+
+	// Begin UThumbnailRenderer Object
+	virtual void Draw(UObject* Object, int32 X, int32 Y, uint32 Width, uint32 Height, FRenderTarget*, FCanvas* Canvas, bool bAdditionalViewFamily) override;
+	// End UThumbnailRenderer Object
+
+	// UObject implementation
+	virtual void BeginDestroy() override;
+
+private:
+	class FSkeletalMeshThumbnailScene* ThumbnailScene;
+
+	UPROPERTY()
+	TObjectPtr<class UCustomizableObjectInstance> CustomizableObjectInstance;
+
+	UPROPERTY()
+	TObjectPtr<class UTexture2D> NoImage;
+};
+
diff --git a/Engine/Plugins/Experimental/Mutable/Source/CustomizableObjectEditor/Private/AssetTypeActions_CustomizableObject.cpp b/Engine/Plugins/Experimental/Mutable/Source/CustomizableObjectEditor/Private/AssetTypeActions_CustomizableObject.cpp
new file mode 100644
--- /dev/null
+++ b/Engine/Plugins/Experimental/Mutable/Source/CustomizableObjectEditor/Private/AssetTypeActions_CustomizableObject.cpp
@@ -0,0 +1,293 @@
+// Copyright Epic Games, Inc. All Rights Reserved.
+
+#include "AssetTypeActions_CustomizableObject.h"
+#include "CustomizableObjectEditorModule.h"
+#include "ContentBrowserModule.h"
+#include "AssetRegistry/AssetRegistryModule.h"
+#include "CustomizableObjectCompiler.h"
+#include "IContentBrowserSingleton.h"
+#include "Framework/MultiBox/MultiBoxBuilder.h"
+#include "Misc/MessageDialog.h"
+#include "UnrealEditorPortabilityHelpers.h"
+
+
+#define LOCTEXT_NAMESPACE "CustomizableObjectEditor"
+
+void FAssetTypeActions_CustomizableObject::GetActions( const TArray<UObject*>& InObjects, FMenuBuilder& MenuBuilder )
+{
+	TArray<TWeakObjectPtr<UCustomizableObject>> Objects = GetTypedWeakObjectPtrs<UCustomizableObject>(InObjects);
+
+	MenuBuilder.AddMenuEntry(
+		NSLOCTEXT("AssetTypeActions_CustomizableObject","CustomizableObject_Edit","Edit"),
+		NSLOCTEXT("AssetTypeActions_CustomizableObject", "CustomizableObject_EditTooltip", "Opens the selected object in the customizable object editor."),
+		FSlateIcon(),
+		FUIAction(
+			FExecuteAction::CreateSP( this, &FAssetTypeActions_CustomizableObject::ExecuteEdit, Objects ),
+			FCanExecuteAction()
+			)
+		);
+
+	MenuBuilder.AddMenuEntry(
+		NSLOCTEXT("AssetTypeActions_CustomizableObject", "CustomizableObject_NewInstance", "Create New Instance"),
+		NSLOCTEXT("AssetTypeActions_CustomizableObject", "CustomizableObject_NewInstanceTooltip", "Creates a new instance of this customizable object."),
+		FSlateIcon(),
+		FUIAction(
+		FExecuteAction::CreateSP( this, &FAssetTypeActions_CustomizableObject::ExecuteNewInstance, Objects ),
+		FCanExecuteAction()
+		)
+		);
+
+	// Recompile sub-menu
+	MenuBuilder.AddSubMenu(
+		NSLOCTEXT("AssetTypeActions_CustomizableObject", "CustomizableObject_RecompileObjects", "Recompile"),
+		NSLOCTEXT("AssetTypeActions_CustomizableObject", "CustomizableObject_RecompileAllObjectsTooltip", "Recompile Customizable Objects actions"),
+		FNewMenuDelegate::CreateSP(this, &FAssetTypeActions_CustomizableObject::MakeRecompileSubMenu, Objects));
+
+
+
+	MenuBuilder.AddMenuEntry(
+		NSLOCTEXT("AssetTypeActions_CustomizableObject", "CustomizableObject_Debug", "Debug"),
+		NSLOCTEXT("AssetTypeActions_CustomizableObject", "CustomizableObject_DebugTooltip", "Open the debugger for the customizable object."),
+		FSlateIcon(),
+		FUIAction(
+			FExecuteAction::CreateSP(this, &FAssetTypeActions_CustomizableObject::ExecuteDebug, Objects),
+			FCanExecuteAction()
+		)
+	);
+
+}
+
+
+void FAssetTypeActions_CustomizableObject::OpenAssetEditor( const TArray<UObject*>& InObjects, TSharedPtr<IToolkitHost> EditWithinLevelEditor )
+{
+	EToolkitMode::Type Mode = EditWithinLevelEditor.IsValid() ? EToolkitMode::WorldCentric : EToolkitMode::Standalone;
+
+	for (auto ObjIt = InObjects.CreateConstIterator(); ObjIt; ++ObjIt)
+	{
+		UCustomizableObject* Object = Cast<UCustomizableObject>(*ObjIt);
+		if (Object != NULL)
+		{
+			ICustomizableObjectEditorModule* CustomizableObjectEditorModule = &FModuleManager::LoadModuleChecked<ICustomizableObjectEditorModule>( "CustomizableObjectEditor" );
+			CustomizableObjectEditorModule->CreateCustomizableObjectEditor(Mode, EditWithinLevelEditor, Object);
+		}
+	}
+}
+
+
+uint32 FAssetTypeActions_CustomizableObject::GetCategories()
+{
+	const ICustomizableObjectEditorModule* CustomizableObjectEditorModule = &FModuleManager::LoadModuleChecked<ICustomizableObjectEditorModule>("CustomizableObjectEditor");
+	return CustomizableObjectEditorModule->GetAssetCategory();
+}
+
+
+
+void FAssetTypeActions_CustomizableObject::ExecuteNewInstance(TArray<TWeakObjectPtr<UCustomizableObject>> Objects)
+{
+	const FString DefaultSuffix = TEXT("_Inst");
+
+	TArray<UObject*> ObjectsToSync;
+
+	for (auto ObjIt = Objects.CreateConstIterator(); ObjIt; ++ObjIt)
+	{
+		UCustomizableObject* Object = (*ObjIt).Get();
+		if ( Object )
+		{
+			// Determine an appropriate name
+			FString Name;
+			FString PackageName;
+			CreateUniqueAssetName( Object->GetOutermost()->GetName(), DefaultSuffix, PackageName, Name );
+			UPackage* Pkg = CreatePackage(*PackageName);
+
+			UCustomizableObject* CustomizableObject = Cast<UCustomizableObject>(Object);
+			UCustomizableObjectInstance* Instance = NewObject<UCustomizableObjectInstance>(Pkg, FName(*Name), RF_Public|RF_Standalone);
+			
+			if (Instance)
+			{
+				// Mark the package dirty...
+				Pkg->MarkPackageDirty();
+
+				Instance->SetObject(CustomizableObject);
+
+				ObjectsToSync.Add(Instance);
+			}
+		}
+	}
+
+	if ( ObjectsToSync.Num() > 0 )
+	{
+		FContentBrowserModule& ContentBrowserModule = FModuleManager::Get().LoadModuleChecked<FContentBrowserModule>("ContentBrowser");
+		ContentBrowserModule.Get().SyncBrowserToAssets( ObjectsToSync, /*bAllowLockedBrowsers=*/true );
+
+		GEditor->GetEditorSubsystem<UAssetEditorSubsystem>()->OpenEditorForAsset(ObjectsToSync[0]);
+	}
+
+}
+
+
+void FAssetTypeActions_CustomizableObject::ExecuteEdit(TArray<TWeakObjectPtr<UCustomizableObject>> Objects)
+{
+	for (auto ObjIt = Objects.CreateConstIterator(); ObjIt; ++ObjIt)
+	{
+		UCustomizableObject* Object = (*ObjIt).Get();
+		if (Object)
+		{
+			GEditor->GetEditorSubsystem<UAssetEditorSubsystem>()->OpenEditorForAsset(Object);
+		}
+	}
+}
+
+
+void FAssetTypeActions_CustomizableObject::ExecuteDebug(TArray<TWeakObjectPtr<UCustomizableObject>> Objects)
+{
+	for (auto ObjIt = Objects.CreateConstIterator(); ObjIt; ++ObjIt)
+	{
+		UCustomizableObject* Object = (*ObjIt).Get();
+		if (!Object) continue;
+
+		ICustomizableObjectEditorModule* CustomizableObjectEditorModule = &FModuleManager::LoadModuleChecked<ICustomizableObjectEditorModule>("CustomizableObjectEditor");
+		if (!CustomizableObjectEditorModule) continue;
+
+		CustomizableObjectEditorModule->CreateCustomizableObjectDebugger(EToolkitMode::Standalone, nullptr , Object);
+	}
+}
+
+
+void FAssetTypeActions_CustomizableObject::MakeRecompileSubMenu(FMenuBuilder& MenuBuilder, TArray<TWeakObjectPtr<UCustomizableObject>> InObjects)
+{
+	// Recompile Selected
+	MenuBuilder.AddMenuEntry(
+		NSLOCTEXT("AssetTypeActions_CustomizableObject", "CustomizableObject_RecompileSelected", "Selected"),
+		NSLOCTEXT("AssetTypeActions_CustomizableObject", "CustomizableObject_RecompileSelectedTooltip", "Recompile selected Customizable Object."),
+		FSlateIcon(),
+		FUIAction(
+			FExecuteAction::CreateSP(this, &FAssetTypeActions_CustomizableObject::RecompileObjects, ERecompileCO::RCO_Selected, InObjects),
+			FCanExecuteAction()
+		)
+	);
+
+	// Recompile All
+	MenuBuilder.AddMenuEntry(
+		NSLOCTEXT("AssetTypeActions_CustomizableObject", "CustomizableObject_RecompileAll", "All"),
+		NSLOCTEXT("AssetTypeActions_CustomizableObject", "CustomizableObject_RecompileAllToolTip", "Recompile all Customizable Object assets in the project. Can take a really long time and freeze the editor while it is compiling."),
+		FSlateIcon(),
+		FUIAction(
+			FExecuteAction::CreateSP(this, &FAssetTypeActions_CustomizableObject::RecompileObjects, ERecompileCO::RCO_All, InObjects),
+			FCanExecuteAction()
+		)
+	);
+
+	// Recompile All Root Objects
+	MenuBuilder.AddMenuEntry(
+		NSLOCTEXT("AssetTypeActions_CustomizableObject", "CustomizableObject_RecompileRootObjects", "All root objects"),
+		NSLOCTEXT("AssetTypeActions_CustomizableObject", "CustomizableObject_RecompileRootObjectsTooltip", "Recompile all Root Customizable Objects in the project. Can take a long time and freeze the editor while it is compiling."),
+		FSlateIcon(),
+		FUIAction(
+			FExecuteAction::CreateSP(this, &FAssetTypeActions_CustomizableObject::RecompileObjects, ERecompileCO::RCO_AllRootObjects, InObjects),
+			FCanExecuteAction()
+		)
+	);
+
+	// Recompile all objects in memory, either loaded by the level, and editor or a during a compilation (child objects)
+	MenuBuilder.AddMenuEntry(
+		NSLOCTEXT("AssetTypeActions_CustomizableObject", "CustomizableObject_RecompileObjectsInMemory", "All loaded objects"),
+		NSLOCTEXT("AssetTypeActions_CustomizableObject", "CustomizableObject_RecompileObjectsInMemoryTooltip", "Recompile all Customizable Objects in memory. This action will compile objects referenced by the current level, CO Editors, etc. Can take a long time and freeze the editor while it is compiling."),
+		FSlateIcon(),
+		FUIAction(
+			FExecuteAction::CreateSP(this, &FAssetTypeActions_CustomizableObject::RecompileObjects, ERecompileCO::RCO_InMemory, InObjects),
+			FCanExecuteAction()
+		)
+	);
+
+}
+
+
+void FAssetTypeActions_CustomizableObject::RecompileObjects(ERecompileCO RecompileType, TArray<TWeakObjectPtr<UCustomizableObject>> InObjects)
+{
+	if (!UCustomizableObjectSystem::GetInstance()) return;
+
+	TArray<FAssetData> ObjectsToRecompile;
+
+	FAssetRegistryModule& AssetRegistryModule = FModuleManager::LoadModuleChecked<FAssetRegistryModule>("AssetRegistry");
+
+	if (RecompileType == ERecompileCO::RCO_Selected)
+	{
+		for (const TWeakObjectPtr<UCustomizableObject>& Object : InObjects)
+		{
+			FAssetData AssetData = AssetRegistryModule.Get().GetAssetByObjectPath(*Object->GetPathName());
+			if (AssetData.IsValid())
+			{
+				ObjectsToRecompile.Add(AssetData);
+			}
+		}
+	}
+	else
+	{
+		FARFilter AssetRegistryFilter;
+		AssetRegistryFilter.ClassPaths.Add(FTopLevelAssetPath(TEXT("/Script/CustomizableObject"), TEXT("CustomizableObject")));
+
+		FText Msg;
+		switch (RecompileType)
+		{
+		case ERecompileCO::RCO_All:
+			Msg = LOCTEXT("AssetTypeActions_CustomizableObject_RecompileAll_Dialog", "Are you sure you want to recompile all Customizable Objects in this project?\n\nThis action might take a long time and leave the editor unresponsive until it finishes.");
+			break;
+		case ERecompileCO::RCO_AllRootObjects:
+			Msg = LOCTEXT("AssetTypeActions_CustomizableObject_RecompileAllRoot_Dialog", "Are you sure you want to recompile all Root Customizable Objects in this project?\n\nThis action might take a long time and leave the editor unresponsive until it finishes.");
+			AssetRegistryFilter.TagsAndValues.Add(FName("IsRoot"), FString::FromInt(1));
+			break;
+		case ERecompileCO::RCO_InMemory:
+			Msg = LOCTEXT("AssetTypeActions_CustomizableObject_RecompileObjectsInMemory_Dialog", "Do you want to recompile all Customizable Objects in memory?\n\nThis action will compile objects referenced by the current level and other loaded objects.");
+			break;
+		default: check(false);
+			break;
+		}
+
+		if (FMessageDialog::Open(EAppMsgType::OkCancel, Msg) == EAppReturnType::Ok)
+		{
+			if (RecompileType == ERecompileCO::RCO_InMemory)
+			{
+				TArray<FAssetData> OutAssets;
+				AssetRegistryModule.Get().GetAssets(AssetRegistryFilter, OutAssets);
+
+				for (const FAssetData& Asset : OutAssets)
+				{
+					if (Asset.IsAssetLoaded())
+					{
+						ObjectsToRecompile.Add(Asset);
+					}
+				}
+			}
+			else
+			{
+				AssetRegistryModule.Get().GetAssets(AssetRegistryFilter, ObjectsToRecompile);
+			}
+		}
+	}
+
+	UCustomizableObjectSystem::GetInstance()->RecompileCustomizableObjects(ObjectsToRecompile);
+}
+
+
+bool FAssetTypeActions_CustomizableObject::AssetsActivatedOverride( const TArray<UObject*>& InObjects, EAssetTypeActivationMethod::Type ActivationType )
+{
+	if ( ActivationType == EAssetTypeActivationMethod::DoubleClicked || ActivationType == EAssetTypeActivationMethod::Opened )
+	{
+		if ( InObjects.Num() == 1 )
+		{
+			GEditor->GetEditorSubsystem<UAssetEditorSubsystem>()->OpenEditorForAsset(InObjects[0]);
+			return true;
+		}
+		else if ( InObjects.Num() > 1 )
+		{
+			GEditor->GetEditorSubsystem<UAssetEditorSubsystem>()->OpenEditorForAssets(InObjects);
+			return true;
+		}
+	}
+
+	return false;
+}
+
+
+#undef LOCTEXT_NAMESPACE
+
diff --git a/Engine/Plugins/Experimental/Mutable/Source/CustomizableObjectEditor/Private/AssetTypeActions_CustomizableObject.h b/Engine/Plugins/Experimental/Mutable/Source/CustomizableObjectEditor/Private/AssetTypeActions_CustomizableObject.h
new file mode 100644
--- /dev/null
+++ b/Engine/Plugins/Experimental/Mutable/Source/CustomizableObjectEditor/Private/AssetTypeActions_CustomizableObject.h
@@ -0,0 +1,45 @@
+// Copyright Epic Games, Inc. All Rights Reserved.
+
+#pragma once
+
+#include "AssetTypeActions_Base.h"
+#include "Runtime/Launch/Resources/Version.h"
+#include "CustomizableObjectClasses.h"
+
+enum class ERecompileCO
+{
+	RCO_All,
+	RCO_AllRootObjects,
+	RCO_Selected,
+	RCO_InMemory
+};
+
+class FAssetTypeActions_CustomizableObject : public FAssetTypeActions_Base
+{
+public:
+	// IAssetTypeActions Implementation
+	FText GetName() const override { return NSLOCTEXT("AssetTypeActions", "AssetTypeActions_CustomizableObject", "Customizable Object"); }
+	FColor GetTypeColor() const override { return FColor(234, 255, 0); }
+	UClass* GetSupportedClass() const override { return UCustomizableObject::StaticClass(); }
+	bool HasActions ( const TArray<UObject*>& InObjects ) const override { return true; }
+	void GetActions( const TArray<UObject*>& InObjects, FMenuBuilder& MenuBuilder ) override;
+	void OpenAssetEditor( const TArray<UObject*>& InObjects, TSharedPtr<class IToolkitHost> EditWithinLevelEditor = TSharedPtr<IToolkitHost>() ) override;
+	uint32 GetCategories() override;
+
+
+	bool AssetsActivatedOverride( const TArray<UObject*>& InObjects, EAssetTypeActivationMethod::Type ActivationType ) override;
+
+	bool CanFilter() override { return true; }
+	bool ShouldForceWorldCentric() override { return false; }
+	void PerformAssetDiff(UObject* OldAsset, UObject* NewAsset, const struct FRevisionInfo& OldRevision, const struct FRevisionInfo& NewRevision) const override {}
+
+private:
+
+	void ExecuteEdit(TArray<TWeakObjectPtr<UCustomizableObject>> Objects);	
+	void ExecuteDebug(TArray<TWeakObjectPtr<UCustomizableObject>> Objects);
+	void ExecuteNewInstance(TArray<TWeakObjectPtr<UCustomizableObject>> Objects);
+
+	void MakeRecompileSubMenu(FMenuBuilder& MenuBuilder, TArray<TWeakObjectPtr<UCustomizableObject>> InObjects);
+	void RecompileObjects(ERecompileCO Mode, TArray<TWeakObjectPtr<UCustomizableObject>> InObjects);
+	
+};
diff --git a/Engine/Plugins/Experimental/Mutable/Source/CustomizableObjectEditor/Private/AssetTypeActions_CustomizableObjectInstance.cpp b/Engine/Plugins/Experimental/Mutable/Source/CustomizableObjectEditor/Private/AssetTypeActions_CustomizableObjectInstance.cpp
new file mode 100644
--- /dev/null
+++ b/Engine/Plugins/Experimental/Mutable/Source/CustomizableObjectEditor/Private/AssetTypeActions_CustomizableObjectInstance.cpp
@@ -0,0 +1,82 @@
+// Copyright Epic Games, Inc. All Rights Reserved.
+
+#include "AssetTypeActions_CustomizableObjectInstance.h"
+#include "CustomizableObjectEditorModule.h"
+#include "Framework/MultiBox/MultiBoxBuilder.h"
+#include "UnrealEditorPortabilityHelpers.h"
+
+#define LOCTEXT_NAMESPACE "AssetTypeActions"
+
+void FAssetTypeActions_CustomizableObjectInstance::GetActions( const TArray<UObject*>& InObjects, FMenuBuilder& MenuBuilder )
+{
+	TArray<TWeakObjectPtr<UCustomizableObjectInstance>> Objects = GetTypedWeakObjectPtrs<UCustomizableObjectInstance>(InObjects);
+
+	MenuBuilder.AddMenuEntry(
+		NSLOCTEXT("AssetTypeActions_CustomizableObjectInstance", "CustomizableObjectInstance_Edit", "Edit"),
+		NSLOCTEXT("AssetTypeActions_CustomizableObjectInstance", "CustomizableObjectInstance_EditTooltip", "Opens the selected instances in the static mesh editor."),
+		FSlateIcon(),
+		FUIAction(
+			FExecuteAction::CreateSP( this, &FAssetTypeActions_CustomizableObjectInstance::ExecuteEdit, Objects ),
+			FCanExecuteAction()
+			)
+		);
+
+}
+
+
+void FAssetTypeActions_CustomizableObjectInstance::OpenAssetEditor( const TArray<UObject*>& InObjects, TSharedPtr<IToolkitHost> EditWithinLevelEditor )
+{
+	EToolkitMode::Type Mode = EditWithinLevelEditor.IsValid() ? EToolkitMode::WorldCentric : EToolkitMode::Standalone;
+
+	for (auto ObjIt = InObjects.CreateConstIterator(); ObjIt; ++ObjIt)
+	{
+		UCustomizableObjectInstance* Object = Cast<UCustomizableObjectInstance>(*ObjIt);
+		if (Object != NULL)
+		{
+			ICustomizableObjectEditorModule* CustomizableObjectEditorModule = &FModuleManager::LoadModuleChecked<ICustomizableObjectEditorModule>( "CustomizableObjectEditor" );
+			CustomizableObjectEditorModule->CreateCustomizableObjectInstanceEditor(Mode, EditWithinLevelEditor, Object);
+		}
+	}
+}
+
+
+uint32 FAssetTypeActions_CustomizableObjectInstance::GetCategories()
+{
+	const ICustomizableObjectEditorModule* CustomizableObjectEditorModule = &FModuleManager::LoadModuleChecked<ICustomizableObjectEditorModule>("CustomizableObjectEditor");
+	return CustomizableObjectEditorModule->GetAssetCategory();
+}
+
+
+void FAssetTypeActions_CustomizableObjectInstance::ExecuteEdit(TArray<TWeakObjectPtr<UCustomizableObjectInstance>> Objects)
+{
+	for (auto ObjIt = Objects.CreateConstIterator(); ObjIt; ++ObjIt)
+	{
+		UCustomizableObjectInstance* Object = (*ObjIt).Get();
+		if ( Object )
+		{
+			GEditor->GetEditorSubsystem<UAssetEditorSubsystem>()->OpenEditorForAsset(Object);
+		}
+	}
+}
+
+
+bool FAssetTypeActions_CustomizableObjectInstance::AssetsActivatedOverride(const TArray<UObject*>& InObjects, EAssetTypeActivationMethod::Type ActivationType)
+{
+	if (ActivationType == EAssetTypeActivationMethod::DoubleClicked || ActivationType == EAssetTypeActivationMethod::Opened)
+	{
+		if (InObjects.Num() == 1)
+		{
+			GEditor->GetEditorSubsystem<UAssetEditorSubsystem>()->OpenEditorForAsset(InObjects[0]);
+			return true;
+		}
+		else if (InObjects.Num() > 1)
+		{
+			GEditor->GetEditorSubsystem<UAssetEditorSubsystem>()->OpenEditorForAssets(InObjects);
+			return true;
+		}
+	}
+	return false;
+}
+
+
+#undef LOCTEXT_NAMESPACE
diff --git a/Engine/Plugins/Experimental/Mutable/Source/CustomizableObjectEditor/Private/AssetTypeActions_CustomizableObjectInstance.h b/Engine/Plugins/Experimental/Mutable/Source/CustomizableObjectEditor/Private/AssetTypeActions_CustomizableObjectInstance.h
new file mode 100644
--- /dev/null
+++ b/Engine/Plugins/Experimental/Mutable/Source/CustomizableObjectEditor/Private/AssetTypeActions_CustomizableObjectInstance.h
@@ -0,0 +1,31 @@
+// Copyright Epic Games, Inc. All Rights Reserved.
+
+#pragma once
+
+#include "AssetTypeActions_Base.h"
+#include "Runtime/Launch/Resources/Version.h"
+#include "CustomizableObjectClasses.h"
+
+
+class FAssetTypeActions_CustomizableObjectInstance : public FAssetTypeActions_Base
+{
+public:
+	// IAssetTypeActions Implementation
+	FText GetName() const override { return NSLOCTEXT("AssetTypeActions", "AssetTypeActions_CustomizableObjectInstance", "Customizable Object Instance"); }
+	FColor GetTypeColor() const override { return FColor(255, 82, 49); }
+	UClass* GetSupportedClass() const override { return UCustomizableObjectInstance::StaticClass(); }
+	bool HasActions ( const TArray<UObject*>& InObjects ) const override { return true; }
+	void GetActions( const TArray<UObject*>& InObjects, FMenuBuilder& MenuBuilder ) override;
+	void OpenAssetEditor( const TArray<UObject*>& InObjects, TSharedPtr<class IToolkitHost> EditWithinLevelEditor = TSharedPtr<IToolkitHost>() ) override;
+	uint32 GetCategories() override;
+
+
+	bool AssetsActivatedOverride(const TArray<UObject*>& InObjects, EAssetTypeActivationMethod::Type ActivationType) override;
+	bool CanFilter() override { return true; }
+	bool ShouldForceWorldCentric() override { return false; }
+	void PerformAssetDiff(UObject* OldAsset, UObject* NewAsset, const struct FRevisionInfo& OldRevision, const struct FRevisionInfo& NewRevision) const override {}
+
+private:
+
+	void ExecuteEdit(TArray<TWeakObjectPtr<UCustomizableObjectInstance>> Objects);	
+};
diff --git a/Engine/Plugins/Experimental/Mutable/Source/CustomizableObjectEditor/Private/CustomizableInstanceDetails.cpp b/Engine/Plugins/Experimental/Mutable/Source/CustomizableObjectEditor/Private/CustomizableInstanceDetails.cpp
new file mode 100644
--- /dev/null
+++ b/Engine/Plugins/Experimental/Mutable/Source/CustomizableObjectEditor/Private/CustomizableInstanceDetails.cpp
@@ -0,0 +1,76 @@
+// Copyright Epic Games, Inc. All Rights Reserved.
+
+#include "CustomizableInstanceDetails.h"
+#include "CustomizableObjectEditorModule.h"
+#include "CustomizableObjectSystem.h"
+#include "DetailLayoutBuilder.h"
+#include "DetailCategoryBuilder.h"
+#include "DetailWidgetRow.h"
+#include "SCustomizableInstanceProperties.h"
+#include "LevelEditorActions.h"
+
+#define LOCTEXT_NAMESPACE "CustomizableInstanceDetails"
+
+
+TSharedRef<IDetailCustomization> FCustomizableInstanceDetails::MakeInstance()
+{
+	return MakeShareable(new FCustomizableInstanceDetails);
+}
+
+
+void FCustomizableInstanceDetails::CustomizeDetails(const TSharedPtr<IDetailLayoutBuilder>& DetailBuilder)
+{
+	const IDetailsView* DetailsView = DetailBuilder->GetDetailsView();
+	check(DetailsView->GetSelectedObjects().Num());
+
+	CustomInstance = Cast<UCustomizableObjectInstance>(DetailsView->GetSelectedObjects()[0].Get());
+	check(CustomInstance.IsValid());
+	
+	LayoutBuilder = DetailBuilder;
+	
+	IDetailCategoryBuilder& MainCategory = DetailBuilder->EditCategory( "Customizable Instance" );
+
+	MainCategory.AddCustomRow( LOCTEXT("CustomizableInstanceDetails", "Instance Parameters") )
+	[
+		SAssignNew(InstancePropertiesWidget, SCustomizableInstanceProperties)
+			.CustomInstance(CustomInstance)
+			.InstanceDetails(SharedThis(this))
+	];
+
+	
+	// If the current instance has texture parameters, show the user interface to set the possible preview
+	// values for those parameters, in case no provider is registered.
+	UCustomizableObjectSystem* System = UCustomizableObjectSystem::GetInstance();
+
+	if (CustomInstance->GetTextureParameters().Num())
+	{
+		if (UCustomizableObjectImageProviderArray* ImageProvider = System->GetEditorExternalImageProvider())
+		{
+			ImageProvider->TexturesChangeDelegate.AddSP(this, &FCustomizableInstanceDetails::UpdateInstance);
+
+			IDetailCategoryBuilder& PreviewCategory = DetailBuilder->EditCategory("Preview Texture Parameter Options");
+
+			TArray<UObject*> Objs;
+			Objs.Add(ImageProvider);
+			PreviewCategory.AddExternalObjectProperty(Objs, FName("Textures"));
+		}
+	}
+}
+
+
+void FCustomizableInstanceDetails::Refresh() const
+{
+	if (IDetailLayoutBuilder* Layout = LayoutBuilder.Pin().Get()) // Raw because we don't want to keep alive the details builder when calling the force refresh details
+	{
+		Layout->ForceRefreshDetails();
+	}
+}
+
+
+void FCustomizableInstanceDetails::UpdateInstance() const
+{
+		CustomInstance->UpdateSkeletalMeshAsync(true, true);
+}
+
+
+#undef LOCTEXT_NAMESPACE
diff --git a/Engine/Plugins/Experimental/Mutable/Source/CustomizableObjectEditor/Private/CustomizableObjectBakeHelpers.cpp b/Engine/Plugins/Experimental/Mutable/Source/CustomizableObjectEditor/Private/CustomizableObjectBakeHelpers.cpp
new file mode 100644
--- /dev/null
+++ b/Engine/Plugins/Experimental/Mutable/Source/CustomizableObjectEditor/Private/CustomizableObjectBakeHelpers.cpp
@@ -0,0 +1,697 @@
+// Copyright Epic Games, Inc. All Rights Reserved.
+
+#include "CustomizableObjectBakeHelpers.h"
+
+#include "CustomizableObjectEditor.h"
+#include "CustomizableObjectMipDataProvider.h"
+#include "BusyCursor.h"
+#include "ObjectTools.h"
+#include "AssetRegistry/AssetRegistryModule.h"
+#include "AssetToolsModule.h"
+#include "Serialization/ArchiveReplaceObjectRef.h"
+#include "Rendering/SkeletalMeshRenderData.h"
+#include "Rendering/SkeletalMeshLODRenderData.h"
+#include "Rendering/SkeletalMeshModel.h"
+#include "Rendering/SkeletalMeshLODModel.h"
+#include "IAssetTools.h"
+
+#include "Misc/ScopeExit.h"
+
+#include "Templates/Decay.h"
+
+namespace
+{
+
+
+	inline uint32_t PackRGBA(uint8_t r, uint8_t g, uint8_t b, uint8_t a)
+	{
+		return r | (g << 8) | (b << 16) | (a << 24);
+	}
+
+	inline float Int8ToFloat_SNORM(const uint8_t input)
+	{
+		return (float)((int8_t)input) / 127.0f;
+	}
+
+	inline float Int8ToFloat_UNORM(const uint8_t input)
+	{
+		return (float)input / 255.0f;
+	}
+
+
+	inline void Decompress16x3bitIndices(const uint8_t* packed, uint8_t* unpacked)
+	{
+		uint32_t tmp, block, i;
+
+		for (block = 0; block < 2; ++block) {
+			tmp = 0;
+
+			// Read three bytes
+			for (i = 0; i < 3; ++i) {
+				tmp |= ((uint32_t)packed[i]) << (i * 8);
+			}
+
+			// Unpack 8x3 bit from last 3 byte block
+			for (i = 0; i < 8; ++i) {
+				unpacked[i] = (tmp >> (i * 3)) & 0x7;
+			}
+
+			packed += 3;
+			unpacked += 8;
+		}
+	}
+
+	inline void DecompressBlockBC1Internal(const uint8_t* block,
+		unsigned char* output, uint32_t outputStride, const uint8_t* alphaValues)
+	{
+		uint32_t temp, code;
+
+		uint16_t color0, color1;
+		uint8_t r0, g0, b0, r1, g1, b1;
+
+		int i, j;
+
+		color0 = *(const uint16_t*)(block);
+		color1 = *(const uint16_t*)(block + 2);
+
+		temp = (color0 >> 11) * 255 + 16;
+		r0 = (uint8_t)((temp / 32 + temp) / 32);
+		temp = ((color0 & 0x07E0) >> 5) * 255 + 32;
+		g0 = (uint8_t)((temp / 64 + temp) / 64);
+		temp = (color0 & 0x001F) * 255 + 16;
+		b0 = (uint8_t)((temp / 32 + temp) / 32);
+
+		temp = (color1 >> 11) * 255 + 16;
+		r1 = (uint8_t)((temp / 32 + temp) / 32);
+		temp = ((color1 & 0x07E0) >> 5) * 255 + 32;
+		g1 = (uint8_t)((temp / 64 + temp) / 64);
+		temp = (color1 & 0x001F) * 255 + 16;
+		b1 = (uint8_t)((temp / 32 + temp) / 32);
+
+		code = *(const uint32_t*)(block + 4);
+
+		if (color0 > color1) {
+			for (j = 0; j < 4; ++j) {
+				for (i = 0; i < 4; ++i) {
+					uint32_t finalColor, positionCode;
+					uint8_t alpha;
+
+					alpha = alphaValues[j * 4 + i];
+
+					finalColor = 0;
+					positionCode = (code >> 2 * (4 * j + i)) & 0x03;
+
+					switch (positionCode) {
+					case 0:
+						finalColor = PackRGBA(r0, g0, b0, alpha);
+						break;
+					case 1:
+						finalColor = PackRGBA(r1, g1, b1, alpha);
+						break;
+					case 2:
+						finalColor = PackRGBA((2 * r0 + r1) / 3, (2 * g0 + g1) / 3, (2 * b0 + b1) / 3, alpha);
+						break;
+					case 3:
+						finalColor = PackRGBA((r0 + 2 * r1) / 3, (g0 + 2 * g1) / 3, (b0 + 2 * b1) / 3, alpha);
+						break;
+					}
+
+					*(uint32_t*)(output + j*outputStride + i * sizeof(uint32_t)) = finalColor;
+				}
+			}
+		}
+		else {
+			for (j = 0; j < 4; ++j) {
+				for (i = 0; i < 4; ++i) {
+					uint32_t finalColor, positionCode;
+					uint8_t alpha;
+
+					alpha = alphaValues[j * 4 + i];
+
+					finalColor = 0;
+					positionCode = (code >> 2 * (4 * j + i)) & 0x03;
+
+					switch (positionCode) {
+					case 0:
+						finalColor = PackRGBA(r0, g0, b0, alpha);
+						break;
+					case 1:
+						finalColor = PackRGBA(r1, g1, b1, alpha);
+						break;
+					case 2:
+						finalColor = PackRGBA((r0 + r1) / 2, (g0 + g1) / 2, (b0 + b1) / 2, alpha);
+						break;
+					case 3:
+						finalColor = PackRGBA(0, 0, 0, alpha);
+						break;
+					}
+
+					*(uint32_t*)(output + j*outputStride + i * sizeof(uint32_t)) = finalColor;
+				}
+			}
+		}
+	}
+
+
+	inline void DecompressBlockBC1(uint32_t x, uint32_t y, uint32_t stride, const uint8_t* blockStorage, unsigned char* image)
+	{
+		static const uint8_t const_alpha[] = {
+			255, 255, 255, 255,
+			255, 255, 255, 255,
+			255, 255, 255, 255,
+			255, 255, 255, 255
+		};
+
+		DecompressBlockBC1Internal(blockStorage, image + x * sizeof(uint32_t) + (y * stride), stride, const_alpha);
+	}
+
+
+	inline void DecompressBlockBC3(uint32_t x, uint32_t y, uint32_t stride, const uint8_t* blockStorage, unsigned char* image)
+	{
+		uint8_t alpha0, alpha1;
+		uint8_t alphaIndices[16];
+
+		uint16_t color0, color1;
+		uint8_t r0, g0, b0, r1, g1, b1;
+
+		int i, j;
+
+		uint32_t temp, code;
+
+		alpha0 = *(blockStorage);
+		alpha1 = *(blockStorage + 1);
+
+		Decompress16x3bitIndices(blockStorage + 2, alphaIndices);
+
+		color0 = *(const uint16_t*)(blockStorage + 8);
+		color1 = *(const uint16_t*)(blockStorage + 10);
+
+		temp = (color0 >> 11) * 255 + 16;
+		r0 = (uint8_t)((temp / 32 + temp) / 32);
+		temp = ((color0 & 0x07E0) >> 5) * 255 + 32;
+		g0 = (uint8_t)((temp / 64 + temp) / 64);
+		temp = (color0 & 0x001F) * 255 + 16;
+		b0 = (uint8_t)((temp / 32 + temp) / 32);
+
+		temp = (color1 >> 11) * 255 + 16;
+		r1 = (uint8_t)((temp / 32 + temp) / 32);
+		temp = ((color1 & 0x07E0) >> 5) * 255 + 32;
+		g1 = (uint8_t)((temp / 64 + temp) / 64);
+		temp = (color1 & 0x001F) * 255 + 16;
+		b1 = (uint8_t)((temp / 32 + temp) / 32);
+
+		code = *(const uint32_t*)(blockStorage + 12);
+
+		for (j = 0; j < 4; j++) {
+			for (i = 0; i < 4; i++) {
+				uint8_t finalAlpha;
+				int alphaCode;
+				uint8_t colorCode;
+				uint32_t finalColor;
+
+				alphaCode = alphaIndices[4 * j + i];
+
+				if (alphaCode == 0) {
+					finalAlpha = alpha0;
+				}
+				else if (alphaCode == 1) {
+					finalAlpha = alpha1;
+				}
+				else {
+					if (alpha0 > alpha1) {
+						finalAlpha = (uint8_t)(((8 - alphaCode)*alpha0 + (alphaCode - 1)*alpha1) / 7);
+					}
+					else {
+						if (alphaCode == 6) {
+							finalAlpha = 0;
+						}
+						else if (alphaCode == 7) {
+							finalAlpha = 255;
+						}
+						else {
+							finalAlpha = (uint8_t)(((6 - alphaCode)*alpha0 + (alphaCode - 1)*alpha1) / 5);
+						}
+					}
+				}
+
+				colorCode = (code >> 2 * (4 * j + i)) & 0x03;
+				finalColor = 0;
+
+				switch (colorCode) {
+				case 0:
+					finalColor = PackRGBA(r0, g0, b0, finalAlpha);
+					break;
+				case 1:
+					finalColor = PackRGBA(r1, g1, b1, finalAlpha);
+					break;
+				case 2:
+					finalColor = PackRGBA((2 * r0 + r1) / 3, (2 * g0 + g1) / 3, (2 * b0 + b1) / 3, finalAlpha);
+					break;
+				case 3:
+					finalColor = PackRGBA((r0 + 2 * r1) / 3, (g0 + 2 * g1) / 3, (b0 + 2 * b1) / 3, finalAlpha);
+					break;
+				}
+
+
+				*(uint32_t*)(image + sizeof(uint32_t) * (i + x) + (stride * (y + j))) = finalColor;
+			}
+		}
+	}
+
+
+	inline void DecompressBlockBC2(uint32_t x, uint32_t y, uint32_t stride, const uint8_t* blockStorage, unsigned char* image)
+	{
+		int i;
+
+		uint8_t alphaValues[16] = { 0 };
+
+		for (i = 0; i < 4; ++i) {
+			const uint16_t* alphaData = (const uint16_t*)(blockStorage);
+
+			alphaValues[i * 4 + 0] = (((*alphaData) >> 0) & 0xF) * 17;
+			alphaValues[i * 4 + 1] = (((*alphaData) >> 4) & 0xF) * 17;
+			alphaValues[i * 4 + 2] = (((*alphaData) >> 8) & 0xF) * 17;
+			alphaValues[i * 4 + 3] = (((*alphaData) >> 12) & 0xF) * 17;
+
+			blockStorage += 2;
+		}
+
+		DecompressBlockBC1Internal(blockStorage, image + x * sizeof(uint32_t) + (y * stride), stride, alphaValues);
+	}
+
+
+	inline void DecompressBlockBC4Internal(const uint8_t* block, uint8_t* output, uint32_t outputStride, const float* colorTable)
+	{
+		uint8_t indices[16];
+		int x, y;
+
+		Decompress16x3bitIndices(block + 2, indices);
+
+		for (y = 0; y < 4; ++y) 
+		{
+			for (x = 0; x < 4; ++x) 
+			{
+				uint8_t v = (uint8_t)FMath::Clamp(colorTable[indices[y * 4 + x]] * 255.0f, 0.0f, 255.0f);
+				*(output + x) = v;
+			}
+
+			output += outputStride;
+		}
+	}
+
+
+	enum class BC4Mode : uint8_t { BC4_UNORM, BC4_SNORM };
+
+
+	inline void DecompressBlockBC4(uint32_t x, uint32_t y, uint32_t stride, BC4Mode mode, const uint8_t* blockStorage, uint8_t* image)
+	{
+		float colorTable[8];
+		float r0, r1;
+
+		if (mode == BC4Mode::BC4_UNORM) 
+		{
+			r0 = Int8ToFloat_UNORM(blockStorage[0]);
+			r1 = Int8ToFloat_UNORM(blockStorage[1]);
+
+			colorTable[0] = r0;
+			colorTable[1] = r1;
+
+			if (r0 > r1) 
+			{
+				// 6 interpolated color values
+				colorTable[2] = (6 * r0 + 1 * r1) / 7.0f; // bit code 010
+				colorTable[3] = (5 * r0 + 2 * r1) / 7.0f; // bit code 011
+				colorTable[4] = (4 * r0 + 3 * r1) / 7.0f; // bit code 100
+				colorTable[5] = (3 * r0 + 4 * r1) / 7.0f; // bit code 101
+				colorTable[6] = (2 * r0 + 5 * r1) / 7.0f; // bit code 110
+				colorTable[7] = (1 * r0 + 6 * r1) / 7.0f; // bit code 111
+			}
+			else 
+			{
+				// 4 interpolated color values
+				colorTable[2] = (4 * r0 + 1 * r1) / 5.0f; // bit code 010
+				colorTable[3] = (3 * r0 + 2 * r1) / 5.0f; // bit code 011
+				colorTable[4] = (2 * r0 + 3 * r1) / 5.0f; // bit code 100
+				colorTable[5] = (1 * r0 + 4 * r1) / 5.0f; // bit code 101
+				colorTable[6] = 0.0f;               // bit code 110
+				colorTable[7] = 1.0f;               // bit code 111
+			}
+		}
+		else if (mode == BC4Mode::BC4_SNORM) 
+		{
+			r0 = Int8ToFloat_SNORM(blockStorage[0]);
+			r1 = Int8ToFloat_SNORM(blockStorage[1]);
+
+			colorTable[0] = r0;
+			colorTable[1] = r1;
+
+			if (r0 > r1) 
+			{
+				// 6 interpolated color values
+				colorTable[2] = (6 * r0 + 1 * r1) / 7.0f; // bit code 010
+				colorTable[3] = (5 * r0 + 2 * r1) / 7.0f; // bit code 011
+				colorTable[4] = (4 * r0 + 3 * r1) / 7.0f; // bit code 100
+				colorTable[5] = (3 * r0 + 4 * r1) / 7.0f; // bit code 101
+				colorTable[6] = (2 * r0 + 5 * r1) / 7.0f; // bit code 110
+				colorTable[7] = (1 * r0 + 6 * r1) / 7.0f; // bit code 111
+			}
+			else 
+			{
+				// 4 interpolated color values
+				colorTable[2] = (4 * r0 + 1 * r1) / 5.0f; // bit code 010
+				colorTable[3] = (3 * r0 + 2 * r1) / 5.0f; // bit code 011
+				colorTable[4] = (2 * r0 + 3 * r1) / 5.0f; // bit code 100
+				colorTable[5] = (1 * r0 + 4 * r1) / 5.0f; // bit code 101
+				colorTable[6] = -1.0f;              // bit code 110
+				colorTable[7] = 1.0f;              // bit code 111
+			}
+		}
+
+		DecompressBlockBC4Internal(blockStorage, image + x + (y * stride), stride, colorTable);
+	}
+
+
+	inline void DecompressBlockBC5(uint32_t x, uint32_t y, uint32_t stride, BC4Mode mode, const uint8_t* blockStorage, uint8_t* image)
+	{
+		// We decompress the two channels separately and interleave them when
+		// writing to the output
+		uint8_t c0[16];
+		uint8_t c1[16];
+
+		int dx, dy;
+
+		DecompressBlockBC4(0, 0, 4, mode, blockStorage, (uint8_t*)c0);
+		DecompressBlockBC4(0, 0, 4, mode, blockStorage + 8, (uint8_t*)c1);
+
+		for (dy = 0; dy < 4; ++dy) 
+		{
+			for (dx = 0; dx < 4; ++dx) 
+			{
+				uint8_t* pDest = image + stride * (y+dy) + (x+dx) * 4;
+				pDest[0] = c0[dy * 4 + dx];
+				pDest[1] = c1[dy * 4 + dx];
+				pDest[2] = 255;
+				pDest[3] = 255;
+			}
+		}
+	}
+
+}
+
+
+UObject* BakeHelper_DuplicateAsset(UObject* Object, const FString& ObjName, const FString& PkgName, bool ResetDuplicatedFlags, TMap<UObject*, UObject*>& ReplacementMap, bool OverwritePackage)
+{
+	FString FinalObjName = ObjName;
+	FString FinalPkgName = PkgName;
+
+	if (!OverwritePackage)
+	{
+		FAssetToolsModule& AssetToolsModule = FModuleManager::GetModuleChecked<FAssetToolsModule>("AssetTools");
+		AssetToolsModule.Get().CreateUniqueAssetName(PkgName, "", FinalPkgName, FinalObjName);
+	}
+
+	UPackage* Package = CreatePackage(*FinalPkgName);
+	Package->FullyLoad();
+
+	FObjectDuplicationParameters Params = InitStaticDuplicateObjectParams(Object, Package, *FinalObjName, RF_AllFlags, nullptr, EDuplicateMode::Normal);
+
+	UObject* DupObject = StaticDuplicateObjectEx(Params);
+	if (DupObject)
+	{	
+		DupObject->SetFlags(RF_Public | RF_Standalone);
+		DupObject->ClearFlags(RF_Transient);
+
+		// The garbage collector is called in the middle of the bake process, and this can destroy this temporary objects. 
+		// We add them to the garbage root to prevent this. This will avoid them being unloaded while the editor is running, but this
+		// action is not used often.
+		DupObject->AddToRoot();
+		DupObject->MarkPackageDirty();
+
+		// Notify the asset registry
+		FAssetRegistryModule::AssetCreated(DupObject);
+
+		// Replace all references
+		ReplacementMap.Add(Object, DupObject);
+
+		constexpr EArchiveReplaceObjectFlags ReplaceFlags = (EArchiveReplaceObjectFlags::IgnoreOuterRef | EArchiveReplaceObjectFlags::IgnoreArchetypeRef);
+		FArchiveReplaceObjectRef<UObject> ReplaceAr(DupObject, ReplacementMap, ReplaceFlags);
+	}
+
+	return DupObject;
+}
+
+
+UTexture2D* BakeHelper_CreateAssetTexture(UTexture2D* SrcTex, const FString& TexObjName, const FString& TexPkgName, const UTexture* OrgTex, bool ResetDuplicatedFlags, TMap<UObject*, UObject*>& ReplacementMap, bool OverwritePackage)
+{
+	const bool IsMutableTexutre = !SrcTex->Source.IsValid();
+	if (IsMutableTexutre)
+	{
+		int sx = SrcTex->GetPlatformData()->SizeX;
+		int sy = SrcTex->GetPlatformData()->SizeY;
+		ETextureSourceFormat PixelFormat = SrcTex->GetPlatformData()->PixelFormat == PF_BC4 ? TSF_G8 : TSF_BGRA8;
+
+		// Begin duplicate texture
+		FString FinalObjName = TexObjName;
+		FString FinalPkgName = TexPkgName;
+
+		if (!OverwritePackage)
+		{
+			FAssetToolsModule& AssetToolsModule = FModuleManager::GetModuleChecked<FAssetToolsModule>("AssetTools");
+			AssetToolsModule.Get().CreateUniqueAssetName(TexPkgName, "", FinalPkgName, FinalObjName);
+		}
+
+		UPackage* Package = CreatePackage(*FinalPkgName);
+		Package->FullyLoad();
+
+		UTexture2D* DupTex = NewObject<UTexture2D>(Package, *FinalObjName, RF_Public | RF_Standalone);
+
+		ReplacementMap.Add(SrcTex, DupTex);
+		
+		// The garbage collector is called in the middle of the bake process, and this can destroy this temporary objects. 
+		// We add them to the garbage root to prevent this. This will avoid them being unloaded while the editor is running, but this
+		// action is not used often.
+		DupTex->AddToRoot();
+		DupTex->MarkPackageDirty();
+
+		// Notify the asset registry
+		FAssetRegistryModule::AssetCreated(DupTex);
+
+		// Replace all references
+		constexpr EArchiveReplaceObjectFlags ReplaceFlags = (EArchiveReplaceObjectFlags::IgnoreOuterRef | EArchiveReplaceObjectFlags::IgnoreArchetypeRef);
+		FArchiveReplaceObjectRef<UObject> ReplaceAr(DupTex, ReplacementMap, ReplaceFlags);
+
+		// End duplicate texture
+
+		CopyTextureProperties(DupTex, SrcTex);
+
+		DupTex->RemoveUserDataOfClass(UMutableTextureMipDataProviderFactory::StaticClass());
+		if (OrgTex)
+		{
+			DupTex->CompressionSettings = OrgTex->CompressionSettings;
+		}
+
+		// Mutable textures only have platform data. We need to build the source data for them to be assets.
+		DupTex->Source.Init(sx, sy, 1, 1, PixelFormat);
+
+		int32 MipCount = SrcTex->GetPlatformData()->Mips.Num();
+		if (!MipCount)
+		{
+			UE_LOG(LogMutable, Warning, TEXT("Bake Instances: Empty texture found [%s]."), *SrcTex->GetName());
+			return DupTex;
+		}
+
+		const uint8_t* src = reinterpret_cast<const uint8_t*>(SrcTex->GetPlatformData()->Mips[0].BulkData.LockReadOnly());
+		check(src); // Mutable texture should always contain platform data
+
+		uint8_t* dst = DupTex->Source.LockMip(0);
+		check(dst);
+
+		switch (SrcTex->GetPlatformData()->PixelFormat)
+		{
+		case PF_R8G8B8A8:
+			FMemory::Memcpy(dst, src, sx*sy * 4);
+			break;
+
+		case PF_DXT1:
+			for (int y = 0; y<sy; y += 4)
+				for (int x = 0; x<sx; x += 4)
+				{
+					DecompressBlockBC1(x, y, sx * 4, src, dst);
+					src += 8;
+				}
+			break;
+
+		case PF_DXT3:
+			for (int y = 0; y<sy; y += 4)
+				for (int x = 0; x<sx; x += 4)
+				{
+					DecompressBlockBC2(x, y, sx * 4, src, dst);
+					src += 16;
+				}
+			break;
+
+		case PF_DXT5:
+			for (int y = 0; y<sy; y += 4)
+				for (int x = 0; x<sx; x += 4)
+				{
+					DecompressBlockBC3(x, y, sx * 4, src, dst);
+					src += 16;
+				}
+			break;
+
+		case PF_BC4:
+			for (int y = 0; y<sy; y += 4)
+				for (int x = 0; x<sx; x += 4)
+				{
+					DecompressBlockBC4(x, y, sx * 1, BC4Mode::BC4_UNORM, src, dst);
+					src += 8;
+				}
+			break;
+
+		case PF_BC5:
+			for (int y = 0; y<sy; y += 4)
+				for (int x = 0; x<sx; x += 4)
+				{
+					DecompressBlockBC5(x, y, sx * 4, BC4Mode::BC4_UNORM, src, dst);
+					src += 16;
+				}
+			break;
+
+		default:
+			// Not implemented...
+			check(false);
+			break;
+		}
+
+		const bool bNeedsRBSwizzle = PixelFormat == TSF_BGRA8;
+		if (bNeedsRBSwizzle)
+		{
+			for (int x = 0; x<sx*sy; ++x)
+			{
+				uint8_t temp = dst[0];
+				dst[0] = dst[2];
+				dst[2] = temp;
+				dst += 4;
+			}
+		}
+
+		SrcTex->GetPlatformData()->Mips[0].BulkData.Unlock();
+
+		DupTex->Source.UnlockMip(0);
+
+		DupTex->UpdateResource();
+
+		return DupTex;
+	}
+	else
+	{
+		return Cast<UTexture2D>(BakeHelper_DuplicateAsset(SrcTex, TexObjName ,TexPkgName, ResetDuplicatedFlags, ReplacementMap, OverwritePackage));
+	}
+}
+
+void BakeHelper_RegenerateImportedModel(USkeletalMesh* SkeletalMesh)
+{
+	if (!SkeletalMesh)
+	{
+		return;
+	}
+	
+	FSkeletalMeshRenderData* SkelResource = SkeletalMesh->GetResourceForRendering();
+	if (!SkelResource)
+	{
+		return;
+	}
+
+	FSkeletalMeshModel* ImportedModel = SkeletalMesh->GetImportedModel();
+	ImportedModel->bGuidIsHash = false;
+	ImportedModel->SkeletalMeshModelGUID = FGuid::NewGuid();
+
+	ImportedModel->LODModels.Empty();
+	int32 OriginalIndex = 0;
+
+	for (int32 LODIndex = 0; LODIndex < SkelResource->LODRenderData.Num(); ++LODIndex)
+	{
+		ImportedModel->LODModels.Add(new FSkeletalMeshLODModel());
+
+		FSkeletalMeshLODRenderData& LODModel = SkelResource->LODRenderData[LODIndex];
+		int32 CurrentSectionInitialVertex = 0;
+
+		ImportedModel->LODModels[LODIndex].ActiveBoneIndices = LODModel.ActiveBoneIndices;
+		ImportedModel->LODModels[LODIndex].NumTexCoords = LODModel.GetNumTexCoords();
+		ImportedModel->LODModels[LODIndex].RequiredBones = LODModel.RequiredBones;
+		ImportedModel->LODModels[LODIndex].NumVertices = LODModel.GetNumVertices();
+
+		// Indices
+		int indexCount = LODModel.MultiSizeIndexContainer.GetIndexBuffer()->Num();
+		ImportedModel->LODModels[LODIndex].IndexBuffer.SetNum(indexCount);
+		for ( int i = 0; i < indexCount; ++i )
+		{
+			ImportedModel->LODModels[LODIndex].IndexBuffer[i] = LODModel.MultiSizeIndexContainer.GetIndexBuffer()->Get(i);
+		}
+
+		ImportedModel->LODModels[LODIndex].Sections.SetNum(LODModel.RenderSections.Num());
+
+		for (int SectionIndex = 0; SectionIndex < LODModel.RenderSections.Num(); ++SectionIndex)
+		{
+			const FSkelMeshRenderSection& RenderSection = LODModel.RenderSections[SectionIndex];
+			FSkelMeshSection& ImportedSection = ImportedModel->LODModels[LODIndex].Sections[SectionIndex];
+
+			// Vertices
+			ImportedSection.NumVertices = RenderSection.NumVertices;
+			ImportedSection.SoftVertices.Empty(RenderSection.NumVertices);
+			ImportedSection.SoftVertices.AddUninitialized(RenderSection.NumVertices);
+
+			for (uint32 i = 0; i < RenderSection.NumVertices; ++i)
+			{
+				const FPositionVertex* PosPtr = static_cast<const FPositionVertex*>(LODModel.StaticVertexBuffers.PositionVertexBuffer.GetVertexData());
+				PosPtr += (CurrentSectionInitialVertex + i);
+
+				check(!LODModel.StaticVertexBuffers.StaticMeshVertexBuffer.GetUseHighPrecisionTangentBasis());
+				const FPackedNormal* TangentPtr = static_cast<const FPackedNormal*>(LODModel.StaticVertexBuffers.StaticMeshVertexBuffer.GetTangentData());
+				TangentPtr += ((CurrentSectionInitialVertex + i) * 2);
+
+				check(LODModel.StaticVertexBuffers.StaticMeshVertexBuffer.GetUseFullPrecisionUVs());
+				
+                using UVsVectorType = typename TDecay<decltype(DeclVal<FSoftSkinVertex>().UVs[0])>::Type;
+                
+                const UVsVectorType* TexCoordPosPtr = static_cast<const UVsVectorType*>(LODModel.StaticVertexBuffers.StaticMeshVertexBuffer.GetTexCoordData());
+				const uint32 NumTexCoords = LODModel.StaticVertexBuffers.StaticMeshVertexBuffer.GetNumTexCoords();
+				TexCoordPosPtr += ((CurrentSectionInitialVertex + i) * NumTexCoords);
+
+				FSoftSkinVertex& Vertex = ImportedSection.SoftVertices[i];
+				for (int32 j = 0; j < MAX_TOTAL_INFLUENCES; j++)
+				{
+					Vertex.InfluenceBones[j] = LODModel.SkinWeightVertexBuffer.GetBoneIndex(i,j);
+					Vertex.InfluenceWeights[j] = LODModel.SkinWeightVertexBuffer.GetBoneWeight(i,j);
+				}
+
+				Vertex.Color = FColor::White;
+
+				Vertex.Position = PosPtr->Position;
+
+				Vertex.TangentX = TangentPtr[0].ToFVector3f();
+				Vertex.TangentZ = TangentPtr[1].ToFVector3f();
+				float TangentSign = TangentPtr[1].Vector.W == 0 ? -1.f : 1.f;
+				Vertex.TangentY = FVector3f::CrossProduct(Vertex.TangentZ, Vertex.TangentX) * TangentSign;
+
+				Vertex.UVs[0] = TexCoordPosPtr[0];
+				Vertex.UVs[1] = NumTexCoords > 1 ? TexCoordPosPtr[1] : UVsVectorType::ZeroVector;
+				Vertex.UVs[2] = NumTexCoords > 2 ? TexCoordPosPtr[2] : UVsVectorType::ZeroVector;
+				Vertex.UVs[3] = NumTexCoords > 3 ? TexCoordPosPtr[3] : UVsVectorType::ZeroVector;
+			}
+
+			CurrentSectionInitialVertex += RenderSection.NumVertices;
+
+			// Triangles
+			ImportedSection.NumTriangles = RenderSection.NumTriangles;
+			ImportedSection.BaseIndex = RenderSection.BaseIndex;
+			ImportedSection.BaseVertexIndex = RenderSection.BaseVertexIndex;
+			ImportedSection.BoneMap = RenderSection.BoneMap;
+			ImportedSection.MaterialIndex = RenderSection.MaterialIndex;
+			ImportedSection.MaxBoneInfluences = RenderSection.MaxBoneInfluences;
+			ImportedSection.OriginalDataSectionIndex = OriginalIndex++;
+		}
+
+		ImportedModel->LODModels[LODIndex].SyncronizeUserSectionsDataArray();
+	}
+}
diff --git a/Engine/Plugins/Experimental/Mutable/Source/CustomizableObjectEditor/Private/CustomizableObjectBakeHelpers.h b/Engine/Plugins/Experimental/Mutable/Source/CustomizableObjectEditor/Private/CustomizableObjectBakeHelpers.h
new file mode 100644
--- /dev/null
+++ b/Engine/Plugins/Experimental/Mutable/Source/CustomizableObjectEditor/Private/CustomizableObjectBakeHelpers.h
@@ -0,0 +1,17 @@
+// Copyright Epic Games, Inc. All Rights Reserved.
+
+#pragma once
+
+#include "Engine/Texture2D.h"
+#include "Engine/SkeletalMesh.h"
+
+extern UObject* BakeHelper_DuplicateAsset(UObject* Object, const FString& ObjName, const FString& PkgName, bool ResetDuplicatedFlags, TMap<UObject*, UObject*>& ReplacementMap, bool OverwritePackage);
+
+/**
+ * Duplicates a texture asset. Duplicates Mutable and non Mutable textures.
+ * 
+ * @param OrgTex Original source texture from which a Mutable texture has been generated. Only required when SrcTex is a Mutable texture.
+ */
+extern UTexture2D* BakeHelper_CreateAssetTexture(UTexture2D* SrcTex, const FString& TexObjName, const FString& TexPkgName, const UTexture* OrgTex, bool ResetDuplicatedFlags, TMap<UObject*, UObject*>& ReplacementMap, bool OverwritePackage);
+
+extern void BakeHelper_RegenerateImportedModel(class USkeletalMesh* Mesh);
