[mutable] Move Mutable plugin from the restricted folder to the experimental plugins folder.

diff --git a/Engine/Plugins/Experimental/Mutable/Source/CustomizableObjectEditor/Private/CustomizableObjectCompileRunnable.cpp b/Engine/Plugins/Experimental/Mutable/Source/CustomizableObjectEditor/Private/CustomizableObjectCompileRunnable.cpp
new file mode 100644
--- /dev/null
+++ b/Engine/Plugins/Experimental/Mutable/Source/CustomizableObjectEditor/Private/CustomizableObjectCompileRunnable.cpp
@@ -0,0 +1,301 @@
+// Copyright Epic Games, Inc. All Rights Reserved.
+
+#include "CustomizableObjectCompileRunnable.h"
+#include "CustomizableObjectCompiler.h"
+#include "CustomizableObjectEditorModule.h"
+#include "CustomizableObjectSystem.h"
+#include "CustomizableObject.h"
+#include "UnrealMutableModelDiskStreamer.h"
+
+#include "StaticMeshResources.h"
+#include "Framework/Notifications/NotificationManager.h"
+#include "Widgets/Notifications/SNotificationList.h"
+#include "AssetRegistry/AssetRegistryModule.h"
+#include "MessageLogModule.h"
+#include "Logging/MessageLog.h"
+#include "Application/ThrottleManager.h"
+#include "HAL/Runnable.h"
+#include "Interfaces/ITargetPlatform.h"
+#include "Interfaces/ITargetPlatformManagerModule.h"
+#include "HAL/RunnableThread.h"
+#include "MeshUtilities.h"
+#include "Misc/ConfigCacheIni.h"
+
+#include "Misc/Paths.h"
+#include "HAL/FileManager.h"
+#include "HAL/PlatformFileManager.h"
+
+#include "UnrealEditorPortabilityHelpers.h"
+#include "UnrealPortabilityHelpers.h"
+#include "PlatformInfo.h"
+#include "Engine/AssetManager.h"
+#include "AssetRegistry/ARFilter.h"
+
+#include "MutableTools/Public/Compiler.h"
+#include "MutableTools/Public/ErrorLog.h"
+#include "MutableRuntime/Public/System.h"
+#include "MutableRuntime/Public/Ptr.h"
+
+
+#define LOCTEXT_NAMESPACE "CustomizableObjectEditor"
+
+
+FCustomizableObjectCompileRunnable::FCustomizableObjectCompileRunnable(mu::NodePtr Root, bool bInDisableTextureLayout)
+	: MutableRoot(Root)
+	, bDisableTextureLayout(bInDisableTextureLayout)
+	, bThreadCompleted(false)
+	, MutableIsDisabled(false)
+{
+}
+
+
+uint32 FCustomizableObjectCompileRunnable::Run()
+{
+	UE_LOG(LogMutable, Verbose, TEXT("PROFILE: [ %16.8f ] FCustomizableObjectCompileRunnable::Run start."), FPlatformTime::Seconds());
+
+	uint32 Result = 1;
+	ErrorMsg = FString();
+
+	if (MutableIsDisabled)
+	{
+		bThreadCompleted = true;
+		UE_LOG(LogMutable, Verbose, TEXT("PROFILE: [ %16.8f ] FCustomizableObjectCompileRunnable::Run end. NOTE: Mutable compile is deactivated in Editor. To reactivate it, go to Project Settings -> Plugins -> Mutable and unmark the option Disable Mutable Compile In Editor"), FPlatformTime::Seconds());
+		return true;
+	}
+
+	mu::CompilerOptionsPtr CompilerOptions = new mu::CompilerOptions();
+
+	CompilerOptions->SetUseDiskCache(Options.bUseDiskCompilation);
+
+	if (Options.OptimizationLevel > 3)
+	{
+		UE_LOG(LogMutable, Verbose, TEXT("Mutable compile optimization level out of range. Clamping to maximum."));
+		Options.OptimizationLevel = 3;
+	}
+
+	switch (Options.OptimizationLevel)
+	{
+	case 0:
+		CompilerOptions->SetOptimisationEnabled(false);
+		CompilerOptions->SetConstReductionEnabled(false);
+		CompilerOptions->SetOptimisationMaxIteration(1);
+		break;
+
+	case 1:
+		CompilerOptions->SetOptimisationEnabled(false);
+		CompilerOptions->SetConstReductionEnabled(true);
+		CompilerOptions->SetOptimisationMaxIteration(1);
+		break;
+
+	case 2:
+		CompilerOptions->SetOptimisationEnabled(true);
+		CompilerOptions->SetConstReductionEnabled(true);
+		CompilerOptions->SetOptimisationMaxIteration(16);
+		break;
+
+	case 3:
+		CompilerOptions->SetOptimisationEnabled(true);
+		CompilerOptions->SetConstReductionEnabled(true);
+		CompilerOptions->SetOptimisationMaxIteration(0);
+		break;
+
+	default:
+		CompilerOptions->SetOptimisationEnabled(false);
+		CompilerOptions->SetConstReductionEnabled(true);
+		CompilerOptions->SetOptimisationMaxIteration(1);
+		break;
+	}
+
+	// Minimum resident mip count.
+	const int MinResidentMips = UTexture::GetStaticMinTextureResidentMipCount();
+	// Data smaller than this will always be loaded, as part of the customizable object compiled model.
+	const int MinRomSize = 128;
+	CompilerOptions->SetDataPackingStrategy(MinRomSize, MinResidentMips);
+
+	CompilerOptions->SetTextureLayoutStrategy(bDisableTextureLayout
+		? mu::CompilerOptions::TextureLayoutStrategy::None
+		: mu::CompilerOptions::TextureLayoutStrategy::Pack);
+
+	// \TODO For now force it to be disabled.
+	//CompilerOptions->SetEnableConcurrency(Options.bUseParallelCompilation);
+	CompilerOptions->SetEnableConcurrency(false);
+
+	mu::CompilerPtr Compiler = new mu::Compiler(CompilerOptions);
+
+	UE_LOG(LogMutable, Verbose, TEXT("PROFILE: [ %16.8f ] FCustomizableObjectCompileRunnable Compile start."), FPlatformTime::Seconds());
+	Model = Compiler->Compile(MutableRoot);
+
+	// Dump all the log messages from the compiler
+	mu::Ptr<const mu::ErrorLog> pLog = Compiler->GetLog();
+	for (int i = 0; i < pLog->GetMessageCount(); ++i)
+	{
+		const char* strMessage = pLog->GetMessageText(i);
+		mu::ErrorLogMessageType MessageType = pLog->GetMessageType(i);
+		mu::ErrorLogMessageAttachedDataView MessageAttachedData = pLog->GetMessageAttachedData(i);
+
+		if (MessageType == mu::ELMT_WARNING)
+		{
+			if (MessageAttachedData.m_unassignedUVs && MessageAttachedData.m_unassignedUVsSize > 0) 
+			{			
+				TSharedPtr<FErrorAttachedData> ErrorAttachedData = MakeShared<FErrorAttachedData>();
+				ErrorAttachedData->UnassignedUVs.Reset();
+				ErrorAttachedData->UnassignedUVs.Append(MessageAttachedData.m_unassignedUVs, MessageAttachedData.m_unassignedUVsSize);
+				ArrayWarning.Add(FError(FText::AsCultureInvariant(strMessage), ErrorAttachedData, pLog->GetMessageContext(i)));
+			}
+			else
+			{
+				ArrayWarning.Add(FError(FText::AsCultureInvariant(strMessage), pLog->GetMessageContext(i)));
+			}
+		}
+		else if (MessageType == mu::ELMT_ERROR)
+		{
+			if (MessageAttachedData.m_unassignedUVs && MessageAttachedData.m_unassignedUVsSize > 0) 
+			{			
+				TSharedPtr<FErrorAttachedData> ErrorAttachedData = MakeShared<FErrorAttachedData>();
+				ErrorAttachedData->UnassignedUVs.Reset();
+				ErrorAttachedData->UnassignedUVs.Append(MessageAttachedData.m_unassignedUVs, MessageAttachedData.m_unassignedUVsSize);
+				ArrayError.Add(FError(FText::AsCultureInvariant(strMessage), ErrorAttachedData, pLog->GetMessageContext(i)));
+			}
+			else
+			{
+				ArrayError.Add(FError(FText::AsCultureInvariant(strMessage), pLog->GetMessageContext(i)));
+			}
+		}
+	}
+
+	Compiler = nullptr;
+
+	bThreadCompleted = true;
+
+	UE_LOG(LogMutable, Verbose, TEXT("PROFILE: [ %16.8f ] FCustomizableObjectCompileRunnable::Run end."), FPlatformTime::Seconds());
+
+	return Result;
+}
+
+
+bool FCustomizableObjectCompileRunnable::IsCompleted() const
+{
+	return bThreadCompleted;
+}
+
+
+const TArray<FCustomizableObjectCompileRunnable::FError>& FCustomizableObjectCompileRunnable::GetArrayError() const
+{
+	return ArrayError;
+}
+
+
+const TArray<FCustomizableObjectCompileRunnable::FError>& FCustomizableObjectCompileRunnable::GetArrayWarning() const
+{
+	return ArrayWarning;
+}
+
+
+FCustomizableObjectSaveDDRunnable::FCustomizableObjectSaveDDRunnable(UCustomizableObject* CustomizableObject, const FCompilationOptions& InOptions)
+{
+	Model = CustomizableObject->GetModel();
+	Options = InOptions;
+	
+	CustomizableObjectHeader.InternalVersion = CustomizableObject->GetCurrentSupportedVersion();
+	CustomizableObjectHeader.VersionId = Options.bIsCooking? FGuid::NewGuid() : CustomizableObject->GetVersionId();
+
+	if (!Options.bIsCooking || Options.bSaveCookedDataToDisk)
+	{
+		// We will be saving all compilation data in two separate files, write CO Data
+		FolderPath = CustomizableObject->GetCompiledDataFolderPath(!InOptions.bIsCooking);
+		CompildeDataFullFileName = FolderPath + CustomizableObject->GetCompiledDataFileName(true, InOptions.TargetPlatform);
+		StreamableDataFullFileName = FolderPath + CustomizableObject->GetCompiledDataFileName(false, InOptions.TargetPlatform);
+
+		// Serialize Customizable Object's data
+		FMemoryWriter64 MemoryWriter(Bytes);
+		CustomizableObject->SaveCompiledData(MemoryWriter, Options.bIsCooking);
+	}
+}
+
+
+uint32 FCustomizableObjectSaveDDRunnable::Run()
+{
+	bool bModelSerialized = Model.get() != nullptr;
+
+	if (Options.bIsCooking && !Options.bSaveCookedDataToDisk)
+	{
+		// Serialize mu::Model and streamable resources 
+		FMemoryWriter64 ModelMemoryWriter(Bytes, false, true);
+		FMemoryWriter64 StreamableMemoryWriter(BulkDataBytes, false, true);
+
+		ModelMemoryWriter << bModelSerialized;
+		if (bModelSerialized)
+		{
+			FUnrealMutableModelBulkStreamer Streamer(&ModelMemoryWriter, &StreamableMemoryWriter);
+			mu::Model::Serialise(Model.get(), Streamer);
+		}
+	}
+	else if(bModelSerialized) // Save CO data + mu::Model and streamable resources to disk
+	{
+		// Create folder...
+		IFileManager& FileManager = IFileManager::Get();
+		FileManager.MakeDirectory(*FolderPath, true);
+
+		// Delete files...
+		FileManager.Delete(*CompildeDataFullFileName, true, false, true);
+		FileManager.Delete(*StreamableDataFullFileName, true, false, true);
+
+		// Create file writers...
+		FArchive* ModelMemoryWriter = FileManager.CreateFileWriter(*CompildeDataFullFileName);
+		FArchive* StreamableMemoryWriter = FileManager.CreateFileWriter(*StreamableDataFullFileName);
+
+		// Serailize headers to validate data
+		*ModelMemoryWriter << CustomizableObjectHeader;
+		*StreamableMemoryWriter << CustomizableObjectHeader;
+
+		// Serialize Customizable Object's Data to disk
+		ModelMemoryWriter->Serialize(reinterpret_cast<void*>(Bytes.GetData()), Bytes.Num() * sizeof(uint8));
+		Bytes.Empty();
+
+		// Serialize mu::Model and streamable resources
+		*ModelMemoryWriter << bModelSerialized;
+
+		FUnrealMutableModelBulkStreamer Streamer(ModelMemoryWriter, StreamableMemoryWriter);
+		mu::Model::Serialise(Model.get(), Streamer);
+
+		// Save to disk
+		ModelMemoryWriter->Flush();
+		StreamableMemoryWriter->Flush();
+
+		ModelMemoryWriter->Close();
+		StreamableMemoryWriter->Close();
+
+		delete ModelMemoryWriter;
+		delete StreamableMemoryWriter;
+	}
+
+	bThreadCompleted = true;
+
+	return 1;
+}
+
+
+TArray64<uint8>& FCustomizableObjectSaveDDRunnable::GetModelBytes()
+{
+	return Bytes;
+}
+
+
+TArray64<uint8>& FCustomizableObjectSaveDDRunnable::GetBulkBytes()
+{
+	return BulkDataBytes;
+}
+
+
+bool FCustomizableObjectSaveDDRunnable::IsCompleted() const
+{
+	return bThreadCompleted;
+}
+
+
+const ITargetPlatform* FCustomizableObjectSaveDDRunnable::GetTargetPlatform() const
+{
+	return Options.TargetPlatform;
+}
+
+#undef LOCTEXT_NAMESPACE
diff --git a/Engine/Plugins/Experimental/Mutable/Source/CustomizableObjectEditor/Private/CustomizableObjectCompileRunnable.h b/Engine/Plugins/Experimental/Mutable/Source/CustomizableObjectEditor/Private/CustomizableObjectCompileRunnable.h
new file mode 100644
--- /dev/null
+++ b/Engine/Plugins/Experimental/Mutable/Source/CustomizableObjectEditor/Private/CustomizableObjectCompileRunnable.h
@@ -0,0 +1,128 @@
+// Copyright Epic Games, Inc. All Rights Reserved.
+
+#pragma once
+
+#include "HAL/Runnable.h"
+#include "HAL/RunnableThread.h"
+#include "Logging/TokenizedMessage.h"
+#include "CustomizableObject.h"
+#include "Widgets/Notifications/SNotificationList.h"
+
+#include "MutableTools/Public/Node.h"
+#include "MutableRuntime/Public/Model.h"
+
+
+class UCustomizableObject;
+class UCustomizableObjectNodeObject;
+class CustomizableObjectCompileAsyncTask;
+
+
+class FCustomizableObjectCompileRunnable : public FRunnable
+{
+public:
+
+	struct FErrorAttachedData
+	{
+		TArray<float> UnassignedUVs;
+	};
+
+	struct FError
+	{
+		FText Message;
+		TSharedPtr<FErrorAttachedData> AttachedData;
+		const void* Context;
+
+		FError(const FText& InMessage, const void* InContext) : Message(InMessage), Context(InContext) {}
+		FError(const FText& InMessage, const TSharedPtr<FErrorAttachedData>& InAttachedData, const void* InContext) 
+			: Message(InMessage), AttachedData(InAttachedData), Context(InContext) {}
+	};
+
+private:
+
+	mu::NodePtr MutableRoot;
+	TArray<FError> ArrayWarning;
+	TArray<FError> ArrayError;
+
+public:
+
+	FCustomizableObjectCompileRunnable(mu::NodePtr Root, bool bInDisableTextureLayout);
+
+	// FRunnable interface
+	uint32 Run() override;
+
+	// Own interface
+
+	//
+	bool IsCompleted() const;
+
+	//
+	const TArray<FError>& GetArrayError() const;
+	
+	//
+	const TArray<FError>& GetArrayWarning() const;
+
+public:
+
+	mu::ModelPtr Model;
+
+	// Texture packing strategy
+	bool bDisableTextureLayout;
+
+	FCompilationOptions Options;
+
+	FString ErrorMsg;
+
+	// Whether the thread has finished running
+	bool bThreadCompleted;
+
+	// If Mutable compile is disabled, to immediately finish the Run() method
+	bool MutableIsDisabled;
+};
+
+
+class FCustomizableObjectSaveDDRunnable : public FRunnable
+{
+public:
+
+	FCustomizableObjectSaveDDRunnable(class UCustomizableObject* CustomizableObject, const FCompilationOptions& Options);
+
+	// FRunnable interface
+	uint32 Run() override;
+
+	// Bytes where the model is stored
+	TArray64<uint8>& GetModelBytes();
+
+	// Bytes where the streamed data files are stored
+	TArray64<uint8>& GetBulkBytes();
+
+	//
+	bool IsCompleted() const;
+
+
+	const ITargetPlatform* GetTargetPlatform() const;
+
+	//
+	FCompilationOptions Options;
+
+private:
+
+	MutableCompiledDataStreamHeader CustomizableObjectHeader;
+
+	// Paths used to save files to disk
+	FString FolderPath;
+	FString CompildeDataFullFileName;
+	FString StreamableDataFullFileName;
+
+	mu::ModelPtr Model;
+
+	// Bytes where the model is stored
+	TArray64<uint8> Bytes;
+
+	// Bytes where the streamed data files are stored
+	TArray64<uint8> BulkDataBytes;
+
+	bool bIsCooking = false;
+
+	// Whether the thread has finished running
+	bool bThreadCompleted = false;
+};
\ No newline at end of file
diff --git a/Engine/Plugins/Experimental/Mutable/Source/CustomizableObjectEditor/Private/CustomizableObjectCompiler.cpp b/Engine/Plugins/Experimental/Mutable/Source/CustomizableObjectEditor/Private/CustomizableObjectCompiler.cpp
new file mode 100644
--- /dev/null
+++ b/Engine/Plugins/Experimental/Mutable/Source/CustomizableObjectEditor/Private/CustomizableObjectCompiler.cpp
@@ -0,0 +1,1596 @@
+// Copyright Epic Games, Inc. All Rights Reserved.
+
+#include "CustomizableObjectCompiler.h"
+#include "CustomizableObjectEditor.h"
+#include "CustomizableObjectEditorModule.h"
+#include "CustomizableObjectSystem.h"
+#include "CustomizableObject.h"
+#include "ICustomizableObjectModule.h"
+#include "GraphTraversal.h"
+
+#include "StaticMeshResources.h"
+#include "Framework/Notifications/NotificationManager.h"
+#include "Widgets/Notifications/SNotificationList.h"
+#include "AssetRegistry/AssetRegistryModule.h"
+#include "MessageLogModule.h"
+#include "Logging/MessageLog.h"
+#include "Misc/UObjectToken.h"
+#include "EdGraphToken.h"
+#include "Application/ThrottleManager.h"
+#include "Animation/MorphTarget.h"
+#include "Interfaces/ITargetPlatform.h"
+#include "Interfaces/ITargetPlatformManagerModule.h"
+#include "Engine/SkeletalMesh.h"
+#include "Engine/StaticMesh.h"
+#include "Animation/Skeleton.h"
+#include "Engine/TextureLODSettings.h"
+#include "Rendering/SkinWeightVertexBuffer.h"
+#include "MeshUtilities.h"
+#include "Curves/CurveLinearColor.h"
+#include "PhysicsEngine/PhysicsAsset.h"
+#include "ClothingAsset.h"
+#include "ClothConfig.h"
+
+#include "UnrealEditorPortabilityHelpers.h"
+#include "UnrealPortabilityHelpers.h"
+#include "PlatformInfo.h"
+#include "Engine/AssetManager.h"
+#include "Engine/StreamableManager.h"
+#include "AssetRegistry/ARFilter.h"
+
+#include "MutableTools/Public/Node.h"
+#include "MutableTools/Public/Streams.h"
+#include "MutableRuntime/Public/Serialisation.h"
+#include "MutableRuntime/Public/MutableTrace.h"
+#include "MutableRuntime/Private/ImagePrivate.h"
+
+#include "ICustomizableObjectPopulationModule.h"
+
+#include "GenerateMutableSource/GenerateMutableSourceImage.h"
+
+#define LOCTEXT_NAMESPACE "CustomizableObjectEditor"
+
+
+UCustomizableObjectNodeObject* GetRootNode(UCustomizableObject* Object, bool &bOutMultipleBaseObjectsFound);
+
+
+
+FCustomizableObjectCompiler::FCustomizableObjectCompiler() : FCustomizableObjectCompilerBase()
+	, PendingTexturesToLoad(false)
+	, CompilationLaunchPending(false)
+	, PreloadingReferencerAssets(false)
+	, CurrentGAsyncLoadingTimeLimit(-1.0f)
+	, CompletedUnrealToMutableTask(0)
+	, MaxConvertToMutableTextureTime(0.2f)
+{
+	bAreExtraBoneInfluencesEnabled = ICustomizableObjectModule::Get().AreExtraBoneInfluencesEnabled();
+}
+
+
+
+bool FCustomizableObjectCompiler::Tick()
+{
+	bool bUpdated = false;
+
+	if (CompileTask.IsValid() && CompileTask->IsCompleted())
+	{
+		UE_LOG(LogMutable, Log, TEXT("PROFILE: [ %16.8f ] Finishing Compilation task."), FPlatformTime::Seconds());
+
+		FinishCompilation();
+
+		if (SaveDDTask.IsValid())
+		{
+			SaveCODerivedData(true);
+		}
+		else
+		{
+			bUpdated = true;
+			State = ECustomizableObjectCompilationState::Completed;
+
+			RemoveCompileNotification();
+
+			NotifyCompilationErrors();
+		}
+		
+		UE_LOG(LogMutable, Verbose, TEXT("PROFILE: [ %16.8f ] Finished Compilation task."), FPlatformTime::Seconds());
+		UE_LOG(LogMutable, Verbose, TEXT("PROFILE: -----------------------------------------------------------"));
+
+		CleanCachedReferencers();
+		UpdateArrayGCProtect();
+
+		CompilationLogsContainer.ClearMessagesArray();
+
+
+	}
+
+	if (SaveDDTask.IsValid() && SaveDDTask->IsCompleted())
+	{
+		State = ECustomizableObjectCompilationState::Completed;
+		bUpdated = true;
+
+		FinishSavingDerivedData();
+		CurrentObject->PostCompileDelegate.Broadcast();
+
+		UE_LOG(LogMutable, Verbose, TEXT("PROFILE: [ %16.8f ] Finished Saving Derived Data task."), FPlatformTime::Seconds());
+		UE_LOG(LogMutable, Verbose, TEXT("PROFILE: -----------------------------------------------------------"));
+	}
+
+	// In editor, when compiling a CO, referencer assets and Unreal to Mutable texture conversion are performed asynchronously
+	UpdatePendingTextureConversion(true);
+
+	if (!PendingTexturesToLoad && CompilationLaunchPending)
+	{
+		CompilationLaunchPending = false;
+		LaunchMutableCompile(true);
+	}
+
+	return bUpdated;
+}
+
+void FCustomizableObjectCompiler::MutableIsDisabledCase(UCustomizableObject* Object)
+{
+	if (Object->IsLocked())
+	{
+		return;
+	}
+
+	if (!Object->IsLocked())
+	{
+		UCustomizableObjectSystem::GetInstance()->LockObject(Object);
+	}
+
+	CurrentObject = Object;
+
+	if (CompileTask.IsValid()) // Don't start compilation if there's a compilation running
+	{
+		return;
+	}
+
+	UE_LOG(LogMutable, Warning, TEXT("Mutable has been disabled. To reenable it, please deactivate the Disable Mutale Option in the Plugins -> Mutable option"), FPlatformTime::Seconds());
+	CompileTask = MakeShareable(new FCustomizableObjectCompileRunnable(nullptr, true));
+	CompileTask->MutableIsDisabled = true;
+	LaunchMutableCompile(false);
+}
+
+
+bool FCustomizableObjectCompiler::IsRootObject(const UCustomizableObject* Object) const
+{
+	// Look for the base object node
+	UCustomizableObjectNodeObject* Root = nullptr;
+	TArray<UCustomizableObjectNodeObject*> ObjectNodes;
+	if (!Object->Source || !Object->Source->Nodes.Num())
+	{
+		// Conservative approach.
+		return true;
+	}
+
+	Object->Source->GetNodesOfClass<UCustomizableObjectNodeObject>(ObjectNodes);
+
+	for (TArray<UCustomizableObjectNodeObject*>::TIterator It(ObjectNodes); It; ++It)
+	{
+		if ((*It)->bIsBase)
+		{
+			Root = *It;
+		}
+	}
+
+	return Root && !Root->ParentObject;
+}
+
+
+void FCustomizableObjectCompiler::PreloadingReferencerAssetsCallback(UCustomizableObject* Object, FCustomizableObjectCompiler* CustomizableObjectCompiler, const FCompilationOptions Options, bool bAsync)
+{
+	UE_LOG(LogMutable, Log, TEXT("PROFILE: [ %16.8f ] Preload asynchronously assets end."), FPlatformTime::Seconds());
+
+	if (CustomizableObjectCompiler->GetCurrentGAsyncLoadingTimeLimit() != -1.0f)
+	{
+		static IConsoleVariable* CVAsyncLoadingTimeLimit = IConsoleManager::Get().FindConsoleVariable(TEXT("s.AsyncLoadingTimeLimit"));
+		if (CVAsyncLoadingTimeLimit)
+		{
+			CVAsyncLoadingTimeLimit->Set(CustomizableObjectCompiler->GetCurrentGAsyncLoadingTimeLimit());
+		}
+		CustomizableObjectCompiler->SetCurrentGAsyncLoadingTimeLimit(-1.0f);
+	}
+
+	CustomizableObjectCompiler->UpdateArrayGCProtect();
+	CustomizableObjectCompiler->SetPreloadingReferencerAssets(false);
+	UCustomizableObjectSystem::GetInstance()->UnlockObject(Object);
+
+	CustomizableObjectCompiler->CompileInternal(Object, Options, bAsync);
+}
+
+
+void FCustomizableObjectCompiler::Compile(UCustomizableObject& Object, const FCompilationOptions& InOptions, bool bAsync)
+{
+	if (UCustomizableObjectSystem::GetInstance()->IsCompilationDisabled())
+	{
+		CompileInternal(&Object, InOptions, bAsync);
+		return;
+	}
+
+	UE_LOG(LogMutable, Log, TEXT("PROFILE: [ %16.8f ] Preload asynchronously assets start."), FPlatformTime::Seconds());
+
+	if (PreloadingReferencerAssets || CompilationLaunchPending || (Object.IsLocked()))
+	{
+		FString Message = FString::Printf(TEXT("Customizable Object %s is already being compiled or updated. Please wait a few seconds and try again."), *Object.GetName());
+		UE_LOG(LogMutable, Warning, TEXT("%s"), *Message);
+		FNotificationInfo Info(LOCTEXT("CustomizableObjectBeingCompilerOrUpdated", "Customizable object compile and update still in process"));
+		Info.bFireAndForget = true;
+		Info.bUseThrobber = true;
+		Info.FadeOutDuration = 1.0f;
+		Info.ExpireDuration = 1.0f;
+		FSlateNotificationManager::Get().AddNotification(Info);
+		return;
+	}
+
+	// Lock object during asynchronous asset loading to avoid instance updates
+	bool LockResult = UCustomizableObjectSystem::GetInstance()->LockObject(&Object);
+
+	if (!LockResult)
+	{
+		FString Message = FString::Printf(TEXT("Customizable Object %s is already being compiled or updated. Please wait a few seconds and try again."), *Object.GetName());
+		UE_LOG(LogMutable, Warning, TEXT("%s"), *Message);
+		FNotificationInfo Info(LOCTEXT("CustomizableObjectBeingCompilerOrUpdated", "Customizable object compile and update still in process"));
+		Info.bFireAndForget = true;
+		Info.bUseThrobber = true;
+		Info.FadeOutDuration = 1.0f;
+		Info.ExpireDuration = 1.0f;
+		FSlateNotificationManager::Get().AddNotification(Info);
+		return;
+	}
+
+	CurrentObject = &Object;
+
+	PreloadingReferencerAssets = true;
+
+	CleanCachedReferencers();
+	UpdateArrayGCProtect();
+	TArray<FName> ArrayReferenceNames;
+	AddCachedReferencers(*Object.GetOuter()->GetPathName(), ArrayReferenceNames);
+
+	TArray<FSoftObjectPath> ArrayAssetToStream;
+
+	for (FAssetData& Element : ArrayAssetData)
+	{
+		ArrayAssetToStream.Add(FSoftObjectPath(Element.ObjectPath.ToString()));
+	}
+
+	MaxConvertToMutableTextureTime = ComputeAsyncLoadingTimeLimit() * 0.001f;
+
+	bool bAssetsLoaded = true;
+
+	if (ArrayAssetToStream.Num() > 0)
+	{
+		FStreamableManager& Streamable = UCustomizableObjectSystem::GetInstance()->GetStreamableManager();
+
+		if (bAsync) 
+		{
+			AddCompileNotification(LOCTEXT("LoadingReferencerAssets", "Loading assets"));
+
+			static IConsoleVariable* CVAsyncLoadingTimeLimit = IConsoleManager::Get().FindConsoleVariable(TEXT("s.AsyncLoadingTimeLimit"));
+			if (CVAsyncLoadingTimeLimit)
+			{
+				CurrentGAsyncLoadingTimeLimit = CVAsyncLoadingTimeLimit->GetFloat();
+				CVAsyncLoadingTimeLimit->Set(MaxConvertToMutableTextureTime * 1000.0f);
+			}
+
+			AsynchronousStreamableHandlePtr = Streamable.RequestAsyncLoad(ArrayAssetToStream, FStreamableDelegate::CreateStatic(PreloadingReferencerAssetsCallback, &Object, this, InOptions, bAsync));
+			bAssetsLoaded = false;
+		}
+		else
+		{
+			Streamable.RequestSyncLoad(ArrayAssetToStream);
+			bAssetsLoaded = true;
+		}
+	}
+
+	if (bAssetsLoaded)
+	{
+		PreloadingReferencerAssetsCallback(&Object, this, InOptions, bAsync);
+	}
+}
+
+
+void FCustomizableObjectCompiler::AddReferencedObjects(FReferenceCollector& Collector)
+{
+	// While compilation takes place, no COs involved can be garbage-collected
+	const int32 MaxIndex = ArrayGCProtect.Num();
+
+	for (int32 i = 0; i < MaxIndex; ++i)
+	{
+		Collector.AddReferencedObject(ArrayGCProtect[i]);
+	}
+}
+
+
+void FCustomizableObjectCompiler::UpdateArrayGCProtect()
+{
+	check(IsInGameThread());
+
+	const int32 MaxIndex = ArrayAssetData.Num();
+	ArrayGCProtect.SetNum(MaxIndex);
+
+	for (int i = 0; i < MaxIndex; ++i)
+	{
+		ArrayGCProtect[i] = Cast<UCustomizableObject>(ArrayAssetData[i].GetAsset());
+	}
+}
+
+
+void FCustomizableObjectCompiler::ProcessChildObjectsRecursively(UCustomizableObject* Object, FAssetRegistryModule& AssetRegistryModule, FMutableGraphGenerationContext &GenerationContext)
+{
+	TArray<FName> ArrayReferenceNames;
+	AddCachedReferencers(*Object->GetOuter()->GetPathName(), ArrayReferenceNames);
+	UpdateArrayGCProtect();
+
+	bool bMultipleBaseObjectsFound = false;
+
+	FAssetData* AssetData;
+	UCustomizableObject* ChildObject;
+
+	for (const FName& ReferenceName : ArrayReferenceNames)
+	{
+		AssetData = GetCachedAssetData(ReferenceName.ToString());
+
+		if (AssetData != nullptr) // Elements in ArrayAssetData are already of static class UCustomizableObject
+		{
+			ChildObject = Cast<UCustomizableObject>(AssetData->GetAsset());
+			
+			if (!ChildObject)
+				continue;
+
+			if (ChildObject != Object && !ChildObject->HasAnyFlags(RF_Transient))
+			{
+				UCustomizableObjectNodeObject* ChildRoot = GetRootNode(ChildObject, bMultipleBaseObjectsFound);
+
+				if (ChildRoot && !bMultipleBaseObjectsFound)
+				{
+					if (ChildRoot->ParentObject == Object)
+					{
+						if (const FGroupNodeIdsTempData* GroupGuid = GenerationContext.DuplicatedGroupNodeIds.FindPair(Object, FGroupNodeIdsTempData(ChildRoot->ParentObjectGroupId)))
+						{
+							ChildRoot->ParentObjectGroupId = GroupGuid->NewGroupNodeId;
+						}
+
+						GenerationContext.GroupIdToExternalNodeMap.Add(ChildRoot->ParentObjectGroupId, ChildRoot);
+						GenerationContext.CustomizableObjectGuidsInCompilation.Add(ChildObject->GetVersionId());
+					}
+				}
+			}
+
+			TArray<UCustomizableObjectNodeObjectGroup*> GroupNodes;
+			ChildObject->Source->GetNodesOfClass<UCustomizableObjectNodeObjectGroup>(GroupNodes);
+
+			if (GroupNodes.Num() > 0) // Only grafs with group nodes should have child grafs
+			{
+				if (ArrayAlreadyProcessedChild.Find(ReferenceName) == INDEX_NONE)
+				{
+					for (int32 i = 0; i < GroupNodes.Num(); ++i)
+					{
+						const FGuid NodeId = GenerationContext.GetNodeIdUnique(GroupNodes[i]);
+						if (NodeId != GroupNodes[i]->NodeGuid)
+						{
+							GenerationContext.DuplicatedGroupNodeIds.Add(ChildObject, FGroupNodeIdsTempData(GroupNodes[i]->NodeGuid, NodeId));
+							GroupNodes[i]->NodeGuid = NodeId;
+						}
+					}
+
+					ArrayAlreadyProcessedChild.Add(ReferenceName);
+					ProcessChildObjectsRecursively(ChildObject, AssetRegistryModule, GenerationContext);
+				}
+			}
+		}
+	}
+}
+
+
+void FCustomizableObjectCompiler::DisplayParameterWarning(FMutableGraphGenerationContext& GenerationContext)
+{
+	for (const TPair<FString, TArray<const UCustomizableObjectNode*>>& It : GenerationContext.ParameterNamesMap)
+	{
+		if (It.Key == "")
+		{
+			FText MessageWarning = LOCTEXT("NodeWithNoName", ". There is at least one node with no name.");
+			CompilerLog(MessageWarning, It.Value, EMessageSeverity::Warning, true);
+		}
+		else if (It.Value.Num() > 1)
+		{
+			FText MessageWarning = FText::Format(LOCTEXT("NodeWithRepeatedName", ". Several nodes have repeated name \"{0}\""), FText::FromString(It.Key));
+			CompilerLog(MessageWarning, It.Value, EMessageSeverity::Warning, true);
+		}
+	}
+}
+
+
+void FCustomizableObjectCompiler::DisplayDuplicatedNodeIdsWarning(FMutableGraphGenerationContext & GenerationContext)
+{
+	for (const TPair<FGuid, TArray<const UCustomizableObjectNode*>>& It : GenerationContext.NodeIdsMap)
+	{
+		if (It.Value.Num() > 1)
+		{
+			FText MessageWarning = LOCTEXT("NodeWithRepeatedIds", ". Several nodes have repeated NodeIds, reconstruct the nodes.");
+			CompilerLog(MessageWarning, It.Value, EMessageSeverity::Warning, true);
+		}
+	}
+}
+
+
+void FCustomizableObjectCompiler::DisplayUnnamedNodeObjectWarning(FMutableGraphGenerationContext& GenerationContext)
+{
+	FText Message = LOCTEXT("Unnamed Node Object", "Unnamed Node Object");
+	for (const UCustomizableObjectNode* It : GenerationContext.NoNameNodeObjectArray)
+	{
+		CompilerLog(Message, It, EMessageSeverity::Warning, true);
+	}
+}
+
+
+void FCustomizableObjectCompiler::DisplayDiscardedPhysicsAssetSingleWarning(FMutableGraphGenerationContext& GenerationContext)
+{
+	if (GenerationContext.DiscartedPhysicsAssetMap.Num())
+	{
+		FString PhysicsAssetName;
+		for (int32 ComponentIndex = 0; ComponentIndex < GenerationContext.NumMeshComponentsInRoot; ++ComponentIndex)
+		{
+			USkeletalMesh* RefSkeletalMesh = GenerationContext.ComponentInfos.IsValidIndex(ComponentIndex) ? GenerationContext.ComponentInfos[ComponentIndex].RefSkeletalMesh : nullptr;
+			if (RefSkeletalMesh && GenerationContext.DiscartedPhysicsAssetMap.Find(RefSkeletalMesh->GetPhysicsAsset()))
+			{
+				PhysicsAssetName = RefSkeletalMesh->GetPhysicsAsset()->GetName();
+				break;
+			}
+		}
+		
+		if(PhysicsAssetName.IsEmpty())
+		{
+			PhysicsAssetName = (*GenerationContext.DiscartedPhysicsAssetMap.FindKey(0))->GetName();
+		}
+
+		FText Message = FText::Format( LOCTEXT("Discarted PhysicsAssets", "{0} and {1} other PhysicsAssets have been discarted because one or more Bodies have no corresponding bones in the SkeletalMesh.\nFor more infromation on the assets that need attention check the Output log."),
+										FText::FromString(PhysicsAssetName), FText::AsNumber(GenerationContext.DiscartedPhysicsAssetMap.Num()));
+		
+		CompilerLog(Message, nullptr, EMessageSeverity::Warning, false);
+	}
+}
+
+
+void FCustomizableObjectCompiler::DisplayOrphanNodesWarning(FMutableGraphGenerationContext& GenerationContext)
+{
+	for (const TPair<FGeneratedKey, FGeneratedData>& It : GenerationContext.Generated)
+	{
+		if (const UCustomizableObjectNode* Node = Cast<UCustomizableObjectNode>(It.Value.Source))
+		{
+			if (Node->GetAllOrphanPins().Num() > 0)
+			{
+				CompilerLog(LOCTEXT("OrphanPinsWarningCompiler", "Node contains deprecated pins"), Node, EMessageSeverity::Warning, false);
+			}
+		}
+	}
+}
+
+
+UCustomizableObject* FCustomizableObjectCompiler::GetRootObject( UCustomizableObject* InObject)
+{
+	// Grab a node to start the search -> Get the root since it should be always present
+	bool bMultipleBaseObjectsFound = false;
+	UCustomizableObjectNodeObject* ObjectRootNode = GetRootNode(InObject, bMultipleBaseObjectsFound);
+	
+	if (ObjectRootNode->ParentObject)
+	{
+		TArray<UCustomizableObject*> VisitedNodes;
+		return GetFullGraphRootObject(ObjectRootNode,VisitedNodes);
+	}
+	else
+	{
+		// No parent object found, return input as the parent of the graph
+		return InObject;
+	}
+}
+
+
+mu::NodeObjectPtr FCustomizableObjectCompiler::GenerateMutableRoot( 
+	UCustomizableObject* Object, 
+	FMutableGraphGenerationContext& GenerationContext, 
+	FText& ErrorMsg, 
+	bool& bOutIsRootObject)
+{
+	if (!Object->Source)
+	{
+		ErrorMsg = LOCTEXT("NoSource", "Object with no valid graph found. Object not build.");
+
+		if (IsRunningCookCommandlet() || IsRunningCookOnTheFly())
+		{
+			UE_LOG(LogMutable, Warning, TEXT("Compilation failed! Missing EDITORONLY data for Customizable Object [%s]. The object might have been loaded outside the Cooking context."), *Object->GetName());
+		}
+
+		return nullptr;
+	}
+
+	bool bMultipleBaseObjectsFound;
+	bOutIsRootObject = false;
+	UCustomizableObjectNodeObject* Root = GetRootNode(Object, bMultipleBaseObjectsFound);
+
+	if (bMultipleBaseObjectsFound)
+	{
+		ErrorMsg = LOCTEXT("MutlipleBase","Multiple base object nodes found. Only one will be used.");
+		return nullptr;
+	}
+
+	if (!Root)
+	{
+		ErrorMsg = LOCTEXT("NoBase","No base object node found. Object not built.");
+		return nullptr;
+	}
+
+	bOutIsRootObject = Root->ParentObject == nullptr;
+
+	UCustomizableObjectNodeObject* ActualRoot = Root;
+	UCustomizableObject* ActualRootObject = Object;
+
+	GenerationContext.bDisableTextureLayoutManagementFlag = Object->bDisableTextureLayoutManagement;
+	ArrayAlreadyProcessedChild.Empty();
+
+	if (Root->ObjectName.IsEmpty())
+	{
+		GenerationContext.NoNameNodeObjectArray.AddUnique(Root);
+	}
+
+	if ((Object->MeshCompileType == EMutableCompileMeshType::Full) || Options.bIsCooking)
+	{
+		if (Root->ParentObject!=nullptr && Options.bIsCooking)
+		{
+			// This happens while packaging.
+			return nullptr;
+		}
+
+		// We cannot load while saving. This should only happen in cooking and all assets should have been preloaded.
+		if (!GIsSavingPackage)
+		{
+			UE_LOG(LogMutable, Log, TEXT("PROFILE: [ %16.8f ] Begin search for children."), FPlatformTime::Seconds());
+
+			TArray<UCustomizableObject*> VisitedObjects;
+			ActualRootObject = Root->ParentObject ? GetFullGraphRootObject(Root, VisitedObjects) : Object;
+			GenerationContext.bDisableTextureLayoutManagementFlag = ActualRootObject->bDisableTextureLayoutManagement;
+
+			if (Root->ParentObject != nullptr)
+			{
+				VisitedObjects.Empty();
+				ActualRoot = GetFullGraphRootNodeObject(ActualRoot, VisitedObjects);
+			}
+
+			// The object doesn't reference a root object but is a root object, look for all the objects that reference it and get their root nodes
+			FAssetRegistryModule& AssetRegistryModule = FModuleManager::LoadModuleChecked<FAssetRegistryModule>("AssetRegistry");
+			ProcessChildObjectsRecursively(ActualRootObject, AssetRegistryModule, GenerationContext);
+			UE_LOG(LogMutable, Log, TEXT("PROFILE: [ %16.8f ] End search for children."), FPlatformTime::Seconds());
+		}
+	}
+	else
+	{
+		// Local, local with children and working set modes: add parents until whole CO graph root
+		TArray<UCustomizableObjectNodeObject*> ArrayNodeObject;
+		TArray<UCustomizableObject*> ArrayCustomizableObject;
+		
+		if (!GetParentsUntilRoot(Object, ArrayNodeObject, ArrayCustomizableObject))
+		{
+			CompilerLog(LOCTEXT("SkeletalMeshCycleFound", "Error! Cycle detected in the Customizable Object hierarchy."), Root);
+			return nullptr;
+		}
+
+		if ((Object->MeshCompileType == EMutableCompileMeshType::AddWorkingSetNoChildren) ||
+			(Object->MeshCompileType == EMutableCompileMeshType::AddWorkingSetAndChildren))
+		{
+			const int32 MaxIndex = Object->WorkingSet.Num();
+			for (int32 i = 0; i < MaxIndex; ++i)
+			{
+				if (Object->WorkingSet[i] != nullptr)
+				{
+					ArrayCustomizableObject.Reset();
+
+					if (!GetParentsUntilRoot(Object->WorkingSet[i].LoadSynchronous(), ArrayNodeObject, ArrayCustomizableObject))
+					{
+						CompilerLog(LOCTEXT("NoReferenceMesh", "Error! Cycle detected in the Customizable Object hierarchy."), Root);
+						return nullptr;
+					}
+				}
+			}
+		}
+
+		if ((Object->MeshCompileType == EMutableCompileMeshType::LocalAndChildren) ||
+			(Object->MeshCompileType == EMutableCompileMeshType::AddWorkingSetAndChildren))
+		{
+			TArray<UCustomizableObjectNodeObjectGroup*> GroupNodes;
+			Object->Source->GetNodesOfClass<UCustomizableObjectNodeObjectGroup>(GroupNodes);
+
+			if (GroupNodes.Num() > 0) // Only grafs with group nodes should have child grafs
+			{
+				FAssetRegistryModule& AssetRegistryModule = FModuleManager::LoadModuleChecked<FAssetRegistryModule>("AssetRegistry");
+				ProcessChildObjectsRecursively(Object, AssetRegistryModule, GenerationContext);
+			}
+		}
+
+		for (int32 i = 0; i < ArrayNodeObject.Num(); ++i)
+		{
+			if (GenerationContext.GroupIdToExternalNodeMap.FindKey(ArrayNodeObject[i]) == nullptr)
+			{
+				GenerationContext.GroupIdToExternalNodeMap.Add(ArrayNodeObject[i]->ParentObjectGroupId, ArrayNodeObject[i]);
+			}
+		}
+
+		TArray<UCustomizableObject*> VisitedObjects;
+
+		if (Root->ParentObject != nullptr)
+		{
+			ActualRoot = GetFullGraphRootNodeObject(ActualRoot, VisitedObjects);
+		}
+
+		VisitedObjects.Empty();
+		ActualRootObject = Root->ParentObject ? GetFullGraphRootObject(Root, VisitedObjects) : Object;
+		GenerationContext.bDisableTextureLayoutManagementFlag = ActualRootObject->bDisableTextureLayoutManagement;
+
+	}
+
+	// Make sure we have a valid Reference SkeletalMesh and Skeleton for each component
+	for (int32 ComponentIndex = 0; ComponentIndex < ActualRoot->NumMeshComponents; ++ComponentIndex)
+	{
+		USkeletalMesh* RefSkeletalMesh = ActualRootObject->GetRefSkeletalMesh(ComponentIndex);
+		if (!RefSkeletalMesh)
+		{
+			CompilerLog(LOCTEXT("NoReferenceMesh", "Error! Missing reference mesh in the Object Properties Tab"), ActualRoot);
+			return nullptr;
+		}
+
+		USkeleton* RefSkeleton = RefSkeletalMesh->GetSkeleton();
+		if(!RefSkeleton)
+		{
+			FText Msg = FText::Format(LOCTEXT("NoReferenceSkeleton", "Error! Missing skeleton in the reference mesh [{0}]"),
+				FText::FromString(GenerationContext.CustomizableObjectWithCycle->GetPathName()));
+
+			CompilerLog(Msg, ActualRoot);
+			return nullptr;
+		}
+
+		// Add a new entry to the list of Component Infos 
+		GenerationContext.ComponentInfos.Add(RefSkeletalMesh);
+
+		// Make sure the Skeleton from the reference mesh is added to the list of referenced Skeletons.
+		GenerationContext.ReferencedSkeletons.Add(RefSkeleton);
+	}
+
+	Object->ReferenceSkeletalMeshes = ActualRootObject->ReferenceSkeletalMeshes;
+
+    GenerationContext.RealTimeMorphTargetsOverrides = ActualRoot->RealTimeMorphSelectionOverrides;
+    GenerationContext.RealTimeMorphTargetsOverrides.Reset();
+
+	// Generate the object expression
+	UE_LOG(LogMutable, Log, TEXT("PROFILE: [ %16.8f ] GenerateMutableSource start."), FPlatformTime::Seconds());
+	mu::NodeObjectPtr MutableRoot = GenerateMutableSource(ActualRoot->OutputPin(), GenerationContext, !bOutIsRootObject);
+	UE_LOG(LogMutable, Log, TEXT("PROFILE: [ %16.8f ] GenerateMutableSource end."), FPlatformTime::Seconds());
+
+    ActualRoot->RealTimeMorphSelectionOverrides = GenerationContext.RealTimeMorphTargetsOverrides;
+	GenerationContext.GenerateClippingCOInternalTags();
+
+	// Generate ReferenceSkeletalMeshes data;
+	PopulateReferenceSkeletalMeshesData(GenerationContext);
+	
+	for (const USkeletalMesh* RefSkeletalMesh : Object->ReferenceSkeletalMeshes)
+	{
+		GenerationContext.CheckPhysicsAssetInSkeletalMesh(RefSkeletalMesh);
+	}
+
+	DisplayParameterWarning( GenerationContext );
+	DisplayUnnamedNodeObjectWarning( GenerationContext );
+	DisplayDuplicatedNodeIdsWarning( GenerationContext );
+	DisplayDiscardedPhysicsAssetSingleWarning( GenerationContext );
+	DisplayOrphanNodesWarning( GenerationContext );
+
+	if (GenerationContext.CustomizableObjectWithCycle)
+	{
+		ErrorMsg = FText::Format(LOCTEXT("CycleDetected","Cycle detected in graph of CustomizableObject {0}. Object not built."),
+			FText::FromString(GenerationContext.CustomizableObjectWithCycle->GetPathName()));
+
+		return nullptr;
+	}
+
+	// Apply the platform specific transform
+	// Format and mips are done per-image now.
+	//MutableRoot = GenerateMutableTransform(MutableRoot, GenerationContext);
+
+	return MutableRoot;
+}
+
+
+void FCustomizableObjectCompiler::AddCachedReferencers(const FName& PathName, TArray<FName>& ArrayReferenceNames)
+{
+	ArrayReferenceNames.Empty();
+	FAssetRegistryModule& AssetRegistryModule = FModuleManager::LoadModuleChecked<FAssetRegistryModule>("AssetRegistry");
+	AssetRegistryModule.Get().GetReferencers(PathName, ArrayReferenceNames, UE::AssetRegistry::EDependencyCategory::Package, UE::AssetRegistry::EDependencyQuery::Hard);
+
+	FARFilter Filter;
+	for (const FName& ReferenceName : ArrayReferenceNames)
+	{
+		if (!IsCachedInAssetData(ReferenceName.ToString()) && !ReferenceName.ToString().StartsWith(TEXT("/TempAutosave")))
+		{
+			Filter.PackageNames.Add(ReferenceName);
+		}
+	}
+
+	Filter.bIncludeOnlyOnDiskAssets = false;
+
+	TArray<FAssetData> ArrayAssetDataTemp;
+	AssetRegistryModule.Get().GetAssets(Filter, ArrayAssetDataTemp);
+
+	// Store only those which have static class type Customizable Object, to avoid loading not needed elements
+	const int32 MaxIndex = ArrayAssetDataTemp.Num();
+	for (int32 i = 0; i < MaxIndex; ++i)
+	{
+		if (ArrayAssetDataTemp[i].GetClass() == UCustomizableObject::StaticClass())
+		{
+			ArrayAssetData.Add(ArrayAssetDataTemp[i]);
+		}
+	}
+}
+
+
+void FCustomizableObjectCompiler::LaunchMutableCompile(bool ShowNotification)
+{
+	if(!Options.bSilentCompilation && ShowNotification)
+	{
+		AddCompileNotification(LOCTEXT("CustomizableObjectCompileInProgress", "Compiling"));
+	}
+
+	// Even for async build, we spawn a thread, so that we can set a large stack. 
+	// Thread names need to be unique, apparently.
+	static int ThreadCount = 0;
+	FString ThreadName = FString::Printf(TEXT("MutableCompile-%03d"), ++ThreadCount);
+	CompileThread = MakeShareable(FRunnableThread::Create(CompileTask.Get(), *ThreadName, 16 * 1024 * 1024, TPri_Normal));
+}
+
+
+void FCustomizableObjectCompiler::SaveCODerivedData(bool ShowNotification)
+{
+	if (!SaveDDTask.IsValid())
+	{
+		return;
+	}
+
+	if (!Options.bSilentCompilation && ShowNotification)
+	{
+		AddCompileNotification(LOCTEXT("SavingCustomizableObjectDerivedData", "Saving Data"));
+	}
+
+	// Even for async saving derived data.
+	static int SDDThreadCount = 0;
+	FString ThreadName = FString::Printf(TEXT("MutableSDD-%03d"), ++SDDThreadCount);
+	SaveDDThread = MakeShareable(FRunnableThread::Create(SaveDDTask.Get(), *ThreadName));
+}
+
+
+void FCustomizableObjectCompiler::CleanCachedReferencers()
+{
+	ArrayAssetData.Empty();
+}
+
+
+bool FCustomizableObjectCompiler::IsCachedInAssetData(const FString& PackageName)
+{
+	const int32 MaxIndex = ArrayAssetData.Num();
+
+	for (int32 i = 0; i < MaxIndex; ++i)
+	{
+		if (ArrayAssetData[i].PackageName.ToString() == PackageName)
+		{
+			return true;
+		}
+	}
+
+	return false;
+}
+
+
+FAssetData* FCustomizableObjectCompiler::GetCachedAssetData(const FString& PackageName)
+{
+	const int32 MaxIndex = ArrayAssetData.Num();
+
+	for (int32 i = 0; i < MaxIndex; ++i)
+	{
+		if (ArrayAssetData[i].PackageName.ToString() == PackageName)
+		{
+			return &ArrayAssetData[i];
+		}
+	}
+
+	return nullptr;
+}
+
+
+void FCustomizableObjectCompiler::UpdatePendingTextureConversion(bool UseTimeLimit)
+{
+	// In editor, when compiling a CO, referencer assets and Unreal to Mutable texture conversion are performed asynchronously
+	if (PendingTexturesToLoad)
+	{
+		float initialTime = FPlatformTime::Seconds();
+		while ( CompletedUnrealToMutableTask < ArrayTextureUnrealToMutableTask.Num() )
+		{
+			UTexture2D* Texture = ArrayTextureUnrealToMutableTask[CompletedUnrealToMutableTask].Texture;
+
+			// If the texture has been set to null, it means it was a duplicate that we already processed.
+			if (!Texture)
+			{ 
+				CompletedUnrealToMutableTask++;
+				continue;
+			}
+
+			// Convert the texture
+			mu::ImagePtr Image = ConvertTextureUnrealToMutable( 
+				Texture, 
+				ArrayTextureUnrealToMutableTask[CompletedUnrealToMutableTask].Node,
+				this,
+				ArrayTextureUnrealToMutableTask[CompletedUnrealToMutableTask].bIsNormalComposite );
+
+			// Assign to all tasks referring to the same texture
+			for (int32 j = CompletedUnrealToMutableTask; j < ArrayTextureUnrealToMutableTask.Num(); ++j)
+			{
+				if (ArrayTextureUnrealToMutableTask[j].Texture == Texture)
+				{
+					if (ArrayTextureUnrealToMutableTask[j].ImageNode.get())
+					{
+						ArrayTextureUnrealToMutableTask[j].ImageNode->SetValue(Image.get());
+						ArrayTextureUnrealToMutableTask[j].Texture = nullptr;
+					}
+
+					else if (ArrayTextureUnrealToMutableTask[j].TableNode.get())
+					{
+						int32 ColumnIndx = ArrayTextureUnrealToMutableTask[j].TableColumn;
+						int32 RowIndx = ArrayTextureUnrealToMutableTask[j].TableRow;
+
+						ArrayTextureUnrealToMutableTask[j].TableNode->SetCell(ColumnIndx, RowIndx, Image.get());
+
+						ArrayTextureUnrealToMutableTask[j].Texture = nullptr;
+					}
+				}
+			}
+
+			CompletedUnrealToMutableTask++;
+
+			if (UseTimeLimit && ((FPlatformTime::Seconds() - initialTime) > MaxConvertToMutableTextureTime))
+			{
+				break;
+			}
+		}
+
+		if (CompletedUnrealToMutableTask >= ArrayTextureUnrealToMutableTask.Num())
+		{
+			PendingTexturesToLoad = false;
+			ArrayTextureUnrealToMutableTask.Empty();
+			CompletedUnrealToMutableTask = 0;
+		}
+	}
+}
+
+
+float FCustomizableObjectCompiler::ComputeAsyncLoadingTimeLimit()
+{
+	float DeltaTimeSeconds = FApp::GetDeltaTime();
+
+	if (DeltaTimeSeconds == 0.0f)
+	{
+		return 100.0f;
+	}
+
+	// AsyncLoadingTimeLimit will depend on the speed of the computer, using delta time as reference
+	float Value = 2.0f / DeltaTimeSeconds;
+	Value = FMath::Clamp(Value, 75.0f, 250.0f);
+	return Value;
+}
+
+void FCustomizableObjectCompiler::GetCompilationMessages(TArray<FText>& OutWarningMessages, TArray<FText>& OutErrorMessages) const
+{
+	CompilationLogsContainer.GetMessages(OutWarningMessages,OutErrorMessages);
+}
+
+
+void FCustomizableObjectCompiler::CompileInternal(UCustomizableObject* Object, const FCompilationOptions& InOptions, bool bAsync)
+{
+	MUTABLE_CPUPROFILER_SCOPE(FCustomizableObjectCompiler::Compile)
+	
+	State = ECustomizableObjectCompilationState::Failed;
+
+	if (!Object) return;
+
+	if(UCustomizableObjectSystem::GetInstance()->IsCompilationDisabled())
+	{
+		MutableIsDisabledCase(Object);
+		return;
+	}
+
+	if (bAsync && CompileTask.IsValid()) // Don't start compilation if there's a compilation running
+	{
+		UE_LOG(LogMutable, Log, TEXT("FCustomizableObjectCompiler::Compile An object is already being compiled."));
+		return;
+	}
+
+	UE_LOG(LogMutable, Log, TEXT("PROFILE: [ %16.8f ] FCustomizableObjectCompiler::Compile start."), FPlatformTime::Seconds());
+
+	CurrentObject = Object;
+
+	// This is redundant but necessary to keep static analysis happy.
+	if (!Object || !CurrentObject) return;
+
+	Options = InOptions;
+	Options.bExtraBoneInfluencesEnabled = bAreExtraBoneInfluencesEnabled;
+	Options.bRealTimeMorphTargetsEnabled = Object->bEnableRealTimeMorphTargets;
+	Options.bClothingEnabled = Object->bEnableClothing;
+
+	if (Object->IsLocked() || !UCustomizableObjectSystem::GetInstance()->LockObject(Object))
+	{
+		UE_LOG(LogMutable, Display, TEXT("Customizable Object is already being compiled or updated %s. Please wait a few seconds and try again."), *Object->GetName());
+		return;
+	}
+
+	UE_LOG(LogMutable, Display, TEXT("Started Customizable Object Compile %s."), *Object->GetName());	
+
+	CompilationLogsContainer.ClearMessageCounters();
+
+	FMutableGraphGenerationContext GenerationContext(Object, this, Options);
+	GenerationContext.ParamNamesToSelectedOptions = ParamNamesToSelectedOptions;
+	if (ParamNamesToSelectedOptions.Num() > 0)
+	{
+		// A partial compilation shouldn't update the derived data cache, because it would force a recompilation at the next restart of the editor
+		Options.bDontUpdateStreamedDataAndCache = true;
+	}
+
+	// If we don't have the target platform yet (in editor) we need to get it
+	if (!GenerationContext.Options.TargetPlatform)
+	{
+		// Set the target platform in the context. For now it is the current platform.
+		ITargetPlatformManagerModule* TPM = GetTargetPlatformManager();
+		check(TPM);
+
+		GenerationContext.Options.TargetPlatform = TPM->GetRunningTargetPlatform();
+		check(GenerationContext.Options.TargetPlatform != nullptr);
+	}
+
+	// Clear Messages from previous Compilations
+	CompilationLogsContainer.ClearMessagesArray();
+
+	// Generate the mutable node expression
+	FText ErrorMessage;
+	bool bIsRootObject = false;
+	mu::NodeObjectPtr MutableRoot = GenerateMutableRoot(Object, GenerationContext, ErrorMessage, bIsRootObject);
+
+	if (!MutableRoot)
+	{
+		if (!ErrorMessage.IsEmpty())
+		{
+			CompilerLog(ErrorMessage, nullptr);
+		}
+		else
+		{
+			CompilerLog(FText::FromString(TEXT("Failed to generate the mutable node graph. Object not built.")), nullptr);
+		}
+
+		if (Object->IsLocked())
+		{
+			UCustomizableObjectSystem::GetInstance()->UnlockObject(Object);
+		}
+
+		RemoveCompileNotification();
+	}
+	else
+	{
+		if (Options.bCheckChildrenGuids)
+		{
+			// Check if all children are exactly the same as in the last compilation, if not, mark the object as modified and change VersionId to force a cache recompilation in other pcs
+			if ((GenerationContext.CustomizableObjectGuidsInCompilation.Num() != Object->CustomizableObjectGuidsInCompilation.Num()) ||
+				(GenerationContext.CustomizableObjectGuidsInCompilation.Intersect(Object->CustomizableObjectGuidsInCompilation).Num() != GenerationContext.CustomizableObjectGuidsInCompilation.Num()))
+			{
+				if (!ParamNamesToSelectedOptions.Num()) // Don't marked the object as modified because of a partial compilation
+				{
+					Object->CustomizableObjectGuidsInCompilation = GenerationContext.CustomizableObjectGuidsInCompilation;
+
+					if (!Options.bIsCooking)
+					{
+						Object->UpdateVersionId();
+						Object->MarkPackageDirty();
+					}
+				}
+			}
+		}
+
+		// Morph target generated data does not need extra processing so move semantics can be used
+		// to avoid a possibly expensive copy.
+		Object->ContributingMorphTargetsInfo = MoveTemp(GenerationContext.ContributingMorphTargetsInfo);
+		Object->MorphTargetReconstructionData = MoveTemp(GenerationContext.MorphTargetReconstructionData);
+
+		// Clothing	
+		Object->ClothMeshToMeshVertData = MoveTemp( GenerationContext.ClothMeshToMeshVertData );
+		
+		Object->ContributingClothingAssetsData = MoveTemp( GenerationContext.ContributingClothingAssetsData );
+		
+		// A clothing backend, e.g. Chaos cloth, can use 2 config files, one owned by the asset, and another that is shared 
+		// among all assets in a SkeletalMesh. When merging different assets in a skeletalmesh we need to make sure only one of 
+		// the shared is used. In that case we will keep the first visited of a type and will be stored separated from the asset.
+		// TODO: Shared configs, which typically controls the quality of the simulation (iterations, etc), probably should be specified 
+		// somewhere else to give more control with which config ends up used. 
+		auto IsSharedConfigData = []( const FCustomizableObjectClothConfigData& ConfigData ) -> bool
+		{
+			 const UClass* ConfigClass = FindObject<UClass>(nullptr, *ConfigData.ClassPath);
+			 return ConfigClass ? static_cast<bool>( Cast<UClothSharedConfigCommon>(ConfigClass->GetDefaultObject() ) ) : false;
+		};
+		
+		// Find shared configs to be used (One of each type) 
+		for ( FCustomizableObjectClothingAssetData& ClothingAssetData : Object->ContributingClothingAssetsData )
+		{
+			 for ( FCustomizableObjectClothConfigData& ClothConfigData : ClothingAssetData.ConfigsData )
+			 {
+				  if ( IsSharedConfigData( ClothConfigData ) )
+				  {
+					  FCustomizableObjectClothConfigData* FoundConfig = Object->ClothSharedConfigsData.FindByPredicate(
+						   [Name = ClothConfigData.ConfigName](const FCustomizableObjectClothConfigData& Other)
+						   {
+							   return Name == Other.ConfigName;
+						   });
+
+					  if (!FoundConfig)
+					  {
+						   FCustomizableObjectClothConfigData& ObjectConfig = Object->ClothSharedConfigsData.AddDefaulted_GetRef();
+						   ObjectConfig = ClothConfigData;
+					  }
+				  }
+			 }
+		}
+		
+		// Remove shared configs
+		for ( FCustomizableObjectClothingAssetData& ClothingAssetData : Object->ContributingClothingAssetsData )
+		{
+			 ClothingAssetData.ConfigsData.RemoveAllSwap(IsSharedConfigData);
+		}
+
+		// Mark the object as modified, used to avoid missing assets in packages. 
+		if (!ParamNamesToSelectedOptions.Num()) // Don't mark the objects as modified because of a partial compilation
+		{
+			if (Object->ReferencedMaterials.Num() != GenerationContext.ReferencedMaterials.Num())
+			{
+				Object->MarkPackageDirty();
+			}
+			else
+			{
+				for (int32 i = 0; i < Object->ReferencedMaterials.Num(); ++i)
+				{
+					if (Object->ReferencedMaterials[i] != GenerationContext.ReferencedMaterials[i])
+					{
+						Object->MarkPackageDirty();
+						break;
+					}
+				}
+			}
+		}
+
+		Object->ReferenceSkeletalMeshesData = GenerationContext.ReferenceSkeletalMeshesData;
+		
+		Object->ReferencedMaterials.Empty(GenerationContext.ReferencedMaterials.Num());
+
+		for (const UMaterialInterface* Material : GenerationContext.ReferencedMaterials)
+		{
+			Object->ReferencedMaterials.Add(Material);
+		}
+
+		Object->ReferencedMaterialSlotNames.Empty(GenerationContext.ReferencedMaterialSlotNames.Num());
+
+		for (const FName& MaterialSlotName : GenerationContext.ReferencedMaterialSlotNames)
+		{
+			Object->ReferencedMaterialSlotNames.Add(MaterialSlotName);
+		}
+
+		Object->ImageProperties.Empty(GenerationContext.ImageProperties.Num());
+
+		for (const FGeneratedImageProperties& ImageProp : GenerationContext.ImageProperties)
+		{
+			Object->ImageProperties.Add({ ImageProp.TextureParameterName,
+										ImageProp.Filter,
+										ImageProp.SRGB,
+										ImageProp.bFlipGreenChannel,
+										ImageProp.LODBias,
+										ImageProp.LODGroup,
+										ImageProp.AddressX, ImageProp.AddressY });
+		}
+
+		Object->GroupNodeMap = GenerationContext.GroupNodeMap;
+		Object->ParameterUIDataMap = GenerationContext.ParameterUIDataMap;
+		Object->StateUIDataMap = GenerationContext.StateUIDataMap;
+
+#if WITH_EDITORONLY_DATA
+		Object->CustomizableObjectPathMap = GenerationContext.CustomizableObjectPathMap;
+#endif
+
+		Object->LODSettings.NumLODsInRoot = GenerationContext.NumLODsInRoot;
+
+		Object->LODSettings.FirstLODAvailable = GenerationContext.FirstLODAvailable;
+
+		Object->LODSettings.bLODStreamingEnabled = GenerationContext.bEnableLODStreaming;
+		Object->LODSettings.NumLODsToStream = GenerationContext.NumMaxLODsToStream;
+
+		// Mark the object as modified, used to avoid missing assets in packages. 
+		if (!Object->PhysicsAssetsMap.OrderIndependentCompareEqual(GenerationContext.PhysicsAssetMap))
+		{
+			if (!ParamNamesToSelectedOptions.Num()) // Don't mark the objects as modified because of a partial compilation
+			{
+				Object->MarkPackageDirty();
+			}
+		}
+
+		Object->PhysicsAssetsMap = GenerationContext.PhysicsAssetMap;
+
+		// Mark the object as modified, used to avoid missing assets in packages. 
+		if (!Object->AnimBPAssetsMap.OrderIndependentCompareEqual(GenerationContext.AnimBPAssetsMap))
+		{
+			if (!ParamNamesToSelectedOptions.Num()) // Don't mark the objects as modified because of a partial compilation
+			{
+				Object->MarkPackageDirty();
+			}
+		}
+	
+		Object->AnimBPAssetsMap = GenerationContext.AnimBPAssetsMap;
+
+		// 
+		if (!ParamNamesToSelectedOptions.Num()) // Don't mark the objects as modified because of a partial compilation
+		{
+			if (Object->ReferencedSkeletons.Num() != GenerationContext.ReferencedSkeletons.Num())
+			{
+				Object->MarkPackageDirty();
+			}
+			else
+			{
+				for (int32 SkeletonIndex = 0; SkeletonIndex < Object->ReferencedSkeletons.Num(); ++SkeletonIndex)
+				{
+					if (Object->ReferencedSkeletons[SkeletonIndex] != GenerationContext.ReferencedSkeletons[SkeletonIndex])
+					{
+						Object->MarkPackageDirty();
+						break;
+					}
+				}
+			}
+		}
+
+		Object->ReferencedSkeletons.Empty(GenerationContext.ReferencedSkeletons.Num());
+
+		for (const USkeleton* Skeleton : GenerationContext.ReferencedSkeletons)
+		{
+			Object->ReferencedSkeletons.Add(Skeleton);
+		}
+
+		if (bIsRootObject && (GenerationContext.MaskOutMaterialCache.Num() > 0 || GenerationContext.MaskOutTextureCache.Num() > 0))
+		{
+			// Load MaskOutCache, if can't, create it.
+			if (!Object->MaskOutCache.Get())
+			{
+				if (!Object->MaskOutCache.LoadSynchronous())
+				{
+					FString MaskOutCacheName = TEXT("MaskOutCache");
+					FString PackageName = Object->GetOutermost()->GetPathName() + FString("_") + MaskOutCacheName;
+					UPackage* Package = CreatePackage(*PackageName);
+					Package->FullyLoad();
+
+					FString ObjectName = Object->GetName() + FString("_") + MaskOutCacheName;
+					Object->MaskOutCache = NewObject<UMutableMaskOutCache>(Package, *ObjectName, RF_Public | RF_Standalone);
+
+					if (!ParamNamesToSelectedOptions.Num()) // Don't marked the object as modified because of a partial compilation
+					{
+						Object->MarkPackageDirty();
+					}
+				}
+			}
+
+			check(Object->MaskOutCache.Get());
+
+			if (!Object->MaskOutCache->Materials.OrderIndependentCompareEqual(GenerationContext.MaskOutMaterialCache) ||
+				!Object->MaskOutCache->Textures.OrderIndependentCompareEqual(GenerationContext.MaskOutTextureCache))
+			{
+				Object->MaskOutCache->Materials = GenerationContext.MaskOutMaterialCache;
+				Object->MaskOutCache->Textures = GenerationContext.MaskOutTextureCache;
+
+				if (!ParamNamesToSelectedOptions.Num()) // Don't marked the object as modified because of a partial compilation
+				{
+					Object->MarkPackageDirty();
+					Object->MaskOutCache->Modify();
+				}
+			}
+		}
+		else
+		{
+			if (!Object->MaskOutCache.ToString().IsEmpty() || Object->MaskOutCache.Get())
+			{
+				if (!ParamNamesToSelectedOptions.Num()) // Don't marked the object as modified because of a partial compilation
+				{
+					Object->MaskOutCache = nullptr;
+					Object->MarkPackageDirty();
+				}
+			}
+		}
+
+		// Initializes a unique identifier for this object
+		if (!Options.bIsCooking)
+		{
+			Object->InitializeIdentifier();
+		}
+
+		if (CompileTask.IsValid()) // Don't start compilation if there's a compilation running
+		{
+			// TODO : warning?
+			UCustomizableObjectSystem::GetInstance()->UnlockObject(CurrentObject);
+			return;
+		}
+
+		// Lock the object to prevent instance updates while compiling. It's ignored and returns false if it's already locked
+		// Will unlock in the FinishCompilation call.
+		if (!Object->IsLocked())
+		{
+			UCustomizableObjectSystem::GetInstance()->LockObject(Object);
+		}
+
+		CompileTask = MakeShareable(new FCustomizableObjectCompileRunnable(MutableRoot, GenerationContext.bDisableTextureLayoutManagementFlag));
+		CompileTask->Options = Options;
+		State = ECustomizableObjectCompilationState::InProgress;
+
+		if (GenerationContext.ArrayTextureUnrealToMutableTask.Num() > 0)
+		{
+			ArrayTextureUnrealToMutableTask.Insert(GenerationContext.ArrayTextureUnrealToMutableTask, ArrayTextureUnrealToMutableTask.Num());
+			PendingTexturesToLoad = true;
+		}
+
+		// If synchronous compilation is requested, proceed the same way if in editor and if packaging
+		if (!bAsync)
+		{
+			UpdatePendingTextureConversion(false);
+			LaunchMutableCompile(false);
+			MUTABLE_CPUPROFILER_SCOPE(WaitForCompletion);
+			CompileThread->WaitForCompletion();
+			FinishCompilation();
+
+			if (SaveDDTask.IsValid())
+			{
+				SaveCODerivedData(false);
+				SaveDDThread->WaitForCompletion();
+				FinishSavingDerivedData();
+			}
+
+			CleanCachedReferencers();
+			UpdateArrayGCProtect();
+
+			State = ECustomizableObjectCompilationState::Completed;
+		}
+		else
+		{
+			// If packaging, convert textures and launch Mutable compile thread
+			if (IsRunningCommandlet())
+			{
+				UpdatePendingTextureConversion(false);
+				LaunchMutableCompile(false);
+			}
+			else
+			{
+				// If asynchronous compilation and not packaging, set CompilationLaunchPending to true and FCustomizableObjectCompiler::Tick()
+				// will do the remaining steps (convert textures asynchronously and launch the Mutable compile thread)
+				CompilationLaunchPending = true;
+
+				if (!Options.bSilentCompilation && PendingTexturesToLoad)
+				{
+					AddCompileNotification(LOCTEXT("ConvertingToMutableTexture", "Converting textures"));
+				}
+			}
+		}
+	}
+
+	// Get a list of generated nodes, to be able to understand error messages later.
+	for (UCustomizableObjectNode* Node : GenerationContext.GeneratedNodes)
+	{
+		GeneratedNodes.Add(Node, Node);
+
+		Node->ResetAttachedErrorData();
+	}
+
+	UE_LOG(LogMutable, Log, TEXT("PROFILE: [ %16.8f ] FCustomizableObjectCompiler::Compile end."), FPlatformTime::Seconds());
+
+	// Population Recompilation
+	if (MutableRoot)
+	{
+		//Checking if there is the population plugin
+		if (FModuleManager::Get().IsModuleLoaded("CustomizableObjectPopulation"))
+		{
+			ICustomizableObjectPopulationModule::Get().RecompilePopulations(Object);
+		}
+	}
+}
+
+
+mu::NodePtr FCustomizableObjectCompiler::Export(UCustomizableObject* Object, const FCompilationOptions& InCompilerOptions)
+{
+	UE_LOG(LogMutable, Display, TEXT("Started Customizable Object Export %s."), *Object->GetName());
+
+	FNotificationInfo Info(LOCTEXT("CustomizableObjectExportInProgress", "Exported Customizable Object"));
+	Info.bFireAndForget = true;
+	Info.bUseThrobber = true;
+	Info.FadeOutDuration = 1.0f;
+	Info.ExpireDuration = 1.0f;
+	FSlateNotificationManager::Get().AddNotification(Info);
+
+	FCompilationOptions CompilerOptions = InCompilerOptions;
+	CompilerOptions.bExtraBoneInfluencesEnabled = bAreExtraBoneInfluencesEnabled;
+		
+	FMutableGraphGenerationContext GenerationContext(Object, this, CompilerOptions);
+	GenerationContext.ParamNamesToSelectedOptions = ParamNamesToSelectedOptions;
+
+	// Generate the mutable node expression
+	FText ErrorMsg;
+	bool bIsRootObject = false;
+	mu::NodeObjectPtr MutableRoot = GenerateMutableRoot(Object, GenerationContext, ErrorMsg, bIsRootObject);
+
+	if (!MutableRoot)
+	{
+		if (!ErrorMsg.IsEmpty())
+		{
+			FCustomizableObjectCompiler::CompilerLog(ErrorMsg, nullptr);
+		}
+		else
+		{
+			FCustomizableObjectCompiler::CompilerLog(LOCTEXT("FailedToGenerate","Failed to generate the mutable node graph. Object not built."), nullptr);
+		}
+		return nullptr;
+	}
+
+	// Ensure the images are converted
+	if (GenerationContext.ArrayTextureUnrealToMutableTask.Num() > 0)
+	{
+		ArrayTextureUnrealToMutableTask.Insert(GenerationContext.ArrayTextureUnrealToMutableTask, ArrayTextureUnrealToMutableTask.Num());
+		PendingTexturesToLoad = true;
+		UpdatePendingTextureConversion(false);
+	}
+
+	return MutableRoot;
+}
+
+
+void FCustomizableObjectCompiler::FinishCompilation()
+{
+	check(CompileTask.IsValid());
+
+	UpdateCompilerLogData();
+	mu::ModelPtr Model = CompileTask->Model;
+
+	CurrentObject->SetModel(Model.get());
+
+	// Reset all instances, as the parameters may need to be rebuilt.
+	for (TObjectIterator<UCustomizableObjectInstance> It; It; ++It)
+	{
+		UCustomizableObjectInstance* Instance = *It;
+		if (Instance &&
+			Instance->GetCustomizableObject() == CurrentObject)
+		{
+			Instance->SetObject(CurrentObject);
+		}
+	}
+
+	// Order matters
+	CompileThread.Reset();
+	CompileTask.Reset();
+
+	if (!Options.bDontUpdateStreamedDataAndCache)
+	{
+		SaveDDTask = MakeShareable(new FCustomizableObjectSaveDDRunnable(CurrentObject, Options));
+	}
+	else
+	{
+		// when skipping the SaveDerivedData task unlock the object so that instances can be updated
+		UCustomizableObjectSystem::GetInstance()->UnlockObject(CurrentObject);
+	}
+
+	UE_LOG(LogMutable, Log, TEXT("Finished Customizable Object Compile."));
+}
+
+void FCustomizableObjectCompiler::FinishSavingDerivedData()
+{
+	check(SaveDDTask.IsValid());
+
+	if (Options.bIsCooking)
+	{
+		CurrentObject->CachePlatformData(SaveDDTask->GetTargetPlatform(), SaveDDTask->GetModelBytes(), SaveDDTask->GetBulkBytes());
+	}
+
+	// Order matters
+	SaveDDThread.Reset();
+	SaveDDTask.Reset();
+
+	// Unlock the object so that instances can be updated
+	UCustomizableObjectSystem::GetInstance()->UnlockObject(CurrentObject);
+
+	RemoveCompileNotification();
+
+	NotifyCompilationErrors();
+}
+
+
+void FCustomizableObjectCompiler::ForceFinishCompilation()
+{
+	if (CompileTask.IsValid())
+	{
+		if (CompileThread.IsValid())
+		{
+			CompileThread->WaitForCompletion();
+		}
+
+		CleanCachedReferencers();
+		UpdateArrayGCProtect();
+
+		UCustomizableObjectSystem::GetInstance()->UnlockObject(CurrentObject);
+	}
+
+	else if (SaveDDTask.IsValid())
+	{
+		SaveDDThread->WaitForCompletion();
+
+		UCustomizableObjectSystem::GetInstance()->UnlockObject(CurrentObject);
+	}
+
+	RemoveCompileNotification();
+}
+
+
+void FCustomizableObjectCompiler::ForceFinishBeforeStartCompilation(UCustomizableObject* Object)
+{
+	CleanCachedReferencers();
+	UpdateArrayGCProtect();
+
+	if (Object && Object->IsLocked())
+	{
+		UCustomizableObjectSystem::GetInstance()->UnlockObject(Object);
+	}
+
+	RemoveCompileNotification();
+}
+
+
+void FCustomizableObjectCompiler::AddCompileNotification(const FText& CompilationStep) const
+{
+	const FText Text = CurrentObject ? FText::FromString(FString::Printf(TEXT("Compiling %s"), *CurrentObject->GetName())) : LOCTEXT("CustomizableObjectCompileInProgress", "Compiling Customizable Object");
+	
+	FCustomizableObjectEditorLogger::CreateLog(Text)
+	.SubText(CompilationStep)
+	.Category(ELoggerCategory::Compilation)
+	.CustomNotification()
+	.FixNotification()
+	.Log();
+}
+
+
+void FCustomizableObjectCompiler::RemoveCompileNotification()
+{
+	FCustomizableObjectEditorLogger::DismissNotification(ELoggerCategory::Compilation);
+}
+
+
+void FCustomizableObjectCompiler::NotifyCompilationErrors() const
+{
+	const uint32 NumWarnings = CompilationLogsContainer.GetWarningCount(false);
+	const uint32 NumErrors = CompilationLogsContainer.GetErrorCount();
+	const bool NoWarningsOrErrors = !(NumWarnings || NumErrors);
+
+	if (Options.bSilentCompilation && NoWarningsOrErrors)
+	{
+		return;
+	}
+	
+	const EMessageSeverity::Type Severity = [&]
+	{
+		if (NumErrors)
+		{
+			return EMessageSeverity::Error;
+		}
+		else if (NumWarnings)
+		{
+			return EMessageSeverity::Warning;
+		}
+		else
+		{
+			return EMessageSeverity::Info;
+		}
+	}();
+
+	const FText Prefix = FText::FromString(CurrentObject ? CurrentObject->GetName() : "Customizable Object");
+
+	const FText Message = NoWarningsOrErrors ?
+		FText::Format(LOCTEXT("CompilationFinishedSuccessfully", "{0} finished compiling successfully"), Prefix) :
+		FText::Format(LOCTEXT("CompilationFinished", "{0} finished compiling successfully with {1} {1}|plural(one=warning,other=warnings) and {2} {2}|plural(one=error,other=errors)"), Prefix, NumWarnings, NumErrors);
+	
+	FCustomizableObjectEditorLogger::CreateLog(Message)
+	.Category(ELoggerCategory::Compilation)
+	.Severity(Severity)
+	.CustomNotification()
+	.Log();
+}
+
+
+void FCustomizableObjectCompiler::CompilerLog(const FText& Message, const TArray<const UCustomizableObjectNode*>& ArrayNode, EMessageSeverity::Type MessageSeverity, bool bAddBaseObjectInfo)
+{
+	// Cache the message for later reference
+	CompilationLogsContainer.AddMessage(Message,ArrayNode,MessageSeverity);
+	
+	FCustomizableObjectEditorLogger::CreateLog(Message)
+	.Severity(MessageSeverity)
+	.Nodes(ArrayNode)
+	.BaseObject(bAddBaseObjectInfo)
+	.Log();
+}
+
+
+void FCustomizableObjectCompiler::CompilerLog(const FText& Message, const UCustomizableObjectNode* Node, EMessageSeverity::Type MessageSeverity, bool AddBaseObjectInfo)
+{
+	TArray<const UCustomizableObjectNode*> ArrayNode;
+	if (Node)
+	{
+		ArrayNode.Add(Node);
+	}
+	CompilerLog(Message, ArrayNode, MessageSeverity, AddBaseObjectInfo);
+}
+
+
+void FCustomizableObjectCompiler::UpdateCompilerLogData()
+{
+	FMessageLogModule& MessageLogModule = FModuleManager::LoadModuleChecked<FMessageLogModule>("MessageLog");
+	MessageLogModule.RegisterLogListing(FName("Mutable"), LOCTEXT("MutableLog", "Mutable"));
+	const TArray<FCustomizableObjectCompileRunnable::FError>& ArrayCompileWarning = CompileTask->GetArrayWarning();
+	const TArray<FCustomizableObjectCompileRunnable::FError>& ArrayCompileError = CompileTask->GetArrayError();
+
+	FText ObjectName = CurrentObject ? FText::FromString(CurrentObject->GetName()) : LOCTEXT("Unknown Object", "Unknown Object");
+
+	int32 i;
+	for (i = 0; i < ArrayCompileError.Num(); ++i)
+	{
+		const UCustomizableObjectNode** pNode = GeneratedNodes.Find(ArrayCompileError[i].Context);
+
+		if (ArrayCompileError[i].AttachedData && pNode)
+		{
+			UCustomizableObjectNode::FAttachedErrorDataView ErrorDataView;
+			ErrorDataView.UnassignedUVs = { ArrayCompileError[i].AttachedData->UnassignedUVs.GetData(),
+											ArrayCompileError[i].AttachedData->UnassignedUVs.Num() };
+
+			const UCustomizableObjectNode* Node = Cast<UCustomizableObjectNode>(*pNode);
+			const_cast<UCustomizableObjectNode*>(Node)->AddAttachedErrorData(ErrorDataView);
+		}
+
+		FText FullMsg = FText::Format(LOCTEXT("MutableMessage", "{0} : {1}"), ObjectName, ArrayCompileError[i].Message);
+		CompilerLog(FullMsg, pNode ? *pNode : nullptr, EMessageSeverity::Error, true);
+		UE_LOG(LogMutable, Warning, TEXT("  %s"), *FullMsg.ToString());
+	}
+
+	for (i = 0; i < ArrayCompileWarning.Num(); ++i)
+	{
+		const UCustomizableObjectNode** pNode = GeneratedNodes.Find(ArrayCompileWarning[i].Context);
+
+		if (ArrayCompileWarning[i].AttachedData && pNode)
+		{
+			UCustomizableObjectNode::FAttachedErrorDataView ErrorDataView;
+			ErrorDataView.UnassignedUVs = { ArrayCompileWarning[i].AttachedData->UnassignedUVs.GetData(),
+											ArrayCompileWarning[i].AttachedData->UnassignedUVs.Num() };
+
+			const UCustomizableObjectNode* Node = Cast<UCustomizableObjectNode>(*pNode);
+			const_cast<UCustomizableObjectNode*>(Node)->AddAttachedErrorData(ErrorDataView);
+		}
+
+		FText FullMsg = FText::Format(LOCTEXT("MutableMessage", "{0} : {1}"), ObjectName, ArrayCompileWarning[i].Message);
+		CompilerLog(FullMsg, pNode?*pNode:nullptr, EMessageSeverity::Warning, true);
+		UE_LOG(LogMutable, Warning, TEXT("  %s"), *FullMsg.ToString());
+	}
+}
+
+
+void FCustomizableObjectCompiler::AddCompileOnlySelectedOption(const FString& ParamName, const FString& OptionValue)
+{
+	ParamNamesToSelectedOptions.Add(ParamName, OptionValue);
+}
+
+
+void FCustomizableObjectCompiler::ClearAllCompileOnlySelectedOption()
+{
+	ParamNamesToSelectedOptions.Empty();
+}
+
+#undef LOCTEXT_NAMESPACE
diff --git a/Engine/Plugins/Experimental/Mutable/Source/CustomizableObjectEditor/Private/CustomizableObjectCompiler.h b/Engine/Plugins/Experimental/Mutable/Source/CustomizableObjectEditor/Private/CustomizableObjectCompiler.h
new file mode 100644
--- /dev/null
+++ b/Engine/Plugins/Experimental/Mutable/Source/CustomizableObjectEditor/Private/CustomizableObjectCompiler.h
@@ -0,0 +1,252 @@
+// Copyright Epic Games, Inc. All Rights Reserved.
+
+#pragma once
+
+#include "CustomizableObjectCompileRunnable.h"
+#include "GenerateMutableSource/GenerateMutableSource.h"
+
+#include "Logging/TokenizedMessage.h"
+#include "CompilationMessageCache.h"
+#include "CustomizableObject.h"
+#include "Widgets/Notifications/SNotificationList.h"
+#include "Engine/StreamableManager.h"
+
+#include "MutableTools/Public/NodeObject.h"
+#include "MutableTools/Public/NodeImageConstant.h"
+
+
+
+class UCustomizableObject;
+class UCustomizableObjectNodeObject;
+class CustomizableObjectCompileAsyncTask;
+
+
+class FCustomizableObjectCompiler : public FCustomizableObjectCompilerBase, public FGCObject
+{
+public:
+
+	CUSTOMIZABLEOBJECTEDITOR_API FCustomizableObjectCompiler();
+	CUSTOMIZABLEOBJECTEDITOR_API virtual ~FCustomizableObjectCompiler() {}
+
+	// FCustomizableObjectCompilerBase interface
+	CUSTOMIZABLEOBJECTEDITOR_API bool IsRootObject(const class UCustomizableObject* Object) const override;
+	
+	CUSTOMIZABLEOBJECTEDITOR_API ECustomizableObjectCompilationState GetCompilationState() const override { return State;  };
+	
+	/** Check for pending compilation process. Returns true if an object has been updated. */
+	CUSTOMIZABLEOBJECTEDITOR_API virtual bool Tick() override;
+
+	/** In case Mutable Compile is disabled, fake the compilation process */
+	CUSTOMIZABLEOBJECTEDITOR_API void MutableIsDisabledCase(UCustomizableObject* Object);
+
+	/** Provided a CO object it provides the root CO it is connected. In other words : it returns the root of the entire
+	 * mutable graph.
+	 * @param InObject Customizable object whose root CO we are asking for.
+	 * @return The CO that is the root of the provided Customizable Object. It can be equal to InObject if the provided
+	 * object does not have any parent.
+	 */
+	virtual UCustomizableObject* GetRootObject(UCustomizableObject* InObject) override;
+
+	/** Generate the Mutable Graph from the Unreal Graph. */
+	mu::NodePtr Export(UCustomizableObject* Object, const FCompilationOptions& Options);
+
+	void CompilerLog(const FText& Message, const TArray<const UCustomizableObjectNode*>& ArrayNode, EMessageSeverity::Type MessageSeverity = EMessageSeverity::Warning, bool bAddBaseObjectInfo = true);
+	void CompilerLog(const FText& Message, const UCustomizableObjectNode* Node = nullptr, EMessageSeverity::Type MessageSeverity = EMessageSeverity::Warning, bool bAddBaseObjectInfo = true);
+	void NotifyCompilationErrors() const;
+
+	void FinishCompilation();
+	void FinishSavingDerivedData();
+	virtual void ForceFinishCompilation() override;
+
+	void AddCompileNotification(const FText& CompilationStep) const;
+	static void RemoveCompileNotification();
+
+	/** Load required assets and compile.
+	 *
+	 * Loads assets which reference Object's package asynchronously before calling ProcessChildObjectsRecursively. */
+	virtual void Compile(UCustomizableObject& Object, const FCompilationOptions& Options, bool bAsync) override;
+	
+	/** FSerializableObject interface */
+	virtual void AddReferencedObjects(FReferenceCollector& Collector) override;
+	virtual FString GetReferencerName() const override
+	{
+		return TEXT("FCustomizableObjectCompiler");
+	}
+
+	/** Simply add CO elements from ArrayAssetData to ArrayGCProtect when they've been loaded from ArrayAssetData */
+	void UpdateArrayGCProtect();
+
+	/** Getter of CurrentGAsyncLoadingTimeLimit */
+	float GetCurrentGAsyncLoadingTimeLimit() { return CurrentGAsyncLoadingTimeLimit; }
+
+	/** Setter for CurrentGAsyncLoadingTimeLimit */
+	void SetCurrentGAsyncLoadingTimeLimit(float Value) { CurrentGAsyncLoadingTimeLimit = Value; }
+
+	/** Getter of PreloadingReferencerAssets */
+	bool GetPreloadingReferencerAssets() { return PreloadingReferencerAssets; }
+
+	/** Setter for PreloadingReferencerAssets */
+	void SetPreloadingReferencerAssets(bool Value) { PreloadingReferencerAssets = Value; }
+
+	// This is used to restrict Group nodes to compile only the SelectedOption and discard all the others
+	void AddCompileOnlySelectedOption(const FString& ParamName, const FString& OptionValue);
+	// This clears all restrictions and let's a full compilation happen
+	void ClearAllCompileOnlySelectedOption();
+
+	// Function to clear the compiler if the editor was closed before starting the compilation 
+	void ForceFinishBeforeStartCompilation(UCustomizableObject* Object);
+
+	// Getter of AsynchronousStreamableHandlePtr
+	TSharedPtr<struct FStreamableHandle> GetAsynchronousStreamableHandlePtr() { return AsynchronousStreamableHandlePtr; }
+
+private:
+
+	// Object containing all error and warning logs raised during compilation.
+	FCompilationMessageCache CompilationLogsContainer;
+	
+private:
+	ECustomizableObjectCompilationState State = ECustomizableObjectCompilationState::None;
+	
+	void CompileInternal(UCustomizableObject* Object, const FCompilationOptions& Options, bool bAsync = false);
+
+	static void PreloadingReferencerAssetsCallback(UCustomizableObject* Object, FCustomizableObjectCompiler* CustomizableObjectCompiler, const FCompilationOptions Options, bool bAsync);
+	
+	void ProcessChildObjectsRecursively(UCustomizableObject* Object, class FAssetRegistryModule& AssetRegistryModule, struct FMutableGraphGenerationContext &GenerationContext);
+
+	//
+	FCompilationOptions Options;
+
+	//
+	TSharedPtr< FCustomizableObjectCompileRunnable > CompileTask;
+
+	//
+	TSharedPtr< FRunnableThread > CompileThread;
+
+	//
+	TSharedPtr< FCustomizableObjectSaveDDRunnable > SaveDDTask;
+
+	//
+	TSharedPtr< FRunnableThread > SaveDDThread;
+
+	// Cache configuration settings from ini files 
+	bool bAreExtraBoneInfluencesEnabled = false;
+
+	// TODO: add references
+	UCustomizableObject* CurrentObject = nullptr;
+
+	/** Array where to put the names of the already processed child in ProcessChildObjectsRecursively */
+	TArray<FName> ArrayAlreadyProcessedChild;
+
+
+	// Will output to Mutable Log the warning and error messages generated during the CO compilation
+	// and update the values of NumWarnings and NumErrors
+	void UpdateCompilerLogData();
+
+
+	/** If duplicated elements are found in each entry of ParameterNamesMap, a warning for
+	the parameters with repeated name will be generated */
+	void DisplayParameterWarning(struct FMutableGraphGenerationContext& GenerationContext);
+	
+	/** If duplicated node ids are found, usually due to duplicating CustomizableObjects Assets, a warning
+	for the nodes with repeated ids will be generated */
+	void DisplayDuplicatedNodeIdsWarning(struct FMutableGraphGenerationContext& GenerationContext);
+	
+	/** Display warnings for unnamed node objects */
+	void DisplayUnnamedNodeObjectWarning(struct FMutableGraphGenerationContext& GenerationContext);
+	
+	/** Display warnings from discarded PhysicsAssets due to SkeletalBodySetups with no corresponding bones
+	in the SkeletalMesh's RefSkeleton */
+	void DisplayDiscardedPhysicsAssetSingleWarning(struct FMutableGraphGenerationContext& GenerationContext);
+	
+	/** Display a warning for each node contains an orphan pin. */
+	void DisplayOrphanNodesWarning(struct FMutableGraphGenerationContext& GenerationContext);
+	
+	mu::NodeObjectPtr GenerateMutableRoot(UCustomizableObject* Object, FMutableGraphGenerationContext& GenerationContext, FText& ErrorMessage, bool& bOutIsRootObject);
+
+	/** Add to ArrayAssetData the FAssetData information of all referencers of static class type UCustomizableObject::StaticClass()
+	* that reference the package given by the PathName parameter
+	* @param PathName            [in]  path to the CO to be analyzed (for instance, CO->GetOuter()->GetPathName())
+	* @param ArrayReferenceNames [out] array with the package names which are referenced by the CO with PathName given as parameter
+	* @return nothing */
+	void AddCachedReferencers(const FName& PathName, TArray<FName>& ArrayReferenceNames);
+
+	/** Launches the compile task in another thread when compiling a CO in the editor
+	* param ShowNotification [in] whether to show the compiling CO notification or not
+	* @return nothing */
+	void LaunchMutableCompile(bool ShowNotification);
+
+	/** Launches the save derived data task in another thread after compiling a CO in the
+	* editor
+	* param ShowNotification [in] whether to show the saving DD notification or not
+	* @return nothing */
+	void SaveCODerivedData(bool ShowNotification);
+
+	/** When compiling a CO in the editor, flag to know when the Unreal textures have been converted to Mutable textures */
+	bool PendingTexturesToLoad;
+
+	/** When compiling a CO in the editor, flag to know if there's a mutable task pending to be launched through LaunchMutableCompile */
+	bool CompilationLaunchPending;
+
+	/** Just used to clean ArrayAssetData
+	* @return nothing */
+	void CleanCachedReferencers();
+
+	/** Will test if package path given by PackageName parameter is one of ArrayAssetData's elements FAssetData::PackageName value
+	* @param PackageName [in] package name to test
+	* @return true if cached, false otherwise */
+	bool IsCachedInAssetData(const FString& PackageName);
+
+	/** Find FAssetData ArrayAssetData with PackageName given by parameter
+	* @param PackageName [in] package name to find in ArrayAssetData
+	* @return pointer to element if any found, nullptr otherwise */
+	FAssetData* GetCachedAssetData(const FString& PackageName);
+
+	/** When compiling a CO, performs the texture conversion from Unreal to Mutable. If in editor, this method is called in
+	* FCustomizableObjectCompiler::Tick with time limit, otherwise it is called in FCustomizableObjectCompiler::Compile with
+	* no time limit to convert all textures in a row
+	* @param UseTimeLimit [in] if true, the method will return after TimeLimit seconds have been used to convert textures, if false, all textures will be processed in a row
+	* @return nothing */
+	void UpdatePendingTextureConversion(bool UseTimeLimit);
+
+	/** Helper function to compute the value for Unreal Engine variable s.AsyncLoadingTimeLimit while asynchronous loading is used.
+	* Also assigned to MaxConvertToMutableTextureTime
+	* @return value in milliseconds to use for AsyncLoadingTimeLimit and MaxConvertToMutableTextureTime */
+	float ComputeAsyncLoadingTimeLimit();
+	
+public:
+	
+	virtual void GetCompilationMessages(TArray<FText>& OutWarningMessages, TArray<FText>& OutErrorMessages) const override;
+	
+private:
+	/** Array with all the packages used to compile current Customizable Object */
+	TArray<FAssetData> ArrayAssetData;
+
+	/** Array used to protect from garbage collection those COs loaded asynchronously */
+	TArray<UCustomizableObject*> ArrayGCProtect;
+
+	/** Flag to know when asynchronous asset loading is being performed */
+	bool PreloadingReferencerAssets;
+
+	/** Copy of GAsyncLoadingTimeLimit while assets are loaded, previous value is restored after asset load */
+	float CurrentGAsyncLoadingTimeLimit;
+
+	/** Array with the textures that need to be asynchronously converted from Unreal to Mutable */
+	TArray<FTextureUnrealToMutableTask> ArrayTextureUnrealToMutableTask;
+
+	/** Counter to know how many of the textures in ArrayTextureUnrealToMutableTask have been converted from Unreal to Mutable */
+	int32 CompletedUnrealToMutableTask;
+
+	/** Time threshold used in UpdatePendingTextureConversion to stop converting textures from Unreal to Mutable until next tick */
+	float MaxConvertToMutableTextureTime;
+
+	// Stores the only option of an Int Param that should be compiled
+	TMap<FString, FString> ParamNamesToSelectedOptions;
+
+	// List of nodes used during generation, so that they can be found when processing mutable runtime error messages.
+	TMap<const void*, const UCustomizableObjectNode*> GeneratedNodes;
+
+	/** Pointer to the Asynchronous Preloading process call back */
+	TSharedPtr<struct FStreamableHandle> AsynchronousStreamableHandlePtr;
+	
+};
+
diff --git a/Engine/Plugins/Experimental/Mutable/Source/CustomizableObjectEditor/Private/CustomizableObjectCustomSettings.h b/Engine/Plugins/Experimental/Mutable/Source/CustomizableObjectEditor/Private/CustomizableObjectCustomSettings.h
new file mode 100644
--- /dev/null
+++ b/Engine/Plugins/Experimental/Mutable/Source/CustomizableObjectEditor/Private/CustomizableObjectCustomSettings.h
@@ -0,0 +1,85 @@
+// Copyright Epic Games, Inc. All Rights Reserved. 
+
+#pragma once
+
+// include "SComboBox.h"
+//#include "SBoxPanel.h"
+#include "Widgets/Views/SListView.h"
+#include "AssetRegistry/AssetData.h"
+//#include "Engine/DataTable.h"
+#include "CustomizableObjectEditorViewportLights.h"
+#include "CustomizableObjectCustomSettings.generated.h"
+
+UCLASS()
+class UCustomizableObjectEmptyClassForSettings : public UObject
+{
+	GENERATED_BODY()
+public:
+	TWeakPtr<class SCustomizableObjectEditorViewportTabBody> Viewport;
+	class UDebugSkelMeshComponent** PreviewSkeletalMeshComp = nullptr;
+};
+
+class SCustomizableObjectCustomSettings : public SVerticalBox
+{
+public:
+	void Construct(const FArguments& InArgs, UCustomizableObjectEmptyClassForSettings* InObject);
+
+private:
+	class FCustomizableObjectEditorViewportClient* ViewportClient;
+
+	// Lighting
+public:
+	void SetViewportLightsByAsset(const FAssetData& InAsset);
+
+private:
+
+	TSharedRef<SWidget> GetLightComboButtonContent();
+	void CloseLightComboButtonContent();
+	FReply OnSaveViewportLightsAsset();
+	FReply OnNewViewportLightsAsset();
+
+	FReply OnPointLightAdded();
+	FReply OnSpotLightAdded();
+	FReply OnLightRemoved();
+	FReply OnLightUnselected();
+	TSharedRef<class ITableRow> OnGenerateWidgetForList(TSharedPtr<FString> Item, const TSharedRef<class STableViewBase>& OwnerTable);
+	void OnListSelectionChanged(TSharedPtr<FString> Selection, ESelectInfo::Type SelectInfo);
+
+	// LightProperties
+	TOptional<float> GetIntensityValue() const;
+	void OnIntensityValueCommited(float Value, ETextCommit::Type CommitType);
+
+	FLinearColor GetLightColorValue() const;
+	FReply OnLightColorBlockMouseButtonDown(const FGeometry& MyGeometry, const FPointerEvent& MouseEvent);
+	void OnSetLightColorFromColorPicker(FLinearColor InColor);
+
+	TOptional<float> GetAttenuationRadius() const;
+	void OnAttenuationRadiusValueCommited(float Value, ETextCommit::Type CommitType);
+
+	TOptional<float> GetLightSourceRadius() const;
+	void OnLightSourceRadiusValueCommited(float Value, ETextCommit::Type CommitType);
+
+	TOptional<float> GetLightSourceLength() const;
+	void OnLightSourceLengthValueCommited(float Value, ETextCommit::Type CommitType);
+
+	TOptional<float> GetLightInnerConeAngle() const;
+	void OnLightInnerConeAngleValueCommited(float Value, ETextCommit::Type CommitType);
+
+	TOptional<float> GetLightOuterConeAngle() const;
+	void OnLightOuterConeAngleValueCommited(float Value, ETextCommit::Type CommitType);
+
+private:
+	TArray<TSharedPtr<FString>> LightNames;
+	TArray<class UPointLightComponent*> Lights;
+
+	static const FString LightPackagePath;
+	TSharedPtr<class SEditableTextBox> LightsAssetNameInputText;
+
+	TSharedPtr<class SComboButton> LightComboButton;
+	TSharedPtr<class STextBlock> LightAssetText;
+	TSharedPtr<SListView<TSharedPtr<FString>>> LightsListView;
+	class ULightComponent* SelectedLight;
+	TSharedPtr<class SExpandableArea> PointLightProperties;
+	TSharedPtr<class SExpandableArea> SpotLightProperties;
+
+};
\ No newline at end of file
diff --git a/Engine/Plugins/Experimental/Mutable/Source/CustomizableObjectEditor/Private/CustomizableObjectCustomSettingsDetails.cpp b/Engine/Plugins/Experimental/Mutable/Source/CustomizableObjectEditor/Private/CustomizableObjectCustomSettingsDetails.cpp
new file mode 100644
--- /dev/null
+++ b/Engine/Plugins/Experimental/Mutable/Source/CustomizableObjectEditor/Private/CustomizableObjectCustomSettingsDetails.cpp
@@ -0,0 +1,25 @@
+// Copyright Epic Games, Inc. All Rights Reserved.
+
+#include "CustomizableObjectCustomSettingsDetails.h"
+#include "DetailCategoryBuilder.h"
+#include "DetailLayoutBuilder.h"
+#include "DetailWidgetRow.h"
+#include "CustomizableObjectCustomSettings.h"
+
+TSharedRef<IDetailCustomization> FCustomizableObjectCustomSettingsDetails::MakeInstance()
+{
+	return MakeShareable(new FCustomizableObjectCustomSettingsDetails);
+}
+
+void FCustomizableObjectCustomSettingsDetails::CustomizeDetails(IDetailLayoutBuilder& DetailBuilder)
+{
+	const IDetailsView* DetailsView = DetailBuilder.GetDetailsView();
+	if (DetailsView->GetSelectedObjects().Num() > 0 && DetailsView->GetSelectedObjects()[0].IsValid())
+	{
+		IDetailCategoryBuilder& MainCategory = DetailBuilder.EditCategory("Custom Settings");
+		MainCategory.AddCustomRow(FText::FromString("Custom Settings"))
+		[
+			SNew(SCustomizableObjectCustomSettings, Cast<UCustomizableObjectEmptyClassForSettings>(DetailsView->GetSelectedObjects()[0].Get()))
+		];
+	}
+}
diff --git a/Engine/Plugins/Experimental/Mutable/Source/CustomizableObjectEditor/Private/CustomizableObjectDebugger.cpp b/Engine/Plugins/Experimental/Mutable/Source/CustomizableObjectEditor/Private/CustomizableObjectDebugger.cpp
new file mode 100644
--- /dev/null
+++ b/Engine/Plugins/Experimental/Mutable/Source/CustomizableObjectEditor/Private/CustomizableObjectDebugger.cpp
@@ -0,0 +1,134 @@
+// Copyright Epic Games, Inc. All Rights Reserved.
+
+#include "CustomizableObjectDebugger.h"
+#include "CustomizableObjectDebuggerActions.h"
+#include "CustomizableObjectEditorModule.h"
+#include "CustomizableObjectEditorStyle.h"
+#include "CustomizableObjectEditor.h"
+
+#include "SCustomizableObjectProperties.h"
+#include "SMutableGraphViewer.h"
+#include "SMutableCodeViewer.h"
+#include "SMutableObjectViewer.h"
+
+#include "Editor.h"
+#include "AssetRegistry/AssetRegistryModule.h"
+#include "DesktopPlatformModule.h"
+#include "ContentBrowserModule.h"
+#include "PropertyEditorModule.h"
+
+#include "Widgets/Notifications/SNotificationList.h"
+#include "Widgets/Docking/SDockTab.h"
+#include "Widgets/Input/STextComboBox.h"
+#include "Widgets/Layout/SSplitter.h"
+#include "Types/SlateEnums.h"
+
+#include "PreviewScene.h"
+#include "ScopedTransaction.h"
+#include "FileHelpers.h"
+#include "BusyCursor.h"
+#include "EditorDirectories.h"
+
+#include "WorkspaceMenuStructureModule.h"
+#include "WorkspaceMenuStructure.h"
+#include "Engine/Selection.h"
+#include "Components/SkeletalMeshComponent.h"
+
+#include "IDetailsView.h"
+#include "WorkspaceMenuStructure.h"
+#include "Framework/Commands/GenericCommands.h"
+#include "Framework/Notifications/NotificationManager.h"
+
+#include "GenericPlatform/GenericPlatformApplicationMisc.h"
+
+#include "HAL/PlatformApplicationMisc.h"
+
+
+#define LOCTEXT_NAMESPACE "CustomizableObjectDebugger"
+
+DEFINE_LOG_CATEGORY_STATIC(LogCustomizableObjectDebugger, Log, All);
+const FName FCustomizableObjectDebugger::MutableNewTabId(TEXT("CustomizableObjectDebugger_NewTab"));
+
+
+void FCustomizableObjectDebugger::InitCustomizableObjectDebugger( 
+	const EToolkitMode::Type Mode, 
+	const TSharedPtr< class IToolkitHost >& InitToolkitHost, 
+	UCustomizableObject* ObjectToEdit )
+{
+	CustomizableObject = ObjectToEdit;
+
+	const TSharedRef<FTabManager::FLayout> StandaloneDefaultLayout = FTabManager::NewLayout( "Standalone_CustomizableObjectDebugger_Layout_v3" )
+	->AddArea
+	(
+		FTabManager::NewPrimaryArea()
+		->SetOrientation(Orient_Horizontal)
+		->Split
+		(
+			FTabManager::NewStack()
+			->SetSizeCoefficient(0.6f)
+			->AddTab(MutableNewTabId, ETabState::ClosedTab)
+		)
+	);
+
+	const bool bCreateDefaultStandaloneMenu = false;
+	const bool bCreateDefaultToolbar = false;
+	FAssetEditorToolkit::InitAssetEditor( Mode, InitToolkitHost, CustomizableObjectDebuggerAppIdentifier, StandaloneDefaultLayout, bCreateDefaultStandaloneMenu, bCreateDefaultToolbar, ObjectToEdit );
+
+	// Open a tab for the object being debugged.
+	TSharedPtr<SDockTab> NewMutableObjectTab = SNew(SDockTab)
+		.Label(FText::FromString(FString::Printf(TEXT("Object [%s]"),*CustomizableObject->GetName())))
+		[
+			SNew(SMutableObjectViewer, CustomizableObject, TabManager, MutableNewTabId)
+		];
+
+	TabManager->InsertNewDocumentTab(MutableNewTabId, FTabManager::ESearchPreference::PreferLiveTab, NewMutableObjectTab.ToSharedRef());
+
+}
+
+
+const FSlateBrush* FCustomizableObjectDebugger::GetDefaultTabIcon() const
+{
+	return FCustomizableObjectEditorStyle::Get().GetBrush("CustomizableObjectEditor.Debug");
+}
+
+
+FName FCustomizableObjectDebugger::GetToolkitFName() const
+{
+	return FName("CustomizableObjectDebugger");
+}
+
+
+FText FCustomizableObjectDebugger::GetBaseToolkitName() const
+{
+	return LOCTEXT("ToolkitName", "Customizable Object Editor");
+}
+
+
+void FCustomizableObjectDebugger::AddReferencedObjects( FReferenceCollector& Collector )
+{
+	Collector.AddReferencedObject( CustomizableObject );
+}
+
+
+FText FCustomizableObjectDebugger::GetToolkitName() const
+{
+	FFormatNamedArguments Args;
+	Args.Add(TEXT("ObjectName"), FText::FromString(GetEditingObject()->GetName()));
+	Args.Add(TEXT("ToolkitName"), GetBaseToolkitName());
+	return FText::Format(LOCTEXT("AppLabelWithAssetName", "Debug {ObjectName} - {ToolkitName}"), Args);
+}
+
+
+FString FCustomizableObjectDebugger::GetWorldCentricTabPrefix() const
+{
+	return LOCTEXT("DebuggerWorldCentricTabPrefix", "CustomizableObjectDebugger ").ToString();
+}
+
+
+FLinearColor FCustomizableObjectDebugger::GetWorldCentricTabColorScale() const
+{
+	return FLinearColor(0.3f, 0.2f, 0.5f, 0.5f);
+}
+
+
+#undef LOCTEXT_NAMESPACE
diff --git a/Engine/Plugins/Experimental/Mutable/Source/CustomizableObjectEditor/Private/CustomizableObjectDebugger.h b/Engine/Plugins/Experimental/Mutable/Source/CustomizableObjectEditor/Private/CustomizableObjectDebugger.h
new file mode 100644
--- /dev/null
+++ b/Engine/Plugins/Experimental/Mutable/Source/CustomizableObjectEditor/Private/CustomizableObjectDebugger.h
@@ -0,0 +1,64 @@
+// Copyright Epic Games, Inc. All Rights Reserved.
+
+#pragma once
+
+#include "ICustomizableObjectDebugger.h"
+#include "Misc/NotifyHook.h"
+#include "EditorUndoClient.h"
+#include "Toolkits/AssetEditorToolkit.h"
+
+
+DECLARE_DELEGATE(FCreatePreviewInstanceFlagDelegate);
+
+
+/**
+ * CustomizableObject Editor class
+ */
+class FCustomizableObjectDebugger : 
+	public ICustomizableObjectDebugger, 
+	public FGCObject,
+	public FNotifyHook
+{
+public:
+
+	/**
+	 * Edits the specified object
+	 *
+	 * @param	Mode					Asset editing mode for this editor (standalone or world-centric)
+	 * @param	InitToolkitHost			When Mode is WorldCentric, this is the level editor instance to spawn this editor within
+	 * @param	ObjectToEdit			The object to edit
+	 */
+	void InitCustomizableObjectDebugger( const EToolkitMode::Type Mode, const TSharedPtr< class IToolkitHost >& InitToolkitHost, UCustomizableObject* ObjectToEdit );
+
+	// FGCObject interface
+	virtual void AddReferencedObjects( FReferenceCollector& Collector ) override; 
+	virtual FString GetReferencerName() const override
+	{
+		return TEXT("FCustomizableObjectDebugger");
+	}
+
+	// IToolkit interface
+	virtual FName GetToolkitFName() const override;
+	virtual FText GetBaseToolkitName() const override;
+	virtual FText GetToolkitName() const override;
+	virtual FString GetWorldCentricTabPrefix() const override;
+	virtual FLinearColor GetWorldCentricTabColorScale() const override;
+
+	// AssetEditorToolkit interface
+	virtual bool CanSaveAsset() const override { return false; }
+	virtual const FSlateBrush* GetDefaultTabIcon() const override;
+
+	// ICustomizableObjectDebugger interface
+	virtual UCustomizableObject* GetCustomizableObject() override { return CustomizableObject; }
+
+
+private:
+
+	/** The currently viewed object. */
+	UCustomizableObject* CustomizableObject = nullptr;
+
+	/**	The tab ids for all the tabs used */
+	static const FName MutableNewTabId;
+
+
+};
diff --git a/Engine/Plugins/Experimental/Mutable/Source/CustomizableObjectEditor/Private/CustomizableObjectDebuggerActions.cpp b/Engine/Plugins/Experimental/Mutable/Source/CustomizableObjectEditor/Private/CustomizableObjectDebuggerActions.cpp
new file mode 100644
--- /dev/null
+++ b/Engine/Plugins/Experimental/Mutable/Source/CustomizableObjectEditor/Private/CustomizableObjectDebuggerActions.cpp
@@ -0,0 +1,35 @@
+// Copyright Epic Games, Inc. All Rights Reserved.
+
+#include "CustomizableObjectDebuggerActions.h"
+#include "CustomizableObjectEditor.h"
+#include "Styling/AppStyle.h"
+#include "CustomizableObjectEditorStyle.h"
+#include "Framework/Commands/InputChord.h"
+
+#define LOCTEXT_NAMESPACE "CustomizableObjectEditor"
+
+
+FCustomizableObjectDebuggerCommands::FCustomizableObjectDebuggerCommands() 
+	: TCommands<FCustomizableObjectDebuggerCommands>
+(
+	"CustomizableObjectDebugger", // Context name for fast lookup
+	NSLOCTEXT("Contexts", "CustomizableObjectDebugger", "CustomizableObject Debugger"),
+	NAME_None, // Parent
+	FCustomizableObjectEditorStyle::GetStyleSetName()
+	)
+{
+}
+
+
+void FCustomizableObjectDebuggerCommands::RegisterCommands()
+{
+	UI_COMMAND(GenerateMutableGraph, "Unreal to Mutable Graph", "Generate a mutable graph from the customizable obejct source graph.", EUserInterfaceActionType::Button, FInputChord());
+	UI_COMMAND(CompileMutableCode, "Unreal to Mutable Code", "Compile the mutable graph of the customizable object and update the previews.", EUserInterfaceActionType::Button, FInputChord());
+	UI_COMMAND(CompileOptions_EnableTextureCompression, "Enable texture compression.", "Only for debug. Do not use.", EUserInterfaceActionType::ToggleButton, FInputChord());
+	UI_COMMAND(CompileOptions_UseParallelCompilation, "Enable compiling in multiple threads.", "This is faster but use more memory.", EUserInterfaceActionType::ToggleButton, FInputChord());
+	UI_COMMAND(CompileOptions_UseDiskCompilation, "Enable compiling using the disk as memory.", "This is very slow but supports compiling huge objects. It requires a lot of free space in the OS disk.", EUserInterfaceActionType::ToggleButton, FInputChord());
+}
+
+
+#undef LOCTEXT_NAMESPACE
+
