[mutable] Move Mutable plugin from the restricted folder to the experimental plugins folder.

diff --git a/Engine/Plugins/Experimental/Mutable/Source/CustomizableObjectEditor/Private/CustomizableObjectDetails.cpp b/Engine/Plugins/Experimental/Mutable/Source/CustomizableObjectEditor/Private/CustomizableObjectDetails.cpp
new file mode 100644
--- /dev/null
+++ b/Engine/Plugins/Experimental/Mutable/Source/CustomizableObjectEditor/Private/CustomizableObjectDetails.cpp
@@ -0,0 +1,381 @@
+// Copyright Epic Games, Inc. All Rights Reserved.
+
+#include "CustomizableObjectDetails.h"
+#include "CustomizableObjectEditorModule.h"
+
+#include "DetailLayoutBuilder.h"
+#include "DetailCategoryBuilder.h"
+#include "PropertyCustomizationHelpers.h"
+#include "LevelEditor.h"
+#include "LevelEditorActions.h"
+
+#include "Runtime/Launch/Resources/Version.h"
+
+#include "MutableRuntime/Public/Model.h"
+
+
+#define LOCTEXT_NAMESPACE "CustomizableObjectDetails"
+
+TSharedRef<IDetailCustomization> FCustomizableObjectDetails::MakeInstance()
+{
+	return MakeShareable( new FCustomizableObjectDetails );
+}
+
+
+class FStateDetailsNode
+{
+public:
+
+	FStateDetailsNode()
+	{
+		StateIndex = -1;
+		ParameterIndex = -1;
+	}
+
+	mu::ModelPtrConst Model;
+	int StateIndex;
+	int ParameterIndex;
+
+	TSharedPtr<FString> GetName()
+	{
+		FString res;
+
+		if (ParameterIndex<0)
+		{
+			res = Model->GetStateName( StateIndex );
+		}
+		else
+		{
+			mu::ParametersPtr TempParams = Model->NewParameters();
+			res = TempParams->GetName( ParameterIndex );
+		}
+
+		return MakeShareable( new FString( res ) );
+	}
+
+
+	TArray<TSharedPtr<FStateDetailsNode>> GetChildrenList()
+	{
+		TArray<TSharedPtr<FStateDetailsNode>> ChildrenList;
+
+		// Is it a state node?
+		if (ParameterIndex<0)
+		{
+			int numParams = Model->GetStateParameterCount( StateIndex );
+			for ( int i=0; i<numParams; ++i )
+			{
+				TSharedPtr<FStateDetailsNode> node = MakeShareable( new FStateDetailsNode );
+
+				node->Model = Model;
+				node->StateIndex = StateIndex;
+				node->ParameterIndex = Model->GetStateParameterIndex( StateIndex, i );
+
+				ChildrenList.Add( node );
+			}
+		}
+
+		return ChildrenList;
+	}
+
+};
+
+
+
+void FCustomizableObjectDetails::CustomizeDetails( IDetailLayoutBuilder& DetailBuilder )
+{
+	const UCustomizableObject* CustomObject = 0;
+
+#if (ENGINE_MAJOR_VERSION == 4 && ENGINE_MINOR_VERSION >= 18) || ENGINE_MAJOR_VERSION >= 5
+	const IDetailsView* DetailsView = DetailBuilder.GetDetailsView();
+	if (DetailsView->GetSelectedObjects().Num())
+	{
+		CustomObject = Cast<const UCustomizableObject>(DetailsView->GetSelectedObjects()[0].Get());
+	}
+#else
+	const IDetailsView& DetailsView = DetailBuilder.GetDetailsView();
+	if (DetailsView.GetSelectedObjects().Num())
+	{
+		CustomObject = Cast<const UCustomizableObject>(DetailsView.GetSelectedObjects()[0].Get());
+	}
+#endif
+
+	IDetailCategoryBuilder& StatesCategory = DetailBuilder.EditCategory( "States" );
+	//StatesCategory.CategoryIcon( "ActorClassIcon.CustomizableObject" );
+
+	DetailBuilder.HideProperty("States");
+	DetailBuilder.HideProperty("CustomizableObjectClassTags");
+	DetailBuilder.HideProperty("PopulationClassTags");
+
+	StatesCategory.AddCustomRow( LOCTEXT("FCustomizableObjectDetails States", "States") )
+	[
+		SNew( SVerticalBox )
+		+ SVerticalBox::Slot()
+		.Padding( 2.0f )
+		[
+			//SNew( SFilterableDetail, LOCTEXT( "States", "States" ), &StatesCategory )
+			//[
+				SAssignNew(StatesTree, STreeView<TSharedPtr< FStateDetailsNode > >)
+					//.Visibility(EVisibility::Collapsed)
+					.SelectionMode(ESelectionMode::Single)
+					.TreeItemsSource( &RootTreeItems )
+					// Called to child items for any given parent item
+					.OnGetChildren( this, &FCustomizableObjectDetails::OnGetChildrenForStateTree )
+					// Generates the actual widget for a tree item
+					.OnGenerateRow( this, &FCustomizableObjectDetails::OnGenerateRowForStateTree ) 
+
+					// Generates the right click menu.
+					//.OnContextMenuOpening(this, &SClassViewer::BuildMenuWidget)
+
+					// Find out when the user selects something in the tree
+					//.OnSelectionChanged( this, &SClassViewer::OnClassViewerSelectionChanged )
+
+					// Allow for some spacing between items with a larger item height.
+					.ItemHeight(20.0f)
+
+					.HeaderRow
+					(
+						SNew(SHeaderRow)
+						.Visibility(EVisibility::Collapsed)
+						+ SHeaderRow::Column(TEXT("State"))
+						.DefaultLabel(NSLOCTEXT("CustomizableObjectDetails","State","State"))
+					)
+			//]
+		]
+//		+ SVerticalBox::Slot()
+//		.Padding( 2.0f )
+//		[
+//			SNew( SFilterableDetail, LOCTEXT( "SaveCollisionFromBuilderBrush", "Save Collision from Brush" ), &StatesCategory )
+//			[
+//				SNew( SButton )
+////				.ToolTipText( Commands.SaveBrushAsCollision->GetDescription() )
+//				//.OnClicked( this, &FCustomizableObjectDetails::OnSetCollisionFromBuilder )
+//				.VAlign( VAlign_Center )
+//				[
+//					SNew( STextBlock )
+//					.Text( LOCTEXT( "SaveCollisionFromBuilderBrush", "Save Collision from Brush" ) )
+//					.Font( IDetailLayoutBuilder::GetDetailFont() )
+//				]
+//			]
+//		]
+	];
+
+	if ( CustomObject &&  CustomObject->GetModel() )
+	{
+		uint32 numElements = CustomObject->GetModel()->GetStateCount();
+		RootTreeItems.SetNumUninitialized(0);
+		for ( uint32 i=0; i<numElements; ++i )
+		{
+			TSharedPtr<FStateDetailsNode> node = MakeShareable( new FStateDetailsNode );
+
+			node->Model = CustomObject->GetModel();
+			node->StateIndex = i;
+
+			RootTreeItems.Add( node );
+		}
+	}
+
+	TSharedRef<IPropertyHandle> Property = DetailBuilder.GetProperty("ReferenceSkeletalMeshes");
+
+#if (ENGINE_MAJOR_VERSION == 4 && ENGINE_MINOR_VERSION >= 18) || ENGINE_MAJOR_VERSION >= 5
+	if (Property->IsValidHandle() && CustomObject)
+	{
+		if (CustomObject->bIsChildObject)
+		{
+			Property->MarkHiddenByCustomization();
+		}
+		else
+		{
+			Property->MarkResetToDefaultCustomized();
+		}
+	}
+#endif
+
+	StatesTree->RequestTreeRefresh();
+}
+
+
+/** The item used for visualizing the class in the tree. */
+class SStateItem : public STableRow< TSharedPtr<FString> >
+{
+public:
+	
+	SLATE_BEGIN_ARGS( SStateItem )
+		: _StateName()
+		, _OwnerTableView()
+		, _HighlightText(NULL)
+		, _TextColor(FLinearColor(1.0f, 1.0f, 1.0f, 1.0f))
+		{}
+
+		/** The classname this item contains. */
+		SLATE_ARGUMENT( TSharedPtr<FString>, StateName )
+		/** The table that owns this item. */
+		SLATE_ARGUMENT( TSharedPtr<STableViewBase>, OwnerTableView )
+		/** The text this item should highlight, if any. */
+		SLATE_ARGUMENT( const FString*, HighlightText )
+		/** The color text this item will use. */
+		SLATE_ARGUMENT( FSlateColor, TextColor )
+		/** The node this item is associated with. */
+		SLATE_ARGUMENT( TSharedPtr<FStateDetailsNode>, AssociatedNode)
+
+	SLATE_END_ARGS()
+
+	/**
+	 * Construct the widget
+	 *
+	 * @param InArgs   A declaration from which to construct the widget
+	 */
+	void Construct( const FArguments& InArgs )
+	{
+		StateName = InArgs._StateName;
+		AssociatedNode = InArgs._AssociatedNode;
+
+		this->ChildSlot
+		[
+			SNew(SHorizontalBox)
+				
+			+SHorizontalBox::Slot()
+				[
+					SNew( SExpanderArrow, SharedThis(this) )
+				]
+
+			+SHorizontalBox::Slot()
+				.FillWidth(1.0f)
+				.Padding( 0.0f, 3.0f, 6.0f, 3.0f )
+				.VAlign(VAlign_Center)
+				[
+					SNew( STextBlock )
+						.Text( FText::FromString(*StateName.Get()) )
+						//.HighlightString(*InArgs._HighlightText)
+						//.ColorAndOpacity( this, &SClassItem::GetTextColor)
+				]
+
+			//+SHorizontalBox::Slot()
+			//	.HAlign(HAlign_Right)
+			//	.VAlign(VAlign_Center)
+			//	.Padding( 0.0f, 0.0f, 6.0f, 0.0f )
+			//	[
+			//		SNew( SComboButton )
+			//			.ContentPadding(FMargin(2.0f))
+			//			//.OnGetMenuContent(this, &SClassItem::GenerateDropDown)
+			//	]
+		];
+		
+		TextColor = InArgs._TextColor;
+
+		STableRow< TSharedPtr<FString> >::ConstructInternal( 
+			STableRow::FArguments()
+				.ShowSelection(true),
+			InArgs._OwnerTableView.ToSharedRef() );	
+	}
+
+private:
+
+	/**
+	 * Generates the drop down menu for the item.
+	 *
+	 * @return		The drop down menu widget.
+	 */
+	/*
+	TSharedRef<SWidget> GenerateDropDown()
+	{
+		// Empty list of commands.
+		TSharedPtr< FUICommandList > Commands;
+
+		const bool bShouldCloseWindowAfterMenuSelection = true;	// Set the menu to automatically close when the user commits to a choice
+		FMenuBuilder MenuBuilder( bShouldCloseWindowAfterMenuSelection, Commands);
+		{ 
+			if (UClass* Class = AssociatedNode->Class.Get())
+			{
+				bool bIsBlueprint(false);
+				bool bHasBlueprint(false);
+
+				ClassViewer::Helpers::GetClassInfo(Class, bIsBlueprint, bHasBlueprint);
+
+				bHasBlueprint = AssociatedNode->Blueprint.IsValid();
+
+				if (bIsBlueprint)
+				{
+					FUIAction Action( FExecuteAction::CreateStatic( &ClassViewer::Helpers::CreateBlueprint, Class ) );
+					MenuBuilder.AddMenuEntry(TEXT("Create Blueprint..."), TEXT("Creates a new Blueprint from the current one."), NAME_None, Action);
+				}
+
+				if (bHasBlueprint)
+				{
+					MenuBuilder.AddMenuSeparator();
+
+					FUIAction Action( FExecuteAction::CreateStatic( &ClassViewer::Helpers::OpenBlueprintTool, ClassViewer::Helpers::GetBlueprint(Class) ) );
+					MenuBuilder.AddMenuEntry(TEXT("Open Blueprint Editor"), TEXT("Open the Blueprint in the editor."), NAME_None, Action);
+				}
+
+				if (bHasBlueprint)
+				{
+					MenuBuilder.AddMenuSeparator();
+
+					FUIAction Action( FExecuteAction::CreateStatic( &ClassViewer::Helpers::FindInContentBrowser, ClassViewer::Helpers::GetBlueprint(Class), Class ) );
+					MenuBuilder.AddMenuEntry(TEXT("Find In Content Browser"), TEXT("Find in Content Browser"), NAME_None, Action);
+				}
+			}
+		}
+
+		return MenuBuilder.MakeWidget();
+	}
+	*/
+
+	/** Returns the text color for the item based on if it is selected or not. */
+	FSlateColor GetTextColor() const
+	{
+		const TSharedPtr< ITypedTableView< TSharedPtr<FString> > > OwnerWidget = OwnerTablePtr.Pin();
+		const TSharedPtr<FString>* MyItem = OwnerWidget->Private_ItemFromWidget( this );
+		const bool bIsSelected = OwnerWidget->Private_IsItemSelected( *MyItem );
+
+		if(bIsSelected)
+		{
+			return FSlateColor::UseForeground();
+		}
+
+		return TextColor;
+	}
+
+private:
+
+	/** The class name for which this item is associated with. */
+	TSharedPtr<FString> StateName;
+
+	/** The text color for this item. */
+	FSlateColor TextColor;
+
+	/** The Class Viewer Node this item is associated with. */
+	TSharedPtr< FStateDetailsNode > AssociatedNode;
+
+};
+
+
+TSharedRef< ITableRow > FCustomizableObjectDetails::OnGenerateRowForStateTree( TSharedPtr<FStateDetailsNode> Item, const TSharedRef< STableViewBase >& OwnerTable )
+{	
+	// If the item was accepted by the filter, leave it bright, otherwise dim it.
+	float AlphaValue = 1.0f;
+	TSharedRef< ITableRow > ReturnRow = SNew(SStateItem)
+		.StateName(Item->GetName())
+		//.HighlightText(&SearchBox->GetText())
+		.TextColor(FLinearColor(1.0f, 1.0f, 1.0f, AlphaValue))
+		.AssociatedNode(Item)
+		.OwnerTableView( OwnerTable );
+
+	// Expand the item if needed.
+	//bool* bIsExpanded = ExpansionStateMap.Find( Item->GetName() );
+	//if( bIsExpanded && *bIsExpanded )
+	//{
+	//	StatesTree->SetItemExpansion( Item, *bIsExpanded );
+	//}
+
+	return ReturnRow;
+}
+
+
+void FCustomizableObjectDetails::OnGetChildrenForStateTree( TSharedPtr<FStateDetailsNode> InParent, TArray< TSharedPtr< FStateDetailsNode > >& OutChildren )
+{
+	// Simply return the children, it's already setup.
+	OutChildren = InParent->GetChildrenList();
+}
+
+#undef LOCTEXT_NAMESPACE
diff --git a/Engine/Plugins/Experimental/Mutable/Source/CustomizableObjectEditor/Private/CustomizableObjectDetails.h b/Engine/Plugins/Experimental/Mutable/Source/CustomizableObjectEditor/Private/CustomizableObjectDetails.h
new file mode 100644
--- /dev/null
+++ b/Engine/Plugins/Experimental/Mutable/Source/CustomizableObjectEditor/Private/CustomizableObjectDetails.h
@@ -0,0 +1,33 @@
+// Copyright Epic Games, Inc. All Rights Reserved.
+
+#pragma once
+
+#include "IDetailCustomization.h"
+#include "Widgets/Views/STreeView.h"
+
+#include "Templates/SharedPointer.h"
+#include "UObject/WeakObjectPtr.h"
+
+class FStateDetailsNode;
+
+
+class FCustomizableObjectDetails : public IDetailCustomization
+{
+public:
+	/** Makes a new instance of this detail layout class for a specific detail view requesting it */
+	static TSharedRef<IDetailCustomization> MakeInstance();
+
+	/** ILayoutDetails interface */
+	virtual void CustomizeDetails( IDetailLayoutBuilder& DetailBuilder ) override;
+
+private:
+
+	TSharedPtr<STreeView<TSharedPtr< FStateDetailsNode > >> StatesTree;
+
+	/** The root items to be displayed in the state tree. */
+	TArray< TSharedPtr< FStateDetailsNode > > RootTreeItems;
+	
+	TSharedRef< ITableRow > OnGenerateRowForStateTree( TSharedPtr<FStateDetailsNode> Item, const TSharedRef< STableViewBase >& OwnerTable );
+	void OnGetChildrenForStateTree( TSharedPtr<FStateDetailsNode> InParent, TArray< TSharedPtr< FStateDetailsNode > >& OutChildren );
+
+};
\ No newline at end of file
