[mutable] Move Mutable plugin from the restricted folder to the experimental plugins folder.

diff --git a/Engine/Plugins/Experimental/Mutable/Source/CustomizableObjectEditor/Private/CustomizableObjectEditor.cpp b/Engine/Plugins/Experimental/Mutable/Source/CustomizableObjectEditor/Private/CustomizableObjectEditor.cpp
new file mode 100644
--- /dev/null
+++ b/Engine/Plugins/Experimental/Mutable/Source/CustomizableObjectEditor/Private/CustomizableObjectEditor.cpp
@@ -0,0 +1,3141 @@
+// Copyright Epic Games, Inc. All Rights Reserved.
+
+#include "CustomizableObjectEditor.h"
+#include "CustomizableObjectEditorModule.h"
+#include "CustomizableObjectEditorStyle.h"
+
+#include "CustomizableObjectInstance.h"
+#include "CustomizableObjectEditorActions.h"
+#include "CustomizableObjectPreviewScene.h"
+#include "CustomizableObjectDebugger.h"
+#include "SCustomizableObjectEditorViewport.h"
+#include "SCustomizableObjectProperties.h"
+#include "SCustomizableObjectSystemVersion.h"
+#include "SCustomizableObjectNodeLayoutBlocksEditor.h"
+#include "SCustomizableObjectNodeLayoutBlocksSelector.h"
+#include "CustomizableObjectPreviewScene.h"
+#include "PinViewer/SPinViewer.h"
+#include "SCustomizableObjectEditorTextureAnalyzer.h"
+#include "SCustomizableObjectEditorPerformanceReport.h"
+#include "SCustomizableObjectEditorTagExplorer.h"
+#include "GraphTraversal.h"
+
+#include "Editor.h"
+#include "AssetRegistry/AssetRegistryModule.h"
+#include "DesktopPlatformModule.h"
+#include "ContentBrowserModule.h"
+#include "PropertyEditorModule.h"
+
+#include "Interfaces/ITargetPlatform.h"
+#include "Interfaces/ITargetPlatformManagerModule.h"
+#include "Widgets/Notifications/SNotificationList.h"
+#include "Widgets/Docking/SDockTab.h"
+#include "Widgets/Input/STextComboBox.h"
+#include "Widgets/Layout/SSplitter.h"
+#include "Widgets/Layout/SSpacer.h"
+#include "Widgets/Input/SSearchBox.h"
+#include "Types/SlateEnums.h"
+#include "Slate/SceneViewport.h"
+#include "SNodePanel.h"
+#include "SGraphNode.h"
+#include "Logging/MessageLog.h"
+#include "Logging/TokenizedMessage.h"
+
+#include "PreviewScene.h"
+#include "ScopedTransaction.h"
+#include "FileHelpers.h"
+#include "BusyCursor.h"
+#include "EditorDirectories.h"
+
+#include "WorkspaceMenuStructureModule.h"
+#include "WorkspaceMenuStructure.h"
+#include "Engine/Selection.h"
+#include "MaterialGraph/MaterialGraphNode.h"
+#include "MaterialGraph/MaterialGraphNode_Comment.h"
+#include "Components/SkeletalMeshComponent.h"
+
+#include "GraphEditorActions.h"
+#include "EdGraphUtilities.h"
+#include "EdGraphSchema_CustomizableObject.h"
+#include "IDetailsView.h"
+#include "WorkspaceMenuStructure.h"
+#include "Framework/Commands/GenericCommands.h"
+#include "Framework/Notifications/NotificationManager.h"
+#include "Framework/MultiBox/MultiBoxExtender.h"
+
+#include "Runtime/Launch/Resources/Version.h"
+#include "AdvancedPreviewSceneModule.h"
+#include "CustomizableObjectCustomSettings.h"
+#include "SCustomizableObjectEditorAdvancedPreviewSettings.h"
+#include "CustomizableObjectEditorNodeContextCommands.h"
+#include "Animation/AnimSingleNodeInstance.h"
+#include "AdvancedPreviewScene.h"
+#include "CustomizableObjectBakeHelpers.h"
+#include "UnrealEditorPortabilityHelpers.h"
+#include "GenericPlatform/GenericPlatformApplicationMisc.h"
+#include "ISinglePropertyView.h"
+#include "Misc/UObjectToken.h"
+
+#include "Widgets/Layout/SScrollBox.h"
+#include "Widgets/Layout/SScrollBar.h"
+
+
+#if (ENGINE_MAJOR_VERSION == 4 && ENGINE_MINOR_VERSION >= 18) || ENGINE_MAJOR_VERSION >= 5
+#include "HAL/PlatformApplicationMisc.h"
+#endif
+
+
+#define LOCTEXT_NAMESPACE "CustomizableObjectEditor"
+
+DEFINE_LOG_CATEGORY_STATIC(LogCustomizableObjectEditor, Log, All);
+
+const FName FCustomizableObjectEditor::ViewportTabId( TEXT( "CustomizableObjectEditor_Viewport" ) );
+const FName FCustomizableObjectEditor::ObjectPropertiesTabId( TEXT( "CustomizableObjectEditor_ObjectProperties" ) );
+const FName FCustomizableObjectEditor::InstancePropertiesTabId( TEXT( "CustomizableObjectEditor_InstanceProperties" ) );
+const FName FCustomizableObjectEditor::GraphTabId( TEXT( "CustomizableObjectEditor_Graph" ) );
+const FName FCustomizableObjectEditor::GraphNodePropertiesTabId( TEXT( "CustomizableObjectEditor_GraphNodeProperties" ) );
+const FName FCustomizableObjectEditor::AdvancedPreviewSettingsTabId(TEXT("CustomizableObjectEditor_AdvancedPreviewSettings"));
+const FName FCustomizableObjectEditor::TextureAnalyzerTabId(TEXT("CustomizableObjectEditor_TextureAnalyzer"));
+const FName FCustomizableObjectEditor::PerformanceReportTabId(TEXT("CustomizableObjectEditor_PerformanceReport"));
+const FName FCustomizableObjectEditor::TagExplorerTabId(TEXT("CustomizableObjectEditor_TagExplorer"));
+const FName FCustomizableObjectEditor::ObjectDebuggerTabId(TEXT("CustomizableObjectEditor_ObjectDebugger"));
+const FName FCustomizableObjectEditor::PopulationClassTagManagerTabId(TEXT("CustomizableObjectEditor_PopulationClassTabManager"));
+
+
+void UUpdateClassWrapper::DelegatedCallback(UCustomizableObjectInstance* Instance)
+{
+	Delegate.ExecuteIfBound();
+}
+
+
+void FCustomizableObjectEditor::RegisterTabSpawners(const TSharedRef<class FTabManager>& InTabManager)
+{
+	WorkspaceMenuCategory = InTabManager->AddLocalWorkspaceMenuCategory(LOCTEXT("WorkspaceMenu_CustomizableObjectEditor", "Customizable Object Editor"));
+	TSharedRef<FWorkspaceItem> WorkspaceMenuCategoryRef = WorkspaceMenuCategory.ToSharedRef();
+	
+	FAssetEditorToolkit::RegisterTabSpawners(InTabManager);
+
+	InTabManager->RegisterTabSpawner(ViewportTabId, FOnSpawnTab::CreateSP(this, &FCustomizableObjectEditor::SpawnTab_Viewport))
+		.SetDisplayName(LOCTEXT("ViewportTab", "Viewport"))
+		.SetGroup(WorkspaceMenuCategoryRef);
+	
+	InTabManager->RegisterTabSpawner(ObjectPropertiesTabId, FOnSpawnTab::CreateSP(this, &FCustomizableObjectEditor::SpawnTab_ObjectProperties))
+		.SetDisplayName(LOCTEXT("ObjectPropertiesTab", "Object Properties"))
+		.SetGroup(WorkspaceMenuCategoryRef);
+	
+	InTabManager->RegisterTabSpawner(InstancePropertiesTabId, FOnSpawnTab::CreateSP(this, &FCustomizableObjectEditor::SpawnTab_InstanceProperties))
+		.SetDisplayName(LOCTEXT("InstancePropertiesTab", "Instance Properties"))
+		.SetGroup(WorkspaceMenuCategoryRef);
+	
+	InTabManager->RegisterTabSpawner(GraphTabId, FOnSpawnTab::CreateSP(this, &FCustomizableObjectEditor::SpawnTab_Graph))
+		.SetDisplayName(LOCTEXT("GraphTab", "Object Graph"))
+		.SetGroup(WorkspaceMenuCategoryRef);
+	
+	InTabManager->RegisterTabSpawner(GraphNodePropertiesTabId, FOnSpawnTab::CreateSP(this, &FCustomizableObjectEditor::SpawnTab_GraphNodeProperties))
+		.SetDisplayName(LOCTEXT("GraphNodePropertiesTab", "Object Graph Node Properties"))
+		.SetGroup(WorkspaceMenuCategoryRef);
+
+	InTabManager->RegisterTabSpawner(AdvancedPreviewSettingsTabId, FOnSpawnTab::CreateSP(this, &FCustomizableObjectEditor::SpawnTab_AdvancedPreviewSettings))
+		.SetDisplayName(LOCTEXT("AdvancedPreviewSettingsTab", "Advanced Preview Settings"))
+		.SetGroup(WorkspaceMenuCategoryRef);
+
+	InTabManager->RegisterTabSpawner(TextureAnalyzerTabId, FOnSpawnTab::CreateSP(this, &FCustomizableObjectEditor::SpawnTab_TextureAnalyzer))
+		.SetDisplayName(LOCTEXT("TextureAnalyzer", "Texture Analyzer"))
+		.SetGroup(WorkspaceMenuCategoryRef);
+
+	InTabManager->RegisterTabSpawner(PerformanceReportTabId, FOnSpawnTab::CreateSP(this, &FCustomizableObjectEditor::SpawnTab_PerformanceReport))
+		.SetDisplayName(LOCTEXT("PerformanceReport", "Performance Report"))
+		.SetGroup(WorkspaceMenuCategoryRef);
+
+	InTabManager->RegisterTabSpawner(TagExplorerTabId, FOnSpawnTab::CreateSP(this, &FCustomizableObjectEditor::SpawnTab_TagExplorer))
+		.SetDisplayName(LOCTEXT("TagExplorerTab", "Tag Explorer"))
+		.SetGroup(WorkspaceMenuCategoryRef);
+}
+
+
+void FCustomizableObjectEditor::UnregisterTabSpawners(const TSharedRef<class FTabManager>& InTabManager)
+{
+	InTabManager->UnregisterTabSpawner( ViewportTabId );
+	InTabManager->UnregisterTabSpawner( ObjectPropertiesTabId );
+	InTabManager->UnregisterTabSpawner( InstancePropertiesTabId );
+	InTabManager->UnregisterTabSpawner( GraphTabId );
+	InTabManager->UnregisterTabSpawner( GraphNodePropertiesTabId );
+	InTabManager->UnregisterTabSpawner(AdvancedPreviewSettingsTabId);
+	InTabManager->UnregisterTabSpawner(TextureAnalyzerTabId);
+	InTabManager->UnregisterTabSpawner(PerformanceReportTabId);
+}
+
+
+TSharedRef<FCustomizableObjectEditor> FCustomizableObjectEditor::Create(const EToolkitMode::Type Mode, const TSharedPtr<IToolkitHost>& InitToolkitHost, UCustomizableObject* ObjectToEdit)
+{
+	TSharedRef<FCustomizableObjectEditor> Editor = MakeShareable(new FCustomizableObjectEditor(*ObjectToEdit));
+	Editor->InitCustomizableObjectEditor(Mode, InitToolkitHost);
+	return Editor;
+}
+
+
+FCustomizableObjectEditor::~FCustomizableObjectEditor()
+{
+	if (PreviewInstance && HelperCallback)
+	{
+		PreviewInstance->UpdatedDelegate.RemoveDynamic(HelperCallback, &UUpdateClassWrapper::DelegatedCallback);
+		HelperCallback->Delegate.Unbind();
+	}
+
+	if (PreviewInstance)
+	{
+		if (PreviewInstance->bSelectedProfileDirty && PreviewInstance->SelectedProfileIndex != INDEX_NONE)
+		{
+			PreviewInstance->SaveParametersToProfile(PreviewInstance->SelectedProfileIndex);
+		}
+	}
+	
+	for (UCustomizableSkeletalComponent* Component : PreviewCustomizableSkeletalComponents)
+	{
+		Component->DestroyComponent();
+	}
+
+	PreviewCustomizableSkeletalComponents.Empty();
+	CustomizableObjectDetailsView.Reset();
+	GEditor->UnregisterForUndo(this);
+	Compiler.ForceFinishCompilation();
+
+	if (Compiler.GetAsynchronousStreamableHandlePtr().IsValid() && !Compiler.GetAsynchronousStreamableHandlePtr()->HasLoadCompleted())
+	{
+		Compiler.GetAsynchronousStreamableHandlePtr()->CancelHandle();
+		Compiler.ForceFinishBeforeStartCompilation(CustomizableObject);
+	}
+
+	FCoreUObjectDelegates::OnObjectModified.Remove(OnObjectModifiedHandle);
+
+	FAssetRegistryModule& AssetRegistryModule = FModuleManager::LoadModuleChecked<FAssetRegistryModule>(TEXT("AssetRegistry"));
+	AssetRegistryModule.Get().OnFilesLoaded().RemoveAll(this);
+
+	FCoreUObjectDelegates::OnObjectPropertyChanged.RemoveAll(this);
+
+	GEngine->ForceGarbageCollection(true);
+}
+
+
+void FCustomizableObjectEditor::InitCustomizableObjectEditor(const EToolkitMode::Type Mode, const TSharedPtr<IToolkitHost>& InitToolkitHost)
+{
+	FAssetRegistryModule& AssetRegistryModule = FModuleManager::LoadModuleChecked<FAssetRegistryModule>(TEXT("AssetRegistry"));
+	if (AssetRegistryModule.Get().IsLoadingAssets())
+	{
+		AssetRegistryModule.Get().OnFilesLoaded().AddRaw(this, &FCustomizableObjectEditor::OnAssetRegistryLoadComplete);
+	}
+	else
+	{
+		AssetRegistryLoaded = true;
+	}
+
+	HelperCallback = nullptr;
+
+	// Support undo/redo
+	CustomizableObject->SetFlags(RF_Transactional);
+
+	GEditor->RegisterForUndo(this);
+
+	// Register our commands. This will only register them if not previously registered
+	FGraphEditorCommands::Register();
+	FCustomizableObjectEditorCommands::Register();
+	FCustomizableObjectEditorViewportCommands::Register();
+	FCustomizableObjectEditorNodeContextCommands::Register();
+
+	BindCommands();
+
+	FPropertyEditorModule& PropPlugin = FModuleManager::LoadModuleChecked<FPropertyEditorModule>("PropertyEditor");
+
+	FDetailsViewArgs DetailsViewArgs;
+	DetailsViewArgs.NotifyHook = this;
+	DetailsViewArgs.NameAreaSettings = FDetailsViewArgs::ENameAreaSettings::HideNameArea;
+	DetailsViewArgs.bAllowSearch = false;
+	//DetailsViewArgs.bShowActorLabel = false;
+	DetailsViewArgs.bShowObjectLabel = false;
+
+	CustomizableObjectDetailsView = PropPlugin.CreateDetailView( DetailsViewArgs );
+	CustomizableObjectDetailsView->SetObject(CustomizableObject);
+
+	CustomizableInstanceDetailsView = PropPlugin.CreateDetailView( DetailsViewArgs );	
+
+	FDetailsViewArgs GraphDetailsViewArgs;
+	GraphDetailsViewArgs.NotifyHook = this;
+	GraphDetailsViewArgs.NameAreaSettings = FDetailsViewArgs::ENameAreaSettings::HideNameArea;
+	GraphDetailsViewArgs.bAllowSearch = false;
+	//GraphDetailsViewArgs.bShowActorLabel = false;
+	GraphDetailsViewArgs.bShowObjectLabel = false;
+	GraphNodeDetailsView = PropPlugin.CreateDetailView( GraphDetailsViewArgs );
+
+	Viewport = SNew(SCustomizableObjectEditorViewportTabBody)
+		.CustomizableObjectEditor(SharedThis(this));
+
+	Viewport->SetCustomizableObject(CustomizableObject);
+	Viewport->SetAssetRegistryLoaded(AssetRegistryLoaded);
+	ViewportClient = Viewport->GetViewportClient();
+
+	// \TODO: Create only when needed?
+	TextureAnalyzer = SNew(SCustomizableObjecEditorTextureAnalyzer).CustomizableObjectEditor(this).CustomizableObjectInstanceEditor(nullptr);
+
+	// \TODO: Create only when needed?
+	TagExplorer = SNew(SCustomizableObjectEditorTagExplorer).CustomizableObjectEditor(this);
+	
+	AdditionalSettings = NewObject<UCustomizableObjectEmptyClassForSettings>();
+	if(AdditionalSettings != nullptr)
+		{
+		AdditionalSettings->Viewport = Viewport;
+		AdditionalSettings->PreviewSkeletalMeshComp = !PreviewSkeletalMeshComponents.IsEmpty() ? &PreviewSkeletalMeshComponents[0] : nullptr;
+	}
+
+	FAdvancedPreviewSceneModule& AdvancedPreviewSceneModule = FModuleManager::LoadModuleChecked<FAdvancedPreviewSceneModule>("AdvancedPreviewScene");
+
+	TSharedPtr<FAdvancedPreviewScene> AdvancedPreviewScene = StaticCastSharedPtr<FAdvancedPreviewScene>(Viewport->GetPreviewScene());
+
+	CustomizableObjectEditorAdvancedPreviewSettings =
+		SNew(SCustomizableObjectEditorAdvancedPreviewSettings, AdvancedPreviewScene.ToSharedRef())
+		.AdditionalSettings(AdditionalSettings)
+		.CustomizableObjectEditor(this);
+	CustomizableObjectEditorAdvancedPreviewSettings->LoadProfileEnvironment();
+	AdvancedPreviewSettingsWidget = CustomizableObjectEditorAdvancedPreviewSettings;
+
+
+	const TSharedRef<FTabManager::FLayout> StandaloneDefaultLayout = FTabManager::NewLayout( "Standalone_CustomizableObjectEditor_Layout_v1.4" )
+	->AddArea
+	(
+		FTabManager::NewPrimaryArea()->SetOrientation(Orient_Horizontal)
+		->Split
+		(
+			FTabManager::NewStack()
+			->SetSizeCoefficient(0.6f)
+			->SetHideTabWell(true)
+			->AddTab(GraphTabId, ETabState::OpenedTab)
+		)
+		->Split
+		(
+			FTabManager::NewSplitter() ->SetOrientation(Orient_Horizontal) 
+			->SetSizeCoefficient(0.4f)
+			->Split
+			(
+				FTabManager::NewSplitter() ->SetOrientation(Orient_Vertical) 
+				->SetSizeCoefficient(0.5f)
+				->Split
+				(
+					FTabManager::NewStack()
+					->SetSizeCoefficient(0.4f)
+					->AddTab(ObjectPropertiesTabId, ETabState::OpenedTab)
+					->AddTab(TagExplorerTabId, ETabState::OpenedTab)
+					->SetForegroundTab(ObjectPropertiesTabId)
+				)
+				->Split
+				(
+					FTabManager::NewStack()
+					->SetSizeCoefficient(0.6f)
+					->AddTab(GraphNodePropertiesTabId, ETabState::OpenedTab)
+				)
+			)
+			->Split
+			(
+				FTabManager::NewSplitter() ->SetOrientation(Orient_Vertical) 
+				->SetSizeCoefficient(0.5f)
+				->Split
+				(
+					FTabManager::NewStack()
+					->SetSizeCoefficient(0.5f)
+					->AddTab(InstancePropertiesTabId, ETabState::OpenedTab)
+					->AddTab(AdvancedPreviewSettingsTabId, ETabState::OpenedTab)
+					->SetForegroundTab(InstancePropertiesTabId)
+				)
+				->Split
+				(
+					FTabManager::NewStack()
+					->SetSizeCoefficient(0.5f)
+					->AddTab(ViewportTabId, ETabState::OpenedTab)
+					->SetHideTabWell(true)
+				)
+			)
+		)
+	);
+
+	const bool bCreateDefaultStandaloneMenu = true;
+	const bool bCreateDefaultToolbar = true;
+	FAssetEditorToolkit::InitAssetEditor( Mode, InitToolkitHost, CustomizableObjectEditorAppIdentifier, StandaloneDefaultLayout, bCreateDefaultStandaloneMenu, bCreateDefaultToolbar, CustomizableObject);
+
+	ExtendToolbar();
+	RegenerateMenusAndToolbars();
+
+	// Clears selection highlight.
+	OnObjectPropertySelectionChanged(NULL);
+	OnInstancePropertySelectionChanged(NULL);
+	OnObjectModifiedHandle = FCoreUObjectDelegates::OnObjectModified.AddRaw(this, &FCustomizableObjectEditor::OnObjectModified);
+
+	FCoreUObjectDelegates::OnObjectPropertyChanged.AddSP(this, &FCustomizableObjectEditor::OnObjectPropertyChanged);
+
+	// Compile for the first time if necessary
+	if (!CustomizableObject->IsCompiled())
+	{
+		CompileObject();
+	}
+	else
+	{
+		CreatePreviewInstance();
+	}
+}
+
+
+FName FCustomizableObjectEditor::GetToolkitFName() const
+{
+	return FName("CustomizableObjectEditor");
+}
+
+
+FText FCustomizableObjectEditor::GetBaseToolkitName() const
+{
+	return LOCTEXT("ToolkitName", "Customizable Object Editor");
+}
+
+
+void FCustomizableObjectEditor::SelectNode(const UCustomizableObjectNode* Node)
+{
+	GraphEditor->JumpToNode(Node);
+}
+
+void FCustomizableObjectEditor::SetPoseAsset(class UPoseAsset* PoseAssetParameter)
+{
+	PoseAsset = PoseAssetParameter;
+
+	if (PoseAsset != nullptr)
+	{
+		Viewport->SetAnimation(nullptr, EAnimationMode::AnimationBlueprint);
+
+		for (UDebugSkelMeshComponent* PreviewSkeletalMeshComponent : PreviewSkeletalMeshComponents)
+		{
+			PreviewSkeletalMeshComponent->SetAnimationMode(EAnimationMode::AnimationBlueprint);
+			PreviewSkeletalMeshComponent->InitAnim(false);
+			PreviewSkeletalMeshComponent->SetAnimation(PoseAsset);
+
+			UAnimSingleNodeInstance* SingleNodeInstance = Cast<UAnimSingleNodeInstance>(PreviewSkeletalMeshComponent->GetAnimInstance());
+			if (SingleNodeInstance)
+			{
+				TArray<FSmartName> ArrayPoseSmartNames = PoseAsset->GetPoseNames();
+				for (int32 i = 0; i < ArrayPoseSmartNames.Num(); ++i)
+				{
+					SingleNodeInstance->SetPreviewCurveOverride(ArrayPoseSmartNames[i].DisplayName, 1.0f, false);
+				}
+			}
+		}
+	}
+}
+
+bool FCustomizableObjectEditor::CreatePreviewComponent(int32 ComponentIndex)
+{
+	if (!PreviewCustomizableSkeletalComponents.IsValidIndex(ComponentIndex))
+	{
+		PreviewCustomizableSkeletalComponents.AddZeroed(ComponentIndex - PreviewCustomizableSkeletalComponents.Num() + 1);
+		check(PreviewCustomizableSkeletalComponents.IsValidIndex(ComponentIndex) && !PreviewCustomizableSkeletalComponents.IsValidIndex(ComponentIndex + 1));
+	}
+
+	PreviewCustomizableSkeletalComponents[ComponentIndex] = NewObject<UCustomizableSkeletalComponent>(UCustomizableSkeletalComponent::StaticClass());
+
+	if (PreviewCustomizableSkeletalComponents[ComponentIndex])
+	{
+		PreviewCustomizableSkeletalComponents[ComponentIndex]->CustomizableObjectInstance = PreviewInstance;
+		PreviewCustomizableSkeletalComponents[ComponentIndex]->ComponentIndex = ComponentIndex;
+
+		if (!PreviewSkeletalMeshComponents.IsValidIndex(ComponentIndex))
+		{
+			PreviewSkeletalMeshComponents.AddZeroed(ComponentIndex - PreviewSkeletalMeshComponents.Num() + 1);
+			check(PreviewSkeletalMeshComponents.IsValidIndex(ComponentIndex) && !PreviewSkeletalMeshComponents.IsValidIndex(ComponentIndex + 1));
+		}
+
+		PreviewSkeletalMeshComponents[ComponentIndex] = NewObject<UDebugSkelMeshComponent>(GetTransientPackage(), NAME_None, RF_Transient);
+
+		if (PreviewSkeletalMeshComponents[ComponentIndex])
+		{
+			PreviewSkeletalMeshComponents[ComponentIndex]->UseInGameBounds(true);
+			PreviewCustomizableSkeletalComponents[ComponentIndex]->AttachToComponent(PreviewSkeletalMeshComponents[ComponentIndex], FAttachmentTransformRules::KeepRelativeTransform);
+
+			return true;
+		}
+	}
+
+	return false;
+}
+
+
+void FCustomizableObjectEditor::CreatePreviewInstance()
+{
+	OnUpdatePreviewInstanceWork.Add([this]()
+	{
+		ViewportClient->ResetCamera();
+	});
+
+	PreviewInstance = CustomizableObject->CreateInstance();
+
+	CustomizableInstanceDetailsView->SetObject(PreviewInstance, true);
+	
+	HelperCallback = NewObject<UUpdateClassWrapper>(GetTransientPackage());
+	PreviewInstance->UpdatedDelegate.AddDynamic(HelperCallback, &UUpdateClassWrapper::DelegatedCallback);
+
+	PreviewStaticMeshComponent = nullptr;
+
+	if (CustomizableObject->ReferenceSkeletalMeshes.Num())
+	{
+		CreatePreviewComponents();
+
+		if (!PreviewSkeletalMeshComponents.IsEmpty() && !PreviewCustomizableSkeletalComponents.IsEmpty())
+		{
+			HelperCallback->Delegate.BindSP(this, &FCustomizableObjectEditor::OnUpdatePreviewInstance);
+
+			if (AssetRegistryLoaded)
+			{
+				// Asset loading works in the Main Thread, there's a risk the sync loading could put to sleep the thread while
+				// waiting for the asset registry to load the content (that never gets executed)
+				Viewport->SetAssetRegistryLoaded(true);
+				PreviewInstance->bBuildParameterDecorations = true;
+				PreviewInstance->UpdateSkeletalMeshAsync(true, true);
+			}
+			else
+			{
+				FNotificationInfo Info(NSLOCTEXT("CustomizableObject", "CustomizableObjectCompileTryLater", "Please wait until asset registry loads all assets"));
+				Info.bFireAndForget = true;
+				Info.bUseThrobber = true;
+				Info.FadeOutDuration = 1.0f;
+				Info.ExpireDuration = 2.0f;
+				FSlateNotificationManager::Get().AddNotification(Info);
+				UpdateSkeletalMeshAfterAssetLoaded = true;
+			}
+		}
+		else
+		{
+			TArray<UDebugSkelMeshComponent*> EmptyArray;
+			Viewport->SetPreviewComponents(EmptyArray);
+			ViewportClient->SetReferenceMeshMissingWarningMessage(true);
+		}
+	}
+	else
+	{
+		ViewportClient->SetReferenceMeshMissingWarningMessage(true);
+	}
+
+	if (GraphEditor.IsValid())
+	{
+		for (TObjectPtr<UEdGraphNode>& Node : CustomizableObject->Source->Nodes)
+		{
+			GraphEditor->RefreshNode(*Node);
+		}
+	}
+}
+
+
+void FCustomizableObjectEditor::OnPreviewInstanceUpdated()
+{
+	RefreshViewport();
+}
+
+
+void FCustomizableObjectEditor::AddReferencedObjects( FReferenceCollector& Collector )
+{
+	Collector.AddReferencedObject( CustomizableObject );
+	Collector.AddReferencedObject( PreviewInstance );
+	Collector.AddReferencedObjects( PreviewCustomizableSkeletalComponents );
+	Collector.AddReferencedObject( PreviewStaticMeshComponent );
+	Collector.AddReferencedObjects( PreviewSkeletalMeshComponents );
+	Collector.AddReferencedObject( HelperCallback );
+	Collector.AddReferencedObject( AdditionalSettings );
+}
+
+
+FCustomizableObjectEditor::FCustomizableObjectEditor(UCustomizableObject& ObjectToEdit) :
+	CustomizableObject(&ObjectToEdit) {}
+
+
+TSharedRef<SDockTab> FCustomizableObjectEditor::SpawnTab_Viewport( const FSpawnTabArgs& Args )
+{
+	check( Args.GetTabId() == ViewportTabId );
+
+	TSharedRef<SDockTab> DockTab = SNew(SDockTab)
+		.Label( FText::FromString( GetTabPrefix() + LOCTEXT("CustomizableObjectViewport_TabTitle", "Preview Inst. Viewport").ToString() ) )
+		[
+			SNew( SVerticalBox )
+
+			+ SVerticalBox::Slot()
+			.Padding( 2.0f )
+			.FillHeight(1.0f)
+			[
+				Viewport.ToSharedRef()
+			]
+		];
+
+	DockTab->SetTabIcon(FCustomizableObjectEditorStyle::Get().GetBrush("CustomizableObjectEditor.Tabs.Preview"));
+
+	return DockTab;
+}
+
+TSharedRef<SDockTab> FCustomizableObjectEditor::SpawnTab_ObjectProperties( const FSpawnTabArgs& Args )
+{
+	check( Args.GetTabId() == ObjectPropertiesTabId );
+
+	TSharedRef<SDockTab> DockTab = SNew(SDockTab)
+		.Label( FText::FromString( GetTabPrefix() + LOCTEXT( "CustomizableObjectProperties_TabTitle", "Object Properties" ).ToString() ) )
+		[
+			CustomizableObjectDetailsView.ToSharedRef()
+			//CustomizableObjectProperties.ToSharedRef()
+		];
+
+	DockTab->SetTabIcon(FCustomizableObjectEditorStyle::Get().GetBrush("CustomizableObjectEditor.Tabs.CustomizableObjectProperties"));
+
+	return DockTab;
+}
+
+
+TSharedRef<SDockTab> FCustomizableObjectEditor::SpawnTab_InstanceProperties( const FSpawnTabArgs& Args )
+{
+	check( Args.GetTabId() == InstancePropertiesTabId );
+
+	TSharedRef<SDockTab> DockTab = SNew(SDockTab)
+		.Label( FText::FromString( GetTabPrefix() + LOCTEXT( "CustomizableInstanceProperties_TabTitle", "Preview Inst. Properties" ).ToString() ) )
+		[
+			CustomizableInstanceDetailsView.ToSharedRef()
+		];
+
+	DockTab->SetTabIcon(FCustomizableObjectEditorStyle::Get().GetBrush("CustomizableObjectEditor.Tabs.CustomizableInstanceProperties"));
+
+	return DockTab;
+}
+
+
+/** Create new tab for the supplied graph - don't call this directly, call SExplorer->FindTabForGraph.*/
+TSharedRef<SGraphEditor> FCustomizableObjectEditor::CreateGraphEditorWidget(UEdGraph* InGraph)
+{
+	check(InGraph != NULL);
+
+	// Check whether the graph has a base node and create one if it doesn't since it is required
+	bool bGraphHasBase = false;
+
+	for (TObjectPtr<UEdGraphNode> AuxNode : InGraph->Nodes)
+	{
+		UCustomizableObjectNodeObject* CustomizableObjectNodeObject = Cast<UCustomizableObjectNodeObject>(AuxNode);
+
+		if (CustomizableObjectNodeObject && CustomizableObjectNodeObject->bIsBase)
+		{
+			bGraphHasBase = true;
+			break;
+		}
+	}
+
+	if (!bGraphHasBase)
+	{
+		FCustomizableObjectSchemaAction_NewNode Action;
+		UCustomizableObjectNodeObject* NodeTemplate = NewObject<UCustomizableObjectNodeObject>();
+
+		Action.NodeTemplate = NodeTemplate;
+		Action.FCustomizableObjectSchemaAction_NewNode::PerformAction(InGraph, nullptr, FVector2D::ZeroVector, false);
+	}
+	
+	GraphEditorCommands = MakeShareable( new FUICommandList );
+	{
+		// Editing commands
+		GraphEditorCommands->MapAction( FGenericCommands::Get().Delete,
+			FExecuteAction::CreateSP( this, &FCustomizableObjectEditor::DeleteSelectedNodes ),
+			FCanExecuteAction::CreateSP( this, &FCustomizableObjectEditor::CanDeleteNodes )
+			);
+
+		GraphEditorCommands->MapAction( FGenericCommands::Get().Copy,
+			FExecuteAction::CreateSP( this, &FCustomizableObjectEditor::CopySelectedNodes ),
+			FCanExecuteAction::CreateSP( this, &FCustomizableObjectEditor::CanCopyNodes )
+			);
+
+		GraphEditorCommands->MapAction( FGenericCommands::Get().Paste,
+			FExecuteAction::CreateSP( this, &FCustomizableObjectEditor::PasteNodes ),
+			FCanExecuteAction::CreateSP( this, &FCustomizableObjectEditor::CanPasteNodes )
+			);
+
+		GraphEditorCommands->MapAction( FGenericCommands::Get().Cut,
+			FExecuteAction::CreateSP( this, &FCustomizableObjectEditor::CutSelectedNodes ),
+			FCanExecuteAction::CreateSP( this, &FCustomizableObjectEditor::CanCutNodes )
+			);
+
+		GraphEditorCommands->MapAction(FGenericCommands::Get().Duplicate,
+			FExecuteAction::CreateSP(this, &FCustomizableObjectEditor::DuplicateSelectedNodes),
+			FCanExecuteAction::CreateSP(this, &FCustomizableObjectEditor::CanDuplicateSelectedNodes)
+			);
+
+		GraphEditorCommands->MapAction(FCustomizableObjectEditorNodeContextCommands::Get().RefreshMaterialNodesInAllChildren,
+			FExecuteAction::CreateSP(this, &FCustomizableObjectEditor::RefreshMaterialNodesInAllChildrenCallback)
+			);
+
+		GraphEditorCommands->MapAction(FCustomizableObjectEditorNodeContextCommands::Get().CreateComment,
+			FExecuteAction::CreateSP(this, &FCustomizableObjectEditor::CreateCommentBoxFromKey)
+		);
+	}
+
+	TSharedRef<SWidget> TitleBarWidget =
+		SNew(SHorizontalBox)
+		+ SHorizontalBox::Slot()
+		.HAlign(HAlign_Left)
+		.FillWidth(10.f)
+		.Padding(5.f)
+		[
+			SNew(SSearchBox)
+			.HintText(LOCTEXT("Search", "Search..."))
+			.ToolTipText(LOCTEXT("Search Nodes, Properties or Values that contain the inserted words", "Search Nodes, Properties or Values that contain the inserted words"))
+			.OnTextCommitted(FOnTextCommitted::CreateSP(this, &FCustomizableObjectEditor::OnEnterText))
+			.SelectAllTextWhenFocused(true)
+		];
+
+	// Create the appearance info
+	FGraphAppearanceInfo AppearanceInfo;	
+	AppearanceInfo.CornerText = LOCTEXT("ApperanceCornerText", "MUTABLE");
+
+	SGraphEditor::FGraphEditorEvents InEvents;
+	InEvents.OnSelectionChanged = SGraphEditor::FOnSelectionChanged::CreateSP( this, &FCustomizableObjectEditor::OnSelectedGraphNodesChanged );
+	InEvents.OnTextCommitted = FOnNodeTextCommitted::CreateSP(this, &FCustomizableObjectEditor::OnNodeTitleCommitted);
+
+	// Make full graph editor
+	return SNew(SGraphEditor)
+		.AdditionalCommands(GraphEditorCommands)
+		.Appearance(AppearanceInfo)
+		.GraphToEdit(InGraph)
+		.GraphEvents(InEvents)
+		.TitleBar(TitleBarWidget)
+		.ShowGraphStateOverlay(false); // Removes graph state overlays (border and text) such as "SIMULATING" and "READ-ONLY"
+}
+
+
+TSharedRef<SDockTab> FCustomizableObjectEditor::SpawnTab_Graph( const FSpawnTabArgs& Args )
+{
+	check( Args.GetTabId().TabType == GraphTabId );
+
+	GraphEditor = CreateGraphEditorWidget(CustomizableObject->Source);
+
+	TSharedRef<SDockTab> DockTab = SNew(SDockTab)
+		.Label( FText::FromString( GetTabPrefix() + LOCTEXT( "SourceGraph", "Source Graph" ).ToString() ) )
+		.TabColorScale( GetTabColorScale() )
+		[
+			GraphEditor.ToSharedRef()
+		];
+
+	DockTab->SetTabIcon(FCustomizableObjectEditorStyle::Get().GetBrush("CustomizableObjectEditor.Tabs.NodeGraph"));
+
+	return DockTab;
+}
+
+
+TSharedRef<SDockTab> FCustomizableObjectEditor::SpawnTab_GraphNodeProperties( const FSpawnTabArgs& Args )
+{
+	check( Args.GetTabId().TabType == GraphNodePropertiesTabId );
+
+	TSharedRef<SScrollBar> ScrollBar = SNew(SScrollBar);
+
+	NodeDetailsSplitter = SNew( SSplitter )
+			.Orientation(Orient_Vertical)
+			+SSplitter::Slot()
+			.Value(0.5f)
+			[
+				SNew(SHorizontalBox)
+				+ SHorizontalBox::Slot()
+				.FillWidth(1)
+				[
+					SNew(SScrollBox)
+					.ExternalScrollbar(ScrollBar)
+					+ SScrollBox::Slot()
+					[
+						SNew(SVerticalBox)
+						+ SVerticalBox::Slot().AutoHeight()
+						[
+							GraphNodeDetailsView.ToSharedRef()
+						]
+					]
+				]
+				+ SHorizontalBox::Slot()
+				.AutoWidth()
+				[
+					SNew(SBox)
+					.WidthOverride(FOptionalSize(16))
+					[
+						ScrollBar
+					]
+				]
+			];
+
+	TSharedRef<SDockTab> DockTab = SNew(SDockTab)
+		.Label(FText::FromString(GetTabPrefix() + LOCTEXT("Graph Node Properties", "Node Properties").ToString()))
+		.TabColorScale(GetTabColorScale())
+		[
+			NodeDetailsSplitter.ToSharedRef()
+		];
+
+	DockTab->SetTabIcon(FCustomizableObjectEditorStyle::Get().GetBrush("CustomizableObjectEditor.Tabs.NodeProperties"));
+
+	return DockTab;
+}
+
+
+TSharedRef<SDockTab> FCustomizableObjectEditor::SpawnTab_AdvancedPreviewSettings(const FSpawnTabArgs& Args)
+{
+	check(Args.GetTabId() == AdvancedPreviewSettingsTabId);
+	TSharedRef<SDockTab> DockTab = SNew(SDockTab)
+		.Label(LOCTEXT("StaticMeshPreviewScene_TabTitle", "Preview Settings"))
+		[
+			AdvancedPreviewSettingsWidget.ToSharedRef()
+		];
+
+	DockTab->SetTabIcon(FCustomizableObjectEditorStyle::Get().GetBrush("CustomizableObjectEditor.Tabs.PreviewSettings"));
+
+	return DockTab;
+}
+
+
+UCustomizableObjectInstance* FCustomizableObjectEditor::GetPreviewInstance()
+{
+	return PreviewInstance;
+}
+
+
+void FCustomizableObjectEditor::CompileObjectUserPressedButton()
+{
+	if (UCustomizableObjectSystem::GetInstance()->IsCompilationDisabled())
+	{
+		UE_LOG(LogMutable, Warning, TEXT("Mutable compile is disabled in Editor. To enable it, go to Project Settings -> Plugins -> Mutable and unmark the option Disable Mutable Compile In Editor"));
+		return;
+	}
+
+	Compiler.ClearAllCompileOnlySelectedOption();
+
+	CustomizableObject->CompileOptions.bCheckChildrenGuids = true;
+	CompileObject();
+	CustomizableObject->CompileOptions.bCheckChildrenGuids = false;
+}
+
+
+void FCustomizableObjectEditor::CompileOnlySelectedObjectUserPressedButton()
+{
+	if (UCustomizableObjectSystem::GetInstance()->IsCompilationDisabled())
+	{
+		UE_LOG(LogMutable, Warning, TEXT("Mutable compile is disabled in Editor. To enable it, go to Project Settings -> Plugins -> Mutable and unmark the option Disable Mutable Compile In Editor"));
+		return;
+	}
+
+	if (PreviewInstance)
+	{
+		Compiler.ClearAllCompileOnlySelectedOption();
+
+		for (const FCustomizableObjectIntParameterValue& IntParam : PreviewInstance->GetIntParameters())
+		{
+			Compiler.AddCompileOnlySelectedOption(IntParam.ParameterName, IntParam.ParameterValueName);
+		}
+
+		CompileObject();
+	}
+}
+
+
+void FCustomizableObjectEditor::BindCommands()
+{
+	const FCustomizableObjectEditorCommands& Commands = FCustomizableObjectEditorCommands::Get();
+
+	// Toolbar
+	// Compile and options
+	ToolkitCommands->MapAction(
+		Commands.Compile,
+		FExecuteAction::CreateSP(this, &FCustomizableObjectEditor::CompileObjectUserPressedButton),
+		FCanExecuteAction(),
+		FIsActionChecked());
+
+	ToolkitCommands->MapAction(
+		Commands.CompileOnlySelected,
+		FExecuteAction::CreateSP(this, &FCustomizableObjectEditor::CompileOnlySelectedObjectUserPressedButton),
+		FCanExecuteAction(),
+		FIsActionChecked());
+
+	// Compile and options
+	ToolkitCommands->MapAction(
+		Commands.ResetCompileOptions,
+		FExecuteAction::CreateSP(this, &FCustomizableObjectEditor::ResetCompileOptions),
+		FCanExecuteAction(),
+		FIsActionChecked());
+
+	ToolkitCommands->MapAction(
+		Commands.CompileOptions_EnableTextureCompression,
+		FExecuteAction::CreateSP(this, &FCustomizableObjectEditor::CompileOptions_TextureCompression_Toggled),
+		FCanExecuteAction(),
+		FIsActionChecked::CreateSP(this, &FCustomizableObjectEditor::CompileOptions_TextureCompression_IsChecked));
+
+	ToolkitCommands->MapAction(
+		Commands.CompileOptions_UseParallelCompilation,
+		FExecuteAction::CreateSP(this, &FCustomizableObjectEditor::CompileOptions_UseParallelCompilation_Toggled),
+		FCanExecuteAction(),
+		FIsActionChecked::CreateSP(this, &FCustomizableObjectEditor::CompileOptions_UseParallelCompilation_IsChecked));
+
+	ToolkitCommands->MapAction(
+		Commands.CompileOptions_UseDiskCompilation,
+		FExecuteAction::CreateSP(this, &FCustomizableObjectEditor::CompileOptions_UseDiskCompilation_Toggled),
+		FCanExecuteAction(),
+		FIsActionChecked::CreateSP(this, &FCustomizableObjectEditor::CompileOptions_UseDiskCompilation_IsChecked));
+
+	// Debug and options
+	ToolkitCommands->MapAction(
+		Commands.Debug,
+		FExecuteAction::CreateSP(this, &FCustomizableObjectEditor::DebugObject),
+		FCanExecuteAction(),
+		FIsActionChecked());
+
+	// Texture Analyzer
+	ToolkitCommands->MapAction(
+		Commands.TextureAnalyzer,
+		FExecuteAction::CreateSP(this, &FCustomizableObjectEditor::OpenTextureAnalyzerTab),
+		FCanExecuteAction(),
+		FIsActionChecked());
+
+	// Performance Report
+	ToolkitCommands->MapAction(
+		Commands.PerformanceReport,
+		FExecuteAction::CreateSP(this, &FCustomizableObjectEditor::OpenPerformanceReportTab),
+		FCanExecuteAction(),
+		FIsActionChecked());
+
+	// Undo-Redo
+	ToolkitCommands->MapAction(
+		FGenericCommands::Get().Undo,
+		FExecuteAction::CreateSP(this, &FCustomizableObjectEditor::UndoGraphAction));
+
+	ToolkitCommands->MapAction(
+		FGenericCommands::Get().Redo,
+		FExecuteAction::CreateSP(this, &FCustomizableObjectEditor::RedoGraphAction));
+}
+
+
+void FCustomizableObjectEditor::UndoGraphAction()
+{
+	GEditor->UndoTransaction();
+}
+
+
+void FCustomizableObjectEditor::RedoGraphAction()
+{
+	// Clear selection, to avoid holding refs to nodes that go away
+	GraphEditor->ClearSelectionSet();
+
+	GEditor->RedoTransaction();
+}
+
+
+bool FCustomizableObjectEditor::GroupNodeIsLinkedToParentByName(UCustomizableObjectNodeObject* Node, UCustomizableObject* Test, const FString& ParentGroupName)
+{
+	TArray<UCustomizableObjectNodeObjectGroup*> GroupNodes;
+	Test->Source->GetNodesOfClass<UCustomizableObjectNodeObjectGroup>(GroupNodes);
+
+	for (UCustomizableObjectNodeObjectGroup* GroupNode : GroupNodes)
+	{
+		if ((Node->ParentObjectGroupId == GroupNode->NodeGuid) && (ParentGroupName == GroupNode->GroupName))
+		{
+			return true;
+		}
+	}
+
+	return false;
+}
+
+
+// Deprecated! Remove on MTBL-283
+void StoreNodeLinkedPins(UCustomizableObjectNode* Node, TMap<FString, int32>& MapLinkedPin)
+{
+	TArray<UEdGraphPin*> ArrayPins = Node->GetAllNonOrphanPins();
+	const int32 MaxIndex = ArrayPins.Num();
+
+	for (int32 i = 0; i < MaxIndex; ++i)
+	{
+		if (ArrayPins[i]->LinkedTo.Num() > 0)
+		{
+			MapLinkedPin.Emplace(ArrayPins[i]->GetDisplayName().ToString(), ArrayPins[i]->LinkedTo.Num());
+		}
+	}
+}
+
+
+// Deprecated! Remove on MTBL-283
+bool CompareLinkedPinMaps(const TMap<FString, int32>& Map0, const TMap<FString, int32>& Map1, const UEdGraphNode* Node)
+{
+	// Compare if there are connected pins in Map0 not present in Map1, viceversa, and also if the number of connections is different for Map0 and Map1
+	bool PinIn0NotIn1 = false;
+	bool PinIn1NotIn0 = false;
+	bool DifferentLinkNumber = false;
+
+	for (const TPair<FString, int32>& Element : Map0)
+	{
+		const int32* NumElement1 = Map1.Find(Element.Key);
+		if (NumElement1 == nullptr)
+		{
+			PinIn0NotIn1 = true;
+			continue;
+		}
+
+		int32 NumElement0 = Element.Value;
+		if (NumElement0 != (*NumElement1))
+		{
+			DifferentLinkNumber = true;
+		}
+	}
+
+	for (const TPair<FString, int32>& Element : Map1)
+	{
+		const int32* NumElement0 = Map0.Find(Element.Key);
+		if (NumElement0 == nullptr)
+		{
+			PinIn1NotIn0 = true;
+		}
+	}
+
+	if (PinIn0NotIn1 || PinIn1NotIn0 || DifferentLinkNumber)
+	{
+		FString ErrorLog = "ERROR: Updated UCustomizableObjectNodeMaterial";
+
+		if (PinIn0NotIn1)
+		{
+			ErrorLog += ", has lost at least one connected pin ";
+		}
+
+		if (PinIn1NotIn0)
+		{
+			ErrorLog += ", has at least one connected pin with name not present before update";
+		}
+
+		if (DifferentLinkNumber)
+		{
+			ErrorLog += ", has at least one pin with different number connection after update";
+		}
+
+		ErrorLog += ". Please verify the node works correctly";
+
+		FCustomizableObjectEditorLogger::CreateLog(FText::FromString(ErrorLog))
+		.Severity(EMessageSeverity::Warning)
+		.Node(*CastChecked<UCustomizableObjectNode>(Node))
+		.Log();
+	}
+
+	return (!PinIn0NotIn1 && !PinIn1NotIn0 && !DifferentLinkNumber);
+}
+
+
+void FCustomizableObjectEditor::RefreshMaterialNodesInAllChildrenCallback()
+{
+	FNotificationInfo Info(NSLOCTEXT("CustomizableObject", "RefreshingMaterialNodesInAllChildren", "Refreshing material nodes in all children"));
+	Info.bFireAndForget = true;
+	Info.bUseThrobber = true;
+	Info.FadeOutDuration = 1.0f;
+	Info.ExpireDuration = 1.0f;
+	FSlateNotificationManager::Get().AddNotification(Info);
+
+	LaunchRefreshMaterialInAllChildren = true;
+}
+
+
+// Deprecated! Remove on MTBL-283
+void FCustomizableObjectEditor::RefreshMaterialNodesInAllChildren()
+{
+	// Verify only one UCustomizableObjectNodeObjectGroup node is selected
+	const FGraphPanelSelectionSet SelectedNodes = GraphEditor->GetSelectedNodes();
+
+	if (SelectedNodes.Num() != 1)
+	{
+		return;
+	}
+
+	FString ParentGroupName = "";
+	for (FGraphPanelSelectionSet::TConstIterator SelectedIter(SelectedNodes); SelectedIter; ++SelectedIter)
+	{
+		if (UCustomizableObjectNodeObjectGroup* Node = Cast<UCustomizableObjectNodeObjectGroup>(*SelectedIter))
+		{
+			ParentGroupName = Node->GroupName;
+		}
+	}
+
+	// Loop through all COs and find if a CO has a UCustomizableObjectNodeObject node linked to the CO being analyzed,
+	// and adding the group name for filtering only for those COs of the selected group
+	FAssetRegistryModule& AssetRegistryModule = FModuleManager::LoadModuleChecked<FAssetRegistryModule>("AssetRegistry");
+	TArray<FAssetData> OutAssetData;
+	AssetRegistryModule.Get().GetAssetsByClass(FTopLevelAssetPath(TEXT("/Script/CustomizableObject"), TEXT("CustomizableObject")), OutAssetData);
+
+	for (auto Itr = OutAssetData.CreateIterator(); Itr; ++Itr)
+	{
+		UCustomizableObject* CustomizableObjectTemp = Cast<UCustomizableObject>(Itr->GetAsset());
+
+		if (CustomizableObjectTemp == nullptr)
+		{
+			continue;
+		}
+		
+		bool MultipleBaseObjectsFound;
+		UCustomizableObjectNodeObject* TempParentNodeObject = GetRootNode(CustomizableObjectTemp, MultipleBaseObjectsFound);
+
+		if (MultipleBaseObjectsFound)
+		{
+			continue;
+		}
+
+		bool IsChild = HasCandidateAsParent(TempParentNodeObject, CustomizableObject);
+
+		if (IsChild)
+		{
+			bool IsLinkedTo = GroupNodeIsLinkedToParentByName(TempParentNodeObject, CustomizableObject, ParentGroupName);
+
+			if (IsLinkedTo)
+			{
+				UEdGraph* Graph = CustomizableObjectTemp->Source;
+
+				if (Graph != nullptr)
+				{
+					TArray<UCustomizableObjectNodeMaterial*> ArrayNode;
+					Graph->GetNodesOfClass<UCustomizableObjectNodeMaterial>(ArrayNode);
+
+					int32 MaxIndex = ArrayNode.Num();
+					int32 NumNodeReconstructed = 0;
+					TMap<FString, int32> MapLinkedPinsBeforeUpdate;
+					TMap<FString, int32> MapLinkedPinsAfterUpdate;
+
+					for (int32 i = 0; i < MaxIndex; ++i)
+					{
+						StoreNodeLinkedPins(ArrayNode[i], MapLinkedPinsBeforeUpdate);
+						ArrayNode[i]->UCustomizableObjectNode::ReconstructNode();
+						StoreNodeLinkedPins(ArrayNode[i], MapLinkedPinsAfterUpdate);
+						NumNodeReconstructed++;
+							
+						CompareLinkedPinMaps(MapLinkedPinsBeforeUpdate, MapLinkedPinsAfterUpdate, ArrayNode[i]);
+					}
+
+					if (NumNodeReconstructed > 0)
+					{
+						Graph->NotifyGraphChanged();
+						Graph->MarkPackageDirty();
+					}
+				}
+			}
+		}
+	}
+}
+
+
+void FCustomizableObjectEditor::PostUndo(bool bSuccess)
+{
+	if (bSuccess)
+	{
+		GraphEditor->ClearSelectionSet();
+		GraphEditor->NotifyGraphChanged();
+		CustomizableObject->MarkPackageDirty();
+		//CustomizableObjectDetailsView->Invalidate(EInvalidateWidget::Layout);
+	}
+}
+
+
+FString FCustomizableObjectEditor::GetDocumentationLink() const
+{
+	return DocumentationURL;
+}
+
+
+void FCustomizableObjectEditor::ExtendToolbar()
+{
+	TSharedPtr<FUICommandList> CommandList = GetToolkitCommands();
+
+	struct Local
+	{
+		static void FillToolbar(FToolBarBuilder& ToolbarBuilder, FCustomizableObjectEditor* Editor, TSharedPtr<FUICommandList> CommandList)
+		{
+			ToolbarBuilder.BeginSection("Compilation");
+			ToolbarBuilder.AddToolBarButton(FCustomizableObjectEditorCommands::Get().Compile);
+			ToolbarBuilder.AddToolBarButton(FCustomizableObjectEditorCommands::Get().CompileOnlySelected);
+			ToolbarBuilder.AddComboButton(
+				FUIAction(),
+				FOnGetContent::CreateSP(Editor, &FCustomizableObjectEditor::GenerateCompileOptionsMenuContent, CommandList.ToSharedRef()),
+				LOCTEXT("Compile_Options_Label", "Compile Options"),
+				LOCTEXT("Compile_Options_Tooltip", "Change Compile Options"),
+				TAttribute<FSlateIcon>(),
+				true);
+
+			ToolbarBuilder.EndSection();
+
+			
+			ToolbarBuilder.BeginSection("Debug");
+			ToolbarBuilder.AddToolBarButton(FCustomizableObjectEditorCommands::Get().Debug);
+			ToolbarBuilder.EndSection();
+
+			
+			ToolbarBuilder.BeginSection("Information");
+			ToolbarBuilder.AddToolBarButton(FCustomizableObjectEditorCommands::Get().TextureAnalyzer);
+			ToolbarBuilder.AddToolBarButton(FCustomizableObjectEditorCommands::Get().PerformanceReport);
+			ToolbarBuilder.EndSection();
+		}
+	};
+
+	TSharedPtr<FExtender> ToolbarExtender = MakeShareable(new FExtender);
+
+	ToolbarExtender->AddToolBarExtension(
+		"Asset",
+		EExtensionHook::After,
+		GetToolkitCommands(),
+		FToolBarExtensionDelegate::CreateStatic(&Local::FillToolbar, this, CommandList));
+
+	AddToolbarExtender(ToolbarExtender);
+
+	ICustomizableObjectEditorModule* CustomizableObjectEditorModule = &FModuleManager::LoadModuleChecked<ICustomizableObjectEditorModule>("CustomizableObjectEditor");
+	AddToolbarExtender(CustomizableObjectEditorModule->GetCustomizableObjectEditorToolBarExtensibilityManager()->GetAllExtenders());
+}
+
+
+TSharedRef<SWidget> FCustomizableObjectEditor::GenerateCompileOptionsMenuContent(TSharedRef<FUICommandList> InCommandList)
+{
+	const bool bShouldCloseWindowAfterMenuSelection = false;
+	FMenuBuilder MenuBuilder(bShouldCloseWindowAfterMenuSelection, InCommandList);
+
+	// settings
+	MenuBuilder.BeginSection("ResetCompileOptions");
+	{
+		MenuBuilder.AddMenuEntry(FCustomizableObjectEditorCommands::Get().ResetCompileOptions);
+	}
+	MenuBuilder.EndSection();
+
+	MenuBuilder.BeginSection("Optimization", LOCTEXT("MutableCompileOptimizationHeading", "Optimization"));
+	{
+		// Level
+		CompileOptimizationStrings.Empty();
+		CompileOptimizationStrings.Add(MakeShareable(new FString(NSLOCTEXT("UnrealEd", "OptimizationNone", "None").ToString())));
+		CompileOptimizationStrings.Add(MakeShareable(new FString(NSLOCTEXT("UnrealEd", "OptimizationMin", "Minimal").ToString())));
+		CompileOptimizationStrings.Add(MakeShareable(new FString(NSLOCTEXT("UnrealEd", "OptimizationMed", "Medium").ToString())));
+		CompileOptimizationStrings.Add(MakeShareable(new FString(NSLOCTEXT("UnrealEd", "OptimizationMax", "Maximum").ToString())));
+
+		CompileOptimizationCombo =
+			SNew(STextComboBox)
+			.OptionsSource(&CompileOptimizationStrings)
+			.InitiallySelectedItem(CompileOptimizationStrings[CustomizableObject ? CustomizableObject->CompileOptions.OptimizationLevel : 0])
+			.OnSelectionChanged(this, &FCustomizableObjectEditor::OnChangeCompileOptimizationLevel)
+			;
+
+		MenuBuilder.AddWidget(CompileOptimizationCombo.ToSharedRef(), LOCTEXT("MutableCompileOptimizationLevel", "Optimization Level"));
+
+		MenuBuilder.AddMenuEntry(FCustomizableObjectEditorCommands::Get().CompileOptions_UseParallelCompilation);
+		MenuBuilder.AddMenuEntry(FCustomizableObjectEditorCommands::Get().CompileOptions_UseDiskCompilation);
+		MenuBuilder.AddMenuEntry(FCustomizableObjectEditorCommands::Get().CompileOptions_EnableTextureCompression);
+	}
+	MenuBuilder.EndSection();
+
+	return MenuBuilder.MakeWidget();
+}
+
+
+FText FCustomizableObjectEditor::GetToolkitName() const
+{
+	FFormatNamedArguments Args;
+	Args.Add(TEXT("ObjectName"), FText::FromString(GetEditingObject()->GetName()));
+	Args.Add(TEXT("ToolkitName"), GetBaseToolkitName());
+	return FText::Format(LOCTEXT("AppLabelWithAssetName", "{ObjectName} - {ToolkitName}"), Args);
+}
+
+
+FString FCustomizableObjectEditor::GetWorldCentricTabPrefix() const
+{
+	return LOCTEXT( "WorldCentricTabPrefix", "CustomizableObject " ).ToString();
+}
+
+
+FLinearColor FCustomizableObjectEditor::GetWorldCentricTabColorScale() const
+{
+	return FLinearColor( 0.3f, 0.2f, 0.5f, 0.5f );
+}
+
+
+UCustomizableObject* FCustomizableObjectEditor::GetCustomizableObject()
+{
+	return CustomizableObject;
+}
+
+
+void FCustomizableObjectEditor::RefreshTool()
+{
+	RefreshViewport();
+}
+
+
+void FCustomizableObjectEditor::RefreshViewport()
+{
+	if (ViewportClient)
+	{
+		ViewportClient->Invalidate();
+	}
+}
+
+
+void FCustomizableObjectEditor::OnObjectPropertySelectionChanged(FProperty* InProperty)
+{
+	CustomizableObject->PostEditChange();
+	RefreshViewport();
+}
+
+
+void FCustomizableObjectEditor::OnInstancePropertySelectionChanged(FProperty* InProperty)
+{
+	RefreshViewport();
+}
+
+
+void FCustomizableObjectEditor::OnObjectModified(UObject* Object)
+{
+	UCustomizableObjectInstance* Instance = Cast<UCustomizableObjectInstance>(Object);
+
+	if ((Instance != nullptr) && (Instance == PreviewInstance))
+	{
+		if (ManagingProjector)
+		{
+			return;
+		}
+
+		if(Instance->ProjectorLayerChange)
+		{
+			// Don't make any changes in the projector if only a projector layer has changed
+			Instance->ProjectorLayerChange = false;
+			return;
+		}
+
+		if (Viewport->GetGizmoHasAssignedData() && !Viewport->GetGizmoAssignedDataIsFromNode())
+		{
+			Instance->LastSelectedProjectorParameter = Viewport->GetGizmoProjectorParameterName();
+			Instance->LastSelectedProjectorParameterWithIndex = Viewport->GetGizmoProjectorParameterNameWithIndex();
+		}
+
+		//UE_LOG(LogMutable, Warning, TEXT("LastSelectedProjectorParameter=%s, LastSelectedProjectorParameterWithIndex=%s"), *(PreviewInstance->LastSelectedProjectorParameter), *(PreviewInstance->LastSelectedProjectorParameterWithIndex));
+
+		// Update projector parameter information (not projector parameter node information,
+		// this is the case where a projector parameter in the object details is selected)
+		const TArray<FCustomizableObjectProjectorParameterValue>& Parameters = PreviewInstance->GetProjectorParameters();
+		const int32 MaxIndex = Parameters.Num();
+		bool AnySelected = false;
+
+		for (int32 i = 0; (!AnySelected && (i < MaxIndex)); ++i)
+		{
+			int32 RangeIndex = Parameters[i].RangeValues.Num() > 0 ? 0 : -1;
+
+			for (; RangeIndex < Parameters[i].RangeValues.Num(); ++RangeIndex)
+			{
+				const FCustomizableObjectProjector& Value = (RangeIndex == -1) ? Parameters[i].Value : Parameters[i].RangeValues[RangeIndex];
+
+				FString ParameterNameWithIndex = Parameters[i].ParameterName;
+
+				if (RangeIndex != -1)
+				{
+					ParameterNameWithIndex += FString::Printf(TEXT("__%d"), RangeIndex);
+				}
+
+				if ((PreviewInstance->GetProjectorState(Parameters[i].ParameterName, RangeIndex) == EProjectorState::Selected) ||
+					((PreviewInstance->GetProjectorState(Parameters[i].ParameterName, RangeIndex) == EProjectorState::TypeChanged) &&
+					(ParameterNameWithIndex != Instance->LastSelectedProjectorParameterWithIndex)))
+				{
+					AnySelected = true;
+					SetProjectorVisibilityForParameter = true;
+					ProjectorParameterName = Parameters[i].ParameterName;
+					ProjectorParameterNameWithIndex = Parameters[i].ParameterName;
+					if (RangeIndex != -1)
+					{
+						ProjectorParameterNameWithIndex += FString::Printf(TEXT("__%d"), RangeIndex);
+					}
+					ProjectorRangeIndex = RangeIndex;
+					ProjectorParameterIndex = i;
+					ProjectorParameterPosition = Value.Position;
+					ProjectorParameterDirection = Value.Direction;
+					ProjectorParameterUp = Value.Up;
+					ProjectorParameterScale = Value.Scale;
+					ProjectorParameterProjectionType = Value.ProjectionType;
+					ProjectionAngle = Value.Angle;
+				}
+				else if ((PreviewInstance->GetProjectorState(Parameters[i].ParameterName, RangeIndex) == EProjectorState::TypeChanged) &&
+					(Parameters[i].ParameterName == Instance->LastSelectedProjectorParameter))
+				{
+					AnySelected = true;
+					SetProjectorTypeForParameter = true;
+					ProjectorParameterName = Parameters[i].ParameterName;
+					ProjectorParameterNameWithIndex = Parameters[i].ParameterName;
+					if (RangeIndex != -1)
+					{
+						ProjectorParameterNameWithIndex += FString::Printf(TEXT("__%d"), RangeIndex);
+					}
+					ProjectorRangeIndex = RangeIndex;
+					ProjectorParameterIndex = i;
+					ProjectorParameterPosition = Value.Position;
+					ProjectorParameterDirection = Value.Direction;
+					ProjectorParameterUp = Value.Up;
+					ProjectorParameterScale = Value.Scale;
+					ProjectorParameterProjectionType = Value.ProjectionType;
+					ProjectionAngle = Value.Angle;
+				}
+			}
+		}
+
+		// Another projector parameter is being selected, and there already a projector parameter selected pending to be hidden,
+		// update its information before selecting the new one
+		bool ResetPendingDone = false;
+		if (AnySelected && SetProjectorVisibilityForParameter && !Instance->LastSelectedProjectorParameter.IsEmpty())
+		{
+			ResetProjectorVisibilityNoUpdate();
+			Instance->LastSelectedProjectorParameter = "";
+			Instance->LastSelectedProjectorParameterWithIndex = "";
+			ResetPendingDone = true;
+		}
+
+		if (!ProjectorConstantNodeSelected && !ProjectorParameterNodeSelected && !ResetPendingDone && !AnySelected && !Viewport->GetIsManipulatingGizmo())
+		{
+			ResetProjectorVisibilityForNonNode = true;
+		}
+		else if (ProjectorParameterNodeSelected)
+		{
+			SelectedProjectorParameterNotNode = true;
+		}
+
+		if (Instance->UnselectProjector)
+		{
+			Instance->UnselectProjector = false;
+
+			ResetProjectorVisibilityNoUpdate();
+			
+			Instance->LastSelectedProjectorParameter = "";
+			Instance->LastSelectedProjectorParameterWithIndex = "";
+		}
+	}
+
+	if (!Instance)
+	{
+		// Sometimes when another CO is open in another editor window/tab, it triggers this callback, so prevent the modification of this object by a callback triggered by another one
+		if(UCustomizableObject* AuxCustomizableObject = Cast<UCustomizableObject>(Object))
+		{
+			AuxCustomizableObject->UpdateVersionId();
+		}
+		else if (UCustomizableObjectNode* Node = Cast<UCustomizableObjectNode>(Object))
+		{
+			if (UCustomizableObjectGraph* Graph = Cast<UCustomizableObjectGraph>(Node->GetOuter()))
+			{
+				if (UCustomizableObject* AuxOuterCustomizableObject = Cast<UCustomizableObject>(Graph->GetOuter()))
+				{
+					AuxOuterCustomizableObject->UpdateVersionId();
+				}
+			}
+		}
+		else if (UCustomizableObjectGraph* Graph = Cast<UCustomizableObjectGraph>(Object))
+		{
+			if (UCustomizableObject* AuxOuterCustomizableObject = Cast<UCustomizableObject>(Graph->GetOuter()))
+			{
+				AuxOuterCustomizableObject->UpdateVersionId();
+			}
+		}
+	}
+}
+
+
+void FCustomizableObjectEditor::CompileObject()
+{
+	// If any projector selected, unselect it in case the projector node is removed with this CO compile
+	ResetProjectorVisibilityNoUpdate();
+
+	// Unselect NodeMeshClipMorph if selected
+	Viewport->SetClipMorphPlaneVisibility(false, nullptr);
+
+	UE_LOG(LogMutable, Verbose, TEXT("PROFILE: -----------------------------------------------------------"));
+	UE_LOG(LogMutable, Verbose, TEXT("PROFILE: [ %16.8f ] FCustomizableObjectEditor::CompileObject start."), FPlatformTime::Seconds());
+
+	if (!AssetRegistryLoaded)
+	{
+		FNotificationInfo Info(NSLOCTEXT("CustomizableObject", "CustomizableObjectCompileTryLater", "Please wait until asset registry loads all assets"));
+		Info.bFireAndForget = true;
+		Info.bUseThrobber = true;
+		Info.FadeOutDuration = 1.0f;
+		Info.ExpireDuration = 2.0f;
+		FSlateNotificationManager::Get().AddNotification(Info);
+		return;
+	}
+
+	Viewport->UpdateGizmoDataToOrigin();
+
+	if (CustomizableObject->Source)
+	{
+		FCompilationOptions Options = CustomizableObject->CompileOptions;
+		Options.bSilentCompilation = false;
+		Compiler.Compile(*CustomizableObject, Options, true);
+	}
+
+	UE_LOG(LogMutable, Log, TEXT("PROFILE: [ %16.8f ] FCustomizableObjectEditor::CompileObject end."), FPlatformTime::Seconds());
+}
+
+
+void FCustomizableObjectEditor::DebugObject()
+{
+	if (!CustomizableObject)
+	{
+		return;
+	}
+
+	// Open a mutable debugger as a standalone window
+	ICustomizableObjectEditorModule* CustomizableObjectEditorModule = &FModuleManager::LoadModuleChecked<ICustomizableObjectEditorModule>("CustomizableObjectEditor");
+	if (!CustomizableObjectEditorModule) return;
+
+	CustomizableObjectEditorModule->CreateCustomizableObjectDebugger(EToolkitMode::Standalone, nullptr, CustomizableObject);
+}
+
+
+void FCustomizableObjectEditor::ResetCompileOptions()
+{
+	const FScopedTransaction Transaction(LOCTEXT("ResetCompilationOptionsTransaction", "Reset Compilation Options"));
+	CustomizableObject->Modify();
+	CustomizableObject->CompileOptions = FCompilationOptions();
+}
+
+void FCustomizableObjectEditor::OnChangeCompileOptimizationLevel(TSharedPtr<FString> NewSelection, ESelectInfo::Type SelectInfo)
+{
+	const FScopedTransaction Transaction(LOCTEXT("ChangedOptimizationLevelTransaction", "Changed Optimization Level"));
+	CustomizableObject->Modify();
+	CustomizableObject->CompileOptions.OptimizationLevel = CompileOptimizationStrings.Find(NewSelection);
+}
+
+
+void FCustomizableObjectEditor::CompileOptions_UseParallelCompilation_Toggled()
+{
+	const FScopedTransaction Transaction(LOCTEXT("ChangedEnableCompilingInMultipleThreadsTransaction", "Changed Enable compiling in multiple threads"));
+	CustomizableObject->Modify();
+	CustomizableObject->CompileOptions.bUseParallelCompilation = !CustomizableObject->CompileOptions.bUseParallelCompilation;
+}
+
+
+bool FCustomizableObjectEditor::CompileOptions_UseParallelCompilation_IsChecked()
+{
+	return CustomizableObject ? CustomizableObject->CompileOptions.bUseParallelCompilation : false;
+}
+
+
+void FCustomizableObjectEditor::CompileOptions_UseDiskCompilation_Toggled()
+{
+	const FScopedTransaction Transaction(LOCTEXT("ChangedEnableCompilingUsingTheDiskAsMemory", "Changed Enable compiling using the disk as memory"));
+	CustomizableObject->Modify();
+	CustomizableObject->CompileOptions.bUseDiskCompilation = !CustomizableObject->CompileOptions.bUseDiskCompilation;
+}
+
+
+bool FCustomizableObjectEditor::CompileOptions_UseDiskCompilation_IsChecked()
+{
+	return CustomizableObject->CompileOptions.bUseDiskCompilation;
+}
+
+
+void FCustomizableObjectEditor::CompileOptions_TextureCompression_Toggled()
+{
+	CustomizableObject->CompileOptions.bTextureCompression = !CustomizableObject->CompileOptions.bTextureCompression;
+	CustomizableObject->Modify();
+}
+
+
+bool FCustomizableObjectEditor::CompileOptions_TextureCompression_IsChecked()
+{
+	return CustomizableObject->CompileOptions.bTextureCompression;
+}
+
+
+void FCustomizableObjectEditor::SaveAsset_Execute()
+{
+	if (PreviewInstance)
+	{
+		if (PreviewInstance->IsSelectedParameterProfileDirty())
+		{
+			PreviewInstance->SaveParametersToProfile(PreviewInstance->SelectedProfileIndex);
+		}
+	}
+
+	UPackage* Package = CustomizableObject->GetOutermost();
+
+	if (Package)
+	{
+		TArray<UPackage*> PackagesToSave;
+		PackagesToSave.Add(Package);
+
+		FEditorFileUtils::PromptForCheckoutAndSave(PackagesToSave, false, false);
+	}
+}
+
+
+void FCustomizableObjectEditor::DeleteSelectedNodes()
+{
+	if (!GraphEditor.IsValid())
+	{
+		return;
+	}
+
+	const FScopedTransaction Transaction(NSLOCTEXT("UnrealEd", "UEdGraphSchema_CustomizableObject", "Delete Nodes"));
+
+	const FGraphPanelSelectionSet SelectedNodes = GraphEditor->GetSelectedNodes();
+	GraphEditor->ClearSelectionSet();
+
+	for (FGraphPanelSelectionSet::TConstIterator NodeIt( SelectedNodes ); NodeIt; ++NodeIt)
+	{
+		if (UEdGraphNode* Node = Cast<UEdGraphNode>(*NodeIt))
+		{
+			if (Node->CanUserDeleteNode())
+			{
+				if (const UEdGraph* GraphObj = Node->GetGraph())
+				{
+					if (const UEdGraphSchema* Schema = GraphObj->GetSchema())
+					{
+						Schema->BreakNodeLinks(*Node);  // Required to notify to all connected nodes (UEdGraphNode::PinConnectionListChanged() and UEdGraphNode::PinConnectionListChanged(...))
+					}
+				}
+
+				Node->DestroyNode();
+			}
+		}
+	}
+}
+
+
+bool FCustomizableObjectEditor::CanDeleteNodes() const
+{
+	if (GraphEditor.IsValid() && GraphEditor->GetSelectedNodes().Num() > 0)
+	{
+		for (auto Itr = GraphEditor->GetSelectedNodes().CreateConstIterator(); Itr; ++Itr)
+		{
+			UCustomizableObjectNode* Node = Cast<UCustomizableObjectNode>(*Itr);
+
+			if (Node && !Node->CanUserDeleteNode())
+			{
+				return false;
+			}
+		}
+
+		return true;
+	}
+
+	return false;
+}
+
+
+void FCustomizableObjectEditor::DuplicateSelectedNodes()
+{
+	CopySelectedNodes();
+	PasteNodes();
+}
+
+
+bool FCustomizableObjectEditor::CanDuplicateSelectedNodes() const
+{
+	return CanCopyNodes();
+}
+
+
+void FCustomizableObjectEditor::OnSelectedGraphNodesChanged(const FGraphPanelSelectionSet& NewSelection)
+{
+	TArray<UObject*> Objects;
+	for ( FGraphPanelSelectionSet::TConstIterator It(NewSelection); It; ++It)
+	{
+		Objects.Add(*It);
+	}
+
+	// Standard details
+	if ( GraphNodeDetailsView.IsValid() )
+	{
+		GraphNodeDetailsView->SetObjects( Objects );
+	}
+
+	if (Objects.Num())
+	{
+		UObject* FirstNode = nullptr;
+		FirstNode = Objects[0];
+		SelectedMeshClipMorphNode = Cast<UCustomizableObjectNodeMeshClipMorph>(FirstNode);
+
+		if (!SelectedMeshClipMorphNode)
+		{
+			SelectedMeshClipWithMeshNode = Cast<UCustomizableObjectNodeMeshClipWithMesh>(FirstNode);
+
+			if (!SelectedProjectorNode || (SelectedProjectorNode && FirstNode && (SelectedProjectorNode != FirstNode)))
+			{
+				SelectedProjectorNode = Cast<UCustomizableObjectNodeProjectorConstant>(FirstNode);
+				if (SelectedProjectorNode != nullptr)
+				{
+					ProjectorConstantNodeSelected = true;
+				}
+			}
+
+			if (!SelectedProjectorParameterNode || (SelectedProjectorParameterNode && FirstNode && (SelectedProjectorParameterNode != FirstNode)))
+			{
+				SelectedProjectorParameterNode = Cast<UCustomizableObjectNodeProjectorParameter>(FirstNode);
+				if (SelectedProjectorParameterNode != nullptr)
+				{
+					ProjectorParameterNodeSelected = true;
+				}
+			}
+		}
+	}
+	else 
+	{
+		SelectedMeshClipMorphNode = nullptr;
+		SelectedMeshClipWithMeshNode = nullptr;
+		SelectedProjectorNode = nullptr;
+		SelectedProjectorParameterNode = nullptr;
+	}
+
+	// Special details. Usually widgets that you want to be able to scale
+	if ( NodeDetailsSplitter.IsValid() )
+	{
+		UObject* FirstNode = 0;		
+		if ( Objects.Num() )
+		{
+			FirstNode = Objects[0];
+		}
+
+		TSharedPtr<SWidget> ChildWidget;
+
+		// Set the size of the child widget to show properly the node widgets
+		FVector2D SplitWeights = { 0.15f,0.85f };
+		
+		// Create the specialised user interface for each node type
+		if ( UCustomizableObjectNodeLayoutBlocks* CustomizableObjectNodeLayoutBlocks = Cast<UCustomizableObjectNodeLayoutBlocks>(FirstNode) )
+		{
+			if (!LayoutBlocksEditor.IsValid())
+			{
+				LayoutBlocksEditor = SNew(SCustomizableObjectNodeLayoutBlocksEditor)
+					.CustomizableObjectEditor(SharedThis(this));
+			}
+
+			LayoutBlocksEditor->SetCurrentLayout(CustomizableObjectNodeLayoutBlocks->Layout);
+			SplitWeights = { 0.0f,1.0f };
+
+			ChildWidget = LayoutBlocksEditor;
+		}
+
+		else if (UCustomizableObjectNodeTable* CustomizableObjectNodeTable = Cast<UCustomizableObjectNodeTable>(FirstNode))
+		{
+			if (!LayoutBlocksEditor.IsValid())
+			{
+				LayoutBlocksEditor = SNew(SCustomizableObjectNodeLayoutBlocksEditor)
+					.CustomizableObjectEditor(SharedThis(this));
+			}
+
+			LayoutBlocksEditor->SetCurrentLayout(nullptr);
+			SplitWeights = { 0.32f,0.68f };
+
+			ChildWidget = LayoutBlocksEditor;
+		}
+
+		else if ( UCustomizableObjectNodeEditMaterial* CustomizableObjectNodeEditMaterial = Cast<UCustomizableObjectNodeEditMaterial>(FirstNode) )
+		{
+			if (!LayoutBlocksSelector.IsValid())
+			{
+				LayoutBlocksSelector = SNew(SCustomizableObjectNodeLayoutBlocksSelector)
+					.CustomizableObjectEditor(SharedThis(this));
+			}
+
+			LayoutBlocksSelector->SetSelectedNode(CustomizableObjectNodeEditMaterial);
+			SplitWeights = { 0.25f,0.75f };
+
+			ChildWidget = LayoutBlocksSelector;
+		}
+		else if (UCustomizableObjectNodeRemoveMeshBlocks* CustomizableObjectNodeRemoveMeshBlocks = Cast< UCustomizableObjectNodeRemoveMeshBlocks>(FirstNode))
+		{
+			if (!LayoutBlocksSelector.IsValid())
+			{
+				LayoutBlocksSelector = SNew(SCustomizableObjectNodeLayoutBlocksSelector)
+					.CustomizableObjectEditor(SharedThis(this));
+			}
+
+			LayoutBlocksSelector->SetSelectedNode(CustomizableObjectNodeRemoveMeshBlocks);
+			SplitWeights = { 0.25f,0.75f };
+
+			ChildWidget = LayoutBlocksSelector;
+		}
+
+		if ( ChildWidget.IsValid() )
+		{
+			if ( NodeDetailsSplitter->GetChildren()->Num()<=1 )
+			{
+				NodeDetailsSplitter->AddSlot();
+			}
+
+			NodeDetailsSplitter->SlotAt(0).SetSizeValue(SplitWeights.X);
+			NodeDetailsSplitter->SlotAt(1).SetSizeValue(SplitWeights.Y);
+
+			NodeDetailsSplitter->SlotAt(1)
+				[
+					SNew(SBorder)
+					.BorderImage( FAppStyle::GetBrush( TEXT("Graph.TitleBackground") ) )
+					.HAlign(HAlign_Fill)
+					[
+						ChildWidget.ToSharedRef()
+					]
+				];
+		}
+		else
+		{
+			NodeDetailsSplitter->SlotAt(0).SetSizeValue(0.5);
+
+			if ( NodeDetailsSplitter->GetChildren()->Num()>1 )
+			{
+				NodeDetailsSplitter->RemoveAt(1);
+			}
+		}
+	}
+
+	if (!ManagingProjector)
+	{
+		SelectedGraphNodesChanged = true;
+	}
+}
+
+
+void FCustomizableObjectEditor::NotifyPostChange( const FPropertyChangedEvent& PropertyChangedEvent, FProperty* PropertyThatChanged  )
+{
+	// Is it a source graph node?
+	const UObject* OuterObject = PropertyThatChanged->GetOwner<UObject>();
+	const UClass* OuterClass = Cast<UClass>(OuterObject);
+	if (OuterClass && OuterClass->IsChildOf(UCustomizableObjectNode::StaticClass()))
+	{
+		FPropertyChangedEvent Event(PropertyThatChanged);
+		CustomizableObject->Source->PostEditChangeProperty(Event);
+		CustomizableObject->PostEditChangeProperty(Event);
+
+		if (GraphEditor.IsValid())
+		{
+			GraphEditor->NotifyGraphChanged();
+		}
+	}
+}
+
+
+bool FCustomizableObjectEditor::IsTickable(void) const
+{
+	return true;
+}
+
+
+void FCustomizableObjectEditor::Tick( float InDeltaTime )
+{
+	const bool bUpdated = Compiler.Tick();
+	
+	if (bUpdated && CustomizableObject->IsCompiled())
+	{
+		if (!PreviewCustomizableSkeletalComponents.IsEmpty() && !PreviewSkeletalMeshComponents.IsEmpty())
+		{
+			CreatePreviewComponents();
+
+			PreviewInstance->UpdateSkeletalMeshAsync(true, true);
+		}
+		else
+		{
+			CreatePreviewInstance();
+		}
+	}
+
+	if (SelectedMeshClipMorphNode)
+	{
+		bool bVisible = SelectedMeshClipMorphNode->BoneName != FName();
+
+		Viewport->SetClipMorphPlaneVisibility(bVisible, SelectedMeshClipMorphNode);
+	}
+	else if (SelectedMeshClipWithMeshNode)
+	{
+		UStaticMesh* ClipMesh = nullptr;
+
+		if (const UEdGraphPin* ConnectedPin = FollowInputPin(*SelectedMeshClipWithMeshNode->ClipMeshPin()))
+		{
+			if (const UEdGraphNode* Node = ConnectedPin->GetOwningNode())
+			{
+				if (const UCustomizableObjectNodeStaticMesh* TypedNode = Cast<UCustomizableObjectNodeStaticMesh>(Node))
+				{
+					ClipMesh = TypedNode->StaticMesh;
+				}
+			}
+		}
+
+		bool bVisible = ClipMesh != nullptr;
+		Viewport->SetClipMeshVisibility(bVisible, ClipMesh, SelectedMeshClipWithMeshNode);
+	}
+	else if (ProjectorConstantNodeSelected)
+	{
+		ManagingProjector = true;
+		Viewport->SetProjectorVisibility(true, SelectedProjectorNode);
+		Viewport->ProjectorParameterChanged(SelectedProjectorNode);
+		ManagingProjector = false;
+		ProjectorConstantNodeSelected = false;
+	}
+	else if (ProjectorParameterNodeSelected)
+	{
+		ManagingProjector = true;
+		Viewport->SetProjectorParameterVisibility(true, SelectedProjectorParameterNode);
+		Viewport->ProjectorParameterChanged(SelectedProjectorParameterNode);
+		ManagingProjector = false;
+		ProjectorParameterNodeSelected = false;
+	}
+	else
+	{
+		if (SetProjectorVisibilityForParameter)
+		{
+			ManagingProjector = true;
+			if (Viewport->AnyProjectorNodeSelected())
+			{
+				GraphEditor->ClearSelectionSet();
+			}
+			FCustomizableObjectProjector ProjectorParameterValue;
+			ProjectorParameterValue.Position = (FVector3f)ProjectorParameterPosition;
+			ProjectorParameterValue.Direction = (FVector3f)ProjectorParameterDirection;
+			ProjectorParameterValue.Up = (FVector3f)ProjectorParameterUp;
+			ProjectorParameterValue.Scale = (FVector3f)ProjectorParameterScale;
+			ProjectorParameterValue.ProjectionType = ProjectorParameterProjectionType;
+			ProjectorParameterValue.Angle = ProjectionAngle;
+			Viewport->SetProjectorVisibility(true, ProjectorParameterName, ProjectorParameterNameWithIndex, ProjectorRangeIndex, ProjectorParameterValue, ProjectorParameterIndex);
+			ManagingProjector = false;
+			SetProjectorVisibilityForParameter = false;
+		}
+
+		if (SetProjectorTypeForParameter)
+		{
+			ManagingProjector = true;
+			if (Viewport->AnyProjectorNodeSelected())
+			{
+				GraphEditor->ClearSelectionSet();
+			}
+			FCustomizableObjectProjector ProjectorParameterValue;
+			ProjectorParameterValue.Position = (FVector3f)ProjectorParameterPosition;
+			ProjectorParameterValue.Direction = (FVector3f)ProjectorParameterDirection;
+			ProjectorParameterValue.Up = (FVector3f)ProjectorParameterUp;
+			ProjectorParameterValue.Scale = (FVector3f)ProjectorParameterScale;
+			ProjectorParameterValue.ProjectionType = ProjectorParameterProjectionType;
+			ProjectorParameterValue.Angle = ProjectionAngle;
+			Viewport->SetProjectorType(true, ProjectorParameterName, ProjectorParameterNameWithIndex, ProjectorRangeIndex, ProjectorParameterValue, ProjectorParameterIndex);
+			ManagingProjector = false;
+			SetProjectorTypeForParameter = false;
+		}
+
+		if (ResetProjectorVisibilityForNonNode && PreviewInstance && !PreviewInstance->AvoidResetProjectorVisibilityForNonNode)
+		{
+			ResetProjectorVisibilityNoUpdate();
+		}
+
+		if (SelectedGraphNodesChanged && !SelectedProjectorParameterNotNode)
+		{
+			ManagingProjector = true;
+			Viewport->ResetProjectorVisibility(false);
+			ManagingProjector = false;
+		}
+		Viewport->SetClipMorphPlaneVisibility(false, nullptr);
+		Viewport->SetClipMeshVisibility(false, nullptr, nullptr);
+	}
+
+	if (SelectedGraphNodesChanged)
+	{
+		SelectedGraphNodesChanged = false;
+	}
+
+	if (SelectedProjectorParameterNotNode)
+	{
+		SelectedProjectorParameterNotNode = false;
+	}
+
+	if (ResetProjectorVisibilityForNonNode)
+	{
+		ResetProjectorVisibilityForNonNode = false;
+    }
+
+	if (PreviewInstance != nullptr)
+	{
+		PreviewInstance->AvoidResetProjectorVisibilityForNonNode = false;
+	}
+
+    // TEMP CODE
+    if ((PreviewInstance != nullptr) && PreviewInstance->TempUpdateGizmoInViewport)
+    {
+        PreviewInstance->TempUpdateGizmoInViewport = false;
+		PreviewInstance->AvoidResetProjectorVisibilityForNonNode = true;
+        Viewport->CopyTransformFromOriginData();
+        PreviewInstance->UpdateSkeletalMeshAsync(true);
+
+        EProjectorState::Type ProjectorState = PreviewInstance->GetProjectorState(PreviewInstance->TempProjectorParameterName, PreviewInstance->TempProjectorParameterRangeIndex);
+        if (ProjectorState != EProjectorState::Selected)
+        {
+            //ResetProjectorVisibilityForNonNode = false;
+			PreviewInstance->ResetProjectorStates();
+			PreviewInstance->SetProjectorState(PreviewInstance->TempProjectorParameterName, PreviewInstance->TempProjectorParameterRangeIndex, EProjectorState::Selected);
+            FCoreUObjectDelegates::BroadcastOnObjectModified(PreviewInstance);
+        }
+	}
+
+	if (LaunchRefreshMaterialInAllChildren)
+	{
+		PendingTimeRefreshMaterialInAllChildren -= InDeltaTime;
+
+		if (PendingTimeRefreshMaterialInAllChildren < 0.0f)
+		{
+			PendingTimeRefreshMaterialInAllChildren = 2.0f;
+			LaunchRefreshMaterialInAllChildren = false;
+			RefreshMaterialNodesInAllChildren();
+		}
+	}
+
+	// Reconstruct marked nodes 
+	if (GraphEditor.IsValid())
+	{
+		for (UEdGraphNode* Node: ReconstructNodes)
+		{
+			GraphEditor->GetCurrentGraph()->Schema.GetDefaultObject()->ReconstructNode(*Node);
+		}
+		if (ReconstructNodes.Num()) 
+		{
+			GraphEditor->NotifyGraphChanged();
+			ReconstructNodes.Empty();
+		}
+	}
+}
+
+
+TStatId FCustomizableObjectEditor::GetStatId() const
+{
+	RETURN_QUICK_DECLARE_CYCLE_STAT(FCustomizableObjectEditor, STATGROUP_Tickables);
+}
+
+
+void FCustomizableObjectEditor::CopySelectedNodes()
+{
+	// Export the selected nodes and place the text on the clipboard
+	const FGraphPanelSelectionSet SelectedNodes = GraphEditor->GetSelectedNodes();
+
+	FString ExportedText;
+
+	for (FGraphPanelSelectionSet::TConstIterator SelectedIter(SelectedNodes); SelectedIter; ++SelectedIter)
+	{
+		if(UEdGraphNode* Node = Cast<UEdGraphNode>(*SelectedIter))
+		{
+			Node->PrepareForCopying();
+		}
+	}
+
+	FEdGraphUtilities::ExportNodesToText(SelectedNodes, /*out*/ ExportedText);
+
+#if (ENGINE_MAJOR_VERSION == 4 && ENGINE_MINOR_VERSION >= 18) || ENGINE_MAJOR_VERSION >= 5
+	FPlatformApplicationMisc::ClipboardCopy(*ExportedText);
+#else
+	FPlatformMisc::ClipboardCopy(*ExportedText);
+#endif
+
+	// Make sure Material remains the owner of the copied nodes
+	for (FGraphPanelSelectionSet::TConstIterator SelectedIter(SelectedNodes); SelectedIter; ++SelectedIter)
+	{
+		if (UMaterialGraphNode* Node = Cast<UMaterialGraphNode>(*SelectedIter))
+		{
+			Node->PostCopyNode();
+		}
+		else if (UMaterialGraphNode_Comment* Comment = Cast<UMaterialGraphNode_Comment>(*SelectedIter))
+		{
+			Comment->PostCopyNode();
+		}
+	}
+}
+
+bool FCustomizableObjectEditor::CanCopyNodes() const
+{
+	// If any of the nodes can be duplicated then we should allow copying
+	const FGraphPanelSelectionSet SelectedNodes = GraphEditor->GetSelectedNodes();
+	for (FGraphPanelSelectionSet::TConstIterator SelectedIter(SelectedNodes); SelectedIter; ++SelectedIter)
+	{
+		UEdGraphNode* Node = Cast<UEdGraphNode>(*SelectedIter);
+		if ((Node != NULL) && Node->CanDuplicateNode())
+		{
+			return true;
+		}
+	}
+
+	return false;
+}
+
+
+void FCustomizableObjectEditor::PasteNodes()
+{
+	PasteNodesHere(GraphEditor->GetPasteLocation());
+}
+
+
+void FCustomizableObjectEditor::PasteNodesHere(const FVector2D& Location)
+{
+	// Undo/Redo support
+	const FScopedTransaction Transaction( NSLOCTEXT("UnrealEd", "CustomizableObjectEditorPaste", "Customizable Object Editor Editor: Paste") );
+	CustomizableObject->Source->Modify();
+	CustomizableObject->Modify();
+
+	// Clear the selection set (newly pasted stuff will be selected)
+	GraphEditor->ClearSelectionSet();
+
+	// Grab the text to paste from the clipboard.
+	FString TextToImport;
+#if (ENGINE_MAJOR_VERSION == 4 && ENGINE_MINOR_VERSION >= 18) || ENGINE_MAJOR_VERSION >= 5
+	FPlatformApplicationMisc::ClipboardPaste(TextToImport);
+#else
+	FPlatformMisc::ClipboardPaste(TextToImport);
+#endif
+
+	// Import the nodes
+	TSet<UEdGraphNode*> PastedNodes;
+	FEdGraphUtilities::ImportNodesFromText(CustomizableObject->Source, TextToImport, /*out*/ PastedNodes);
+	
+	//Average position of nodes so we can move them while still maintaining relative distances to each other
+	FVector2D AvgNodePosition(0.0f,0.0f);
+
+	for (TSet<UEdGraphNode*>::TIterator It(PastedNodes); It; ++It)
+	{
+		UEdGraphNode* Node = *It;
+		AvgNodePosition.X += Node->NodePosX;
+		AvgNodePosition.Y += Node->NodePosY;
+	}
+
+	if ( PastedNodes.Num() > 0 )
+	{
+		float InvNumNodes = 1.0f/float(PastedNodes.Num());
+		AvgNodePosition.X *= InvNumNodes;
+		AvgNodePosition.Y *= InvNumNodes;
+	}
+
+	for (TSet<UEdGraphNode*>::TIterator It(PastedNodes); It; ++It)
+	{
+		UEdGraphNode* Node = *It;
+		if (UCustomizableObjectNode* GraphNode = Cast<UCustomizableObjectNode>(Node))
+		{
+			// There can be only one default mesh paint texture.
+			//UMaterialExpressionTextureBase* TextureSample = Cast<UMaterialExpressionTextureBase>( NewExpression );
+			//if( TextureSample )
+			//{
+			//	TextureSample->IsDefaultMeshpaintTexture = false;
+			//}
+
+			//NewExpression->UpdateParameterGuid(true, true);
+
+			//UMaterialExpressionFunctionInput* FunctionInput = Cast<UMaterialExpressionFunctionInput>( NewExpression );
+			//if( FunctionInput )
+			//{
+			//	FunctionInput->ConditionallyGenerateId(true);
+			//	FunctionInput->ValidateName();
+			//}
+
+			//UMaterialExpressionFunctionOutput* FunctionOutput = Cast<UMaterialExpressionFunctionOutput>( NewExpression );
+			//if( FunctionOutput )
+			//{
+			//	FunctionOutput->ConditionallyGenerateId(true);
+			//	FunctionOutput->ValidateName();
+			//}
+		}
+		//else if (UMaterialGraphNode_Comment* CommentNode = Cast<UMaterialGraphNode_Comment>(Node))
+		//{
+		//	CommentNode->MaterialDirtyDelegate = Material->MaterialGraph->MaterialDirtyDelegate;
+		//	CommentNode->MaterialExpressionComment->Material = Material;
+		//	Material->EditorComments.Add(CommentNode->MaterialExpressionComment);
+		//}
+
+		// Select the newly pasted stuff
+		GraphEditor->SetNodeSelection(Node, true);
+
+		Node->NodePosX = (Node->NodePosX - AvgNodePosition.X) + Location.X ;
+		Node->NodePosY = (Node->NodePosY - AvgNodePosition.Y) + Location.Y ;
+
+		Node->SnapToGrid(SNodePanel::GetSnapGridSize());
+
+		// Give new node a different Guid from the old one
+		Node->CreateNewGuid();
+	}
+
+	// Force new pasted Material Expressions to have same connections as graph nodes
+	//Material->MaterialGraph->LinkMaterialExpressionsFromGraph();
+
+	for (UEdGraphNode* PastedNode : PastedNodes) // Perform possible BackwardsCompatibleFixup. PostLoads have already been called at this point.
+	{
+		if (UCustomizableObjectNode* TypedNode = Cast<UCustomizableObjectNode>(PastedNode))
+		{
+			TypedNode->BackwardsCompatibleFixup();
+		}
+	}
+	
+	for (UEdGraphNode* PastedNode : PastedNodes)
+	{
+		if (UCustomizableObjectNode* TypedNode = Cast<UCustomizableObjectNode>(PastedNode))
+		{
+			TypedNode->PostBackwardsCompatibleFixup();
+		}
+	}
+	
+	// Update UI
+	GraphEditor->NotifyGraphChanged();
+
+	CustomizableObject->PostEditChange();
+	CustomizableObject->MarkPackageDirty();
+}
+
+
+bool FCustomizableObjectEditor::CanPasteNodes() const
+{
+	FString ClipboardContent;
+
+#if (ENGINE_MAJOR_VERSION == 4 && ENGINE_MINOR_VERSION >= 18) || ENGINE_MAJOR_VERSION >= 5
+	FPlatformApplicationMisc::ClipboardPaste(ClipboardContent);
+#else
+	FPlatformMisc::ClipboardPaste(ClipboardContent);
+#endif
+
+	return FEdGraphUtilities::CanImportNodesFromText(CustomizableObject->Source, ClipboardContent);
+}
+
+
+void FCustomizableObjectEditor::CutSelectedNodes()
+{
+	CopySelectedNodes();
+	// Cut should only delete nodes that can be duplicated
+	DeleteSelectedNodes();
+}
+
+
+bool FCustomizableObjectEditor::CanCutNodes() const
+{
+	return CanCopyNodes() && CanDeleteNodes();
+}
+
+
+void FCustomizableObjectEditor::CreateCommentBoxFromKey()
+{
+	CreateCommentBox(GraphEditor->GetPasteLocation());
+}
+
+UEdGraphNode* FCustomizableObjectEditor::CreateCommentBox(const FVector2D& InTargetPosition)
+{
+	//const FScopedTransaction Transaction(NSLOCTEXT("UnrealEd", "UEdGraphSchema_CustomizableObject", "Add Comment Box"));
+
+	UEdGraphNode_Comment* CommentTemplate = NewObject<UEdGraphNode_Comment>();
+
+	UEdGraphNode_Comment* NewComment = nullptr;
+	{
+		CustomizableObject->Modify();
+
+		// const FGraphPanelSelectionSet& SelectionSet = GraphEditor->GetSelectedNodes();
+		FSlateRect Bounds;
+		FVector2D Location;
+		FVector2D Size;
+
+		if (GraphEditor->GetBoundsForSelectedNodes(Bounds, 50.f))
+		{
+			Location.X = Bounds.Left;
+			Location.Y = Bounds.Top;
+			Size = Bounds.GetSize();
+		}
+		else
+		{
+			Location.X = InTargetPosition.X;
+			Location.Y = InTargetPosition.Y;
+			Size.X = 400;
+			Size.Y = 100;
+		}
+
+		NewComment = FEdGraphSchemaAction_NewNode::SpawnNodeFromTemplate<UEdGraphNode_Comment>(GraphEditor->GetCurrentGraph(), CommentTemplate, InTargetPosition, true);
+		NewComment->NodePosX = Location.X;
+		NewComment->NodePosY = Location.Y;
+		NewComment->NodeWidth = Size.X;
+		NewComment->NodeHeight = Size.Y;
+		NewComment->NodeComment = FString(TEXT("Comment"));
+	}
+
+	CustomizableObject->MarkPackageDirty();
+	GraphEditor->NotifyGraphChanged();
+
+	return NewComment;
+}
+
+
+void FCustomizableObjectEditor::OnAssetRegistryLoadComplete()
+{
+	AssetRegistryLoaded = true;
+
+	Viewport->SetAssetRegistryLoaded(true);
+
+	if (UpdateSkeletalMeshAfterAssetLoaded)
+	{
+		UpdateSkeletalMeshAfterAssetLoaded = false;
+
+		PreviewInstance->UpdateSkeletalMeshAsync(true, true);		
+	}
+
+	if (!CustomizableObject->IsCompiled())
+	{
+		CompileObject();
+	}
+}
+
+
+void FCustomizableObjectEditor::OnObjectPropertyChanged(UObject* Object, FPropertyChangedEvent& PropertyChangedEvent)
+{
+	if (Object->IsA(UCustomizableObjectGraph::StaticClass()) && PropertyChangedEvent.Property)
+	{
+		UClass* OuterClass = PropertyChangedEvent.Property->GetOwnerClass();
+		if (OuterClass->IsChildOf(UCustomizableObjectNodeProjectorConstant::StaticClass()))
+		{
+			if (Viewport.IsValid() && SelectedProjectorNode)
+			{
+				Viewport->ProjectorParameterChanged(SelectedProjectorNode);
+			}
+		}
+		else if (OuterClass->IsChildOf(UCustomizableObjectNodeProjectorParameter::StaticClass()))
+		{
+			if (Viewport.IsValid() && SelectedProjectorParameterNode)
+			{
+				Viewport->ProjectorParameterChanged(SelectedProjectorParameterNode);
+			}
+		}
+	}
+}
+
+
+TSharedPtr<SCustomizableObjectEditorAdvancedPreviewSettings> FCustomizableObjectEditor::GetCustomizableObjectEditorAdvancedPreviewSettings()
+{
+	return CustomizableObjectEditorAdvancedPreviewSettings;
+}
+
+
+void FCustomizableObjectEditor::ResetProjectorVisibilityNoUpdate()
+{
+	ManagingProjector = true;
+	Viewport->SetGizmoCallUpdateSkeletalMesh(false);
+	Viewport->ResetProjectorVisibility(true);
+	Viewport->SetGizmoCallUpdateSkeletalMesh(true);
+	ManagingProjector = false;
+}
+
+
+void FCustomizableObjectEditor::OnEnterText(const FText & NewText, ETextCommit::Type TextType)
+{
+	bool found = false;
+
+	if (TextType == ETextCommit::OnEnter)
+	{
+		if (!GraphEditor.IsValid() || !GraphEditor->GetCurrentGraph())
+		{
+			return;
+		}
+
+		UEdGraph* Graph = GraphEditor->GetCurrentGraph();
+		
+		for (int i = 0; i < Graph->Nodes.Num(); ++i)
+		{
+			UObject* Node = Graph->Nodes[i];
+
+			//Find Node Coincidence 
+			FString NodeName = Graph->Nodes[i]->GetNodeTitle(ENodeTitleType::ListView).ToString();
+			int32 cPos = NodeName.Find("\n");
+			
+			if (cPos >= 0)
+			{
+				NodeName.RemoveAt(cPos);
+				NodeName.InsertAt(cPos, " ");
+			}
+
+			if (NodeName.Contains(NewText.ToString(), ESearchCase::IgnoreCase))
+			{
+				LogSearchResult(Graph->Nodes[i], "Node", found, NodeName);
+				found = true;
+			}
+
+			//Find Variable method 1 (using contains of string)
+			FString VariableName = NewText.ToString();
+			VariableName.RemoveSpacesInline();
+
+			for (TFieldIterator<FProperty> It(Node->GetClass()); It; ++It)
+			{
+				FProperty* Property = *It;
+				FString ResultName = Property->GetName();
+				
+				//Find Variable name
+				if (ResultName.Contains(VariableName))
+				{
+					LogSearchResult(Graph->Nodes[i], "Variable", found, ResultName);
+					found = true;
+				}
+
+				//Find Variable String Content
+				if (FStrProperty *StringProperty = CastField<FStrProperty>(Property))
+				{
+					FString* StringResult = Property->ContainerPtrToValuePtr<FString>(Node);
+
+					if (StringResult->Contains(NewText.ToString()))
+					{
+						LogSearchResult(Graph->Nodes[i], "Value", found, *StringResult);
+						found = true;
+					}
+				}
+
+				//Find Variable Enum Content (Hard coded)
+				//TODO: improve using full Reflection
+				if (FEnumProperty *EnumProperty = CastField<FEnumProperty>(Property))
+				{
+					UEnum* EnumResult = EnumProperty->GetEnum();
+					FString EnumType = EnumResult->CppType;
+					FString StringResult;
+
+					FFieldClass* SourceObjectClass = Property->GetClass();
+
+					if (EnumType == "ECustomizableObjectGroupType")
+					{
+						ECustomizableObjectGroupType* value = EnumProperty->ContainerPtrToValuePtr<ECustomizableObjectGroupType>(Node);
+						StringResult = EnumResult->GetDisplayNameTextByIndex((int32)*value).ToString();
+					}
+					else if (EnumType == "ENodeEnabledState")
+					{
+						ENodeEnabledState* value = EnumProperty->ContainerPtrToValuePtr<ENodeEnabledState>(Node);
+						StringResult = EnumResult->GetDisplayNameTextByIndex((int32)*value).ToString();
+					}
+					else if (EnumType == "ECustomizableObjectAutomaticLODStrategy")
+					{
+						ECustomizableObjectAutomaticLODStrategy* value = EnumProperty->ContainerPtrToValuePtr<ECustomizableObjectAutomaticLODStrategy>(Node);
+						StringResult = EnumResult->GetDisplayNameTextByIndex((int32)*value).ToString();
+					}
+					else if (EnumType == "ECustomizableObjectProjectorType")
+					{
+						ECustomizableObjectProjectorType* value = EnumProperty->ContainerPtrToValuePtr<ECustomizableObjectProjectorType>(Node);
+						StringResult = EnumResult->GetDisplayNameTextByIndex((int32)*value).ToString();
+					}
+					else if (EnumType == "EColorArithmeticOperation")
+					{
+						EColorArithmeticOperation* value = EnumProperty->ContainerPtrToValuePtr<EColorArithmeticOperation>(Node);
+						StringResult = EnumResult->GetDisplayNameTextByIndex((int32)*value).ToString();
+					}
+					else if (EnumType == "ECustomizableObjectNodeMaterialVariationType")
+					{
+						ECustomizableObjectNodeMaterialVariationType* value = EnumProperty->ContainerPtrToValuePtr<ECustomizableObjectNodeMaterialVariationType>(Node);
+						StringResult = EnumResult->GetDisplayNameTextByIndex((int32)*value).ToString();
+					}
+
+					if (StringResult.Contains(NewText.ToString()))
+					{
+						LogSearchResult(Graph->Nodes[i], "Value", found, StringResult);
+						found = true;
+					}
+				}
+
+				//Find Variable Array Content (Hard coded)
+				//TODO: improve using full Reflection
+				if (FArrayProperty* ArrayProperty = CastField<FArrayProperty>(Property))
+				{
+					FProperty* AProperty = ArrayProperty->Inner;
+					FString ArrayType = AProperty->GetCPPType();
+
+					if (ArrayType == "FGroupProjectorParameterImage")
+					{
+						TArray<FGroupProjectorParameterImage>  Array =
+							*ArrayProperty->ContainerPtrToValuePtr<TArray<FGroupProjectorParameterImage>>(Node);
+
+						//Class vars
+						FString OptionName = "OptionName";
+
+						for (int a = 0; a < Array.Num(); ++a)
+						{
+							if (OptionName.Contains(VariableName))
+							{
+								LogSearchResult(Graph->Nodes[i], "Variable", found, OptionName);
+								found = true;
+							}
+
+							FString OptionNameContent = Array[a].OptionName;
+
+							if (OptionNameContent.Contains(NewText.ToString()))
+							{
+								LogSearchResult(Graph->Nodes[i], "Value", found, Array[a].OptionName);
+								found = true;
+							}
+						}
+					}
+					else if (ArrayType == "FGroupProjectorParameterPose")
+					{
+						TArray<FGroupProjectorParameterPose>  Array =
+							*ArrayProperty->ContainerPtrToValuePtr<TArray<FGroupProjectorParameterPose>>(Node);
+
+						//Class vars
+						FString PoseName = "PoseName";
+
+						for (int a = 0; a < Array.Num(); ++a)
+						{
+							if (PoseName.Contains(VariableName))
+							{
+								LogSearchResult(Graph->Nodes[i], "Variable", found, PoseName);
+								found = true;
+							}
+
+							FString OptionNameContent = Array[a].PoseName;
+
+							if (OptionNameContent.Contains(NewText.ToString()))
+							{
+								LogSearchResult(Graph->Nodes[i], "Value", found, Array[a].PoseName);
+								found = true;
+							}
+						}
+					}
+					else if (ArrayType == "FCustomizableObjectNodeFloatDescription")
+					{
+						TArray<FCustomizableObjectNodeFloatDescription>  Array =
+							*ArrayProperty->ContainerPtrToValuePtr<TArray<FCustomizableObjectNodeFloatDescription>>(Node);
+						//Class vars
+						FString Name = "Name";
+
+						for (int a = 0; a < Array.Num(); ++a)
+						{
+							if (Name.Contains(VariableName))
+							{
+								LogSearchResult(Graph->Nodes[i], "Variable", found, Name);
+								found = true;
+							}
+
+							FString OptionNameContent = Array[a].Name;
+
+							if (OptionNameContent.Contains(NewText.ToString()))
+							{
+								LogSearchResult(Graph->Nodes[i], "Value", found, Array[a].Name);
+								found = true;
+							}
+						}
+					}
+					else if (ArrayType == "FCustomizableObjectNodeEnumValue")
+					{
+						TArray<FCustomizableObjectNodeEnumValue>  Array =
+							*ArrayProperty->ContainerPtrToValuePtr<TArray<FCustomizableObjectNodeEnumValue>>(Node);
+						//Class vars
+						FString Name = "Name";
+
+						for (int a = 0; a < Array.Num(); ++a)
+						{
+							if (Name.Contains(VariableName))
+							{
+								LogSearchResult(Graph->Nodes[i], "Variable", found, Name);
+								found = true;
+							}
+
+							FString OptionNameContent = Array[a].Name;
+
+							if (OptionNameContent.Contains(NewText.ToString()))
+							{
+								LogSearchResult(Graph->Nodes[i], "Value", found, Array[a].Name);
+								found = true;
+							}
+						}
+					}
+					else if (ArrayType == "FCustomizableObjectNodeExtendMaterialImage")
+					{
+						TArray<FCustomizableObjectNodeExtendMaterialImage>  Array =
+							*ArrayProperty->ContainerPtrToValuePtr<TArray<FCustomizableObjectNodeExtendMaterialImage>>(Node);
+						//Class vars
+						FString Name = "Name";
+
+						for (int a = 0; a < Array.Num(); ++a)
+						{
+							if (Name.Contains(VariableName))
+							{
+								LogSearchResult(Graph->Nodes[i], "Variable", found, Name);
+								found = true;
+							}
+
+							FString OptionNameContent = Array[a].Name;
+
+							if (OptionNameContent.Contains(NewText.ToString()))
+							{
+								LogSearchResult(Graph->Nodes[i], "Value", found, Array[a].Name);
+								found = true;
+							}
+						}
+					}
+					else if (ArrayType == "FCustomizableObjectNodeMaterialImage")
+					{
+						TArray<FCustomizableObjectNodeMaterialImage>  Array =
+							*ArrayProperty->ContainerPtrToValuePtr<TArray<FCustomizableObjectNodeMaterialImage>>(Node);
+						//Class vars
+						FString Name = "Name";
+
+						for (int a = 0; a < Array.Num(); ++a)
+						{
+							if (Name.Contains(VariableName))
+							{
+								LogSearchResult(Graph->Nodes[i], "Variable", found, Name);
+								found = true;
+							}
+
+							FString OptionNameContent = Array[a].Name;
+
+							if (OptionNameContent.Contains(NewText.ToString()))
+							{
+								LogSearchResult(Graph->Nodes[i], "Value", found, Array[a].Name);
+								found = true;
+							}
+						}
+					}
+					else if (ArrayType == "FCustomizableObjectNodeMaterialVector")
+					{
+						TArray<FCustomizableObjectNodeMaterialVector>  Array =
+							*ArrayProperty->ContainerPtrToValuePtr<TArray<FCustomizableObjectNodeMaterialVector>>(Node);
+						//Class vars
+						FString Name = "Name";
+
+						for (int a = 0; a < Array.Num(); ++a)
+						{
+							if (Name.Contains(VariableName))
+							{
+								LogSearchResult(Graph->Nodes[i], "Variable", found, Name);
+								found = true;
+							}
+
+							FString NameContent = Array[a].Name;
+
+							if (NameContent.Contains(NewText.ToString()))
+							{
+								LogSearchResult(Graph->Nodes[i], "Value", found, Array[a].Name);
+								found = true;
+							}
+						}
+					}
+					else if (ArrayType == "FCustomizableObjectNodeMaterialScalar")
+					{
+						TArray<FCustomizableObjectNodeMaterialScalar>  Array =
+							*ArrayProperty->ContainerPtrToValuePtr<TArray<FCustomizableObjectNodeMaterialScalar>>(Node);
+						//Class vars
+						FString Name = "Name";
+
+						for (int a = 0; a < Array.Num(); ++a)
+						{
+							if (Name.Contains(VariableName))
+							{
+								LogSearchResult(Graph->Nodes[i], "Variable", found, Name);
+								found = true;
+							}
+
+							FString NameContent = Array[a].Name;
+
+							if (NameContent.Contains(NewText.ToString()))
+							{
+								LogSearchResult(Graph->Nodes[i], "Value", found, Array[a].Name);
+								found = true;
+							}
+						}
+					}
+					else if (ArrayType == "FCustomizableObjectMaterialVariation")
+					{
+						TArray<FCustomizableObjectMaterialVariation>  Array =
+							*ArrayProperty->ContainerPtrToValuePtr<TArray<FCustomizableObjectMaterialVariation>>(Node);
+						//Class vars
+						FString Tag = "Tag";
+
+						for (int a = 0; a < Array.Num(); ++a)
+						{
+							if (Tag.Contains(VariableName))
+							{
+								LogSearchResult(Graph->Nodes[i], "Variable", found, Tag);
+								found = true;
+							}
+
+							FString OptionNameContent = Array[a].Tag;
+
+							if (OptionNameContent.Contains(NewText.ToString()))
+							{
+								LogSearchResult(Graph->Nodes[i], "Value", found, Array[a].Tag);
+								found = true;
+							}
+						}
+					}
+					else if (ArrayType == "FString")
+					{
+						TArray<FString>  Array =
+							*ArrayProperty->ContainerPtrToValuePtr<TArray<FString>>(Node);
+
+						for (int a = 0; a < Array.Num(); ++a)
+						{
+							if (Array[a].Contains(NewText.ToString()))
+							{
+								LogSearchResult(Graph->Nodes[i], "Value", found, Array[a]);
+								found = true;
+							}
+						}
+
+					}
+					else if (ArrayType == "FCustomizableObjectState")
+					{
+						TArray<FCustomizableObjectState>  Array =
+							*ArrayProperty->ContainerPtrToValuePtr<TArray<FCustomizableObjectState>>(Node);
+						//Class vars
+						FString Name = "Name";
+						FString RuntimeParameters = "RuntimeParameters";
+
+						for (int a = 0; a < Array.Num(); ++a)
+						{
+							if (Name.Contains(VariableName))
+							{
+								LogSearchResult(Graph->Nodes[i], "Variable", found, Name);
+								found = true;
+							}
+
+							if (Array[a].Name.Contains(NewText.ToString()))
+							{
+								LogSearchResult(Graph->Nodes[i], "Value", found, Array[a].Name);
+								found = true;
+							}
+
+							if (RuntimeParameters.Contains(VariableName))
+							{
+								LogSearchResult(Graph->Nodes[i], "Variable", found, Name);
+								found = true;
+							}
+
+							for (int s = 0; s < Array[a].RuntimeParameters.Num(); ++s)
+							{
+								if (Array[a].RuntimeParameters[s].Contains(NewText.ToString()))
+								{
+									LogSearchResult(Graph->Nodes[i], "Value", found, Array[a].RuntimeParameters[s]);
+									found = true;
+								}
+							}
+						}
+					}
+					else if (ArrayType == "FCustomizableObjectNodeSkeletalMeshLOD")
+					{
+						TArray<FCustomizableObjectNodeSkeletalMeshLOD>  Array =
+							*ArrayProperty->ContainerPtrToValuePtr<TArray<FCustomizableObjectNodeSkeletalMeshLOD>>(Node);
+						
+						FString LODs = "LODs";
+
+						for (int a = 0; a < Array.Num(); ++a)
+						{
+							if (LODs.Contains(VariableName))
+							{
+								LogSearchResult(Graph->Nodes[i], "Variable", found, LODs);
+								found = true;
+							}
+
+							for (int s = 0; s < Array[a].Materials.Num(); ++s)
+							{
+								FString Name = "Name";
+
+								if (Name.Contains(VariableName))
+								{
+									LogSearchResult(Graph->Nodes[i], "Variable", found, Name);
+									found = true;
+								}
+
+								if (Array[a].Materials[s].Name.Contains(NewText.ToString()))
+								{
+									LogSearchResult(Graph->Nodes[i], "Value", found, Array[a].Materials[s].Name);
+									found = true;
+								}
+							}
+						}
+					}
+					else if (ArrayType == "FCustomizableObjectNodeStaticMeshLOD")
+					{
+						TArray<FCustomizableObjectNodeStaticMeshLOD>  Array =
+							*ArrayProperty->ContainerPtrToValuePtr<TArray<FCustomizableObjectNodeStaticMeshLOD>>(Node);
+
+						FString LODs = "LODs";
+
+						for (int a = 0; a < Array.Num(); ++a)
+						{
+							if (LODs.Contains(VariableName))
+							{
+								LogSearchResult(Graph->Nodes[i], "Variable", found, LODs);
+								found = true;
+							}
+
+							for (int s = 0; s < Array[a].Materials.Num(); ++s)
+							{
+								FString Name = "Name";
+
+								if (Name.Contains(VariableName))
+								{
+									LogSearchResult(Graph->Nodes[i], "Variable", found, Name);
+									found = true;
+								}
+
+								if (Array[a].Materials[s].Name.Contains(NewText.ToString()))
+								{
+									LogSearchResult(Graph->Nodes[i], "Value", found, Array[a].Materials[s].Name);
+									found = true;
+								}
+							}
+						}
+					}
+				}				
+			}
+		}
+
+		const FText Text = found ?
+			LOCTEXT("SearchCompleted", "Search completed") :
+			FText::FromString("No Results for: " + NewText.ToString());
+
+		FCustomizableObjectEditorLogger::CreateLog(Text)
+			.Category(ELoggerCategory::GraphSearch)
+			.CustomNotification()
+			.Log();
+	}
+}
+
+
+void FCustomizableObjectEditor::LogSearchResult(UEdGraphNode* Node, FString Type, bool IsFirst, FString Result) const
+{
+	if (!IsFirst)
+	{
+		FCustomizableObjectEditorLogger::CreateLog(LOCTEXT("SearchResults", "Search Results:"))
+		.Notification(false)
+		.Log();
+	}
+	
+	FCustomizableObjectEditorLogger::CreateLog(FText::FromString(Type + ": " + Result))
+	.Node(*CastChecked<UCustomizableObjectNode>(Node))
+	.BaseObject()
+	.Notification(false)
+	.Log();
+}
+
+
+void FCustomizableObjectEditor::UpdatePreviewVisibility()
+{
+	for (UDebugSkelMeshComponent* PreviewSkeletalMeshComponent : PreviewSkeletalMeshComponents)
+	{
+		if (PreviewSkeletalMeshComponent)
+		{
+			if (PreviewInstance->SkeletalMeshStatus != ESkeletalMeshState::UpdateError)
+			{
+				PreviewSkeletalMeshComponent->SetVisibility(true, true);
+			}
+			else
+			{
+				PreviewSkeletalMeshComponent->SetVisibility(false, true);
+			}
+		}
+	}
+}
+
+
+void FCustomizableObjectEditor::OnUpdatePreviewInstance()
+{
+	UpdatePreviewVisibility();
+		
+	for (UDebugSkelMeshComponent* PreviewSkeletalMeshComponent : PreviewSkeletalMeshComponents)
+	{
+		PreviewSkeletalMeshComponent->RegisterComponentWithWorld(PreviewSkeletalMeshComponent->GetWorld());
+		PreviewSkeletalMeshComponent->bComponentUseFixedSkelBounds = true; // First bounds computed would be using physics asset
+		PreviewSkeletalMeshComponent->UpdateBounds();
+	}
+	
+	Viewport->SetPreviewComponents(PreviewSkeletalMeshComponents);
+	Viewport->SetDrawDefaultUVMaterial();
+	ViewportClient->Invalidate();
+	ViewportClient->ReSetAnimation();
+	ViewportClient->SetReferenceMeshMissingWarningMessage(false);
+
+	if (AssetRegistryLoaded)
+	{
+		// If the instance is updated due to a change in a parameter projector, set again the LastSelectedProjectorParameter variable with the
+		// current projector the viewport gizmo has to continue having it as selected when rebuilding the parameter details widget.
+		// Otherwise, reset the value.
+
+		if (Viewport->GetGizmoHasAssignedData() && !Viewport->GetGizmoAssignedDataIsFromNode())
+		{
+			PreviewInstance->LastSelectedProjectorParameter = Viewport->GetGizmoProjectorParameterName();
+			PreviewInstance->LastSelectedProjectorParameterWithIndex = Viewport->GetGizmoProjectorParameterNameWithIndex();
+
+			if (!PreviewInstance->ProjectorAlphaChange && !Viewport->GetIsManipulatingGizmo())
+			{
+				TArray<UObject*> instances;
+				instances.Add(PreviewInstance);
+				CustomizableInstanceDetailsView->SetObjects(instances, true);
+			}
+
+			PreviewInstance->ProjectorAlphaChange = false;
+		}
+		else
+		{
+			PreviewInstance->LastSelectedProjectorParameter = "";
+			PreviewInstance->LastSelectedProjectorParameterWithIndex = "";
+		}
+
+		//UE_LOG(LogMutable, Warning, TEXT("LastSelectedProjectorParameter=%s, LastSelectedProjectorParameterWithIndex=%s"), *(PreviewInstance->LastSelectedProjectorParameter), *(PreviewInstance->LastSelectedProjectorParameterWithIndex));
+	}
+
+	if (TextureAnalyzer.IsValid())
+	{
+		TextureAnalyzer->RefreshTextureAnalyzerTable(PreviewInstance);
+	}
+
+	for (int32 ComponentIndex = 0; ComponentIndex < PreviewInstance->SkeletalMeshes.Num(); ++ComponentIndex)
+	{
+		BakeHelper_RegenerateImportedModel(PreviewInstance->SkeletalMeshes[ComponentIndex]);
+	}
+
+	// Do postponed work.
+	for (const TFunction<void()>& Work : OnUpdatePreviewInstanceWork)
+	{
+		Work();		
+	}
+	OnUpdatePreviewInstanceWork.Reset();
+}
+
+
+void FCustomizableObjectEditor::UpdateObjectProperties()
+{
+	if (CustomizableObjectDetailsView.IsValid())
+	{
+#if (ENGINE_MAJOR_VERSION == 4 && ENGINE_MINOR_VERSION >= 18) || ENGINE_MAJOR_VERSION >= 5
+		CustomizableObjectDetailsView->ForceRefresh();
+#endif
+	}
+}
+
+
+void FCustomizableObjectEditor::OnNodeTitleCommitted(const FText& NewText, ETextCommit::Type CommitInfo, UEdGraphNode* NodeBeingChanged)
+{
+	if (NodeBeingChanged)
+	{
+		const FScopedTransaction Transaction(LOCTEXT("RenameNode", "Rename Node"));
+		NodeBeingChanged->Modify();
+		NodeBeingChanged->OnRenameNode(NewText.ToString());
+	}
+}
+
+
+void FCustomizableObjectEditor::UpdateGraphNodeProperties()
+{
+	OnSelectedGraphNodesChanged(FGraphPanelSelectionSet());
+	OnSelectedGraphNodesChanged(GraphEditor->GetSelectedNodes());
+}
+
+
+bool FCustomizableObjectEditor::GetAssetRegistryLoaded()
+{
+	return AssetRegistryLoaded;
+}
+
+
+void FCustomizableObjectEditor::OpenTextureAnalyzerTab()
+{
+	TabManager->TryInvokeTab(TextureAnalyzerTabId);
+}
+
+
+void FCustomizableObjectEditor::OpenPerformanceReportTab()
+{
+	TabManager->TryInvokeTab(PerformanceReportTabId);
+}
+
+
+TSharedRef<SDockTab> FCustomizableObjectEditor::SpawnTab_TextureAnalyzer(const FSpawnTabArgs& Args)
+{
+	check(Args.GetTabId() == TextureAnalyzerTabId);
+	
+	return SNew(SDockTab)
+	.Label(LOCTEXT("Texture Analyzer", "Texture Analyzer"))
+	[
+		TextureAnalyzer.ToSharedRef()
+	];
+}
+
+
+TSharedRef<SDockTab> FCustomizableObjectEditor::SpawnTab_PerformanceReport(const FSpawnTabArgs& Args)
+{
+	check(Args.GetTabId() == PerformanceReportTabId);
+	check(CustomizableObject);
+
+	if (!PerformanceReport.IsValid())
+	{
+		PerformanceReport = SNew(SCustomizableObjecEditorPerformanceReport).CustomizableObject(CustomizableObject);
+	}
+
+	return SNew(SDockTab)
+	.Label(LOCTEXT("Performance Report", "Performance Report"))
+	[
+		PerformanceReport.ToSharedRef()
+	];
+}
+
+
+TSharedRef<SDockTab> FCustomizableObjectEditor::SpawnTab_TagExplorer(const FSpawnTabArgs& Args)
+{
+	check(Args.GetTabId() == TagExplorerTabId);
+
+	return SNew(SDockTab)
+	.Label(LOCTEXT("Tag_Explorer", "Tag Explorer"))
+	[
+		TagExplorer.ToSharedRef()
+	];
+}
+
+
+UCustomizableObject* FCustomizableObjectEditor::GetAbsoluteCOParent(const UCustomizableObjectNodeObject* const Root)
+{
+	if (Root->ParentObject != nullptr)
+	{
+		//Get all the NodeObjects
+		TArray<UCustomizableObjectNodeObject*> ObjectNodes;
+		Root->ParentObject->Source->GetNodesOfClass<UCustomizableObjectNodeObject>(ObjectNodes);
+
+		for (int i = 0; i < ObjectNodes.Num(); ++i)
+		{
+			//Getting the parent of the root
+			if (ObjectNodes[i]->ParentObject == nullptr)
+			{
+				return Root->ParentObject;
+			}
+			else
+			{
+				return GetAbsoluteCOParent(ObjectNodes[0]);
+			}
+		}
+	}
+
+	return nullptr;
+}
+
+
+void FCustomizableObjectEditor::AddCachedReferencers(const FName& PathName, TArray<FName>& ArrayReferenceNames, TArray<FAssetData>& ArrayAssetData)
+{
+	ArrayReferenceNames.Empty();
+	FAssetRegistryModule& AssetRegistryModule = FModuleManager::LoadModuleChecked<FAssetRegistryModule>("AssetRegistry");
+	AssetRegistryModule.Get().GetReferencers(PathName, ArrayReferenceNames);
+
+	FARFilter Filter;
+	for (const FName& ReferenceName : ArrayReferenceNames)
+	{
+		bool IsCachedInAssetData = false;
+
+		const int32 MaxIndex = ArrayAssetData.Num();
+
+		for (int32 i = 0; i < MaxIndex; ++i)
+		{
+			if (ArrayAssetData[i].PackageName.ToString() == ReferenceName.ToString())
+			{
+				IsCachedInAssetData = true;
+			}
+		}
+
+		if (!IsCachedInAssetData && !ReferenceName.ToString().StartsWith(TEXT("/TempAutosave")))
+		{
+			Filter.PackageNames.Add(ReferenceName);
+		}
+	}
+
+	Filter.bIncludeOnlyOnDiskAssets = false;
+
+	TArray<FAssetData> ArrayAssetDataTemp;
+	AssetRegistryModule.Get().GetAssets(Filter, ArrayAssetDataTemp);
+
+	// Store only those which have static class type Customizable Object, to avoid loading not needed elements
+	const int32 MaxIndex = ArrayAssetDataTemp.Num();
+	for (int32 i = 0; i < MaxIndex; ++i)
+	{
+		if (ArrayAssetDataTemp[i].GetClass() == UCustomizableObject::StaticClass())
+		{
+			ArrayAssetData.Add(ArrayAssetDataTemp[i]);
+		}
+	}
+}
+
+
+void FCustomizableObjectEditor::GetExternalChildObjects(const UCustomizableObject* const Object, TArray<UCustomizableObject*>& ExternalChildren, const bool bRecursively, const EObjectFlags ExcludeFlags)
+{
+	TArray<FAssetData> ArrayAssetData;
+	TArray<FName> ArrayAlreadyProcessedChild;
+	TArray<FName> ArrayReferenceNames;
+
+	AddCachedReferencers(*Object->GetOuter()->GetPathName(), ArrayReferenceNames, ArrayAssetData);
+
+	bool bMultipleBaseObjectsFound = false;
+
+	FAssetData* AssetData = nullptr;
+	UCustomizableObject* ChildObject;
+
+	for (const FName& ReferenceName : ArrayReferenceNames)
+	{
+		const int32 MaxIndex = ArrayAssetData.Num();
+
+		for (int32 i = 0; i < MaxIndex; ++i)
+		{
+			if (ArrayAssetData[i].PackageName.ToString() == ReferenceName.ToString())
+			{
+				AssetData = &ArrayAssetData[i];
+			}
+		}
+
+		if (AssetData != nullptr) // Elements in ArrayAssetData are already of static class UCustomizableObject
+		{
+			ChildObject = Cast<UCustomizableObject>(AssetData->GetAsset());
+			if (!ChildObject) continue;
+
+			if (ChildObject != Object && !ChildObject->HasAnyFlags(ExcludeFlags))
+			{
+				ExternalChildren.AddUnique(ChildObject);
+			}
+
+			TArray<UCustomizableObjectNodeObjectGroup*> GroupNodes;
+			ChildObject->Source->GetNodesOfClass<UCustomizableObjectNodeObjectGroup>(GroupNodes);
+
+			if (GroupNodes.Num() > 0) // Only grafs with group nodes should have child grafs
+			{
+				if (ArrayAlreadyProcessedChild.Find(ReferenceName) == INDEX_NONE)
+				{
+					ArrayAlreadyProcessedChild.Add(ReferenceName);
+					if (bRecursively)
+					{
+						GetExternalChildObjects(ChildObject, ExternalChildren, bRecursively);
+					}
+				}
+			}
+		}
+	}
+}
+
+
+void FCustomizableObjectEditor::MarkForReconstruct(UEdGraphNode* Node)
+{
+	ReconstructNodes.Add(Node);
+}
+
+
+void FCustomizableObjectEditor::CreatePreviewComponents()
+{
+	// Getting the number of mesh components from the root node
+	int32 NumMeshComponents = 0;
+
+	if (CustomizableObject->Source)
+	{
+		TArray<UCustomizableObjectNodeObject*> RootNode;
+		CustomizableObject->Source->GetNodesOfClass< UCustomizableObjectNodeObject>(RootNode);
+
+		for (int32 i = 0; i < RootNode.Num(); ++i)
+		{
+			if (RootNode[i]->bIsBase)
+			{
+				NumMeshComponents = RootNode[i]->NumMeshComponents;
+				break;
+			}
+		}
+	}
+
+	if (NumMeshComponents != PreviewSkeletalMeshComponents.Num())
+	{
+		if (NumMeshComponents > PreviewSkeletalMeshComponents.Num())
+		{
+			PreviewSkeletalMeshComponents.AddZeroed(NumMeshComponents - PreviewSkeletalMeshComponents.Num());
+			PreviewCustomizableSkeletalComponents.AddZeroed(NumMeshComponents - PreviewCustomizableSkeletalComponents.Num());
+		}
+		else
+		{
+			PreviewSkeletalMeshComponents.SetNumZeroed(NumMeshComponents, false);
+			PreviewCustomizableSkeletalComponents.SetNumZeroed(NumMeshComponents, false);
+		}
+
+		for (int32 ComponentIndex = 0; ComponentIndex < PreviewSkeletalMeshComponents.Num(); ++ComponentIndex)
+		{
+			if (!PreviewSkeletalMeshComponents[ComponentIndex])
+			{
+				if (!CreatePreviewComponent(ComponentIndex))
+				{
+					PreviewSkeletalMeshComponents.Empty();
+					PreviewCustomizableSkeletalComponents.Empty();
+					break;
+				}
+			}
+		}
+	}
+}
+
+
+
+#undef LOCTEXT_NAMESPACE
