[mutable] Move Mutable plugin from the restricted folder to the experimental plugins folder.

diff --git a/Engine/Plugins/Experimental/Mutable/Source/CustomizableObjectEditor/Private/CustomizableObjectEditor.h b/Engine/Plugins/Experimental/Mutable/Source/CustomizableObjectEditor/Private/CustomizableObjectEditor.h
new file mode 100644
--- /dev/null
+++ b/Engine/Plugins/Experimental/Mutable/Source/CustomizableObjectEditor/Private/CustomizableObjectEditor.h
@@ -0,0 +1,482 @@
+// Copyright Epic Games, Inc. All Rights Reserved.
+
+#pragma once
+
+#include "ICustomizableObjectEditor.h"
+#include "Misc/NotifyHook.h"
+#include "TickableEditorObject.h"
+#include "EditorUndoClient.h"
+#include "Toolkits/AssetEditorToolkit.h"
+#include "GraphEditor.h"
+#include "CustomizableObjectCompiler.h"
+#include "Widgets/Input/SNumericEntryBox.h"
+#include "Framework/MultiBox/MultiBoxBuilder.h"
+#include "StatePerformanceTesting.h"
+#include "CustomizableObjectInstance.h"
+#include "SCustomizableObjectEditorViewport.h"
+#include "CustomizableObjectEditorViewportClient.h"
+#include "Animation/DebugSkelMeshComponent.h"
+#include "UObject/ObjectMacros.h"
+#include "Templates/Function.h"
+
+#include "CustomizableObjectEditor.generated.h"
+
+DECLARE_DELEGATE(FCreatePreviewInstanceFlagDelegate);
+
+
+/**
+* Wrapper UObject class for the UCustomizableObjectInstance::FObjectInstanceUpdatedDelegate dynamic multicast delegate
+*/
+UCLASS()
+class UUpdateClassWrapper : public UObject
+{
+public:
+	GENERATED_BODY()
+
+	/** Method to assign for the callback */
+	UFUNCTION()
+	void DelegatedCallback(UCustomizableObjectInstance* Instance);
+
+	FCreatePreviewInstanceFlagDelegate Delegate;
+};
+
+
+/** Statistics for the Texture Analyzer */
+UCLASS(Transient, MinimalAPI, meta = (DisplayName = "Texture Stats"))
+class UCustomizableObjectEditorTextureStats : public UObject
+{
+	GENERATED_BODY()
+
+public:
+
+	/** Texture - double click to open */
+	UPROPERTY(VisibleAnywhere, AssetRegistrySearchable, Category = "Stats", meta = (DisplayName = "Texture", ColumnWidth = "40"))
+	FString TextureName;
+
+	/** Texture - double click to open */
+	UPROPERTY(VisibleAnywhere, AssetRegistrySearchable, Category = "Stats", meta = (DisplayName = "Name", ColumnWidth = "50"))
+	FString TextureParameterName;
+
+	/** Material - double click to open */
+	UPROPERTY(VisibleAnywhere, AssetRegistrySearchable, Category = "Stats", meta = (DisplayName = "Material", ColumnWidth = "50"))
+	FString MaterialName;
+
+	/** Parent Material - double click to open */
+	UPROPERTY(VisibleAnywhere, AssetRegistrySearchable, Category = "Stats", meta = (DisplayName = "Parent", ColumnWidth = "50"))
+	FString MaterialParameterName;
+
+	/** Used to open the texture in the editor*/
+	UPROPERTY()
+	TObjectPtr<UTexture> Texture;
+
+	/** Used to open the material in the editor*/
+	UPROPERTY()
+	TObjectPtr<UMaterialInterface> Material;
+
+	/** Used to open the parent material in the editor*/
+	UPROPERTY()
+	TObjectPtr<UMaterialInterface> ParentMaterial;
+
+	/** Resolution of the texture */
+	UPROPERTY(VisibleAnywhere, AssetRegistrySearchable, Category = "Stats", meta = (DisplayName = "Resolution X", ColumnWidth = "40", DisplayRight = "true"))
+	int32 ResolutionX;
+	UPROPERTY(VisibleAnywhere, AssetRegistrySearchable, Category = "Stats", meta = (DisplayName = "Resolution Y", ColumnWidth = "40", DisplayRight = "true"))
+	int32 ResolutionY;
+
+	/** The memory used in KB */
+	UPROPERTY(VisibleAnywhere, AssetRegistrySearchable, Category = "Stats", meta = (DisplayName = "Size Kb", ColumnWidth = "90"))
+	FString Size;
+
+	/** The texture format, e.g. PF_DXT1 */
+	UPROPERTY(VisibleAnywhere, AssetRegistrySearchable, Category = "Stats", meta = (ColumnWidth = "96"))
+	TEnumAsByte<EPixelFormat> Format;
+
+	/** LOD Bias for this texture. (Texture LODBias + Texture group) */
+	UPROPERTY(VisibleAnywhere, AssetRegistrySearchable, Category = "Stats", meta = (DisplayName = "LOD Bias", ColumnWidth = "70"))
+	int32 LODBias;
+	
+	/** Says if the texture is being streamed */
+	UPROPERTY(VisibleAnywhere, AssetRegistrySearchable, Category = "Stats", meta = (DisplayName = "Streamed", ColumnWidth = "70"))
+	FString IsStreamed;
+
+	/** The Level of detail group of the texture */
+	UPROPERTY(VisibleAnywhere, AssetRegistrySearchable, Category = "Stats", meta = (DisplayName = "Texture Group", ColumnWidth = "70"))
+	TEnumAsByte<enum TextureGroup> LODGroup;
+	
+	/** The Component of the texture */
+	UPROPERTY(VisibleAnywhere, AssetRegistrySearchable, Category = "Stats", meta = (DisplayName = "Component", ColumnWidth = "70"))
+	int32 Component;
+
+};
+
+
+/**
+ * CustomizableObject Editor class
+ */
+class FCustomizableObjectEditor : 
+	public ICustomizableObjectEditor, 
+	public FGCObject,
+	public FNotifyHook,
+	public FTickableEditorObject,
+	public FEditorUndoClient
+{
+public:
+	/**
+	 * Create a new Customizable Object editor. Called immediately after construction.
+	 *
+	 * @param	Mode					Asset editing mode for this editor (standalone or world-centric)
+	 * @param	InitToolkitHost			When Mode is WorldCentric, this is the level editor instance to spawn this editor within
+	 * @param	ObjectToEdit			The object to edit
+	 */
+	static TSharedRef<FCustomizableObjectEditor> Create(const EToolkitMode::Type Mode, const TSharedPtr<IToolkitHost>& InitToolkitHost, UCustomizableObject* ObjectToEdit);
+
+	virtual ~FCustomizableObjectEditor() override;
+
+	virtual void RegisterTabSpawners(const TSharedRef<class FTabManager>& TabManager) override;
+	virtual void UnregisterTabSpawners(const TSharedRef<class FTabManager>& TabManager) override;
+	void CreateCommentBoxFromKey();
+
+	/**
+	 * Initialize a new Customizable Object editor. Called immediately after construction.
+	 * Required due to being unable to use SharedThis in the constructor.
+	 *
+	 * See static Create(...) function.
+	 */
+	void InitCustomizableObjectEditor(const EToolkitMode::Type Mode, const TSharedPtr<IToolkitHost>& InitToolkitHost);
+
+	// FGCObject interface
+	virtual void AddReferencedObjects( FReferenceCollector& Collector ) override; 
+	virtual FString GetReferencerName() const override
+	{
+		return TEXT("FCustomizableObjectEditor");
+	}
+
+	// IToolkit interface
+	virtual FName GetToolkitFName() const override;
+	virtual FText GetBaseToolkitName() const override;
+	virtual FText GetToolkitName() const override;
+	virtual FString GetWorldCentricTabPrefix() const override;
+	virtual FLinearColor GetWorldCentricTabColorScale() const override;
+
+	// ICustomizableObjectEditor interface
+	virtual UCustomizableObject* GetCustomizableObject() override;
+	virtual void RefreshTool() override;
+	virtual void RefreshViewport() override;
+	virtual UCustomizableObjectInstance* GetPreviewInstance() override;
+	virtual bool CanPasteNodes() const override;
+	virtual void PasteNodesHere(const FVector2D& Location) override;
+	virtual void SelectNode(const UCustomizableObjectNode* Node) override;
+	virtual void SetPoseAsset(class UPoseAsset* PoseAssetParameter) override;
+	
+	/** Called to undo the last action */
+	void UndoGraphAction();
+
+	/** Called to redo the last undone action */
+	void RedoGraphAction();
+
+	/** Utility method: Test whether the CO Node Object given as parameter is linked to any of the CO Node Object Group nodes
+	* in the Test CO given as parameter */
+	static bool GroupNodeIsLinkedToParentByName(UCustomizableObjectNodeObject* Node, UCustomizableObject* Test, const FString& ParentGroupName);
+
+	/** Callback to flag LaunchRefreshMaterialInAllChildren and start the timing for the Editor notification before actually calling
+	* the method that does the task, RefreshMaterialNodesInAllChildren */
+	void RefreshMaterialNodesInAllChildrenCallback();
+
+	/** Called to refresh all the Customizable Object Material Node nodes children of a Customizalbe Object Node Object Group */
+	void RefreshMaterialNodesInAllChildren();
+
+	// FEditorUndoClient Interface
+	virtual void PostUndo(bool bSuccess) override;
+	virtual void PostRedo(bool bSuccess) override { PostUndo(bSuccess); }
+
+	/** FNotifyHook interface */
+	virtual void NotifyPostChange( const FPropertyChangedEvent& PropertyChangedEvent, FProperty* PropertyThatChanged ) override;
+
+	/** FTickableGameObject interface */
+	virtual bool IsTickable(void) const override;
+	virtual void Tick( float InDeltaTime ) override;
+	virtual TStatId GetStatId() const override;
+
+	// Delegates
+	void DeleteSelectedNodes();
+	bool CanDeleteNodes() const;
+	void DuplicateSelectedNodes();
+	bool CanDuplicateSelectedNodes() const;
+	void OnSelectedGraphNodesChanged(const FGraphPanelSelectionSet& NewSelection);
+	/**
+	 * Called when a node's title is committed for a rename
+	 *
+	 * @param	NewText				New title text
+	 * @param	CommitInfo			How text was committed
+	 * @param	NodeBeingChanged	The node being changed
+	 */
+	void OnNodeTitleCommitted(const FText& NewText, ETextCommit::Type CommitInfo, UEdGraphNode* NodeBeingChanged);
+
+	virtual void UpdateGraphNodeProperties();
+
+	/** Getter of AssetRegistryLoaded */
+	bool GetAssetRegistryLoaded();
+
+	/** Callback to notify the editor when the PreviewInstance has been updated */
+	void OnUpdatePreviewInstance();
+
+	/** Called when the Object Properties needs to be updated */
+	virtual void UpdateObjectProperties() override;
+
+	/** Getter of CustomizableObjectEditorAdvancedPreviewSettings */
+	TSharedPtr<class SCustomizableObjectEditorAdvancedPreviewSettings> GetCustomizableObjectEditorAdvancedPreviewSettings();
+
+	/** Add a node to be reconstructed at the end of the tick. */
+	void MarkForReconstruct(UEdGraphNode* Node);
+
+	TSharedPtr<class SCustomizableObjectNodeLayoutBlocksEditor> GetLayoutBlocksEditor() { return LayoutBlocksEditor; }
+
+private:
+	explicit FCustomizableObjectEditor(UCustomizableObject& ObjectToEdit);
+
+	TSharedRef<SDockTab> SpawnTab_Viewport(const FSpawnTabArgs& Args);
+	TSharedRef<SDockTab> SpawnTab_ObjectProperties(const FSpawnTabArgs& Args);
+	TSharedRef<SDockTab> SpawnTab_InstanceProperties(const FSpawnTabArgs& Args);
+	TSharedRef<SDockTab> SpawnTab_Graph(const FSpawnTabArgs& Args);
+	TSharedRef<SDockTab> SpawnTab_GraphNodeProperties(const FSpawnTabArgs& Args);
+	TSharedRef<SDockTab> SpawnTab_AdvancedPreviewSettings(const FSpawnTabArgs& Args);
+	TSharedRef<SDockTab> SpawnTab_TextureAnalyzer(const FSpawnTabArgs& Args);
+	TSharedRef<SDockTab> SpawnTab_PerformanceReport(const FSpawnTabArgs& Args);
+	TSharedRef<SDockTab> SpawnTab_TagExplorer(const FSpawnTabArgs& Args);
+	
+	void CreatePreviewInstance();
+
+	/** Returns true if successful */
+	bool CreatePreviewComponent(int32 ComponentIndex);
+
+	/** Binds commands associated with the Static Mesh Editor. */
+	void BindCommands();
+	
+	/** Command list for the graph editor */
+	TSharedPtr<FUICommandList> GraphEditorCommands;
+
+	// Compile the customizable object.
+	void CompileObject();
+	void CompileObjectUserPressedButton();
+	void CompileOnlySelectedObjectUserPressedButton();
+	
+	/** Debug the object as a raw mutable data in the internal tools. */
+	void DebugObject();
+
+	// Compile options menu callbacks
+	TSharedRef<SWidget> GenerateCompileOptionsMenuContent(TSharedRef<FUICommandList> InCommandList);
+	void ResetCompileOptions();
+	TSharedPtr<STextComboBox> CompileOptimizationCombo;
+	TArray< TSharedPtr<FString> > CompileOptimizationStrings;
+	void CompileOptions_UseParallelCompilation_Toggled();
+	bool CompileOptions_UseParallelCompilation_IsChecked();
+	void CompileOptions_UseDiskCompilation_Toggled();
+	bool CompileOptions_UseDiskCompilation_IsChecked();
+	void CompileOptions_TextureCompression_Toggled();
+	bool CompileOptions_TextureCompression_IsChecked();
+	void OnChangeCompileOptimizationLevel(TSharedPtr<FString> NewSelection, ESelectInfo::Type SelectInfo);
+
+	/** Save Customizable Object open in editor */
+	void SaveAsset_Execute() override;
+
+	/** Callback when selection changes in the Property Tree. */
+	void OnObjectPropertySelectionChanged(FProperty* InProperty);
+
+	/** Callback when selection changes in the Property Tree. */
+	void OnInstancePropertySelectionChanged(FProperty* InProperty);
+
+	/** Callback for the object modified event */
+	void OnObjectModified(UObject* Object);
+	
+	/** */
+	void OnPreviewInstanceUpdated();
+	
+	/** */
+	TSharedRef<SGraphEditor> CreateGraphEditorWidget(UEdGraph* InGraph);
+
+	/** Copy the currently selected nodes */
+	void CopySelectedNodes();
+	/** Whether we are able to copy the currently selected nodes */
+	bool CanCopyNodes() const;
+
+	/** Paste the contents of the clipboard */
+	void PasteNodes();
+
+	/** Cut the currently selected nodes */
+	void CutSelectedNodes();
+	/** Whether we are able to cut the currently selected nodes */
+	bool CanCutNodes() const;
+
+	virtual UEdGraphNode* CreateCommentBox(const FVector2D& NodePos) override;
+
+	/** Callback for the asset registry initial load */
+	void OnAssetRegistryLoadComplete();
+
+	/** Updates the visibility of PreviewSkeletalMeshComponent */
+	void UpdatePreviewVisibility();
+
+	/** Handler for when an asset's property has changed */
+	void OnObjectPropertyChanged(UObject* Object, FPropertyChangedEvent& PropertyChangedEvent);
+
+	/** Utility method to reset current's projector visibility with no Skeletal Mesh Update */
+	void ResetProjectorVisibilityNoUpdate();
+
+	/** Searches a node that contains the inserted word */
+	void OnEnterText(const FText& NewText, ETextCommit::Type TextType);
+
+	/** Logs the sersults of the search */
+	void LogSearchResult(UEdGraphNode* Node, FString Type, bool IsFirst, FString result) const;
+
+	/** Open the Texture Analyzer tab */
+	void OpenTextureAnalyzerTab();
+
+	/** Open the Performance Report tab */
+	void OpenPerformanceReportTab();
+
+	/** Creates the necessary components for the preview of the CO instance */
+	void CreatePreviewComponents();
+
+public:
+
+	// Helpers to get the absolute parent of a Customizable Object
+	static UCustomizableObject* GetAbsoluteCOParent(const UCustomizableObjectNodeObject* const Root);
+	static void AddCachedReferencers(const FName& PathName, TArray<FName>& ArrayReferenceNames, TArray<FAssetData>& ArrayAssetData);
+	static void GetExternalChildObjects(const UCustomizableObject* const Object, TArray<UCustomizableObject*>& ExternalChildren, const bool bRecursively = true, const EObjectFlags ExcludeFlags = EObjectFlags::RF_Transient);
+
+
+	/**	The tab ids for all the tabs used */
+	static const FName ViewportTabId;
+	static const FName ObjectPropertiesTabId;
+	static const FName InstancePropertiesTabId;
+	static const FName GraphTabId;
+	static const FName GraphNodePropertiesTabId;
+	static const FName SystemPropertiesTabId;
+	static const FName AdvancedPreviewSettingsTabId;
+	static const FName TextureAnalyzerTabId;
+	static const FName PerformanceReportTabId;
+	static const FName TagExplorerTabId;
+	static const FName ObjectDebuggerTabId;
+	static const FName PopulationClassTagManagerTabId;
+
+private:
+	/** The currently viewed object. */
+	UCustomizableObject* CustomizableObject;
+	UCustomizableObjectInstance* PreviewInstance = nullptr;
+	TArray<UCustomizableSkeletalComponent*> PreviewCustomizableSkeletalComponents;
+	UStaticMeshComponent* PreviewStaticMeshComponent = nullptr;
+	TArray<UDebugSkelMeshComponent*> PreviewSkeletalMeshComponents;
+
+	/** Object compiler */
+	FCustomizableObjectCompiler Compiler;
+
+	/** List of open tool panels; used to ensure only one exists at any one time */
+	TMap< FName, TWeakPtr<class SDockableTab> > SpawnedToolPanels;
+
+	/** Preview Viewport widget */
+	TSharedPtr<SCustomizableObjectEditorViewportTabBody> Viewport;
+	TSharedPtr<FCustomizableObjectEditorViewportClient> ViewportClient;
+
+	TSharedPtr<class IDetailsView> CustomizableInstanceDetailsView;
+
+
+	/** Property View */
+	TSharedPtr<class IDetailsView> CustomizableObjectDetailsView;
+
+	/** */
+	TSharedPtr<SGraphEditor> GraphEditor;
+	TSharedPtr<class IDetailsView> GraphNodeDetailsView;
+
+	/** This splitter can be used to add more widgets outside the property editor of the currently selected node. */
+	TSharedPtr<class SSplitter> NodeDetailsSplitter;
+
+	/** This widget shows the layout blocks and they can be edited. It is kept alive for some problems re-registering toolbar commands.*/
+	TSharedPtr<class SCustomizableObjectNodeLayoutBlocksEditor> LayoutBlocksEditor;
+	TSharedPtr<class SCustomizableObjectNodeLayoutBlocksSelector> LayoutBlocksSelector;
+
+	/** Widget to select which node pins are visible. */
+	TSharedPtr<class SCustomizableObjectNodePinViewer> NodePinViewer;
+
+	UCustomizableObjectNodeMeshClipMorph* SelectedMeshClipMorphNode = nullptr;
+
+	UCustomizableObjectNodeMeshClipWithMesh* SelectedMeshClipWithMeshNode = nullptr;
+
+	UCustomizableObjectNodeProjectorConstant* SelectedProjectorNode = nullptr;
+
+	UCustomizableObjectNodeProjectorParameter* SelectedProjectorParameterNode = nullptr;
+
+	/** Handle for the OnObjectModified event */
+	FDelegateHandle OnObjectModifiedHandle;
+
+	bool ProjectorConstantNodeSelected = false;
+	bool ProjectorParameterNodeSelected = false;
+	bool SelectedGraphNodesChanged = false;
+	bool SelectedProjectorParameterNotNode = false;
+	bool ResetProjectorVisibilityForNonNode = false;
+	bool SetProjectorVisibilityForParameter = false;
+	bool SetProjectorTypeForParameter = false;
+	FString ProjectorParameterName;
+	FString ProjectorParameterNameWithIndex;
+	int32 ProjectorRangeIndex;
+	int32 ProjectorParameterIndex = -1;
+	FVector3f ProjectorParameterPosition;
+	FVector3f ProjectorParameterDirection;
+	FVector3f ProjectorParameterUp;
+	FVector3f ProjectorParameterScale;
+	ECustomizableObjectProjectorType ProjectorParameterProjectionType;
+	float ProjectionAngle;
+	bool ManagingProjector = false;
+
+	/** Flag to know when the asset registry initial loading has completed */
+	bool AssetRegistryLoaded = false;
+
+	/** Flag to know whether the lasts steps when making a new preview instance need to be done when asset registry completes asset loading */
+	bool UpdateSkeletalMeshAfterAssetLoaded = false;
+	
+	/** UObject class to be able to use the update callback */
+	UUpdateClassWrapper* HelperCallback;
+	
+	/** Scene preview settings widget, upcast of CustomizableObjectEditorAdvancedPreviewSettings */
+	TSharedPtr<SWidget> AdvancedPreviewSettingsWidget;
+
+	/** Scene preview settings widget */
+	TSharedPtr<class SCustomizableObjectEditorAdvancedPreviewSettings> CustomizableObjectEditorAdvancedPreviewSettings;
+
+	/** Advanced scene preview settings */
+	class UCustomizableObjectEmptyClassForSettings* AdditionalSettings;
+	
+	/** When launching the refresh material in all children command from Customizable Object Node Object Group, if the COs are not loaded, the
+	* Editor will remaing stuck since the command is executed on the game thread. Use this flag to launch the Editor notification */
+	bool LaunchRefreshMaterialInAllChildren = false;
+
+	/** Used together with LaunchRefreshMaterialInAllChildren as a timer to launch the RefreshMaterialNodesInAllChildren method
+	* which will perform the task */
+	float PendingTimeRefreshMaterialInAllChildren = 2.0;
+
+	/** Pose asset when doing drag and drop of an UPoseAsset to the viewport */
+	UPoseAsset* PoseAsset;
+
+	/** Texture Analyzer table widget which shows the information of the transient textures used in the customizable object instance */
+	TSharedPtr<class SCustomizableObjecEditorTextureAnalyzer> TextureAnalyzer;
+
+	/** Performance report widget to test and analyze the current cusomizable object resource demands */
+	TSharedPtr<class SCustomizableObjecEditorPerformanceReport> PerformanceReport;
+
+	/** Widget to explore all the tags related with the Customizable Object open in the editor */
+	TSharedPtr<class SCustomizableObjectEditorTagExplorer> TagExplorer;
+
+	/** Adds the customizable Object Editor commands to the default toolbar */
+	void ExtendToolbar();
+
+	/** List of ndes to be reconstructed at the end of the tick. */
+	TArray<UEdGraphNode*> ReconstructNodes;
+
+	/** URL to open when pressing the documentation button generated by UE */
+	const FString DocumentationURL{ TEXT("https://work.anticto.com/w/mutable/unreal-engine-4/") };
+
+	/** Postponed work to do when OnUpdatePreviewInstance is called. Emptied at the end on each OnUpdatePreviewInstance call. */
+	TArray<TFunction<void()>> OnUpdatePreviewInstanceWork;
+	
+protected:
+	/** @return the documentation location for this editor */
+	virtual FString GetDocumentationLink() const override;
+};
diff --git a/Engine/Plugins/Experimental/Mutable/Source/CustomizableObjectEditor/Private/CustomizableObjectEditorActions.cpp b/Engine/Plugins/Experimental/Mutable/Source/CustomizableObjectEditor/Private/CustomizableObjectEditorActions.cpp
new file mode 100644
--- /dev/null
+++ b/Engine/Plugins/Experimental/Mutable/Source/CustomizableObjectEditor/Private/CustomizableObjectEditorActions.cpp
@@ -0,0 +1,78 @@
+// Copyright Epic Games, Inc. All Rights Reserved.
+
+#include "CustomizableObjectEditorActions.h"
+#include "CustomizableObjectEditor.h"
+#include "SCustomizableObjectEditorViewport.h"
+#include "Styling/AppStyle.h"
+#include "CustomizableObjectEditorStyle.h"
+#include "Framework/Commands/InputChord.h"
+
+#define LOCTEXT_NAMESPACE "CustomizableObjectEditor"
+
+
+FCustomizableObjectEditorCommands::FCustomizableObjectEditorCommands() 
+	: TCommands<FCustomizableObjectEditorCommands>
+(
+	"CustomizableObjectEditor", // Context name for fast lookup
+	NSLOCTEXT("Contexts", "CustomizableObjectEditor", "CustomizableObject Editor"), // Localized context name for displaying
+	NAME_None, // Parent
+	FCustomizableObjectEditorStyle::GetStyleSetName()
+	)
+{
+}
+
+
+void FCustomizableObjectEditorCommands::RegisterCommands()
+{
+	UI_COMMAND(Compile, "Compile", "Compile the source graph of the customizable object and update the previews.", EUserInterfaceActionType::Button, FInputChord());
+	UI_COMMAND(CompileOnlySelected, "Compile Only Selected", "Compile the source graph of the customizable object and update the previews, only for the selected options in the preview. The rest of options are discarded. If they are selected, press again this button to see their effect in the preview.", EUserInterfaceActionType::Button, FInputChord());
+	UI_COMMAND(ResetCompileOptions, "Reset Compilation Options", "Set reasonable defaults for the compilation options.", EUserInterfaceActionType::Button, FInputChord());
+	UI_COMMAND(CompileOptions_EnableTextureCompression, "Enable texture compression.", "Only for debug. Do not use.", EUserInterfaceActionType::ToggleButton, FInputChord());
+	UI_COMMAND(CompileOptions_UseParallelCompilation, "Enable compiling in multiple threads.", "This is faster but use more memory.", EUserInterfaceActionType::ToggleButton, FInputChord());
+	UI_COMMAND(CompileOptions_UseDiskCompilation, "Enable compiling using the disk as memory.", "This is very slow but supports compiling huge objects. It requires a lot of free space in the OS disk.", EUserInterfaceActionType::ToggleButton, FInputChord());
+
+	UI_COMMAND(Debug, "Debug", "Debug the object.", EUserInterfaceActionType::Button, FInputChord());
+	UI_COMMAND(DebugOptions_OnlySelected, "Debug Only Selected", "Debug only for the selected options in the preview. The rest of options are discarded.", EUserInterfaceActionType::Button, FInputChord());
+	UI_COMMAND(DebugOptions_EnableTextureCompression, "Enable texture compression.", "Only for debug. Do not use.", EUserInterfaceActionType::ToggleButton, FInputChord());
+
+	UI_COMMAND(PerformanceReport, "Performance Report", "Open the Performance Report window to set up and perform all tests relevant to Customizable Objects and access worst cases data.", EUserInterfaceActionType::Button, FInputChord());
+	UI_COMMAND(TextureAnalyzer, "Texture Memory Analyzer", "Open the Texture Analyzer window to check all the information of the textures created by Mutable.", EUserInterfaceActionType::Button, FInputChord());
+}
+
+
+FCustomizableObjectEditorViewportCommands::FCustomizableObjectEditorViewportCommands() 
+	: TCommands<FCustomizableObjectEditorViewportCommands>
+(
+	"CustomizableObjectEditorViewport", // Context name for fast lookup
+	NSLOCTEXT("Contexts", "CustomizableObjectEditorViewport", "CustomizableObject Editor Viewport"), // Localized context name for displaying
+	NAME_None, // Parent
+	FCustomizableObjectEditorStyle::GetStyleSetName()
+	)
+{
+}
+
+
+void FCustomizableObjectEditorViewportCommands::RegisterCommands()
+{
+	UI_COMMAND( SetDrawUVs, "UV", "Toggles display of the static mesh's UVs for the specified channel.", EUserInterfaceActionType::ToggleButton, FInputChord() );
+	UI_COMMAND( SetShowGrid, "Grid", "Displays the viewport grid.", EUserInterfaceActionType::ToggleButton, FInputChord());
+	UI_COMMAND( SetShowSky, "Sky", "Displays the viewport sky.", EUserInterfaceActionType::ToggleButton, FInputChord() );
+	UI_COMMAND( SetShowBounds, "Bounds", "Toggles display of the bounds of the static mesh.", EUserInterfaceActionType::ToggleButton, FInputChord() );
+	UI_COMMAND( SetShowCollision, "Collision", "Toggles display of the simplified collision mesh of the static mesh, if one has been assigned.", EUserInterfaceActionType::ToggleButton, FInputChord() );
+	UI_COMMAND( SetCameraLock, "Camera Lock", "Toggles viewport navigation between orbit and freely moving about.", EUserInterfaceActionType::ToggleButton, FInputChord() );
+	UI_COMMAND( SaveThumbnail, "Save Thumbnail", "Saves the viewpoint position and orientation in the Preview Pane for use as the thumbnail preview in the Content Browser.", EUserInterfaceActionType::Button, FInputChord() );
+
+	UI_COMMAND( SetShowNormals, "Normals", "Toggles display of vertex normals in the Preview Pane.", EUserInterfaceActionType::ToggleButton, FInputChord() );
+	UI_COMMAND( SetShowTangents, "Tangents", "Toggles display of vertex tangents in the Preview Pane.", EUserInterfaceActionType::ToggleButton, FInputChord() );
+	UI_COMMAND( SetShowBinormals, "Binormals", "Toggles display of vertex binormals (orthogonal vector to normal and tangent) in the Preview Pane.", EUserInterfaceActionType::ToggleButton, FInputChord() );
+	UI_COMMAND( SetShowPivot, "Show Pivot", "Display the pivot location of the static mesh.", EUserInterfaceActionType::ToggleButton, FInputChord() );
+
+	UI_COMMAND(BakeInstance, "Bake Instance", "Create baked unreal resources for the current preview instance.", EUserInterfaceActionType::Button, FInputChord());
+	UI_COMMAND(StateChangeTest, "State test", "Make a state change test, iterating through all possible states and runtime parameters for this instance", EUserInterfaceActionType::Button, FInputChord());
+	UI_COMMAND(StateChangeShowData, "Show or hide test results", "Show or hide test results", EUserInterfaceActionType::ToggleButton, FInputChord());
+	UI_COMMAND(StateChangeShowGeometryData, "Show instance geometry data", "Show instance geometry data", EUserInterfaceActionType::ToggleButton, FInputChord());
+}
+
+
+#undef LOCTEXT_NAMESPACE
+
diff --git a/Engine/Plugins/Experimental/Mutable/Source/CustomizableObjectEditor/Private/CustomizableObjectEditorModule.cpp b/Engine/Plugins/Experimental/Mutable/Source/CustomizableObjectEditor/Private/CustomizableObjectEditorModule.cpp
new file mode 100644
--- /dev/null
+++ b/Engine/Plugins/Experimental/Mutable/Source/CustomizableObjectEditor/Private/CustomizableObjectEditorModule.cpp
@@ -0,0 +1,467 @@
+// Copyright Epic Games, Inc. All Rights Reserved.
+
+#include "CustomizableObjectEditorModule.h"
+#include "CustomizableObjectEditor.h"
+#include "CustomizableObjectInstanceEditor.h"
+#include "CustomizableObjectDebugger.h"
+#include "CustomizableObjectInstance.h"
+#include "CustomizableSkeletalMeshActor.h"
+#include "CustomizableObjectInstanceFactory.h"
+#include "CustomizableObjectDetails.h"
+#include "CustomizableInstanceDetails.h"
+#include "CustomizableObjectSystem.h"		// For defines related to memory function replacements.
+#include "CustomizableObjectEditorStyle.h"
+#include "Nodes/CustomizableObjectNodeCopyMaterial.h"
+#include "Nodes/CustomizableObjectNodeEditMaterial.h"
+#include "Nodes/CustomizableObjectNodeExtendMaterial.h"
+#include "Nodes/CustomizableObjectNodeLayoutBlocks.h"
+#include "Nodes/CustomizableObjectNodeMeshClipMorph.h"
+#include "Nodes/CustomizableObjectNodeMeshMorph.h"
+#include "Nodes/CustomizableObjectNodeMorphMaterial.h"
+#include "Nodes/CustomizableObjectNodeProjectorConstant.h"
+#include "Nodes/CustomizableObjectNodeProjectorParameter.h"
+#include "Nodes/CustomizableObjectNodeMeshReshape.h"
+#include "Nodes/CustomizableObjectNodeLayoutBlocksDetails.h"
+#include "Nodes/CustomizableObjectNodeParentedMaterialDetails.h"
+#include "Nodes/CustomizableObjectNodeEditMaterialBaseDetails.h"
+#include "Nodes/CustomizableObjectNodeMorphMaterialDetails.h"
+#include "Nodes/CustomizableObjectNodeObjectDetails.h"
+#include "CustomizableObjectNodeObjectGroupDetails.h"
+#include "Nodes/CustomizableObjectNodeMeshMorphDetails.h"
+#include "Nodes/CustomizableObjectNodeProjectorParameterDetails.h"
+#include "Nodes/CustomizableObjectNodeMeshClipMorphDetails.h"
+#include "Nodes/CustomizableObjectNodeMeshClipWithMeshDetails.h"
+#include "Nodes/CustomizableObjectNodeExternalPinDetails.h"
+#include "Nodes/CustomizableObjectNodeCopyMaterialDetails.h"
+#include "Nodes/CustomizableObjectNodeSkeletalMeshDetails.h"
+#include "Nodes/CustomizableObjectNodePinViewerDetails.h"
+#include "Nodes/CustomizableObjectNodeTableDetails.h"
+#include "Nodes/CustomizableObjectNodeMeshReshapeSelectionDetails.h"
+#include "CustomizableObjectIdentifierCustomization.h"
+#include "Modules/ModuleManager.h"
+#include "Interfaces/IPluginManager.h"
+#include "AssetToolsModule.h"
+#include "AssetTypeActions_CustomizableObject.h"
+#include "AssetTypeActions_CustomizableObjectInstance.h"
+#include "Styling/SlateStyle.h"
+#include "Brushes/SlateImageBrush.h"
+#include "MessageLogModule.h"
+#include "Logging/MessageLog.h"
+#include "CustomizableObjectEditorSettings.h"
+#include "ISettingsModule.h"
+#include "ISettingsSection.h"
+#include "CustomizableObjectCustomSettingsDetails.h"
+#include "AssetRegistry/AssetRegistryModule.h"
+#include "HAL/FileManager.h"
+#include "CoreGlobals.h"
+#include "CustomizableObjectCustomSettings.h"
+#include "Settings/ProjectPackagingSettings.h"
+#include "UnrealEditorPortabilityHelpers.h"
+#include "CookOnTheSide/CookOnTheFlyServer.h" // This is necessary to register the cook delegate.
+#include "Nodes/CustomizableObjectNodeStaticMesh.h"
+#include "CustomizableObjectDGGUI.h"
+#include "Kismet/GameplayStatics.h"
+
+
+const FName CustomizableObjectEditorAppIdentifier = FName(TEXT("CustomizableObjectEditorApp"));
+const FName CustomizableObjectInstanceEditorAppIdentifier = FName(TEXT("CustomizableObjectInstanceEditorApp"));
+const FName CustomizableObjectDebuggerAppIdentifier = FName(TEXT("CustomizableObjectDebuggerApp"));
+
+#define LOCTEXT_NAMESPACE "MutableSettings"
+
+static TAutoConsoleVariable<bool> CVarMutableOnCookStartEnabled(
+	TEXT("b.OnCookStartEnabled"),
+	true,
+	TEXT("If enabled, Customizable Objects will be compiled before the actual cook starts. Compiled data will be stored on de DDC and cached during BeginCache.\n"),
+	ECVF_Scalability);
+
+/**
+ * StaticMesh editor module
+ */
+class FCustomizableObjectEditorModule : public ICustomizableObjectEditorModule
+{
+public:
+	// IModuleInterface interface
+	void StartupModule() override;
+	void ShutdownModule() override;
+
+	// ICustomizableObjectEditorModule interface
+	TSharedRef<ICustomizableObjectEditor> CreateCustomizableObjectEditor( const EToolkitMode::Type Mode, const TSharedPtr< IToolkitHost >& InitToolkitHost, UCustomizableObject* CustomizableObject ) override;
+	TSharedRef<ICustomizableObjectInstanceEditor> CreateCustomizableObjectInstanceEditor( const EToolkitMode::Type Mode, const TSharedPtr< IToolkitHost >& InitToolkitHost, UCustomizableObjectInstance* CustomizableObjectInstance ) override;
+	TSharedRef<ICustomizableObjectDebugger> CreateCustomizableObjectDebugger(const EToolkitMode::Type Mode, const TSharedPtr< IToolkitHost >& InitToolkitHost, UCustomizableObject* CustomizableObject) override;
+	virtual EAssetTypeCategories::Type GetAssetCategory() const override;
+	virtual FCustomizableObjectEditorLogger& GetLogger() override;
+
+	virtual TSharedPtr<FExtensibilityManager> GetCustomizableObjectEditorToolBarExtensibilityManager() override { return CustomizableObjectEditor_ToolBarExtensibilityManager; }
+	virtual TSharedPtr<FExtensibilityManager> GetCustomizableObjectEditorMenuExtensibilityManager() override { return CustomizableObjectEditor_MenuExtensibilityManager; };
+
+public:
+	bool HandleSettingsSaved();
+	void RegisterSettings();
+
+private:	
+	TSharedPtr<FExtensibilityManager> CustomizableObjectEditor_ToolBarExtensibilityManager;
+	TSharedPtr<FExtensibilityManager> CustomizableObjectEditor_MenuExtensibilityManager;
+
+	/** List of registered custom details to remove later. */
+	TArray<FName> RegisteredCustomDetails;
+
+	/** Custom asset category. */
+	EAssetTypeCategories::Type CustomizableObjectAssetCategory;
+
+	/** Register Custom details. Also adds them to RegisteredCustomDetails list. */
+	void RegisterCustomDetails(FPropertyEditorModule& PropertyModule, const UClass* Class, FOnGetDetailCustomizationInstance DetailLayoutDelegate);
+
+	FCustomizableObjectEditorLogger Logger;
+
+	// Command to look for Customizable Object Instance in the player pawn of the current world and open a DGGUI to edit its parameters
+	IConsoleCommand* LaunchDGGUICommand;
+	static void ToggleDGGUI(const TArray<FString>& Arguments);
+
+	// Command to look for Customizable Object Instance in the player pawn of the current world and open its Customizable Object Instance Editor
+	IConsoleCommand* LaunchCOIECommand;
+	static void OpenCOIE(const TArray<FString>& Arguments);
+};
+
+IMPLEMENT_MODULE( FCustomizableObjectEditorModule, CustomizableObjectEditor );
+
+
+static void* CustomMalloc(std::size_t Size_t, uint32_t Alignment)
+{
+	return FMemory::Malloc(Size_t, Alignment);
+}
+
+
+static void CustomFree(void* mem)
+{
+	return FMemory::Free(mem);
+}
+
+
+FCustomizableObjectCompilerBase* NewCompiler()
+{
+	return new FCustomizableObjectCompiler();
+}
+
+
+void FCustomizableObjectEditorModule::StartupModule()
+{
+	UCustomizableObjectSystem::GetInstance()->SetNewCompilerFunc(NewCompiler);
+
+	// Register the thumbnail renderers
+	//UThumbnailManager::Get().RegisterCustomRenderer(UCustomizableObject::StaticClass(), UCustomizableObjectThumbnailRenderer::StaticClass());
+	//UThumbnailManager::Get().RegisterCustomRenderer(UCustomizableObjectInstance::StaticClass(), UCustomizableObjectInstanceThumbnailRenderer::StaticClass());
+	
+	// Property views
+	FPropertyEditorModule& PropertyModule = FModuleManager::LoadModuleChecked<FPropertyEditorModule>("PropertyEditor");
+	RegisterCustomDetails(PropertyModule, UCustomizableObject::StaticClass(), FOnGetDetailCustomizationInstance::CreateStatic(&FCustomizableObjectDetails::MakeInstance));
+	RegisterCustomDetails(PropertyModule, UCustomizableObjectInstance::StaticClass(), FOnGetDetailCustomizationInstance::CreateStatic(&FCustomizableInstanceDetails::MakeInstance));
+	RegisterCustomDetails(PropertyModule, UCustomizableObjectNodeLayoutBlocks::StaticClass(), FOnGetDetailCustomizationInstance::CreateStatic(&FCustomizableObjectNodeLayoutBlocksDetails::MakeInstance));
+	RegisterCustomDetails(PropertyModule, UCustomizableObjectNodeExtendMaterial::StaticClass(), FOnGetDetailCustomizationInstance::CreateStatic(&FCustomizableObjectNodeParentedMaterialDetails::MakeInstance));
+	RegisterCustomDetails(PropertyModule, UCustomizableObjectNodeEditMaterialBase::StaticClass(), FOnGetDetailCustomizationInstance::CreateStatic(&FCustomizableObjectNodeEditMaterialBaseDetails::MakeInstance));
+	RegisterCustomDetails(PropertyModule, UCustomizableObjectNodeEditMaterial::StaticClass(), FOnGetDetailCustomizationInstance::CreateStatic(&FCustomizableObjectNodePinViewerDetails::MakeInstance));
+	RegisterCustomDetails(PropertyModule, UCustomizableObjectNodeMorphMaterial::StaticClass(), FOnGetDetailCustomizationInstance::CreateStatic(&FCustomizableObjectNodeMorphMaterialDetails::MakeInstance));
+	RegisterCustomDetails(PropertyModule, UCustomizableObjectNodeObject::StaticClass(), FOnGetDetailCustomizationInstance::CreateStatic(&FCustomizableObjectNodeObjectDetails::MakeInstance));
+	RegisterCustomDetails(PropertyModule, UCustomizableObjectNodeObjectGroup::StaticClass(), FOnGetDetailCustomizationInstance::CreateStatic(&FCustomizableObjectNodeObjectGroupDetails::MakeInstance));
+	RegisterCustomDetails(PropertyModule, UCustomizableObjectNodeProjectorParameter::StaticClass(), FOnGetDetailCustomizationInstance::CreateStatic(&FCustomizableObjectNodeProjectorParameterDetails::MakeInstance));
+	RegisterCustomDetails(PropertyModule, UCustomizableObjectNodeProjectorConstant::StaticClass(), FOnGetDetailCustomizationInstance::CreateStatic(&FCustomizableObjectNodeProjectorParameterDetails::MakeInstance));
+	RegisterCustomDetails(PropertyModule, UCustomizableObjectNodeMeshMorph::StaticClass(), FOnGetDetailCustomizationInstance::CreateStatic(&FCustomizableObjectNodeMeshMorphDetails::MakeInstance));
+	RegisterCustomDetails(PropertyModule, UCustomizableObjectNodeMeshClipMorph::StaticClass(), FOnGetDetailCustomizationInstance::CreateStatic(&FCustomizableObjectNodeMeshClipMorphDetails::MakeInstance));
+	RegisterCustomDetails(PropertyModule, UCustomizableObjectNodeMeshClipWithMesh::StaticClass(), FOnGetDetailCustomizationInstance::CreateStatic(&FCustomizableObjectNodeMeshClipWithMeshDetails::MakeInstance));
+	RegisterCustomDetails(PropertyModule, UCustomizableObjectNodeExternalPin::StaticClass(), FOnGetDetailCustomizationInstance::CreateStatic(&FCustomizableObjectNodeExternalPinDetails::MakeInstance));
+	RegisterCustomDetails(PropertyModule, UCustomizableObjectEmptyClassForSettings::StaticClass(), FOnGetDetailCustomizationInstance::CreateStatic(&FCustomizableObjectCustomSettingsDetails::MakeInstance));
+	RegisterCustomDetails(PropertyModule, UCustomizableObjectNodeMaterial::StaticClass(), FOnGetDetailCustomizationInstance::CreateStatic(&FCustomizableObjectNodePinViewerDetails::MakeInstance));
+	RegisterCustomDetails(PropertyModule, UCustomizableObjectNodeCopyMaterial::StaticClass(), FOnGetDetailCustomizationInstance::CreateStatic(&FCustomizableObjectNodeCopyMaterialDetails::MakeInstance));
+	RegisterCustomDetails(PropertyModule, UCustomizableObjectNodeSkeletalMesh::StaticClass(), FOnGetDetailCustomizationInstance::CreateStatic(&FCustomizableObjectNodeSkeletalMeshDetails::MakeInstance));
+	RegisterCustomDetails(PropertyModule, UCustomizableObjectNodeStaticMesh::StaticClass(), FOnGetDetailCustomizationInstance::CreateStatic(&FCustomizableObjectNodePinViewerDetails::MakeInstance));
+	RegisterCustomDetails(PropertyModule, UCustomizableObjectNodeTable::StaticClass(), FOnGetDetailCustomizationInstance::CreateStatic(&FCustomizableObjectNodeTableDetails::MakeInstance));
+
+	//Custom properties
+	PropertyModule.RegisterCustomPropertyTypeLayout("CustomizableObjectIdentifier", FOnGetPropertyTypeCustomizationInstance::CreateStatic(&FCustomizableObjectIdentifierCustomization::MakeInstance));
+	PropertyModule.RegisterCustomPropertyTypeLayout(FMeshReshapeBoneReference::StaticStruct()->GetFName(), FOnGetPropertyTypeCustomizationInstance::CreateStatic(&FMeshReshapeBonesReferenceCustomization::MakeInstance));
+
+	PropertyModule.NotifyCustomizationModuleChanged();
+
+	// Register factory
+	UCustomizableObjectInstanceFactory* CustomizableObjectInstanceFactory = NewObject<UCustomizableObjectInstanceFactory>();
+	CustomizableObjectInstanceFactory->NewActorClass = ACustomizableSkeletalMeshActor::StaticClass();
+	UEditorEngine* engine = GEditor;
+	engine->ActorFactories.Add(CustomizableObjectInstanceFactory);
+	
+	// New category
+	IAssetTools& AssetTools = FModuleManager::LoadModuleChecked<FAssetToolsModule>("AssetTools").Get();
+	CustomizableObjectAssetCategory = AssetTools.RegisterAdvancedAssetCategory(FName(TEXT("Mutable")), LOCTEXT("MutableAssetsCategory", "Mutable"));
+
+	// Asset actions
+	FAssetToolsModule& AssetToolsModule = FModuleManager::LoadModuleChecked< FAssetToolsModule >( "AssetTools" );
+	
+	TSharedPtr<FAssetTypeActions_CustomizableObject> CustomizableObjectAssetTypeActions = MakeShareable( new FAssetTypeActions_CustomizableObject );	
+	AssetToolsModule.Get().RegisterAssetTypeActions( CustomizableObjectAssetTypeActions.ToSharedRef() );
+
+	TSharedPtr<FAssetTypeActions_CustomizableObjectInstance> CustomizableObjectInstanceAssetTypeActions = MakeShareable( new FAssetTypeActions_CustomizableObjectInstance );
+	AssetToolsModule.Get().RegisterAssetTypeActions( CustomizableObjectInstanceAssetTypeActions.ToSharedRef() );
+	
+
+	// Additional UI style
+	FCustomizableObjectEditorStyle::Initialize();
+
+	RegisterSettings();
+
+	// Create the message log category
+	FMessageLogModule& MessageLogModule = FModuleManager::LoadModuleChecked<FMessageLogModule>("MessageLog");
+	MessageLogModule.RegisterLogListing(FName("Mutable"), LOCTEXT("MutableLog", "Mutable"));
+
+	CustomizableObjectEditor_ToolBarExtensibilityManager = MakeShareable(new FExtensibilityManager);
+	CustomizableObjectEditor_MenuExtensibilityManager = MakeShareable(new FExtensibilityManager);
+
+	LaunchDGGUICommand = IConsoleManager::Get().RegisterConsoleCommand(
+		TEXT("mutable.ToggleDGGUI"),
+		TEXT("Looks for a Customizable Object Instance within the player pawn and opens a UI to modify its parameters, or closes it if it's open. Specify slot ID to control which component is modified."),
+		FConsoleCommandWithArgsDelegate::CreateStatic(&FCustomizableObjectEditorModule::ToggleDGGUI));
+
+	LaunchCOIECommand = IConsoleManager::Get().RegisterConsoleCommand(
+		TEXT("mutable.OpenCOIE"),
+		TEXT("Looks for a Customizable Object Instance within the player pawn and opens its Customizable Object Instance Editor. Specify slot ID to control which component is edited."),
+		FConsoleCommandWithArgsDelegate::CreateStatic(&FCustomizableObjectEditorModule::OpenCOIE));
+
+}
+
+
+void FCustomizableObjectEditorModule::ShutdownModule()
+{
+	if( FModuleManager::Get().IsModuleLoaded( "PropertyEditor" ) )
+	{
+		FPropertyEditorModule& PropertyModule = FModuleManager::GetModuleChecked<FPropertyEditorModule>("PropertyEditor");
+		for (const auto& ClassName : RegisteredCustomDetails)
+		{
+			PropertyModule.UnregisterCustomClassLayout(ClassName);
+		}
+
+		PropertyModule.NotifyCustomizationModuleChanged();
+	}
+
+	CustomizableObjectEditor_ToolBarExtensibilityManager.Reset();
+	CustomizableObjectEditor_MenuExtensibilityManager.Reset();
+
+	FCustomizableObjectEditorStyle::Shutdown();
+}
+
+
+TSharedRef<ICustomizableObjectEditor> FCustomizableObjectEditorModule::CreateCustomizableObjectEditor( const EToolkitMode::Type Mode, const TSharedPtr< IToolkitHost >& InitToolkitHost, UCustomizableObject* CustomizableObject )
+{
+	return FCustomizableObjectEditor::Create(Mode, InitToolkitHost,CustomizableObject);
+}
+
+
+TSharedRef<ICustomizableObjectInstanceEditor> FCustomizableObjectEditorModule::CreateCustomizableObjectInstanceEditor( const EToolkitMode::Type Mode, const TSharedPtr< IToolkitHost >& InitToolkitHost, UCustomizableObjectInstance* CustomizableObjectInstance )
+{
+	TSharedRef<FCustomizableObjectInstanceEditor> NewCustomizableObjectInstanceEditor(new FCustomizableObjectInstanceEditor());
+	NewCustomizableObjectInstanceEditor->InitCustomizableObjectInstanceEditor(Mode, InitToolkitHost,CustomizableObjectInstance);
+	return NewCustomizableObjectInstanceEditor;
+}
+
+
+TSharedRef<ICustomizableObjectDebugger> FCustomizableObjectEditorModule::CreateCustomizableObjectDebugger(const EToolkitMode::Type Mode, const TSharedPtr< IToolkitHost >& InitToolkitHost, UCustomizableObject* CustomizableObject)
+{
+	TSharedRef<FCustomizableObjectDebugger> NewCustomizableObjectDebugger(new FCustomizableObjectDebugger());
+	NewCustomizableObjectDebugger->InitCustomizableObjectDebugger(Mode, InitToolkitHost, CustomizableObject);
+	return NewCustomizableObjectDebugger;
+}
+
+
+FCustomizableObjectEditorLogger& FCustomizableObjectEditorModule::GetLogger()
+{
+	return Logger;
+}
+
+
+bool FCustomizableObjectEditorModule::HandleSettingsSaved()
+{
+	UCustomizableObjectEditorSettings* CustomizableObjectSettings = GetMutableDefault<UCustomizableObjectEditorSettings>();
+
+	if (CustomizableObjectSettings != nullptr)
+	{
+		CustomizableObjectSettings->SaveConfig();
+		
+		FEditorCompileSettings CompileSettings;
+		CompileSettings.bDisableCompilation = CustomizableObjectSettings->bDisableMutableCompileInEditor;
+		CompileSettings.bEnableAutomaticCompilation = CustomizableObjectSettings->bEnableAutomaticCompilation;
+		CompileSettings.bCompileObjectsSynchronously = CustomizableObjectSettings->bCompileObjectsSynchronously;
+		CompileSettings.bCompileRootObjectsOnStartPIE = CustomizableObjectSettings->bCompileRootObjectsOnStartPIE;
+		
+		UCustomizableObjectSystem::GetInstance()->EditorSettingsChanged(CompileSettings);
+	}
+
+    return true;
+}
+
+
+void FCustomizableObjectEditorModule::RegisterSettings()
+{
+	ISettingsModule* SettingsModule = FModuleManager::GetModulePtr<ISettingsModule>("Settings");
+
+	if (SettingsModule != nullptr)
+    {
+        ISettingsSectionPtr SettingsSectionPtr = SettingsModule->RegisterSettings("Project", "Plugins", "CustomizableObjectSettings",
+            LOCTEXT("MutableSettings_Setting", "Mutable"),
+            LOCTEXT("MutableSettings_Setting_Desc", "Mutable Settings"),
+            GetMutableDefault<UCustomizableObjectEditorSettings>()
+        );
+
+        if (SettingsSectionPtr.IsValid())
+        {
+            SettingsSectionPtr->OnModified().BindRaw(this, &FCustomizableObjectEditorModule::HandleSettingsSaved);
+        }
+
+		if (UCustomizableObjectSystem::GetInstance() != nullptr)
+		{
+			UCustomizableObjectEditorSettings* CustomizableObjectSettings = GetMutableDefault<UCustomizableObjectEditorSettings>();
+			if (CustomizableObjectSettings != nullptr)
+			{
+				FEditorCompileSettings CompileSettings;
+				CompileSettings.bDisableCompilation = CustomizableObjectSettings->bDisableMutableCompileInEditor;
+				CompileSettings.bEnableAutomaticCompilation = CustomizableObjectSettings->bEnableAutomaticCompilation;
+				CompileSettings.bCompileObjectsSynchronously = CustomizableObjectSettings->bCompileObjectsSynchronously;
+				CompileSettings.bCompileRootObjectsOnStartPIE = CustomizableObjectSettings->bCompileRootObjectsOnStartPIE;
+				
+				UCustomizableObjectSystem::GetInstance()->EditorSettingsChanged(CompileSettings);
+			}
+		}
+    }
+}
+
+
+void FCustomizableObjectEditorModule::RegisterCustomDetails(FPropertyEditorModule& PropertyModule, const UClass* Class, FOnGetDetailCustomizationInstance DetailLayoutDelegate)
+{
+	const FName ClassName = FName(Class->GetName());
+	PropertyModule.RegisterCustomClassLayout(ClassName, DetailLayoutDelegate);
+
+	RegisteredCustomDetails.Add(ClassName);
+}
+
+
+const UWorld* GetWorldForCurrentCOI()
+{
+	UWorld* CurrentWorld = nullptr;
+	const TIndirectArray<FWorldContext>& WorldContexts = GEngine->GetWorldContexts();
+	for (const FWorldContext& Context : WorldContexts)
+	{
+		if ((Context.WorldType == EWorldType::Game) && (Context.World() != NULL))
+		{
+			CurrentWorld = Context.World();
+		}
+	}
+	// Fall back to GWorld if we don't actually have a world.
+	if (CurrentWorld == nullptr)
+	{
+		CurrentWorld = GWorld;
+	}
+	return CurrentWorld;
+}
+
+
+UCustomizableSkeletalComponent* GetPlayerCustomizableSkeletalComponent(const int32 SlotID, const UWorld* CurrentWorld, const int32 PlayerIndex = 0)
+{
+	// Get customizable skeletal component attached to player pawn
+	UCustomizableSkeletalComponent* SelectedCustomizableSkeletalComponent = nullptr;
+	{
+		APawn* PlayerPawn = UGameplayStatics::GetPlayerPawn(CurrentWorld, PlayerIndex);
+		int32 IndexFound = INDEX_NONE;
+		for (TObjectIterator<UCustomizableSkeletalComponent> CustomizableSkeletalComponent; CustomizableSkeletalComponent; ++CustomizableSkeletalComponent)
+		{
+			if (CustomizableSkeletalComponent->IsValidLowLevel() && !CustomizableSkeletalComponent->IsTemplate())
+			{
+				AActor* CustomizableActor = CustomizableSkeletalComponent->GetAttachmentRootActor();
+				if (CustomizableActor && PlayerPawn == CustomizableActor)
+				{
+					++IndexFound;
+					SelectedCustomizableSkeletalComponent = *CustomizableSkeletalComponent;
+					if (IndexFound == SlotID)
+					{
+						break;
+					}
+				}
+			}
+		}
+	}
+
+
+	// If none found, try getting a component without caring about the actor
+	if (!SelectedCustomizableSkeletalComponent)
+	{
+		APawn* PlayerPawn = UGameplayStatics::GetPlayerPawn(CurrentWorld, PlayerIndex);
+		int32 IndexFound = INDEX_NONE;
+		for (TObjectIterator<UCustomizableSkeletalComponent> CustomizableSkeletalComponent; CustomizableSkeletalComponent; ++CustomizableSkeletalComponent)
+		{
+			if (CustomizableSkeletalComponent->IsValidLowLevel() && !CustomizableSkeletalComponent->IsTemplate())
+			{
+				++IndexFound;
+				SelectedCustomizableSkeletalComponent = *CustomizableSkeletalComponent;
+				if (IndexFound == SlotID)
+				{
+					break;
+				}
+			}
+		}
+	}
+
+	return SelectedCustomizableSkeletalComponent;
+}
+
+
+void FCustomizableObjectEditorModule::ToggleDGGUI(const TArray<FString>& Arguments)
+{
+	int32 SlotID = INDEX_NONE;
+	if (Arguments.Num() >= 1)
+	{
+		SlotID = FCString::Atoi(*Arguments[0]);
+	}
+	const UWorld* CurrentWorld = GetWorldForCurrentCOI();
+	const int32 PlayerIndex = 0;
+	if (UDGGUI::CloseExistingDGGUI(CurrentWorld))
+	{
+		return;
+	}
+	else if (UCustomizableSkeletalComponent* SelectedCustomizableSkeletalComponent = GetPlayerCustomizableSkeletalComponent(SlotID, CurrentWorld, PlayerIndex))
+	{
+		UDGGUI::OpenDGGUI(SlotID, SelectedCustomizableSkeletalComponent, CurrentWorld, PlayerIndex);
+	}
+}
+
+
+void FCustomizableObjectEditorModule::OpenCOIE(const TArray<FString>& Arguments)
+{
+	int32 SlotID = INDEX_NONE;
+	if (Arguments.Num() >= 1)
+	{
+		SlotID = FCString::Atoi(*Arguments[0]);
+	}
+	const UWorld* CurrentWorld = GetWorldForCurrentCOI();
+	const int32 PlayerIndex = 0;
+
+	// Open the Customizable Object Instance Editor
+	if (UCustomizableSkeletalComponent* SelectedCustomizableSkeletalComponent = GetPlayerCustomizableSkeletalComponent(SlotID, CurrentWorld, PlayerIndex))
+	{
+		if (UCustomizableObjectInstance* COInstance = SelectedCustomizableSkeletalComponent->CustomizableObjectInstance)
+		{
+			FAssetToolsModule& AssetToolsModule = FModuleManager::LoadModuleChecked<FAssetToolsModule>(TEXT("AssetTools"));
+			TWeakPtr<IAssetTypeActions> WeakAssetTypeActions = AssetToolsModule.Get().GetAssetTypeActionsForClass(UCustomizableObjectInstance::StaticClass());
+
+			if (TSharedPtr<IAssetTypeActions> AssetTypeActions = WeakAssetTypeActions.Pin())
+			{
+				TArray<UObject*> AssetsToEdit;
+				AssetsToEdit.Add(COInstance);
+				AssetTypeActions->OpenAssetEditor(AssetsToEdit);
+			}
+		}
+	}
+}
+
+
+EAssetTypeCategories::Type FCustomizableObjectEditorModule::GetAssetCategory() const
+{
+	return CustomizableObjectAssetCategory;
+}
+
+
+#undef LOCTEXT_NAMESPACE
diff --git a/Engine/Plugins/Experimental/Mutable/Source/CustomizableObjectEditor/Private/CustomizableObjectEditorNodeContextCommands.cpp b/Engine/Plugins/Experimental/Mutable/Source/CustomizableObjectEditor/Private/CustomizableObjectEditorNodeContextCommands.cpp
new file mode 100644
--- /dev/null
+++ b/Engine/Plugins/Experimental/Mutable/Source/CustomizableObjectEditor/Private/CustomizableObjectEditorNodeContextCommands.cpp
@@ -0,0 +1,14 @@
+// Copyright Epic Games, Inc. All Rights Reserved.
+
+#include "CustomizableObjectEditorNodeContextCommands.h"
+
+#define LOCTEXT_NAMESPACE "CustomizableObjectEditorNodeContextCommands"
+
+void FCustomizableObjectEditorNodeContextCommands::RegisterCommands()
+{
+	UI_COMMAND(RefreshMaterialNodesInAllChildren, "Refresh all material nodes in all children", "Refresh all material nodes in all children.", EUserInterfaceActionType::Button, FInputChord());
+	
+	FKey CommentSpawnKey = FKey(TEXT("C"));
+	UI_COMMAND(CreateComment, "Create a Comment node", "Create a Comment node", EUserInterfaceActionType::Button, FInputChord(CommentSpawnKey));
+}
+#undef LOCTEXT_NAMESPACE
diff --git a/Engine/Plugins/Experimental/Mutable/Source/CustomizableObjectEditor/Private/CustomizableObjectEditorSettings.cpp b/Engine/Plugins/Experimental/Mutable/Source/CustomizableObjectEditor/Private/CustomizableObjectEditorSettings.cpp
new file mode 100644
--- /dev/null
+++ b/Engine/Plugins/Experimental/Mutable/Source/CustomizableObjectEditor/Private/CustomizableObjectEditorSettings.cpp
@@ -0,0 +1,9 @@
+// Copyright Epic Games, Inc. All Rights Reserved.
+
+#include "CustomizableObjectEditorSettings.h"
+
+
+UCustomizableObjectEditorSettings::UCustomizableObjectEditorSettings()
+{
+
+}
diff --git a/Engine/Plugins/Experimental/Mutable/Source/CustomizableObjectEditor/Private/CustomizableObjectEditorStyle.cpp b/Engine/Plugins/Experimental/Mutable/Source/CustomizableObjectEditor/Private/CustomizableObjectEditorStyle.cpp
new file mode 100644
--- /dev/null
+++ b/Engine/Plugins/Experimental/Mutable/Source/CustomizableObjectEditor/Private/CustomizableObjectEditorStyle.cpp
@@ -0,0 +1,185 @@
+// Copyright Epic Games, Inc. All Rights Reserved.
+
+#include "CustomizableObjectEditorStyle.h"
+
+#include "Framework/Application/SlateApplication.h"
+#include "Styling/AppStyle.h"
+#include "Slate/SlateGameResources.h"
+#include "Styling/SlateStyleRegistry.h"
+#include "Styling/SlateTypes.h"
+#include "Styling/CoreStyle.h"
+#include "Interfaces/IPluginManager.h"
+#include "Settings/EditorStyleSettings.h"
+#include "Styling/StarshipCoreStyle.h"
+#include "Styling/StyleColors.h"
+
+
+TSharedPtr<FSlateStyleSet> FCustomizableObjectEditorStyle::CustomizableObjectEditorStyleInstance;
+
+
+void FCustomizableObjectEditorStyle::Initialize()
+{
+	// Only register once
+	if (!CustomizableObjectEditorStyleInstance.IsValid())
+	{
+		CustomizableObjectEditorStyleInstance = Create();
+		FSlateStyleRegistry::RegisterSlateStyle(*CustomizableObjectEditorStyleInstance);
+	}
+}
+
+
+void FCustomizableObjectEditorStyle::Shutdown()
+{
+	if (CustomizableObjectEditorStyleInstance.IsValid())
+	{
+		FSlateStyleRegistry::UnRegisterSlateStyle(*CustomizableObjectEditorStyleInstance.Get());
+		ensure(CustomizableObjectEditorStyleInstance.IsUnique());
+		CustomizableObjectEditorStyleInstance.Reset();
+	}
+}
+
+
+const ISlateStyle& FCustomizableObjectEditorStyle::Get()
+{
+	return *CustomizableObjectEditorStyleInstance;
+}
+
+
+FName FCustomizableObjectEditorStyle::GetStyleSetName()
+{
+	static FName StyleSetName(TEXT("CustomizableObjectEditorStyle"));
+	return StyleSetName;
+}
+
+
+FString FCustomizableObjectEditorStyle::RelativePathToPluginPath(const FString& RelativePath, const ANSICHAR* Extension)
+{
+	static FString ContentDir = IPluginManager::Get().FindPlugin(TEXT("Mutable"))->GetContentDir();
+	return (ContentDir / RelativePath) + Extension;
+}
+
+
+#define IMAGE_PLUGIN_BRUSH( RelativePath, ... ) FSlateImageBrush( FCustomizableObjectEditorStyle::RelativePathToPluginPath( RelativePath, ".png" ), __VA_ARGS__ )
+#define IMAGE_BRUSH(Style, RelativePath, ... ) FSlateImageBrush( Style->RootToContentDir( RelativePath, TEXT(".png") ), __VA_ARGS__ )
+#define IMAGE_BRUSH_SVG( Style, RelativePath, ... ) FSlateVectorImageBrush( Style->RootToContentDir(RelativePath, TEXT(".svg")), __VA_ARGS__)
+
+TSharedRef<FSlateStyleSet> FCustomizableObjectEditorStyle::Create()
+{
+	const FVector2D Icon20x20(20.0f, 20.0f);
+	const FVector2D Icon16x16(16.0f, 16.0f);
+	const FVector2D Icon8x8(8.0f, 8.0f);
+	const FVector2D Icon40x40(40.0f, 40.0f);
+	const FVector2D Icon64x64(64.0f, 64.0f);
+
+	TSharedRef< FSlateStyleSet > Style = MakeShareable(new FSlateStyleSet("CustomizableObjectEditorStyle"));
+	Style->SetContentRoot(FPaths::EngineContentDir() / TEXT("Editor/Slate"));
+
+	Style->Set("CustomizableObjectEditor.Save", new IMAGE_BRUSH_SVG(Style, "Starship/Common/SaveCurrent", Icon40x40));
+	Style->Set("CustomizableObjectEditor.Save.Small", new IMAGE_BRUSH_SVG(Style, "Starship/Common/SaveCurrent", Icon20x20));
+	Style->Set("CustomizableObjectEditor.Compile", new IMAGE_BRUSH_SVG(Style, "Starship/MainToolbar/compile", Icon40x40));
+	Style->Set("CustomizableObjectEditor.Compile.Small", new IMAGE_BRUSH_SVG(Style, "Starship/MainToolbar/compile", Icon20x20));
+	Style->Set("CustomizableObjectEditor.CompileOnlySelected", new IMAGE_BRUSH_SVG(Style, "Starship/MainToolbar/compile", Icon40x40));
+	Style->Set("CustomizableObjectEditor.CompileOnlySelected.Small", new IMAGE_BRUSH_SVG(Style, "Starship/MainToolbar/compile", Icon20x20));
+	Style->Set("CustomizableObjectEditor.ResetCompileOptions", new IMAGE_BRUSH_SVG(Style, "Starship/MainToolbar/compile", Icon40x40));
+	Style->Set("CustomizableObjectEditor.ResetCompileOptions.Small", new IMAGE_BRUSH_SVG(Style, "Starship/MainToolbar/compile", Icon20x20));
+	Style->Set("CustomizableObjectEditor.Debug", new IMAGE_BRUSH_SVG(Style, "Starship/Common/Debug", Icon40x40));
+	Style->Set("CustomizableObjectEditor.Debug.Small", new IMAGE_BRUSH_SVG(Style, "Starship/Common/Debug", Icon20x20));
+	Style->Set("CustomizableObjectEditor.TextureAnalyzer", new IMAGE_BRUSH_SVG(Style, "Starship/Common/IssueTracker", Icon40x40));
+	Style->Set("CustomizableObjectEditor.TextureAnalyzer.Small", new IMAGE_BRUSH_SVG(Style, "Starship/Common/IssueTracker", Icon20x20));
+	Style->Set("CustomizableObjectEditor.PerofrmanceReport", new IMAGE_BRUSH_SVG(Style, "Starship/MainToolbar/simulate", Icon40x40));
+	Style->Set("CustomizableObjectEditor.PerofrmanceReport.Small", new IMAGE_BRUSH_SVG(Style, "Starship/MainToolbar/simulate", Icon20x20));
+
+	Style->Set("CustomizableObjectEditor.Tabs.Preview", new IMAGE_BRUSH_SVG(Style, "Starship/Persona/PersonaPreviewAnimation", Icon20x20));
+	Style->Set("CustomizableObjectEditor.Tabs.PreviewSettings", new IMAGE_BRUSH_SVG(Style, "Starship/AssetIcons/CameraActor_64", Icon20x20));
+	Style->Set("CustomizableObjectEditor.Tabs.CustomizableObjectProperties", new IMAGE_BRUSH_SVG(Style, "Starship/Persona/SectionSelection", Icon20x20));
+	Style->Set("CustomizableObjectEditor.Tabs.CustomizableInstanceProperties", new IMAGE_BRUSH_SVG(Style, "Starship/Common/Details", Icon20x20));
+	Style->Set("CustomizableObjectEditor.Tabs.NodeGraph", new IMAGE_BRUSH_SVG(Style, "Starship/Common/blueprint", Icon20x20));
+	Style->Set("CustomizableObjectEditor.Tabs.NodeProperties", new IMAGE_BRUSH_SVG(Style, "Starship/Blueprints/icon_BlueprintEditor_EventGraph", Icon20x20));
+
+	Style->Set("CustomizableObjectInstanceEditor.SetShowWireframe", new IMAGE_BRUSH(Style, "Icons/icon_StaticMeshEd_Wireframe_40x", Icon40x40));
+	Style->Set("CustomizableObjectInstanceEditor.SetShowWireframe.Small", new IMAGE_BRUSH(Style, "Icons/icon_StaticMeshEd_Wireframe_40x", Icon20x20));
+
+	Style->Set("CustomizableObjectInstanceEditor.Save", new IMAGE_BRUSH_SVG(Style, "Starship/Common/SaveCurrent", Icon40x40));
+	Style->Set("CustomizableObjectInstanceEditor.Save.Small", new IMAGE_BRUSH_SVG(Style, "Starship/Common/SaveCurrent", Icon20x20));
+	Style->Set("CustomizableObjectInstanceEditor.ShowParentCO", new IMAGE_BRUSH_SVG(Style, "Starship/MainToolbar/eject", Icon40x40));
+	Style->Set("CustomizableObjectInstanceEditor.ShowParentCO.Small", new IMAGE_BRUSH_SVG(Style, "Starship/MainToolbar/eject", Icon20x20));
+	Style->Set("CustomizableObjectInstanceEditor.EditParentCO", new IMAGE_BRUSH_SVG(Style, "Starship/Common/EditorModes", Icon40x40));
+	Style->Set("CustomizableObjectInstanceEditor.EditParentCO.Small", new IMAGE_BRUSH_SVG(Style, "Starship/Common/EditorModes", Icon20x20));
+	Style->Set("CustomizableObjectInstanceEditor.TextureAnalyzer", new IMAGE_BRUSH_SVG(Style, "Starship/Common/IssueTracker", Icon40x40));
+	Style->Set("CustomizableObjectInstanceEditor.TextureAnalyzer.Small", new IMAGE_BRUSH_SVG(Style, "Starship/Common/IssueTracker", Icon20x20));
+
+	Style->Set("CustomizableObjectEditorViewport.SetShowWireframe", new IMAGE_BRUSH(Style, "Icons/icon_StaticMeshEd_Wireframe_40x", Icon40x40));
+	Style->Set("CustomizableObjectEditorViewport.SetShowWireframe.Small", new IMAGE_BRUSH(Style, "Icons/icon_StaticMeshEd_Wireframe_40x", Icon20x20));
+	Style->Set("CustomizableObjectEditorViewport.SetRealtimePreview", new IMAGE_BRUSH(Style, "Icons/icon_MatEd_Realtime_40x", Icon40x40));
+	Style->Set("CustomizableObjectEditorViewport.SetRealtimePreview.Small", new IMAGE_BRUSH(Style, "Icons/icon_MatEd_Realtime_40x", Icon20x20));
+	Style->Set("CustomizableObjectEditorViewport.SetShowBounds", new IMAGE_BRUSH(Style, "Icons/icon_StaticMeshEd_Bounds_40x", Icon40x40));
+	Style->Set("CustomizableObjectEditorViewport.SetShowBounds.Small", new IMAGE_BRUSH(Style, "Icons/icon_StaticMeshEd_Bounds_40x", Icon20x20));
+	Style->Set("CustomizableObjectEditorViewport.SetShowNormals", new IMAGE_BRUSH(Style, "Icons/icon_StaticMeshEd_Normals_40x", Icon40x40));
+	Style->Set("CustomizableObjectEditorViewport.SetShowNormals.Small", new IMAGE_BRUSH(Style, "Icons/icon_StaticMeshEd_Normals_40x", Icon20x20));
+	Style->Set("CustomizableObjectEditorViewport.SetShowTangents", new IMAGE_BRUSH(Style, "Icons/icon_StaticMeshEd_Tangents_40x", Icon40x40));
+	Style->Set("CustomizableObjectEditorViewport.SetShowTangents.Small", new IMAGE_BRUSH(Style, "Icons/icon_StaticMeshEd_Tangents_40x", Icon20x20));
+	Style->Set("CustomizableObjectEditorViewport.SetShowBinormals", new IMAGE_BRUSH(Style, "Icons/icon_StaticMeshEd_Binormals_40x", Icon40x40));
+	Style->Set("CustomizableObjectEditorViewport.SetShowBinormals.Small", new IMAGE_BRUSH(Style, "Icons/icon_StaticMeshEd_Binormals_40x", Icon20x20));
+	Style->Set("CustomizableObjectEditorViewport.SaveThumbnail", new IMAGE_BRUSH(Style, "Icons/icon_StaticMeshEd_SaveThumbnail_40x", Icon40x40));
+	Style->Set("CustomizableObjectEditorViewport.SaveThumbnail.Small", new IMAGE_BRUSH(Style, "Icons/icon_StaticMeshEd_SaveThumbnail_40x", Icon20x20));
+	Style->Set("CustomizableObjectEditorViewport.SetShowPivot", new IMAGE_BRUSH(Style, "Icons/icon_StaticMeshEd_ShowPivot_40x", Icon40x40));
+	Style->Set("CustomizableObjectEditorViewport.SetShowPivot.Small", new IMAGE_BRUSH(Style, "Icons/icon_StaticMeshEd_ShowPivot_40x", Icon20x20));
+
+	Style->Set("CustomizableObjectEditorViewport.SetShowGrid", new IMAGE_BRUSH_SVG(Style, "Starship/Common/Grid", Icon20x20));
+	Style->Set("CustomizableObjectEditorViewport.SetShowSky", new IMAGE_BRUSH_SVG(Style, "Starship/Common/Reflections", Icon20x20));
+	Style->Set("CustomizableObjectEditorViewport.SetDrawUVs", new IMAGE_BRUSH_SVG(Style, "Starship/Common/SetDrawUVs", Icon20x20));
+	Style->Set("CustomizableObjectEditorViewport.BakeInstance", new IMAGE_BRUSH_SVG(Style, "Starship/MainToolbar/compile", Icon20x20));
+
+	Style->Set("COEditorViewportLODCommands.TranslateMode", new IMAGE_BRUSH_SVG(Style, "Starship/EditorViewport/translate", Icon20x20));
+	Style->Set("COEditorViewportLODCommands.RotateMode", new IMAGE_BRUSH_SVG(Style, "Starship/EditorViewport/rotate", Icon20x20));
+	Style->Set("COEditorViewportLODCommands.ScaleMode", new IMAGE_BRUSH_SVG(Style, "Starship/EditorViewport/scale", Icon20x20));
+	Style->Set("COEditorViewportLODCommands.RotationGridSnap", new IMAGE_BRUSH_SVG(Style, "Starship/EditorViewport/angle", Icon20x20));
+
+	Style->Set("LayoutEditorCommands.AddBlock", new IMAGE_BRUSH_SVG(Style, "Starship/Common/Layout", Icon40x40));
+	Style->Set("LayoutEditorCommands.AddBlock.Small", new IMAGE_BRUSH_SVG(Style, "Starship/Common/Layout", Icon20x20));
+	Style->Set("LayoutEditorCommands.RemoveBlock", new IMAGE_BRUSH_SVG(Style, "Starship/Common/LayoutRemove", Icon40x40));
+	Style->Set("LayoutEditorCommands.RemoveBlock.Small", new IMAGE_BRUSH_SVG(Style, "Starship/Common/LayoutRemove", Icon20x20));
+	Style->Set("LayoutEditorCommands.GenerateBlocks", new IMAGE_BRUSH_SVG(Style, "Starship/Persona/AnimationPreviewMesh", Icon40x40));
+	Style->Set("LayoutEditorCommands.GenerateBlocks.Small", new IMAGE_BRUSH_SVG(Style, "Starship/Persona/AnimationPreviewMesh", Icon20x20));
+
+	Style->Set("LayoutBlockSelector.SelectAll", new IMAGE_BRUSH_SVG(Style, "Starship/Common/LayoutLoad", Icon40x40));
+	Style->Set("LayoutBlockSelector.SelectAll.Small", new IMAGE_BRUSH_SVG(Style, "Starship/Common/LayoutLoad", Icon20x20));
+	Style->Set("LayoutBlockSelector.SelectNone", new IMAGE_BRUSH_SVG(Style, "Starship/Common/LayoutSave", Icon40x40));
+	Style->Set("LayoutBlockSelector.SelectNone.Small", new IMAGE_BRUSH_SVG(Style, "Starship/Common/LayoutSave", Icon20x20));
+
+	Style->Set("Nodes.ArrowUp", new IMAGE_BRUSH(Style, "Old/ArrowUp", Icon16x16));
+	Style->Set("Nodes.ArrowDown", new IMAGE_BRUSH(Style, "Old/ArrowDown", Icon16x16));
+
+	Style->Set("CustomizableObjectDebugger.GenerateMutableGraph", new IMAGE_BRUSH_SVG(Style, "Starship/Common/blueprint", Icon40x40));
+	Style->Set("CustomizableObjectDebugger.GenerateMutableGraph.Small", new IMAGE_BRUSH_SVG(Style, "Starship/Common/blueprint", Icon20x20));
+	Style->Set("CustomizableObjectDebugger.CompileMutableCode", new IMAGE_BRUSH_SVG(Style, "Starship/MainToolbar/Compile", Icon40x40));
+	Style->Set("CustomizableObjectDebugger.CompileMutableCode.Small", new IMAGE_BRUSH_SVG(Style, "Starship/MainToolbar/Compile", Icon20x20));
+
+	//Style->Set("ClassIcon.CustomizableObject", new IMAGE_PLUGIN_BRUSH("Icons/Icon_CustomObject_16x16", Icon16x16));
+	//Style->Set("ClassThumbnail.CustomizableObject", new IMAGE_PLUGIN_BRUSH("Icons/Icon_CustomObject_64x64.png", Icon64x64));
+	//Style->Set("ClassIcon.CustomizableObjectInstance", new IMAGE_PLUGIN_BRUSH("Icons/Icon_CustomObjectInstance_16x16.png", Icon16x16));
+	//Style->Set("ClassThumbnail.CustomizableObjectInstance", new IMAGE_PLUGIN_BRUSH("Icons/Icon_CustomObjectInstance_64x64.png", Icon64x64));
+
+
+	FTableRowStyle PerformanceReportRowStyle = FAppStyle::Get().GetWidgetStyle<FTableRowStyle>("TableView.Row");
+	PerformanceReportRowStyle.SetEvenRowBackgroundBrush(*FAppStyle::Get().GetBrush("Brushes.AccentBrown"));
+	PerformanceReportRowStyle.SetOddRowBackgroundBrush(*FAppStyle::Get().GetBrush("Brushes.AccentGray"));
+
+	Style->Set("CustomizableObjectPerformanceReport.Row", PerformanceReportRowStyle);
+
+	FTableRowStyle PerformanceReportAllEvenRowStyle = PerformanceReportRowStyle;
+	PerformanceReportAllEvenRowStyle.SetOddRowBackgroundBrush(PerformanceReportRowStyle.EvenRowBackgroundBrush);
+
+	Style->Set("CustomizableObjectPerformanceReportAllEven.Row", PerformanceReportAllEvenRowStyle);
+
+	FTableRowStyle PerformanceReportAllOddRowStyle = PerformanceReportRowStyle;
+	PerformanceReportAllOddRowStyle.SetEvenRowBackgroundBrush(PerformanceReportRowStyle.OddRowBackgroundBrush);
+
+	Style->Set("CustomizableObjectPerformanceReportAllOdd.Row", PerformanceReportAllOddRowStyle);
+
+	return Style;
+}
+
+#undef IMAGE_PLUGIN_BRUSH
+#undef IMAGE_BRUSH
+#undef IMAGE_BRUSH_SVG
diff --git a/Engine/Plugins/Experimental/Mutable/Source/CustomizableObjectEditor/Private/CustomizableObjectEditorUtilities.h b/Engine/Plugins/Experimental/Mutable/Source/CustomizableObjectEditor/Private/CustomizableObjectEditorUtilities.h
new file mode 100644
--- /dev/null
+++ b/Engine/Plugins/Experimental/Mutable/Source/CustomizableObjectEditor/Private/CustomizableObjectEditorUtilities.h
@@ -0,0 +1,24 @@
+// Copyright Epic Games, Inc. All Rights Reserved.
+
+#pragma once
+
+#include "CoreMinimal.h"
+
+static bool CompareNames(const TSharedPtr<FString>& sp1, const TSharedPtr<FString>& sp2)
+{
+    if (const FString* s1 = sp1.Get())
+    {
+        if (const FString* s2 = sp2.Get())
+        {
+            return (s1->Compare(*s2, ESearchCase::IgnoreCase) < 0);
+        }
+        else
+        {
+            return false;
+        }
+    }
+    else
+    {
+        return true;
+    }
+}
\ No newline at end of file
