[mutable] Move Mutable plugin from the restricted folder to the experimental plugins folder.

diff --git a/Engine/Plugins/Experimental/Mutable/Source/CustomizableObjectEditor/Private/CustomizableObjectEditorViewportClient.cpp b/Engine/Plugins/Experimental/Mutable/Source/CustomizableObjectEditor/Private/CustomizableObjectEditorViewportClient.cpp
new file mode 100644
--- /dev/null
+++ b/Engine/Plugins/Experimental/Mutable/Source/CustomizableObjectEditor/Private/CustomizableObjectEditorViewportClient.cpp
@@ -0,0 +1,3376 @@
+// Copyright Epic Games, Inc. All Rights Reserved.
+
+#include "CustomizableObjectEditorViewportClient.h"
+#include "CustomizableObjectEditorActions.h"
+#include "CustomizableObjectBakeHelpers.h"
+
+#include "MouseDeltaTracker.h"
+#include "PreviewScene.h"
+#include "ISocketManager.h"
+#include "Slate/SceneViewport.h"
+#include "Components/ReflectionCaptureComponent.h"
+#include "Components/SphereReflectionCaptureComponent.h"
+
+#include "Templates/Decay.h"
+#include "Templates/UnrealTemplate.h"
+
+#include "CustomizableObjectEditor.h"
+#include "BusyCursor.h"
+#include "MeshBuild.h"
+#include "PreviewScene.h"
+#include "ObjectTools.h"
+#include "Components/SpotLightComponent.h"
+#include "Engine/Canvas.h"
+#include "Materials/MaterialInstance.h"
+#include "Materials/MaterialInstanceDynamic.h"
+#include "DynamicMeshBuilder.h"
+#include "SceneManagement.h"
+#include "Animation/PoseAsset.h"
+#include "AssetRegistry/AssetRegistryModule.h"
+#include "AssetToolsModule.h"
+#include "Preferences/PersonaOptions.h"
+#include "UnrealEdGlobals.h"
+#include "EditorModeManager.h"
+#include "Editor/UnrealEdEngine.h"
+#include "Widgets/Layout/SUniformGridPanel.h"
+#include "Widgets/Input/SEditableTextBox.h"
+#include "ScopedTransaction.h"
+#include "IContentBrowserSingleton.h"
+#include "ContentBrowserModule.h"
+#include "FileHelpers.h"
+
+#include "ScopedTransaction.h"
+#include "IContentBrowserSingleton.h"
+#include "ContentBrowserModule.h"
+#include "Widgets/Input/SButton.h"
+#include "HAL/PlatformFileManager.h"
+#include "HAL/FileManager.h"
+#include "Misc/MessageDialog.h"
+#include "PackageTools.h"
+#include "UnrealEdGlobals.h"
+#include "Framework/Notifications/NotificationManager.h"
+#include "Widgets/Notifications/SNotificationList.h"
+#include "AssetViewerSettings.h"
+#include "UObject/UObjectGlobals.h"
+#include "Editor/EditorPerProjectUserSettings.h"
+#include "AdvancedPreviewScene.h"
+#include "UnrealWidget.h"
+
+#include "UnrealPortabilityHelpers.h"
+#include "UnrealEditorPortabilityHelpers.h"
+#include "FileHelpers.h"
+
+#include "CustomizableObjectPreviewScene.h"
+#include "UnrealEditorPortabilityHelpers.h"
+#include "UnrealBakeHelpers.h"
+#include "CustomizableObjectMipDataProvider.h"
+#include "EditorStyleSet.h"
+
+#define LOCTEXT_NAMESPACE "CustomizableObjectEditor" 
+
+
+GizmoRTSProxy::GizmoRTSProxy() :
+	AnyGizmoSelected(false)
+	, DataOriginParameter(nullptr)
+	, DataOriginConstant(nullptr)
+	, HasAssignedData(false)
+	, AssignedDataIsFromNode(false)
+	, ProjectorParameterIndex(-1)
+	, ProjectorGizmoEdited(false)
+	, CallUpdateSkeletalMesh(true)
+	, bManipulatingGizmo(false)
+	, ProjectionType(ECustomizableObjectProjectorType::Planar)
+{
+
+}
+
+
+void GizmoRTSProxy::CopyTransformFromOriginData()
+{
+    if (const TSharedPtr<ICustomizableObjectInstanceEditor> Editor = CustomizableObjectEditorPtr.Pin())
+    {
+	    if (const UCustomizableObjectInstance* Instance = Editor->GetPreviewInstance())
+        {
+			Value = Instance->GetProjector(ProjectorParameterName, ProjectorRangeIndex);
+        }
+    }
+}
+
+
+bool GizmoRTSProxy::UpdateOriginData()
+{
+	if (!HasAssignedData)
+	{
+		return false;
+	}
+
+	if (DataOriginParameter != nullptr)
+	{
+		DataOriginParameter->DefaultValue = Value;
+		return true;
+	}
+	else if (DataOriginConstant != nullptr)
+	{
+		DataOriginConstant->Value = Value;
+		return true;
+	}
+	else
+	{
+		if (CustomizableObjectEditorPtr.IsValid())
+		{
+			if (UCustomizableObjectInstance* Instance = CustomizableObjectEditorPtr.Pin()->GetPreviewInstance())
+			{
+				TArray<FCustomizableObjectProjectorParameterValue>& ProjectorParameters = Instance->GetProjectorParameters();
+				int32 Index = Instance->FindProjectorParameterNameIndex(ProjectorParameterName);
+				
+				for (int32 i = 0; i < Instance->GetProjectorParameters().Num(); ++i)
+				{
+					if (ProjectorParameters[i].ParameterName == ProjectorParameterName)
+					{
+						if (ProjectorRangeIndex >= 0 && !ProjectorParameters[i].RangeValues.IsValidIndex(ProjectorRangeIndex))
+						{
+							return false;
+						}
+					}
+				}
+				
+				if (!Instance->RemovedProjectorParameterNameWithIndex.IsEmpty() && Instance->RemovedProjectorParameterNameWithIndex == ProjectorParameterNameWithIndex) // If the projector has been removed, do not update
+				{
+					Index = -1;
+				}
+				Instance->RemovedProjectorParameterNameWithIndex = FString(""); // Consume the projector removed tag
+
+				if (Index != -1)
+				{
+					if (!ProjectorGizmoEdited)
+					{
+						Instance->PreEditChange(NULL);
+					}
+
+					Instance->SetProjectorValue(ProjectorParameterName,
+						(FVector)Value.Position,
+						(FVector)Value.Direction,
+						(FVector)Value.Up,
+						(FVector)Value.Scale,
+						Value.Angle,
+						ProjectorRangeIndex);
+
+					if (!ProjectorGizmoEdited)
+					{
+						Instance->PostEditChange(); // Avoid unnecessary updates
+					}
+
+					if (CallUpdateSkeletalMesh)
+					{
+						Instance->UpdateSkeletalMeshAsync(true);
+					}
+					return true;
+				}
+			}
+		}
+	}
+
+	return false;
+}
+
+
+bool GizmoRTSProxy::Modify()
+{
+	if (DataOriginParameter != nullptr)
+	{
+		DataOriginParameter->Modify();
+		return true;
+	}
+	else if (DataOriginConstant != nullptr)
+	{
+		DataOriginConstant->Modify();
+		return true;
+	}
+
+	return false;
+}
+
+
+bool GizmoRTSProxy::ModifyGraph()
+{
+	if ((DataOriginParameter == nullptr) && (DataOriginConstant == nullptr))
+	{
+		return false;
+	}
+
+	if (DataOriginParameter != nullptr)
+	{
+		DataOriginParameter->GetGraph()->Modify();
+		return true;
+	}
+	else if (DataOriginConstant != nullptr)
+	{
+		DataOriginConstant->GetGraph()->Modify();
+		return true;
+	}
+
+	return false;
+}
+
+
+void GizmoRTSProxy::CleanOriginData()
+{
+	HasAssignedData = false;
+	AssignedDataIsFromNode = false;
+	AnyGizmoSelected = false;
+	DataOriginParameter = nullptr;
+	DataOriginConstant = nullptr;
+	ProjectorParameterName = FString("");
+	ProjectorParameterNameWithIndex = FString("");
+	ProjectorRangeIndex = -1;
+	ProjectorParameterIndex = -1;
+}
+
+
+bool GizmoRTSProxy::ProjectorHasInitialValues(FCustomizableObjectProjector& Parameter)
+{
+	if ((Parameter.Position == FVector3f(0, 0, 0)) &&
+		(Parameter.Direction == FVector3f(1, 0, 0)) &&
+		(Parameter.Up == FVector3f(0, 1, 0)) &&
+		(Parameter.Scale == FVector3f(10, 10, 100)) &&
+		(Parameter.ProjectionType == ECustomizableObjectProjectorType::Planar) &&
+		(Parameter.Angle == 2.0f * PI))
+	{
+		return true;
+	}
+
+	return false;
+}
+
+
+FCustomizableObjectProjector GizmoRTSProxy::SetProjectorInitialValue(TArray<TWeakObjectPtr<UDebugSkelMeshComponent>>& SkeletalMeshComponents, float TotalLength)
+{
+	FCustomizableObjectProjector Result;
+	bool bFoundComponents = false;
+	FBoxSphereBounds Bounds;
+
+	for (TWeakObjectPtr<USkeletalMeshComponent> SkeletalMeshComponent : SkeletalMeshComponents)
+	{
+		if (!bFoundComponents)
+		{
+			Bounds = SkeletalMeshComponent->Bounds;
+		}
+		else
+		{
+			Bounds = Bounds + SkeletalMeshComponent->Bounds;
+		}
+	}
+
+	if (bFoundComponents)
+	{
+		Result.Position = FVector3f(Bounds.Origin + FVector(0.0f, 1.0f, 0.0f) * TotalLength * 0.5f);	// LWC_TODO: Precision Loss
+		Result.Scale = FVector3f(40.0f, 40.0f, 40.0f);
+		Result.Up = FVector3f(0.0f, 0.0f, 1.0f);
+		Result.Direction = FVector3f(0.0f, -1.0f, 0.0f);
+		Result.ProjectionType = ECustomizableObjectProjectorType::Planar;
+		Result.Angle = 2.0f * PI;
+	}
+
+	return Result;
+}
+
+
+bool GizmoRTSProxy::IsProjectorParameterSelected()
+{
+	if (AnyGizmoSelected && HasAssignedData && !AssignedDataIsFromNode && (DataOriginParameter == nullptr) && (DataOriginConstant == nullptr))
+	{
+		return true;
+	}
+
+	return false;
+}
+
+
+void GizmoRTSProxy::ProjectorParameterChanged(UCustomizableObjectNodeProjectorParameter* Node)
+{
+	if (AnyGizmoSelected && AssignedDataIsFromNode && (Node == DataOriginParameter))
+	{
+		switch (Node->ParameterSetModified)
+		{
+			case 0:
+			{
+				ProjectionType = Node->ProjectionType;
+				Value.ProjectionType = Node->ProjectionType;
+				break;
+			}
+			case 1:
+			{
+				Value.Angle = Node->DefaultValue.Angle;
+				break;
+			}
+			case 2:
+			{
+				Value.Position = Node->DefaultValue.Position;
+				Value.Direction = Node->DefaultValue.Direction;
+				Value.Up = Node->DefaultValue.Up;
+				Node->ParameterSetModified = -1;
+				break;
+			}
+			default:
+			{
+				UE_LOG(LogMutable, Warning, TEXT("ERROR: wrong parameter set modified value %d"), Node->ParameterSetModified);
+				break;
+			}
+		}
+	}
+}
+
+
+void GizmoRTSProxy::ProjectorParameterChanged(UCustomizableObjectNodeProjectorConstant* Node)
+{
+	if (AnyGizmoSelected && AssignedDataIsFromNode && (Node == DataOriginConstant))
+	{
+		switch (Node->ParameterSetModified)
+		{
+			case 0:
+			{
+				ProjectionType = Node->ProjectionType;
+				Value.ProjectionType = Node->ProjectionType;
+				break;
+			}
+			case 1:
+			{
+				Value.Angle = Node->Value.Angle;
+				break;
+			}
+			case 2:
+			{
+				Value.Position = Node->Value.Position;
+				Value.Direction = Node->Value.Direction;
+				Value.Up = Node->Value.Up;
+				break;
+			}
+			default:
+			{
+				UE_LOG(LogMutable, Warning, TEXT("ERROR: wrong parameter set modified value %d"), Node->ParameterSetModified);
+				break;
+			}
+		}
+	}
+}
+
+
+void GizmoRTSProxy::SetCallUpdateSkeletalMesh(bool InValue)
+{
+	CallUpdateSkeletalMesh = InValue;
+}
+
+
+void GizmoRTSProxy::SetProjectorUpdatedInViewport(bool InValue)
+{
+	if (!AnyGizmoSelected || DataOriginParameter || DataOriginConstant)
+	{
+		return;
+	}
+
+	if(InValue != bManipulatingGizmo)
+	{
+		bManipulatingGizmo = InValue;
+
+		if (CustomizableObjectEditorPtr.IsValid())
+		{
+			UCustomizableObjectInstance* Inst = CustomizableObjectEditorPtr.Pin()->GetPreviewInstance();
+
+			if (Inst)
+			{
+				Inst->ProjectorUpdatedInViewport = InValue;
+			}
+		}
+	}
+}
+
+
+FString GizmoRTSProxy::GetProjectorParameterName()
+{
+	return ProjectorParameterName;
+}
+
+
+FString GizmoRTSProxy::GetProjectorParameterNameWithIndex()
+{
+	return ProjectorParameterNameWithIndex;
+}
+
+
+bool GizmoRTSProxy::GetHasAssignedData()
+{
+	return HasAssignedData;
+}
+
+
+bool GizmoRTSProxy::GetAssignedDataIsFromNode()
+{
+	return AssignedDataIsFromNode;
+}
+
+
+FCustomizableObjectEditorViewportClient::FCustomizableObjectEditorViewportClient(TWeakPtr<ICustomizableObjectInstanceEditor> InCustomizableObjectEditor, FPreviewScene* InPreviewScene)
+	: FEditorViewportClient(&GLevelEditorModeTools(), InPreviewScene)
+	, CustomizableObjectEditorPtr(InCustomizableObjectEditor)
+	, CustomizableObject(nullptr)
+	, BakingOverwritePermission(false)
+	, AssetRegistryLoaded(false)
+{
+	// load config
+	ConfigOption = UPersonaOptions::StaticClass()->GetDefaultObject<UPersonaOptions>();
+	check (ConfigOption);
+
+	bCameraMove = false;
+	bUsingOrbitCamera = true;
+
+	bShowSockets = true;
+	bDrawUVs = false;
+	bDrawNormals = false;
+	bDrawTangents = false;
+	bDrawBinormals = false;
+	bShowPivot = false;
+	Widget->SetDefaultVisibility(false);
+	bCameraLock = true;
+	bDrawSky = true;
+
+	bActivateOrbitalCamera = true;
+	bSetOrbitalOnPerspectiveMode = true;
+
+	MaterialToDrawInUVs = 0;
+	MaterialToDrawInUVsLOD = 0;
+	MaterialToDrawInUVsIndex = 0;
+	UVChannelToDrawInUVs = 0;
+	MaterialToDrawInUVsComponent = 0;
+
+	bManipulating = false;
+	LastHitProxyWidget = 0;
+
+	bReferenceMeshMissingWarningMessageVisible = false;
+
+	DrawHelper.bDrawPivot = false;
+	DrawHelper.bDrawWorldBox = false;
+	DrawHelper.bDrawKillZ = false;
+	DrawHelper.bDrawGrid = false;
+	DrawHelper.GridColorAxis = FColor(160, 160, 160);
+	DrawHelper.GridColorMajor = FColor(144, 144, 144);
+	DrawHelper.GridColorMinor = FColor(128, 128, 128);
+	DrawHelper.PerspectiveGridSize = 2048.0f;
+	DrawHelper.NumCells = DrawHelper.PerspectiveGridSize / (32);
+	SetShowGrid();
+
+	SetViewMode(VMI_Lit);
+
+	EngineShowFlags.SetSeparateTranslucency(true);
+	EngineShowFlags.SetSnap(0);
+	EngineShowFlags.SetCompositeEditorPrimitives(true);
+
+	EngineShowFlags.ScreenSpaceReflections = 1;
+	EngineShowFlags.AmbientOcclusion = 1;
+	EngineShowFlags.Grid = ConfigOption->bShowGrid;
+
+	OverrideNearClipPlane(1.0f);
+
+	SetPreviewComponent(nullptr);
+
+	// add capture component for reflection
+	USphereReflectionCaptureComponent* CaptureComponent = NewObject<USphereReflectionCaptureComponent>();
+
+	const FTransform CaptureTransform(FRotator(0, 0, 0), FVector(0.f, 0.f, 100.f), FVector(1.f));
+	PreviewScene->AddComponent(CaptureComponent, CaptureTransform);
+	CaptureComponent->UpdateReflectionCaptureContents(PreviewScene->GetWorld());
+
+	// now add the ClipMorph plane
+	ClipMorphNode = nullptr;
+	bClipMorphVisible = false;
+	bClipMorphLocalStartOffset = true;
+	ClipMorphMaterial = LoadObject<UMaterial>(NULL, TEXT("Material'/Engine/EditorMaterials/LevelGridMaterial.LevelGridMaterial'"), NULL, LOAD_None, NULL);
+	check(ClipMorphMaterial);
+
+	// clip mesh preview
+	ClipMeshNode = nullptr;
+	ClipMeshComp = NewObject<UStaticMeshComponent>();
+	PreviewScene->AddComponent(ClipMeshComp, FTransform());
+	ClipMeshComp->SetVisibility(false);
+
+	WidgetMode = UE::Widget::WM_Translate;
+	BoundSphere.W = 100.f;
+
+	const float FOVMin = 5.f;
+	const float FOVMax = 170.f;
+	ViewFOV = FMath::Clamp<float>(53.43f, FOVMin, FOVMax);
+
+	SetRealtime(true);
+	if (GEditor->PlayWorld)
+	{
+		AddRealtimeOverride(false, LOCTEXT("RealtimeOverrideMessage_InstanceViewport", "Instance Viewport")); // We are PIE, don't start in realtime mode
+	}
+
+	GizmoProxy.CustomizableObjectEditorPtr = InCustomizableObjectEditor;
+	WidgetVisibility = false;
+
+	IsPlayingAnimation = false;
+	AnimationBeingPlayed = nullptr;
+
+	// Lighting 
+	SelectedLightComponent = nullptr;
+	bIsEditingLightEnabled = false;
+
+	StateChangeShowGeometryDataFlag = false;
+
+	// Register delegate to update the show flags when the post processing is turned on or off
+	UAssetViewerSettings::Get()->OnAssetViewerSettingsChanged().AddRaw(this, &FCustomizableObjectEditorViewportClient::OnAssetViewerSettingsChanged);
+	// Set correct flags according to current profile settings
+	SetAdvancedShowFlagsForScene(UAssetViewerSettings::Get()->Profiles[GetMutableDefault<UEditorPerProjectUserSettings>()->AssetViewerProfileIndex].bPostProcessingEnabled);
+
+	// Set profile so changes in scene lighting affect and match this editor too
+	UEditorPerProjectUserSettings* PerProjectSettings = GetMutableDefault<UEditorPerProjectUserSettings>();
+	UAssetViewerSettings* DefaultSettings = UAssetViewerSettings::Get();
+	PerProjectSettings->AssetViewerProfileIndex = DefaultSettings->Profiles.IsValidIndex(PerProjectSettings->AssetViewerProfileIndex) ? PerProjectSettings->AssetViewerProfileIndex : 0;
+	int32 ProfileIndex = PerProjectSettings->AssetViewerProfileIndex;
+	FAdvancedPreviewScene* PreviewSceneCasted = static_cast<FAdvancedPreviewScene*>(PreviewScene);
+	PreviewSceneCasted->SetProfileIndex(ProfileIndex);
+
+	TransparentPlaneMaterialXY = (UMaterial*)StaticLoadObject(UMaterial::StaticClass(), NULL, TEXT("/Engine/EditorMaterials/WidgetVertexColorMaterial.WidgetVertexColorMaterial"), NULL, LOAD_None, NULL);
+}
+
+
+void FCustomizableObjectEditorViewportClient::UpdateCameraSetup()
+{
+	static FRotator CustomOrbitRotation(-33.75, -135, 0);
+	if ( (SkeletalMeshComponents.Num() && SkeletalMeshComponents[0].IsValid() && SkeletalMeshComponents[0]->SkeletalMesh)
+		||
+		(StaticMeshComponent.IsValid() && StaticMeshComponent->GetStaticMesh()) )
+	{
+		BoundSphere = GetCameraTarget();
+		FVector CustomOrbitZoom(0, BoundSphere.W / (75.0f * (float)PI / 360.0f), 0);
+		FVector CustomOrbitLookAt = BoundSphere.Center;
+
+		SetCameraSetup(CustomOrbitLookAt, CustomOrbitRotation, CustomOrbitZoom, CustomOrbitLookAt, GetViewLocation(), GetViewRotation() );
+
+		UpdateFloor();
+
+		EnableCameraLock(bActivateOrbitalCamera);
+		FBox Box( BoundSphere.Center - FVector(BoundSphere.W) / 2.0f, BoundSphere.Center + FVector(BoundSphere.W) / 2.0f );
+		FocusViewportOnBox( Box, false );
+	}
+}
+
+
+void FCustomizableObjectEditorViewportClient::UpdateFloor()
+{
+	// Move the floor to the bottom of the bounding box of the mesh, rather than on the origin
+	FVector Bottom(0.0f);
+	bool bFoundSkelMesh = false;
+
+	for (TWeakObjectPtr<USkeletalMeshComponent> SkeletalMeshComponent : SkeletalMeshComponents)
+	{
+		if (SkeletalMeshComponent.IsValid())
+		{
+			SkeletalMeshComponent->bComponentUseFixedSkelBounds = true;
+			SkeletalMeshComponent->UpdateBounds();
+
+			Bottom = SkeletalMeshComponent->Bounds.GetBoxExtrema(0);
+			bFoundSkelMesh = true;
+		}
+	}
+
+	// TODO: Optimize
+	if (bFoundSkelMesh)
+	{
+
+	}
+	else if (StaticMeshComponent.IsValid())
+	{
+		StaticMeshComponent->UpdateBounds();
+
+		Bottom = StaticMeshComponent->Bounds.GetBoxExtrema(0);
+	}
+
+	FAdvancedPreviewScene* AdvancedScene = static_cast<FAdvancedPreviewScene*>(PreviewScene);
+	if (AdvancedScene != nullptr)
+	{
+		const UStaticMeshComponent* FloorMeshComponent = AdvancedScene->GetFloorMeshComponent();
+		if (FloorMeshComponent != nullptr)
+		{
+			UStaticMeshComponent* FloorMeshComponentCasted = const_cast<UStaticMeshComponent*>(FloorMeshComponent);
+			FloorMeshComponentCasted->SetWorldLocation(FVector(0.0f, 0.0f, -1.0f/* Does not seem to work Bottom.Z*/));
+		}
+	}
+}
+
+
+FCustomizableObjectEditorViewportClient::~FCustomizableObjectEditorViewportClient()
+{
+	UAssetViewerSettings::Get()->OnAssetViewerSettingsChanged().RemoveAll(this);
+}
+
+
+void FCustomizableObjectEditorViewportClient::Tick(float DeltaSeconds)
+{
+	FEditorViewportClient::Tick(DeltaSeconds);
+
+	UpdateFloor();
+
+	//PreviewScene->GetWorld()->Tick(LEVELTICK_All, DeltaSeconds);
+}
+
+
+void DrawEllipse(FPrimitiveDrawInterface* PDI, const FVector& Base, const FVector& X, const FVector& Y, const FLinearColor& Color, float Radius1, float Radius2, int32 NumSides, uint8 DepthPriority, float Thickness, float DepthBias, bool bScreenSpace)
+{
+	const float	AngleDelta = 2.0f * PI / NumSides;
+	FVector	LastVertex = Base + X * Radius1;
+
+	for (int32 SideIndex = 0; SideIndex < NumSides; SideIndex++)
+	{
+		const FVector Vertex = Base + (X * FMath::Cos(AngleDelta * (SideIndex + 1)) * Radius1 + Y * FMath::Sin(AngleDelta * (SideIndex + 1)) * Radius2);
+		PDI->DrawLine(LastVertex, Vertex, Color, DepthPriority, Thickness, DepthBias, bScreenSpace);
+		LastVertex = Vertex;
+	}
+}
+
+
+void FCustomizableObjectEditorViewportClient::Draw(const FSceneView* View, FPrimitiveDrawInterface* PDI)
+{
+	FEditorViewportClient::Draw(View, PDI);
+
+	//DrawHelper.Draw( View, PDI );	
+	if( StaticMeshComponent.IsValid() && StaticMeshComponent->GetStaticMesh() && (bDrawNormals || bDrawTangents || bDrawBinormals ) )
+	{
+		/* TODO
+		FStaticMeshRenderData* RenderData = &StaticMeshComponent->StaticMesh->LODModels[CustomizableObjectEditorPtr.Pin()->GetCurrentLODIndex()];
+		uint16* Indices = (uint16*)RenderData->IndexBuffer.Indices.GetData();
+		uint32 NumIndices = RenderData->IndexBuffer.Indices.Num();
+
+		FMatrix LocalToWorldInverseTranspose = StaticMeshComponent->ComponentToWorld.ToMatrixWithScale().Inverse().GetTransposed();
+		for (uint32 i = 0; i < NumIndices; i++)
+		{
+			const FVector& VertexPos = RenderData->PositionVertexBuffer.VertexPosition( Indices[i] );
+
+			const FVector WorldPos = StaticMeshComponent->ComponentToWorld.TransformPosition( VertexPos );
+			const FVector& Normal = RenderData->VertexBuffer.VertexTangentZ( Indices[i] ); 
+			const FVector& Binormal = RenderData->VertexBuffer.VertexTangentY( Indices[i] ); 
+			const FVector& Tangent = RenderData->VertexBuffer.VertexTangentX( Indices[i] ); 
+
+			const float Len = 5.0f;
+
+			if( bDrawNormals )
+			{
+				PDI->DrawLine( WorldPos, WorldPos+LocalToWorldInverseTranspose.TransformVector( Normal ).SafeNormal() * Len, FLinearColor( 0.0f, 1.0f, 0.0f), SDPG_World );
+			}
+
+			if( bDrawTangents )
+			{
+				PDI->DrawLine( WorldPos, WorldPos+LocalToWorldInverseTranspose.TransformVector( Tangent ).SafeNormal() * Len, FLinearColor( 1.0f, 0.0f, 0.0f), SDPG_World );
+			}
+
+			if( bDrawBinormals )
+			{
+				PDI->DrawLine( WorldPos, WorldPos+LocalToWorldInverseTranspose.TransformVector( Binormal ).SafeNormal() * Len, FLinearColor( 0.0f, 0.0f, 1.0f), SDPG_World );
+			}
+		}	
+		*/
+	}
+
+
+	//if( bShowPivot )
+	//{
+	//	FMatrix Transform;
+	//	if (StaticMeshComponent.IsValid())
+	//	{
+	//		Transform = StaticMeshComponent->ComponentToWorld.ToMatrixWithScale();
+	//	}
+	//	else if (SkeletalMeshComponent.IsValid())
+	//	{
+	//		Transform = SkeletalMeshComponent->ComponentToWorld.ToMatrixWithScale();
+	//	}
+
+	//	FUnrealEdUtils::DrawWidget(View, PDI, Transform, 0, 0, EAxisList::All, EWidgetMovementMode::WMM_Translate, false);
+	//}
+
+	// Draw the widgets
+	for ( TMap< int, TSharedPtr<FCustomizableObjectWidget> >::TIterator it(ProjectorWidgets); it; ++it )
+	{
+		it.Value()->Render( View, PDI );
+	}
+	
+	// Draw Selected light Visualizer
+	if (bIsEditingLightEnabled && SelectedLightComponent)
+	{
+		if (USpotLightComponent* SpotLightComp = Cast<USpotLightComponent>(SelectedLightComponent))
+		{
+			FTransform TransformNoScale = SpotLightComp->GetComponentToWorld();
+			TransformNoScale.RemoveScaling();
+
+			// Draw point light source shape
+			DrawWireCapsule(PDI, TransformNoScale.GetTranslation(), -TransformNoScale.GetUnitAxis(EAxis::Z), TransformNoScale.GetUnitAxis(EAxis::Y), TransformNoScale.GetUnitAxis(EAxis::X),
+				FColor(231, 239, 0, 255), SpotLightComp->SourceRadius, 0.5f * SpotLightComp->SourceLength + SpotLightComp->SourceRadius, 25, SDPG_World);
+
+			// Draw outer light cone
+			DrawWireSphereCappedCone(PDI, TransformNoScale, SpotLightComp->AttenuationRadius, SpotLightComp->OuterConeAngle, 32, 8, 10, FColor(200, 255, 255), SDPG_World);
+
+			// Draw inner light cone (if non zero)
+			if (SpotLightComp->InnerConeAngle > UE_KINDA_SMALL_NUMBER)
+			{
+				DrawWireSphereCappedCone(PDI, TransformNoScale, SpotLightComp->AttenuationRadius, SpotLightComp->InnerConeAngle, 32, 8, 10, FColor(150, 200, 255), SDPG_World);
+			}
+		}
+		else if (UPointLightComponent* PointLightComp = Cast<UPointLightComponent>(SelectedLightComponent))
+		{
+			FTransform LightTM = PointLightComp->GetComponentToWorld();
+
+			// Draw light radius
+			DrawWireSphereAutoSides(PDI, FTransform(LightTM.GetTranslation()), FColor(200, 255, 255), PointLightComp->AttenuationRadius, SDPG_World);
+
+			// Draw point light source shape
+			DrawWireCapsule(PDI, LightTM.GetTranslation(), -LightTM.GetUnitAxis(EAxis::Z), LightTM.GetUnitAxis(EAxis::Y), LightTM.GetUnitAxis(EAxis::X),
+				FColor(231, 239, 0, 255), PointLightComp->SourceRadius, 0.5f * PointLightComp->SourceLength + PointLightComp->SourceRadius, 25, SDPG_World);
+		}
+	}
+
+	
+	// Draw the clip morph axis, planes and bounds if necessary
+	if (bClipMorphVisible && SkeletalMeshComponents.Num())
+	{
+		float MaxSphereRadius = 0.f;
+
+		for (TWeakObjectPtr<USkeletalMeshComponent> SkeletalMeshComponent : SkeletalMeshComponents)
+		{
+			if (SkeletalMeshComponent.IsValid())
+			{
+				MaxSphereRadius = FMath::Max(MaxSphereRadius, SkeletalMeshComponent->Bounds.SphereRadius);
+			}
+		}
+
+		if (MaxSphereRadius <= 0.f)
+		{
+			MaxSphereRadius = 1.f;
+		}
+
+		float PlaneRadius1 = MaxSphereRadius * 0.1f;
+		float PlaneRadius2 = PlaneRadius1 * 0.5f;
+
+		FMatrix PlaneMatrix = FMatrix(ClipMorphNormal, ClipMorphYAxis, ClipMorphXAxis, ClipMorphOrigin + ClipMorphOffset);
+
+		// Start Plane
+		DrawDirectionalArrow(PDI, PlaneMatrix, FColor::Red, MorphLength, MorphLength * 0.1f, 0, 0.1f);
+		DrawBox(PDI, PlaneMatrix, FVector(0.01f, PlaneRadius1, PlaneRadius1), Helper_GetMaterialProxy(ClipMorphMaterial), 0);
+
+		// End Plane + Ellipse
+		PlaneMatrix.SetOrigin(ClipMorphOrigin + ClipMorphOffset + ClipMorphNormal * MorphLength);
+		DrawBox(PDI, PlaneMatrix, FVector(0.01f, PlaneRadius2, PlaneRadius2), Helper_GetMaterialProxy(ClipMorphMaterial), 0);
+		DrawEllipse(PDI, ClipMorphOrigin + ClipMorphOffset + ClipMorphNormal * MorphLength, ClipMorphXAxis, ClipMorphYAxis, FColor::Red, Radius1, Radius2, 15, 1, 0.f, 0, false);
+	}
+
+	if (GizmoProxy.AnyGizmoSelected)
+	{
+		FColor Color;
+		if (GizmoProxy.AssignedDataIsFromNode)
+		{
+			Color = FColor::Red;
+		}
+		else
+		{
+			Color = FColor::Emerald;
+		}
+
+		switch (GizmoProxy.ProjectionType)
+		{
+			case ECustomizableObjectProjectorType::Planar:
+			{
+				FVector Scale = FVector(GizmoProxy.Value.Scale.Z, GizmoProxy.Value.Scale.X, GizmoProxy.Value.Scale.Y);
+				FVector Min = FVector(0.f, -0.5f, -0.5f);
+				FVector Max = FVector(1.0f, 0.5f, 0.5f);
+				FBox Box = FBox(Min * Scale, Max * Scale);
+				FMatrix Mat = GetWidgetCoordSystem();
+				Mat.SetOrigin(GetWidgetLocation());
+				DrawWireBox(PDI, Mat, Box, Color, 1, 0.f);
+				break;
+			}
+			case ECustomizableObjectProjectorType::Cylindrical:
+			{
+				// Draw the cylinder
+				FVector Scale = FVector(GizmoProxy.Value.Scale.Z, GizmoProxy.Value.Scale.X, GizmoProxy.Value.Scale.Y);
+				FMatrix Mat = GetWidgetCoordSystem();
+				FVector Location = GetWidgetLocation();
+				Mat.SetOrigin(Location);
+				FVector TransformedX = Mat.TransformVector(FVector(1, 0, 0));
+				FVector TransformedY = Mat.TransformVector(FVector(0, 1, 0));
+				FVector TransformedZ = Mat.TransformVector(FVector(0, 0, 1));
+
+				FVector Min = FVector(0.f, -0.5f, -0.5f);
+				FVector Max = FVector(1.0f, 0.5f, 0.5f);
+				FBox Box = FBox(Min * Scale, Max * Scale);
+				FVector BoxExtent = Box.GetExtent();
+				float CylinderHalfHeight = BoxExtent.X;
+				//float CylinderRadius = (BoxExtent.Y + BoxExtent.Z) * 0.5f;
+				float CylinderRadius = FMath::Abs(BoxExtent.Y);
+
+				DrawWireCylinder(PDI, Location + TransformedX * CylinderHalfHeight, TransformedY, TransformedZ, TransformedX, Color, CylinderRadius, CylinderHalfHeight, 16, SDPG_World, 0.1f, 0, false);
+
+				// Draw the arcs: the locations are Location with an offset towards the local forward direction
+				FVector Location0 = Location - TransformedX * CylinderHalfHeight * 0.8f + TransformedX * CylinderHalfHeight;
+				FVector Location1 = Location + TransformedX * CylinderHalfHeight * 0.8f + TransformedX * CylinderHalfHeight;
+				FMatrix Mat0 = Mat;
+				FMatrix Mat1 = Mat;
+				Mat0.SetOrigin(Location0);
+				Mat1.SetOrigin(Location1);
+				DrawCylinderArc(PDI, Mat0, FVector(0.0f, 0.0f, 0.0f), FVector(0, 1, 0), FVector(0, 0, 1), FVector(1, 0, 0),  CylinderRadius, CylinderHalfHeight * 0.1f, 16, Helper_GetMaterialProxy(TransparentPlaneMaterialXY), SDPG_World, FColor(255, 85, 0, 192), GizmoProxy.Value.Angle);
+				DrawCylinderArc(PDI, Mat1, FVector(0.0f, 0.0f, 0.0f), FVector(0, 1, 0), FVector(0, 0, 1), FVector(1, 0, 0), CylinderRadius, CylinderHalfHeight * 0.1f, 16, Helper_GetMaterialProxy(TransparentPlaneMaterialXY), SDPG_World, FColor(255, 85, 0, 192), GizmoProxy.Value.Angle);
+				break;
+			}
+			case ECustomizableObjectProjectorType::Wrapping:
+            {
+                FVector Scale = FVector(GizmoProxy.Value.Scale.Z, GizmoProxy.Value.Scale.X, GizmoProxy.Value.Scale.Y);
+                FVector Min = FVector(0.f, -0.5f, -0.5f);
+                FVector Max = FVector(1.0f, 0.5f, 0.5f);
+                FBox Box = FBox(Min * Scale, Max * Scale);
+                FMatrix Mat = GetWidgetCoordSystem();
+                Mat.SetOrigin(GetWidgetLocation());
+                DrawWireBox(PDI, Mat, Box, Color, 1, 0.f);
+				break;
+			}
+			default:
+			{
+				UE_LOG(LogMutable, Warning, TEXT("ERROR: wrong projector type for projector %d"), *GizmoProxy.ProjectorParameterName);
+				break;
+			}
+		}
+	}
+}
+
+
+void FCustomizableObjectEditorViewportClient::Draw(FViewport* InViewport, FCanvas* Canvas)
+{
+	// Defensive check to avoid unreal crashing inside render if the mesh is degenereated
+	for (TWeakObjectPtr<USkeletalMeshComponent> SkeletalMeshComponent : SkeletalMeshComponents)
+	{
+		if (SkeletalMeshComponent.IsValid() && SkeletalMeshComponent->GetSkeletalMesh() && Helper_GetLODInfoArray(SkeletalMeshComponent->GetSkeletalMesh()).Num() == 0)
+		{
+			SkeletalMeshComponent->SetSkeletalMesh(nullptr);
+		}
+	}
+
+	FEditorViewportClient::Draw(InViewport, Canvas);
+
+	FSceneViewFamilyContext ViewFamily( FSceneViewFamily::ConstructionValues( InViewport, GetScene(), EngineShowFlags ));
+	FSceneView* View = CalcSceneView(&ViewFamily);
+
+	const int32 HalfX = InViewport->GetSizeXY().X/2;
+	const int32 HalfY = InViewport->GetSizeXY().Y/2;
+
+	//int32 CurrentLODLevel = CustomizableObjectEditorPtr.Pin()->GetCurrentLODIndex();
+
+	int32 YPos = 6;
+	//if ( StaticMeshComponent.IsValid() || SkeletalMeshComponent.IsValid() )
+	//{
+		/* TODO
+		int32 NumVertices = 0;
+		int32 NumTriangles = 0;
+		int32 NumUVChannels = 0;
+		FBoxSphereBounds Bounds(ForceInit);
+		if (StaticMeshComponent && StaticMeshComponent->StaticMesh)
+		{
+			Bounds = StaticMeshComponent->StaticMesh->Bounds;
+			NumTriangles = StaticMeshComponent->StaticMesh->LODModels[CurrentLODLevel].IndexBuffer.Indices.Num()/3;
+			NumVertices = StaticMeshComponent->StaticMesh->LODModels[CurrentLODLevel].VertexBuffer.GetNumVertices();
+			NumUVChannels = StaticMeshComponent->StaticMesh->LODModels[CurrentLODLevel].VertexBuffer.GetNumTexCoords();
+		}
+		else if (SkeletalMeshComponent && SkeletalMeshComponent->SkeletalMesh)
+		{
+			Bounds = SkeletalMeshComponent->SkeletalMesh->Bounds;
+			NumVertices = SkeletalMeshComponent->SkeletalMesh->LODModels[CurrentLODLevel].NumVertices;
+			NumTriangles = SkeletalMeshComponent->SkeletalMesh->LODModels[CurrentLODLevel].MultiSizeIndexContainer.GetIndexBuffer()->Num()/3;
+			NumUVChannels = SkeletalMeshComponent->SkeletalMesh->LODModels[CurrentLODLevel].NumTexCoords;
+		}
+
+		DrawShadowedString(Canvas,
+			6,
+			YPos,
+			*FString::Printf(LocalizeSecure(LocalizeUnrealEd("Triangles_F"), NumTriangles )),
+			GEngine->GetSmallFont(),
+			FLinearColor::White
+			);
+		YPos += 18;
+
+		DrawShadowedString(Canvas,
+			6,
+			YPos,
+			*FString::Printf(LocalizeSecure(LocalizeUnrealEd("Vertices_F"), NumVertices )),
+			GEngine->GetSmallFont(),
+			FLinearColor::White
+			);
+		YPos += 18;
+
+		DrawShadowedString(Canvas,
+			6,
+			YPos,
+			*FString::Printf(LocalizeSecure(LocalizeUnrealEd("UVChannels_F"), NumUVChannels )),
+			GEngine->GetSmallFont(),
+			FLinearColor::White
+			);
+		YPos += 18;
+
+
+		DrawShadowedString(Canvas,
+			6,
+			YPos,
+			*FString::Printf( LocalizeSecure( LocalizeUnrealEd("ApproxSize_F"), int32(Bounds.BoxExtent.X * 2.0f),
+			int32(Bounds.BoxExtent.Y * 2.0f),
+			int32(Bounds.BoxExtent.Z * 2.0f) ) ),
+			GEngine->GetSmallFont(),
+			FLinearColor::White
+			);
+			*/
+	//}
+	YPos += 18;
+
+	if(bDrawUVs)
+	{
+		DrawUVs(InViewport, Canvas, YPos, MaterialToDrawInUVs);
+	}
+
+	if (bReferenceMeshMissingWarningMessageVisible)
+	{
+		Canvas->DrawShadowedString(
+			6,
+			2,
+			*NSLOCTEXT("UnrealEd", "UVOverlay_F", "Warning! No reference mesh is set in the Object Properties tab.").ToString(),
+			GEngine->GetSmallFont(),
+			FLinearColor::Red
+			);
+	}
+
+	if (StateChangeShowGeometryDataFlag)
+	{
+		ShowInstanceGeometryInformation(Canvas);
+	}
+}
+
+namespace
+{
+	template<typename Real>
+	FVector2D ClampUVRange(Real U, Real V)
+	{
+		return FVector2D( FMath::Wrap(U, Real(0), Real(1)), FMath::Wrap(V, Real(0), Real(1)) );
+	}
+}
+
+void FCustomizableObjectEditorViewportClient::DrawUVs(FViewport* InViewport, FCanvas* InCanvas, int32 InTextYPos, const FString& MaterialName )
+{
+	//use the overriden LOD level
+	// TODO
+	const uint32 LODLevel = MaterialToDrawInUVsLOD; //FMath::Clamp(StaticMeshComponent->ForcedLodModel - 1, 0, StaticMesh->RenderData->LODResources.Num() - 1);
+
+	// TODO
+	int32 UVChannel = UVChannelToDrawInUVs; //StaticMeshEditorPtr.Pin()->GetCurrentUVChannel();
+
+	const uint32 ComponentIndex = MaterialToDrawInUVsComponent;
+
+	//draw a string showing what UV channel and LOD is being displayed
+	InCanvas->DrawShadowedString( 
+		6,
+		InTextYPos,
+		*FText::Format( NSLOCTEXT("UnrealEd", "UVOverlay_F", "Showing UV channel {0} for LOD {1}"), FText::AsNumber(UVChannel), FText::AsNumber(LODLevel) ).ToString(),
+		GEngine->GetSmallFont(),
+		FLinearColor::White
+		);
+	InTextYPos += 18;
+
+	//calculate scaling
+	const uint32 BorderWidth = 5;
+	const uint32 MinY = InTextYPos + BorderWidth;
+	const uint32 MinX = BorderWidth;
+	const FVector2D UVBoxOrigin(MinX, MinY);
+	const FVector2D BoxOrigin( MinX - 1, MinY - 1 );
+	const uint32 UVBoxScale = FMath::Min(InViewport->GetSizeXY().X - MinX, InViewport->GetSizeXY().Y - MinY) - BorderWidth;
+	const uint32 BoxSize = UVBoxScale + 2;
+	const FVector2D Box[ 4 ] = {
+		BoxOrigin,									// topleft
+		BoxOrigin + FVector2D( BoxSize, 0 ),		// topright
+		BoxOrigin + FVector2D( BoxSize, BoxSize ),	// bottomright
+		BoxOrigin + FVector2D( 0, BoxSize ),		// bottomleft
+	};
+
+	const FVector Color(1.0f, 1.0f, 1.0f);
+
+	//draw texture border
+	FLinearColor BorderColor = FLinearColor::White;
+	FBatchedElements* BatchedElements = InCanvas->GetBatchedElements(FCanvas::ET_Line);
+	FHitProxyId HitProxyId = InCanvas->GetHitProxyId();
+
+	// Reserve line vertices (4 border lines, then up to the maximum number of graph lines)
+	BatchedElements->AddReserveLines( 4 );
+
+	// Left
+	BatchedElements->AddLine( FVector( Box[ 0 ], 0.0f ), FVector( Box[ 1 ], 0.0f ), BorderColor, HitProxyId );
+	BatchedElements->AddLine( FVector( Box[ 1 ], 0.0f ), FVector( Box[ 2 ], 0.0f ), BorderColor, HitProxyId );
+	BatchedElements->AddLine( FVector( Box[ 2 ], 0.0f ), FVector( Box[ 3 ], 0.0f ), BorderColor, HitProxyId );
+	BatchedElements->AddLine( FVector( Box[ 3 ], 0.0f ), FVector( Box[ 0 ], 0.0f ), BorderColor, HitProxyId );
+
+	if ( StaticMeshComponent.IsValid() && StaticMeshComponent->GetStaticMesh() && StaticMeshComponent->GetStaticMesh()->GetRenderData() )
+	{
+		FStaticMeshLODResources* RenderData = &StaticMeshComponent->GetStaticMesh()->GetRenderData()->LODResources[LODLevel];
+		
+		if( RenderData && ( ( uint32 )UVChannel < RenderData->VertexBuffers.StaticMeshVertexBuffer.GetNumTexCoords() ) )
+		{
+			//draw triangles
+			FIndexArrayView Indices = RenderData->IndexBuffer.GetArrayView();
+			uint32 NumIndices = Indices.Num();
+		
+			BatchedElements->AddReserveLines( NumIndices );
+
+			for (uint32 i = 0; i < NumIndices - 2; i += 3)
+			{
+				FVector2D UV1( RenderData->VertexBuffers.StaticMeshVertexBuffer.GetVertexUV( Indices[ i + 0 ], UVChannel ) );
+				FVector2D UV2( RenderData->VertexBuffers.StaticMeshVertexBuffer.GetVertexUV( Indices[ i + 1 ], UVChannel ) );
+				FVector2D UV3( RenderData->VertexBuffers.StaticMeshVertexBuffer.GetVertexUV( Indices[ i + 2 ], UVChannel ) );
+	
+				// Draw lines in black unless the UVs are outside of the 0.0 - 1.0 range.  For out-of-bounds
+				// UVs, we'll draw the line segment in red
+				
+				// If we are supporting a version lower than LWC get the right real type. 
+				using Vector2DRealType = TDecay<decltype( DeclVal<FVector2D>().X )>::Type; 
+
+				constexpr Vector2DRealType Epsilon = static_cast<Vector2DRealType>(1e-4);
+				constexpr Vector2DRealType One     = static_cast<Vector2DRealType>(1);
+				constexpr Vector2DRealType Zero    = static_cast<Vector2DRealType>(0);
+
+				FLinearColor UV12LineColor = FLinearColor::Black;
+				if (UV1.X < -Epsilon || UV1.X > One + Epsilon ||
+					UV2.X < -Epsilon || UV2.X > One + Epsilon ||
+					UV1.Y < -Epsilon || UV1.Y > One + Epsilon ||
+					UV2.Y < -Epsilon || UV2.Y > One + Epsilon)
+				{
+					UV12LineColor = FLinearColor(0.6f, 0.0f, 0.0f);
+				}
+				FLinearColor UV23LineColor = FLinearColor::Black;
+				if (UV3.X < -Epsilon || UV3.X > One + Epsilon ||
+					UV2.X < -Epsilon || UV2.X > One + Epsilon ||
+					UV3.Y < -Epsilon || UV3.Y > One + Epsilon ||
+					UV2.Y < -Epsilon || UV2.Y > One + Epsilon)
+				{
+					UV23LineColor = FLinearColor(0.6f, 0.0f, 0.0f);
+				}
+				FLinearColor UV31LineColor = FLinearColor::Black;
+				if (UV3.X < -Epsilon || UV3.X > One + Epsilon ||
+					UV1.X < -Epsilon || UV1.X > One + Epsilon ||
+					UV3.Y < -Epsilon || UV3.Y > One + Epsilon ||
+					UV1.Y < -Epsilon || UV1.Y > One + Epsilon)
+				{
+					UV31LineColor = FLinearColor(0.6f, 0.0f, 0.0f);
+				}
+
+				UV1 = ClampUVRange(UV1.X, UV1.Y) * UVBoxScale + UVBoxOrigin;
+				UV2 = ClampUVRange(UV2.X, UV2.Y) * UVBoxScale + UVBoxOrigin;
+				UV3 = ClampUVRange(UV3.X, UV3.Y) * UVBoxScale + UVBoxOrigin;
+
+				BatchedElements->AddLine(FVector(UV1, Zero), FVector(UV2, Zero), BorderColor, HitProxyId);
+				BatchedElements->AddLine(FVector(UV2, Zero), FVector(UV3, Zero), BorderColor, HitProxyId);
+				BatchedElements->AddLine(FVector(UV3, Zero), FVector(UV1, Zero), BorderColor, HitProxyId);
+			}
+		}
+	}
+
+	else if (SkeletalMeshComponents.Num())
+	{
+		int32 CurrentComponentIndex = 0;
+
+		for (TWeakObjectPtr<USkeletalMeshComponent> SkeletalMeshComponent : SkeletalMeshComponents)
+		{
+			if (!SkeletalMeshComponent.IsValid() || !SkeletalMeshComponent->GetSkeletalMesh() || CurrentComponentIndex != ComponentIndex)
+			{
+				CurrentComponentIndex++;
+				continue;
+			}
+
+			bool bFoundMaterial = false;
+
+			const FSkeletalMeshRenderData* MeshRes = SkeletalMeshComponent->GetSkeletalMesh()->GetResourceForRendering();
+			if (UVChannel < (int32)MeshRes->LODRenderData[LODLevel].GetNumTexCoords())
+			{
+				// Find material index from name
+				const FSkeletalMeshLODRenderData& lodModel = MeshRes->LODRenderData[LODLevel];
+				int MaterialIndex = 0;
+				int MaterialIndexCount = 0;
+				{
+					for (int s = 0; s < lodModel.RenderSections.Num(); ++s)
+					{
+						int SectionMaterial = lodModel.RenderSections[s].MaterialIndex;
+						UMaterialInterface* Material = SkeletalMeshComponent->GetMaterial(SectionMaterial);
+
+						if (!Material)
+						{
+							continue;
+						}
+
+						const UMaterial* BaseMaterial = Material->GetBaseMaterial();
+						if (BaseMaterial && BaseMaterial->GetName() == MaterialName)
+						{
+							MaterialIndex = s;
+
+							if (MaterialIndexCount == MaterialToDrawInUVsIndex)
+							{
+								bFoundMaterial = true;
+								break;
+							}
+							MaterialIndexCount++;
+						}
+					}
+				}
+
+				if (!bFoundMaterial)
+				{
+					continue;
+				}
+
+				const FStaticMeshVertexBuffer& Vertices = lodModel.StaticVertexBuffers.StaticMeshVertexBuffer;
+
+				TArray<uint32> Indices;
+				lodModel.MultiSizeIndexContainer.GetIndexBuffer(Indices);
+
+				uint32 NumTriangles = lodModel.RenderSections[MaterialIndex].NumTriangles;
+				int IndexIndex = lodModel.RenderSections[MaterialIndex].BaseIndex;
+
+				BatchedElements->AddReserveLines(NumTriangles * 3);
+
+				for (uint32 FaceIndex = 0
+					; FaceIndex < NumTriangles
+					; ++FaceIndex, IndexIndex += 3)
+				{
+					FVector2D UV1(Vertices.GetVertexUV(Indices[IndexIndex + 0], UVChannel));
+					FVector2D UV2(Vertices.GetVertexUV(Indices[IndexIndex + 1], UVChannel));
+					FVector2D UV3(Vertices.GetVertexUV(Indices[IndexIndex + 2], UVChannel));
+
+					// Draw lines in black unless the UVs are outside of the 0.0 - 1.0 range.  For out-of-bounds
+					// UVs, we'll draw the line segment in red
+
+					// If we are supporting a version lower than LWC get the right real type. 
+					using Vector2DRealType = TDecay<decltype(DeclVal<FVector2D>().X)>::Type;
+
+					constexpr Vector2DRealType Epsilon = static_cast<Vector2DRealType>(1e-4);
+					constexpr Vector2DRealType One = static_cast<Vector2DRealType>(1);
+					constexpr Vector2DRealType Zero = static_cast<Vector2DRealType>(0);
+
+					FLinearColor UV12LineColor = FLinearColor::Black;
+					if (UV1.X < -Epsilon || UV1.X > One + Epsilon ||
+						UV2.X < -Epsilon || UV2.X > One + Epsilon ||
+						UV1.Y < -Epsilon || UV1.Y > One + Epsilon ||
+						UV2.Y < -Epsilon || UV2.Y > One + Epsilon)
+					{
+						UV12LineColor = FLinearColor(0.6f, 0.0f, 0.0f);
+					}
+					FLinearColor UV23LineColor = FLinearColor::Black;
+					if (UV3.X < -Epsilon || UV3.X > One + Epsilon ||
+						UV2.X < -Epsilon || UV2.X > One + Epsilon ||
+						UV3.Y < -Epsilon || UV3.Y > One + Epsilon ||
+						UV2.Y < -Epsilon || UV2.Y > One + Epsilon)
+					{
+						UV23LineColor = FLinearColor(0.6f, 0.0f, 0.0f);
+					}
+					FLinearColor UV31LineColor = FLinearColor::Black;
+					if (UV3.X < -Epsilon || UV3.X > One + Epsilon ||
+						UV1.X < -Epsilon || UV1.X > One + Epsilon ||
+						UV3.Y < -Epsilon || UV3.Y > One + Epsilon ||
+						UV1.Y < -Epsilon || UV1.Y > One + Epsilon)
+					{
+						UV31LineColor = FLinearColor(0.6f, 0.0f, 0.0f);
+					}
+
+					UV1 = ClampUVRange(UV1.X, UV1.Y) * UVBoxScale + UVBoxOrigin;
+					UV2 = ClampUVRange(UV2.X, UV2.Y) * UVBoxScale + UVBoxOrigin;
+					UV3 = ClampUVRange(UV3.X, UV3.Y) * UVBoxScale + UVBoxOrigin;
+
+					BatchedElements->AddLine( FVector(UV1, Zero), FVector(UV2, Zero), BorderColor, HitProxyId );
+					BatchedElements->AddLine( FVector(UV2, Zero), FVector(UV3, Zero), BorderColor, HitProxyId );
+					BatchedElements->AddLine( FVector(UV3, Zero), FVector(UV1, Zero), BorderColor, HitProxyId );
+				}
+			}
+
+			if (bFoundMaterial)
+			{
+				break;
+			}
+		}
+	}
+}
+
+
+float FCustomizableObjectEditorViewportClient::GetFloorOffset() const
+{
+	FAdvancedPreviewScene* AdvancedScene = static_cast<FAdvancedPreviewScene*>(PreviewScene);
+	if (AdvancedScene != nullptr)
+	{
+		const UStaticMeshComponent* FloorMeshComponent = AdvancedScene->GetFloorMeshComponent();
+		if (FloorMeshComponent != nullptr)
+		{
+			return FloorMeshComponent->GetComponentLocation().Z;
+		}
+	}
+
+	return 0.0f;
+}
+
+
+void FCustomizableObjectEditorViewportClient::SetClipMorphPlaneVisibility(bool bVisible, const FVector& Origin, const FVector& Normal, float InMorphLength, const FBoxSphereBounds& Bounds, float InRadius1, float InRadius2, float InRotationAngle)
+{
+	/*
+	//EditorClipMorphPlaneComp->SetVisibility(bVisible);
+	bClipMorphVisible = bVisible;
+	//EditorClipMorphPlaneComp->SetWorldLocationAndRotation(Origin, Normal.Rotation());
+	ClipMorphOrigin = Origin;
+	ClipMorphNormal = Normal;
+	MorphLength = InMorphLength;
+	MorphBounds = Bounds;
+	Radius1 = InRadius1;
+	Radius2 = InRadius2;
+	RotationAngle = InRotationAngle;
+	*/
+}
+
+
+void FCustomizableObjectEditorViewportClient::SetClipMorphPlaneVisibility(bool bVisible, UCustomizableObjectNodeMeshClipMorph* NodeMeshClipMorph)
+{
+	if (!bClipMorphVisible && !bVisible)
+	{
+		return;
+	}
+
+	bClipMorphVisible = bVisible;
+
+	if (NodeMeshClipMorph && (ClipMorphNode != NodeMeshClipMorph || NodeMeshClipMorph->bUpdateViewportWidget))
+	{
+		NodeMeshClipMorph->bUpdateViewportWidget = false;
+
+		bClipMorphLocalStartOffset = NodeMeshClipMorph->bLocalStartOffset;
+		MorphLength = NodeMeshClipMorph->B;
+		Radius1 = NodeMeshClipMorph->Radius;
+		Radius2 = NodeMeshClipMorph->Radius2;
+		RotationAngle = NodeMeshClipMorph->RotationAngle;
+		ClipMorphOrigin = NodeMeshClipMorph->Origin;
+		ClipMorphLocalOffset = NodeMeshClipMorph->StartOffset;
+
+		NodeMeshClipMorph->FindLocalAxes(ClipMorphXAxis, ClipMorphYAxis, ClipMorphNormal);
+
+		if (bClipMorphLocalStartOffset)
+		{
+			ClipMorphOffset = ClipMorphLocalOffset.X * ClipMorphXAxis
+				+ ClipMorphLocalOffset.Y * ClipMorphYAxis
+				+ ClipMorphLocalOffset.Z * ClipMorphNormal;
+		}
+		else
+		{
+			ClipMorphOffset = ClipMorphLocalOffset;
+		}
+	}
+
+	ClipMorphNode = NodeMeshClipMorph;
+
+	if (WidgetVisibility != bClipMorphVisible)
+	{
+		Widget->SetDefaultVisibility(bClipMorphVisible);
+		WidgetVisibility = bClipMorphVisible;
+	}
+}
+
+
+void FCustomizableObjectEditorViewportClient::SetClipMeshVisibility(bool bVisible, UStaticMesh* ClipMesh, UCustomizableObjectNodeMeshClipWithMesh* MeshNode)
+{
+	if (!ClipMeshNode && !MeshNode)
+	{
+		return;
+	}
+
+	ClipMeshComp->SetStaticMesh(ClipMesh);
+	ClipMeshComp->SetVisibility(bVisible);
+	ClipMeshNode = bVisible ? MeshNode : nullptr;
+
+	if (ClipMeshNode)
+	{
+		ClipMeshComp->SetWorldTransform(ClipMeshNode->Transform);
+	}
+
+	if (WidgetVisibility != bVisible)
+	{
+		Widget->SetDefaultVisibility(bVisible);
+		WidgetVisibility = bVisible;
+	}
+}
+
+
+void FCustomizableObjectEditorViewportClient::SetProjectorVisibility(bool bVisible, FString ProjectorParameterName, FString ProjectorParameterNameWithIndex, int32 ProjectorRangeIndex, const FCustomizableObjectProjector& Data, int32 ProjectorParameterIndex)
+{
+	if (GizmoProxy.ProjectorParameterIndex == ProjectorParameterIndex && GizmoProxy.ProjectorRangeIndex == ProjectorRangeIndex)
+	{
+		return;
+	}
+
+	if (GizmoProxy.HasAssignedData)
+	{
+		GizmoProxy.UpdateOriginData();
+		GizmoProxy.CleanOriginData();
+	}
+	GizmoProxy.ProjectorParameterName = ProjectorParameterName;
+	GizmoProxy.ProjectorParameterNameWithIndex = ProjectorParameterNameWithIndex;
+	GizmoProxy.ProjectorRangeIndex = ProjectorRangeIndex;
+	GizmoProxy.Value = Data;
+	GizmoProxy.ProjectionType = Data.ProjectionType;
+
+	if (GizmoProxy.ProjectorHasInitialValues(GizmoProxy.Value))
+	{
+		GizmoProxy.Value = GizmoProxy.SetProjectorInitialValue(SkeletalMeshComponents, BoundSphere.W * 2.0f);
+	}
+
+	GizmoProxy.ProjectorParameterIndex = ProjectorParameterIndex;
+	GizmoProxy.AnyGizmoSelected = true;
+	GizmoProxy.HasAssignedData = true;
+
+	if (!WidgetVisibility)
+	{
+		Widget->SetDefaultVisibility(true);
+		WidgetVisibility = true;
+	}
+}
+
+
+void FCustomizableObjectEditorViewportClient::SetProjectorType(bool bVisible, FString ProjectorParameterName, FString ProjectorParameterNameWithIndex, int32 ProjectorRangeIndex, const FCustomizableObjectProjector& Data, int32 ProjectorParameterIndex)
+{
+	GizmoProxy.ProjectorParameterName = ProjectorParameterName;
+	GizmoProxy.ProjectorParameterNameWithIndex = ProjectorParameterNameWithIndex;
+	GizmoProxy.ProjectorRangeIndex = ProjectorRangeIndex;
+	GizmoProxy.Value = Data;
+	GizmoProxy.ProjectionType = Data.ProjectionType;
+
+	GizmoProxy.ProjectorParameterIndex = ProjectorParameterIndex;
+	GizmoProxy.AnyGizmoSelected = true;
+	GizmoProxy.HasAssignedData = true;
+}
+
+
+void FCustomizableObjectEditorViewportClient::SetProjectorVisibility(bool bVisible, UCustomizableObjectNodeProjectorConstant* InProjector)
+{
+	if (InProjector == nullptr)
+	{
+		if ((GizmoProxy.HasAssignedData) && (GizmoProxy.DataOriginConstant != nullptr))
+		{
+			GizmoProxy.UpdateOriginData();
+			GizmoProxy.CleanOriginData();
+		}
+
+		if (WidgetVisibility)
+		{
+			Widget->SetDefaultVisibility(false);
+			WidgetVisibility = false;
+		}
+	}
+	else
+	{
+		if (!GizmoProxy.HasAssignedData ||
+		   (GizmoProxy.HasAssignedData && (GizmoProxy.DataOriginConstant != InProjector)))
+		{
+			GizmoProxy.UpdateOriginData();
+			GizmoProxy.CleanOriginData();
+			GizmoProxy.DataOriginConstant = InProjector;
+			GizmoProxy.Value = InProjector->Value;
+
+			if (GizmoProxy.ProjectorHasInitialValues(GizmoProxy.Value))
+			{
+				GizmoProxy.Value = GizmoProxy.SetProjectorInitialValue(SkeletalMeshComponents, BoundSphere.W * 2.0f);
+			}
+
+			GizmoProxy.AnyGizmoSelected = true;
+			GizmoProxy.HasAssignedData = true;
+			GizmoProxy.AssignedDataIsFromNode = true;
+			GizmoProxy.ProjectionType = InProjector->ProjectionType;
+
+			if (!WidgetVisibility)
+			{
+				Widget->SetDefaultVisibility(true);
+				WidgetVisibility = true;
+			}
+		}
+	}
+}
+
+
+void FCustomizableObjectEditorViewportClient::SetProjectorParameterVisibility(bool bVisible, UCustomizableObjectNodeProjectorParameter* InProjectorParameter)
+{
+	if (InProjectorParameter == nullptr)
+	{
+		if ((GizmoProxy.HasAssignedData) && (GizmoProxy.DataOriginParameter != nullptr))
+		{
+			GizmoProxy.UpdateOriginData();
+			GizmoProxy.CleanOriginData();
+
+			if (WidgetVisibility)
+			{
+				Widget->SetDefaultVisibility(false);
+				WidgetVisibility = false;
+			}
+		}
+	}
+	else
+	{
+		if (!GizmoProxy.HasAssignedData ||
+			(GizmoProxy.HasAssignedData && (GizmoProxy.DataOriginParameter != InProjectorParameter)))
+		{
+			GizmoProxy.UpdateOriginData();
+			GizmoProxy.CleanOriginData();
+			GizmoProxy.DataOriginParameter = InProjectorParameter;
+			GizmoProxy.Value = InProjectorParameter->DefaultValue;
+
+			if (GizmoProxy.ProjectorHasInitialValues(GizmoProxy.Value))
+			{
+				GizmoProxy.Value = GizmoProxy.SetProjectorInitialValue(SkeletalMeshComponents, BoundSphere.W * 2.0f);
+			}
+
+			GizmoProxy.AnyGizmoSelected = true;
+			GizmoProxy.HasAssignedData = true;
+			GizmoProxy.AssignedDataIsFromNode = true;
+			GizmoProxy.ProjectionType = InProjectorParameter->ProjectionType;
+
+			if (!WidgetVisibility)
+			{
+				Widget->SetDefaultVisibility(true);
+				WidgetVisibility = true;
+			}
+		}
+	}
+}
+
+
+void FCustomizableObjectEditorViewportClient::ResetProjectorVisibility()
+{
+	if (bManipulating || !GizmoProxy.AnyGizmoSelected)
+	{
+		return;
+	}
+
+	GizmoProxy.UpdateOriginData();
+	GizmoProxy.CleanOriginData();
+
+	if (WidgetVisibility)
+	{
+		Widget->SetDefaultVisibility(false);
+		WidgetVisibility = false;
+	}
+}
+
+
+void FCustomizableObjectEditorViewportClient::SetVisibilityForWireframeMode(bool bIsWireframeMode)
+{
+	EngineShowFlags.SetDirectLighting(!bIsWireframeMode && bDrawSky);
+}
+
+
+FSphere FCustomizableObjectEditorViewportClient::GetCameraTarget()
+{
+	bool bFoundTarget = false;
+	FSphere Sphere(FVector(0,0,0), 100.0f); // default
+
+	for (TWeakObjectPtr<USkeletalMeshComponent> SkeletalMeshComponent : SkeletalMeshComponents)
+	{
+		if (SkeletalMeshComponent.IsValid())
+		{
+			FTransform ComponentToWorld = SkeletalMeshComponent->GetComponentTransform();
+			SkeletalMeshComponent.Get()->CalcBounds(ComponentToWorld);
+
+			FBoxSphereBounds Bounds = SkeletalMeshComponent.Get()->CalcBounds(FTransform::Identity);
+
+			if (!bFoundTarget)
+			{
+				Sphere = Bounds.GetSphere();
+			}
+			else
+			{
+				Sphere += Bounds.GetSphere();
+			}
+
+			bFoundTarget = true;
+		}
+	}
+
+	if(!bFoundTarget && StaticMeshComponent.IsValid())
+	{
+		FTransform ComponentToWorld = StaticMeshComponent->GetComponentTransform();
+		StaticMeshComponent.Get()->CalcBounds(ComponentToWorld);
+
+		if( !bFoundTarget )
+		{
+			FBoxSphereBounds Bounds = StaticMeshComponent.Get()->CalcBounds(FTransform::Identity);
+			Sphere = Bounds.GetSphere();
+		}
+	}
+
+	return Sphere;
+}
+
+
+FLinearColor FCustomizableObjectEditorViewportClient::GetBackgroundColor() const
+{
+	FLinearColor BackgroundColor = FColor(55, 55, 55);
+
+	return BackgroundColor;
+}
+
+
+void FCustomizableObjectEditorViewportClient::SetPreviewComponent(UStaticMeshComponent* InStaticMeshComponent)
+{
+	StaticMeshComponent = InStaticMeshComponent;
+	SkeletalMeshComponents.Reset();
+
+	if (StaticMeshComponent.IsValid() && StaticMeshComponent->GetStaticMesh())
+	{
+		SetViewLocation( -FVector(0, StaticMeshComponent->GetStaticMesh()->GetBounds().SphereRadius / (75.0f * (float)PI / 360.0f), 0) );
+		SetViewRotation( FRotator(0, 90.f, 0) );
+		//LockLocation = FVector(0,StaticMeshComponent->StaticMesh->ThumbnailDistance,0);
+		//LockRot = StaticMeshComponent->StaticMesh->ThumbnailAngle;
+	}
+
+	UpdateCameraSetup();
+}
+
+
+void FCustomizableObjectEditorViewportClient::SetPreviewComponents(TArray<UDebugSkelMeshComponent*>& InSkeletalMeshComponents)
+{
+	SkeletalMeshComponents.Reset(InSkeletalMeshComponents.Num());
+	
+	for (UDebugSkelMeshComponent* SkeletalMeshComponent : InSkeletalMeshComponents)
+	{
+		SkeletalMeshComponents.Add(SkeletalMeshComponent);
+
+	}
+
+	StaticMeshComponent = nullptr;
+}
+
+void FCustomizableObjectEditorViewportClient::ResetCamera()
+{
+	float MaxSphereRadius = 0.0f;
+	for (const TWeakObjectPtr<UDebugSkelMeshComponent>& SkeletalMeshComponent : SkeletalMeshComponents)
+	{
+		if (SkeletalMeshComponent->GetSkeletalMesh())
+		{
+			MaxSphereRadius = FMath::Max(MaxSphereRadius, SkeletalMeshComponent->GetSkeletalMesh()->GetBounds().SphereRadius);
+		}
+	}
+	
+	SetViewLocation(-FVector(0, MaxSphereRadius / (75.0f * (float)PI / 360.0f), 0));
+	SetViewRotation(FRotator(0, 90.f, 0));
+
+	UpdateCameraSetup();
+}
+
+
+void FCustomizableObjectEditorViewportClient::SetReferenceMeshMissingWarningMessage(bool bVisible)
+{
+	bReferenceMeshMissingWarningMessageVisible = bVisible;
+}
+
+
+void FCustomizableObjectEditorViewportClient::SetDrawUVOverlay()
+{
+	bDrawUVs = !bDrawUVs;
+	Invalidate();
+}
+
+
+void FCustomizableObjectEditorViewportClient::SetDrawUVOverlayMaterial(const FString& MaterialName, FString UVChannel)
+{
+	// Get LOD Index
+	FString NameWithLOD, ComponentString;
+	MaterialName.Split(FString("_Component_"), &NameWithLOD, &ComponentString);
+	check(ComponentString.IsNumeric());
+	MaterialToDrawInUVsComponent = FCString::Atoi(*ComponentString);
+
+	FString Name, LODIndex;
+	bool bSplit = NameWithLOD.Split(FString(" LOD_"), &Name, &LODIndex);
+	check(bSplit && LODIndex.IsNumeric());
+
+	MaterialToDrawInUVsLOD = FCString::Atoi(*LODIndex);
+	UVChannelToDrawInUVs = FCString::Atoi(*UVChannel);
+
+	// Get Material Index, added if the name of the material already exists within the skeletal mesh.
+	FString DuplicatedMaterialIndex;
+	bSplit = Name.Split(FString("__"), &MaterialToDrawInUVs, &DuplicatedMaterialIndex);
+
+	if (bSplit && DuplicatedMaterialIndex.IsNumeric())
+	{
+		MaterialToDrawInUVsIndex = FCString::Atoi(*DuplicatedMaterialIndex);
+	}
+	else
+	{
+		MaterialToDrawInUVs = Name;
+		MaterialToDrawInUVsIndex = 0;
+	}
+
+	Invalidate();
+}
+
+bool FCustomizableObjectEditorViewportClient::IsSetDrawUVOverlayChecked() const
+{
+	return bDrawUVs;
+}
+
+void FCustomizableObjectEditorViewportClient::SetShowGrid()
+{
+	DrawHelper.bDrawGrid = !DrawHelper.bDrawGrid;
+
+	FAdvancedPreviewScene* AdvancedScene = static_cast<FAdvancedPreviewScene*>(PreviewScene);
+	if (AdvancedScene != nullptr)
+	{
+		AdvancedScene->SetFloorVisibility(DrawHelper.bDrawGrid,true);
+	}
+
+	EngineShowFlags.Grid = DrawHelper.bDrawGrid;
+
+	Invalidate();
+}
+
+bool FCustomizableObjectEditorViewportClient::IsSetShowGridChecked() const
+{
+	return DrawHelper.bDrawGrid;
+}
+
+void FCustomizableObjectEditorViewportClient::SetShowSky()
+{
+	bDrawSky = !bDrawSky;
+	FAdvancedPreviewScene* PreviewSceneCasted = static_cast<FAdvancedPreviewScene*>(PreviewScene);
+	PreviewSceneCasted->SetEnvironmentVisibility(bDrawSky, true);
+	Invalidate();
+}
+
+bool FCustomizableObjectEditorViewportClient::IsSetShowSkyChecked() const
+{
+	return bDrawSky;
+}
+
+void FCustomizableObjectEditorViewportClient::SetShowBounds()
+{
+	EngineShowFlags.Bounds = 1 - EngineShowFlags.Bounds;
+	Invalidate();
+}
+
+bool FCustomizableObjectEditorViewportClient::IsSetShowBoundsChecked() const
+{
+	return EngineShowFlags.Bounds ? true : false;
+}
+
+void FCustomizableObjectEditorViewportClient::SetShowCollision()
+{
+	EngineShowFlags.Collision = 1 - EngineShowFlags.Collision;
+	Invalidate();
+}
+
+bool FCustomizableObjectEditorViewportClient::IsSetShowCollisionChecked() const
+{
+	return EngineShowFlags.Collision ? true : false;
+}
+
+void FCustomizableObjectEditorViewportClient::SetShowPivot()
+{
+	bShowPivot = !bShowPivot;
+	Widget->SetDefaultVisibility(bShowPivot);
+	WidgetVisibility = bShowPivot;
+	Invalidate();
+}
+
+bool FCustomizableObjectEditorViewportClient::IsSetShowPivotChecked() const
+{
+	return bShowPivot;
+}
+
+
+bool FCustomizableObjectEditorViewportClient::InputKey(const FInputKeyEventArgs& EventArgs)
+{
+	const int32 HitX = EventArgs.Viewport->GetMouseX();
+	const int32 HitY = EventArgs.Viewport->GetMouseY();
+
+	const bool bMouseButtonDown = EventArgs.Viewport->KeyState(EKeys::LeftMouseButton) || EventArgs.Viewport->KeyState(EKeys::MiddleMouseButton) || EventArgs.Viewport->KeyState(EKeys::RightMouseButton);
+
+	bool bHandled = false;
+
+	if (EventArgs.Event == IE_Pressed && !bManipulating && !bMouseButtonDown)
+	{
+		if (EventArgs.Key == EKeys::F)
+		{
+			bHandled = true;
+			UpdateCameraSetup();
+		}
+		else if (EventArgs.Key == EKeys::W)
+		{
+			bHandled = true;
+			WidgetMode = UE::Widget::WM_Translate;
+		}
+		else if (EventArgs.Key == EKeys::E)
+		{
+			bHandled = true;
+			WidgetMode = UE::Widget::WM_Rotate;
+		}
+		else if (EventArgs.Key == EKeys::R)
+		{
+			bHandled = true;
+			WidgetMode = UE::Widget::WM_Scale;
+		}
+	}
+
+	if (EventArgs.Event == IE_Released && bManipulating)
+	{
+		GizmoProxy.SetProjectorUpdatedInViewport(false);
+	}
+
+	// Pass keys to standard controls, if we didn't consume input
+	return (bHandled)
+		? true
+		: FEditorViewportClient::InputKey(EventArgs);
+}
+
+
+bool FCustomizableObjectEditorViewportClient::InputWidgetDelta(FViewport* InViewport, EAxisList::Type CurrentAxis, FVector& Drag, FRotator& Rot, FVector& Scale)
+{
+	if (!GizmoProxy.AnyGizmoSelected && !bClipMorphVisible && !ClipMeshNode && !bIsEditingLightEnabled)
+	{
+		return false;
+	}
+
+	// Get some useful info about buttons being held down
+	const bool bCtrlDown = InViewport->KeyState(EKeys::LeftControl) || InViewport->KeyState(EKeys::RightControl);
+	const bool bShiftDown = InViewport->KeyState(EKeys::LeftShift) || InViewport->KeyState(EKeys::RightShift);
+	const bool bMouseButtonDown = InViewport->KeyState(EKeys::LeftMouseButton) || InViewport->KeyState(EKeys::MiddleMouseButton) || InViewport->KeyState(EKeys::RightMouseButton);
+
+	bool bHandled = false;
+
+	if (bManipulating && CurrentAxis != EAxisList::None)
+	{
+		bHandled = true;
+
+		if (GizmoProxy.AnyGizmoSelected)
+		{
+			if (WidgetMode == UE::Widget::WM_Translate)
+			{
+				GizmoProxy.Value.Position += (FVector3f)Drag;
+				GizmoProxy.ProjectorGizmoEdited = true;
+			}
+			else if (WidgetMode == UE::Widget::WM_Rotate)
+			{
+				GizmoProxy.Value.Direction = (FVector3f)Rot.RotateVector((FVector)GizmoProxy.Value.Direction);
+				GizmoProxy.Value.Up = (FVector3f)Rot.RotateVector((FVector)GizmoProxy.Value.Up);
+				GizmoProxy.ProjectorGizmoEdited = true;
+			}
+			else if (WidgetMode == UE::Widget::WM_Scale)
+			{
+				GizmoProxy.Value.Scale.X += Scale.Y;
+				GizmoProxy.Value.Scale.Y += Scale.Z;
+				GizmoProxy.Value.Scale.Z += Scale.X;
+				GizmoProxy.ProjectorGizmoEdited = true;
+			}
+
+			InViewport->Invalidate();
+
+			if (GizmoProxy.ProjectorGizmoEdited)
+			{
+				GizmoProxy.SetProjectorUpdatedInViewport(true);
+				GizmoProxy.UpdateOriginData();
+				GizmoProxy.ProjectorGizmoEdited = false;
+			}
+		}
+
+		else if (bClipMorphVisible && ClipMorphNode)
+		{
+			if (WidgetMode == UE::Widget::WM_Translate)
+			{
+				if (CurrentAxis == EAxisList::Z)
+				{
+					float dragZ = bClipMorphLocalStartOffset ? FVector::DotProduct(Drag, ClipMorphNormal) : Drag.Z;
+					ClipMorphLocalOffset.Z += dragZ;
+					ClipMorphOffset += (bClipMorphLocalStartOffset) ? dragZ * ClipMorphNormal : FVector(0,0,dragZ);
+				}
+				else if(CurrentAxis == EAxisList::X)
+				{
+					float dragX = bClipMorphLocalStartOffset ? FVector::DotProduct(Drag, ClipMorphXAxis) : Drag.X;
+					ClipMorphLocalOffset.X += dragX;
+					ClipMorphOffset += (bClipMorphLocalStartOffset) ? dragX * ClipMorphXAxis : FVector(dragX, 0, 0);
+				}
+				else if (CurrentAxis == EAxisList::Y)
+				{
+					float dragY = bClipMorphLocalStartOffset ? FVector::DotProduct(Drag, ClipMorphYAxis) : Drag.Y;
+					ClipMorphLocalOffset.Y += dragY;
+					ClipMorphOffset += (bClipMorphLocalStartOffset) ? dragY * ClipMorphYAxis : FVector(0, dragY, 0);
+				}
+				ClipMorphNode->StartOffset = ClipMorphLocalOffset;
+			}
+			else if (WidgetMode == UE::Widget::WM_Rotate)
+			{
+				bool bClipMorphViewPortRotation = false;
+
+				if (CurrentAxis == EAxisList::X)
+				{
+					bClipMorphViewPortRotation = true;
+					float Angle = ClipMorphNode->bInvertNormal ? Rot.GetComponentForAxis(EAxis::X) : -Rot.GetComponentForAxis(EAxis::X);
+					ClipMorphNormal = ClipMorphNormal.RotateAngleAxis(Angle, ClipMorphXAxis);
+				}
+				else if (CurrentAxis == EAxisList::Y)
+				{
+					bClipMorphViewPortRotation = true;
+					float Angle = Rot.GetComponentForAxis(EAxis::Y);
+					ClipMorphNormal = ClipMorphNormal.RotateAngleAxis(Angle, ClipMorphYAxis);
+				}
+
+				if (bClipMorphViewPortRotation)
+				{
+					ClipMorphNormal.Normalize();
+					ClipMorphNode->Normal = ClipMorphNormal;
+					ClipMorphNode->FindLocalAxes(ClipMorphXAxis, ClipMorphYAxis, ClipMorphNormal);
+
+					if (bClipMorphLocalStartOffset)
+					{
+						ClipMorphLocalOffset.Z = FVector::DotProduct(ClipMorphOffset, ClipMorphNormal);
+						ClipMorphLocalOffset.Y = FVector::DotProduct(ClipMorphOffset, ClipMorphYAxis);
+						ClipMorphLocalOffset.X = FVector::DotProduct(ClipMorphOffset, ClipMorphXAxis);
+					}
+
+					ClipMorphNode->StartOffset = ClipMorphLocalOffset;
+				}
+			}
+		}
+
+		else if (ClipMeshNode)
+		{
+			if (WidgetMode == UE::Widget::WM_Translate)
+			{
+				ClipMeshNode->Transform.AddToTranslation(Drag);
+			}
+			else if (WidgetMode == UE::Widget::WM_Rotate)
+			{
+				ClipMeshNode->Transform.ConcatenateRotation(Rot.Quaternion());
+			}
+			if (WidgetMode == UE::Widget::WM_Scale)
+			{
+				ClipMeshNode->Transform.SetScale3D(ClipMeshNode->Transform.GetScale3D() + Scale);
+			}
+
+			ClipMeshComp->SetWorldTransform(ClipMeshNode->Transform);
+		}
+
+		else if (bIsEditingLightEnabled && SelectedLightComponent)
+		{
+			if (WidgetMode == UE::Widget::WM_Translate)
+			{
+				SelectedLightComponent->AddWorldOffset(Drag);
+				SelectedLightComponent->MarkForNeededEndOfFrameRecreate();
+			}
+
+			else if (WidgetMode == UE::Widget::WM_Rotate)
+			{
+				SelectedLightComponent->AddWorldRotation(Rot.Quaternion());
+				SelectedLightComponent->MarkForNeededEndOfFrameRecreate();
+			}
+		}
+	}
+
+	return bHandled;
+}
+
+
+void FCustomizableObjectEditorViewportClient::TrackingStarted(const struct FInputEventState& InInputState, bool bIsDraggingWidget, bool bNudge)
+{
+	if (bIsDraggingWidget)
+	{
+		if (InInputState.IsLeftMouseButtonPressed() && (Widget->GetCurrentAxis() & EAxisList::XYZ) != 0)
+		{
+			bManipulating = true;
+
+			if (WidgetMode == UE::Widget::WM_Translate)
+			{
+				GEditor->BeginTransaction(LOCTEXT("CustomizableObjectEditor_TranslateProjector", "Translate Projector"));
+			}
+			else if (WidgetMode == UE::Widget::WM_Rotate)
+			{
+				GEditor->BeginTransaction(LOCTEXT("CustomizableObjectEditor_RotateProjector", "Rotate Projector"));
+			}
+			else if (WidgetMode == UE::Widget::WM_Scale)
+			{
+				GEditor->BeginTransaction(LOCTEXT("CustomizableObjectEditor_ScaleProjector", "Scale Projector"));
+			}
+
+			if (GizmoProxy.AnyGizmoSelected)
+			{
+				GizmoProxy.Modify();
+				GizmoProxy.ModifyGraph();
+			}
+		}
+	}
+}
+
+
+void FCustomizableObjectEditorViewportClient::TrackingStopped()
+{
+	if (bManipulating)
+	{
+		bManipulating = false;
+		GEditor->EndTransaction();
+	}
+
+	Invalidate();
+}
+
+
+UE::Widget::EWidgetMode FCustomizableObjectEditorViewportClient::GetWidgetMode() const
+{
+	return WidgetMode;
+}
+
+
+FVector FCustomizableObjectEditorViewportClient::GetWidgetLocation() const
+{
+	if (GizmoProxy.AnyGizmoSelected)
+	{
+		return (FVector)GizmoProxy.Value.Position;
+	}
+
+	if (bClipMorphVisible)
+	{
+		return ClipMorphOrigin + ClipMorphOffset;
+	}
+
+	if (ClipMeshNode)
+	{
+		return ClipMeshNode->Transform.GetTranslation();
+	}
+
+	if (bIsEditingLightEnabled && SelectedLightComponent)
+	{
+		return SelectedLightComponent->GetComponentLocation();
+	}
+
+	return FVector::ZeroVector;
+}
+
+
+FMatrix FCustomizableObjectEditorViewportClient::GetWidgetCoordSystem() const
+{
+	if (GizmoProxy.AnyGizmoSelected)
+	{
+		FVector3f YVector = FVector3f::CrossProduct(GizmoProxy.Value.Direction, GizmoProxy.Value.Up);
+		return FMatrix((FVector)GizmoProxy.Value.Direction, (FVector)YVector, (FVector)GizmoProxy.Value.Up, FVector::ZeroVector);
+	}
+
+	if (bClipMorphVisible)
+	{
+		if (bClipMorphLocalStartOffset)
+		{
+			return FMatrix(-ClipMorphXAxis, -ClipMorphYAxis, -ClipMorphNormal, FVector::ZeroVector);
+		}
+
+		return FMatrix(FVector(1, 0, 0), FVector(0,1,0), FVector(0,0,1), FVector::ZeroVector);
+	}
+
+	if (ClipMeshNode)
+	{
+		return ClipMeshNode->Transform.ToMatrixNoScale().RemoveTranslation();
+	}
+
+	if (bIsEditingLightEnabled && SelectedLightComponent)
+	{
+		FMatrix Rotation = SelectedLightComponent->GetComponentTransform().ToMatrixNoScale();
+		Rotation.SetOrigin(FVector::ZeroVector);
+		return Rotation;
+	}
+
+	return FMatrix::Identity;
+}
+
+
+ECoordSystem FCustomizableObjectEditorViewportClient::GetWidgetCoordSystemSpace() const
+{
+	return ModeTools->GetCoordSystem();
+}
+
+
+void FCustomizableObjectEditorViewportClient::SetWidgetCoordSystemSpace(ECoordSystem NewCoordSystem)
+{
+	ModeTools->SetCoordSystem(NewCoordSystem);
+	Invalidate();
+}
+
+void FCustomizableObjectEditorViewportClient::SetViewportType(ELevelViewportType InViewportType)
+{
+	// Getting camera mode on perspective view
+	if (ViewportType == ELevelViewportType::LVT_Perspective)
+	{
+		bSetOrbitalOnPerspectiveMode = bActivateOrbitalCamera;
+	}
+
+	// Set Camera mode
+	if (InViewportType == ELevelViewportType::LVT_Perspective || ViewportType == ELevelViewportType::LVT_Perspective)
+	{
+		if (InViewportType == ELevelViewportType::LVT_Perspective)
+		{
+			SetCameraMode(bSetOrbitalOnPerspectiveMode);
+		}
+		else
+		{
+			SetCameraMode(false);
+		}
+	}
+
+	// Set Camera view
+	FEditorViewportClient::SetViewportType(InViewportType);
+}
+
+
+
+bool FCustomizableObjectEditorViewportClient::CanSetWidgetMode(UE::Widget::EWidgetMode NewMode) const
+{
+	return true;
+}
+
+
+void FCustomizableObjectEditorViewportClient::SetAnimation(UAnimationAsset* Animation, EAnimationMode::Type AnimationType)
+{
+	bool bFoundComponent = false;
+	for (TWeakObjectPtr<USkeletalMeshComponent> SkeletalMeshComponent : SkeletalMeshComponents)
+	{
+		if (SkeletalMeshComponent.IsValid() && Animation != nullptr
+			&& SkeletalMeshComponent->GetSkeletalMesh() != nullptr
+			&& SkeletalMeshComponent->GetSkeletalMesh()->GetSkeleton() == Animation->GetSkeleton()
+			)
+		{
+			SkeletalMeshComponent->SetAnimationMode(AnimationType);
+			SkeletalMeshComponent->PlayAnimation(Animation, true);
+			SkeletalMeshComponent->SetPlayRate(1.f);
+			SetRealtime(true);
+			IsPlayingAnimation = true;
+			AnimationBeingPlayed = Animation;
+
+			bFoundComponent = true;
+		}
+	}
+	
+	if(!bFoundComponent)
+	{
+		IsPlayingAnimation = false;
+		AnimationBeingPlayed = nullptr;
+	}
+}
+
+
+void FCustomizableObjectEditorViewportClient::ReSetAnimation()
+{
+	if ((IsPlayingAnimation == true) && (AnimationBeingPlayed != nullptr))
+	{
+		if (Cast<UPoseAsset>(AnimationBeingPlayed))
+		{
+			SetAnimation(AnimationBeingPlayed, EAnimationMode::AnimationBlueprint);
+		}
+		else
+		{
+			SetAnimation(AnimationBeingPlayed, EAnimationMode::AnimationSingleNode);
+		}
+	}
+}
+
+
+void FCustomizableObjectEditorViewportClient::AddLightToScene(ULightComponent* AddedLight)
+{
+	if (AddedLight)
+	{
+		PreviewScene->AddComponent(AddedLight, AddedLight->GetComponentTransform());
+	}
+
+	if (WidgetVisibility && !bIsEditingLightEnabled)
+	{
+		return;
+	}
+
+	SelectedLightComponent = AddedLight;
+	bIsEditingLightEnabled = AddedLight != nullptr;
+
+	if (!WidgetVisibility)
+	{
+		Widget->SetDefaultVisibility(true);
+		WidgetVisibility = true;
+	}
+}
+
+
+void FCustomizableObjectEditorViewportClient::RemoveLightFromScene(ULightComponent* RemovedLight)
+{
+	PreviewScene->RemoveComponent(RemovedLight);
+
+	if (RemovedLight == SelectedLightComponent)
+	{
+		SelectedLightComponent = nullptr;
+		bIsEditingLightEnabled = false;
+
+		Widget->SetDefaultVisibility(false);
+		WidgetVisibility = false;
+	}
+}
+
+
+void FCustomizableObjectEditorViewportClient::SetSelectedLight(ULightComponent* SelectedLight)
+{
+	if (WidgetVisibility && !bIsEditingLightEnabled)
+	{
+		return;
+	}
+
+	SelectedLightComponent = SelectedLight;
+	bIsEditingLightEnabled = SelectedLight != nullptr;
+
+	// Activate Gizmo and Set Light Transform or disable lights if !bIsEditingLightEnabled
+	if (WidgetVisibility != bIsEditingLightEnabled)
+	{
+		Widget->SetDefaultVisibility(bIsEditingLightEnabled);
+		WidgetVisibility = bIsEditingLightEnabled;
+	}
+}
+
+
+void FCustomizableObjectEditorViewportClient::SetFloorOffset(float NewValue)
+{
+	for (TWeakObjectPtr<USkeletalMeshComponent> SkeletalMeshComponent : SkeletalMeshComponents)
+	{
+		USkeletalMesh* Mesh = SkeletalMeshComponent.IsValid() ? SkeletalMeshComponent->GetSkeletalMesh() : nullptr;
+
+		if (Mesh)
+		{
+			// This value is saved in a UPROPERTY for the mesh, so changes are transactional
+			FScopedTransaction Transaction(LOCTEXT("SetFloorOffset", "Set Floor Offset"));
+			Mesh->Modify();
+
+			Mesh->SetFloorOffset(NewValue);
+			UpdateCameraSetup(); // This does the actual moving of the floor mesh
+			Invalidate();
+		}
+	}
+}
+
+
+//-------------------------------------------------------------------------------------------------
+
+class SMutableSelectFolderDlg : public SWindow
+{
+public:
+	SLATE_BEGIN_ARGS(SMutableSelectFolderDlg)
+	{
+	}
+
+	SLATE_ARGUMENT(FText, DefaultAssetPath)
+	SLATE_ARGUMENT(FText, DefaultFileName)
+	SLATE_END_ARGS()
+
+	SMutableSelectFolderDlg() : UserResponse(EAppReturnType::Cancel)
+	{
+		AddAllMaterialTextures = false;
+	}
+
+	void Construct(const FArguments& InArgs);
+
+public:
+	/** Displays the dialog in a blocking fashion */
+	EAppReturnType::Type ShowModal();
+
+	/** Gets the resulting asset path */
+	FString GetAssetPath();
+
+	/** FileName getter */
+	FString GetFileName();
+
+	/** Getter for AddAllMaterialTextures */
+	bool GetAddAllMaterialTextures();
+
+protected:
+	void OnPathChange(const FString& NewPath);
+	FReply OnButtonClick(EAppReturnType::Type ButtonID);
+	void OnNameChange(const FText& NewName, ETextCommit::Type CommitInfo);
+	void OnBoolParameterChanged(ECheckBoxState InCheckboxState);
+
+	EAppReturnType::Type UserResponse;
+	FText AssetPath;
+	FText FileName;
+	bool AddAllMaterialTextures;
+};
+
+
+//-------------------------------------------------------------------------------------------------
+
+void RemoveRestrictedChars(FString& String)
+{
+	// Remove restricted chars, according to FPaths::ValidatePath, RestrictedChars = "/?:&\\*\"<>|%#@^ ";
+
+	String = String.Replace(TEXT("/"), TEXT(""));
+	String = String.Replace(TEXT("?"), TEXT(""));
+	String = String.Replace(TEXT(":"), TEXT(""));
+	String = String.Replace(TEXT("&"), TEXT(""));
+	String = String.Replace(TEXT("\\"), TEXT(""));
+	String = String.Replace(TEXT("*"), TEXT(""));
+	String = String.Replace(TEXT("\""), TEXT(""));
+	String = String.Replace(TEXT("<"), TEXT(""));
+	String = String.Replace(TEXT(">"), TEXT(""));
+	String = String.Replace(TEXT("|"), TEXT(""));
+	String = String.Replace(TEXT("%"), TEXT(""));
+	String = String.Replace(TEXT("#"), TEXT(""));
+	String = String.Replace(TEXT("@"), TEXT(""));
+	String = String.Replace(TEXT("^"), TEXT(""));
+	String = String.Replace(TEXT(" "), TEXT(""));
+}
+
+
+//-------------------------------------------------------------------------------------------------
+void FCustomizableObjectEditorViewportClient::BakeInstance()
+{
+	if (!AssetRegistryLoaded)
+	{
+		FNotificationInfo Info(NSLOCTEXT("CustomizableObject", "CustomizableObjectCompileTryLater", "Please wait until asset registry loads all assets"));
+		Info.bFireAndForget = true;
+		Info.bUseThrobber = true;
+		Info.FadeOutDuration = 1.0f;
+		Info.ExpireDuration = 2.0f;
+		FSlateNotificationManager::Get().AddNotification(Info);
+		return;
+	}
+
+	bool bHasSkeletalMesh = false;
+	int32 NumComponents = SkeletalMeshComponents.Num();
+	for (int32 ComponentIndex = 0; ComponentIndex < NumComponents; ++ComponentIndex)
+	{
+		if (SkeletalMeshComponents[ComponentIndex].IsValid() && SkeletalMeshComponents[ComponentIndex]->GetSkeletalMesh())
+		{
+			bHasSkeletalMesh = true;
+		}
+	}
+		
+	if (!bHasSkeletalMesh)
+	{
+		return;
+	}
+
+	UCustomizableObjectInstance* Instance = CustomizableObjectEditorPtr.Pin()->GetPreviewInstance();
+	FString ObjectName = Instance->GetCustomizableObject()->GetName();
+	FText DefaultFileName = FText::Format(LOCTEXT("DefaultFileName", "{0}"), FText::AsCultureInvariant(ObjectName));
+	bool AddAllMaterialTextures = false;
+
+	TSharedRef<SMutableSelectFolderDlg> FolderDlg =
+		SNew(SMutableSelectFolderDlg)
+		.DefaultAssetPath(LOCTEXT("", ""))
+		.DefaultFileName(DefaultFileName);
+
+	if (FolderDlg->ShowModal() != EAppReturnType::Cancel)
+	{
+		// Make sure we can create the asset without conflicts
+		//IAssetTools& AssetTools = FModuleManager::GetModuleChecked<FAssetToolsModule>("AssetTools").Get();
+		//if (!AssetTools.CanCreateAsset(ObjName, PkgName, LOCTEXT("BakeMutableInstance", "Baking a Mutable instance")))
+		//{
+		//	return;
+		//}
+
+		FString FileName = FolderDlg->GetFileName();
+		ObjectName = FileName;
+
+		TCHAR InvalidCharacter = '0';
+		FString InvalidCharacters = FPaths::GetInvalidFileSystemChars();
+
+		for (int32 i = 0; i < InvalidCharacters.Len(); ++i)
+		{
+			TCHAR Char = InvalidCharacters[i];
+			FString SearchedChar = FString::Chr(Char);
+			if (ObjectName.Contains(SearchedChar))
+			{
+				InvalidCharacter = InvalidCharacters[i];
+				break;
+			}
+		}
+
+		AddAllMaterialTextures = FolderDlg->GetAddAllMaterialTextures();
+		BakingOverwritePermission = false;
+		FString CustomObjectPath = Instance->GetCustomizableObject()->GetPathName();
+		FString AssetPath = FolderDlg->GetAssetPath();
+		FString FullAssetPath = AssetPath + FString("/") + ObjectName + FString(".") + ObjectName;
+
+		if (InvalidCharacter != '0')
+		{
+			FText ErrorString = FText::FromString(FString::Chr(InvalidCharacter));
+			FMessageDialog::Open(EAppMsgType::Ok, FText::Format(LOCTEXT("FCustomizableObjectEditorViewportClient_BakeInstance_InvalidCharacter", "The selected contains an invalid character ({0})."), ErrorString));
+		}
+		else if (CustomObjectPath == FullAssetPath)
+		{
+			FMessageDialog::Open(EAppMsgType::Ok, LOCTEXT("FCustomizableObjectEditorViewportClient_BakeInstance_OverwriteCO", "The selected path would overwrite the instance's parent Customizable Object."));
+		}
+		else
+		{
+			TArray<UPackage*> PackagesToSave;
+
+			for (int32 ComponentIndex = 0; ComponentIndex < NumComponents; ++ComponentIndex)
+			{
+				USkeletalMesh* Mesh = SkeletalMeshComponents.Num() && SkeletalMeshComponents[ComponentIndex].IsValid() ?
+					SkeletalMeshComponents[ComponentIndex]->GetSkeletalMesh() : nullptr;
+				
+				if (!Mesh)
+				{
+					continue;
+				}
+
+				if (NumComponents > 1)
+				{
+					ObjectName = FileName + "_Component_" + FString::FromInt(ComponentIndex);
+				}
+
+				TMap<UObject*, UObject*> ReplacementMap;
+				TArray<FString> ArrayCachedElement;
+				TArray<UObject*> ArrayCachedObject;
+
+				if (AddAllMaterialTextures)
+				{
+					UMaterialInstance* Inst;
+					UMaterial* Material;
+					UTexture* Texture;
+					FString MaterialName;
+					FString ResourceName;
+					FString PackageName;
+					UObject* DuplicatedObject;
+					TArray<TMap<int, UTexture*>> TextureReplacementMaps;
+
+					// Duplicate Mutable generated textures
+					for (int m = 0; m < Mesh->GetMaterials().Num(); ++m)
+					{
+						Inst = Cast<UMaterialInstance>(Mesh->GetMaterials()[m].MaterialInterface);
+						MaterialName = (Inst->Parent != nullptr) ? Inst->Parent->GetName() : "Material";
+						Material = Inst->Parent->GetMaterial();
+
+						if (Material != nullptr)
+						{
+							TMap<int, UTexture*> ReplacementTextures;
+							TextureReplacementMaps.Add(ReplacementTextures);
+
+							TArray<FName> ParameterNames = GetTextureParameterNames(Material);
+
+							for (int32 i = 0; i < ParameterNames.Num(); i++)
+							{
+								if (Inst->GetTextureParameterValue(ParameterNames[i], Texture))
+								{
+									UTexture2D* SrcTex = Cast<UTexture2D>(Texture);
+									if (!SrcTex) continue;
+
+									bool bIsMutableTexture = false;
+
+									for (UAssetUserData* UserData : *SrcTex->GetAssetUserDataArray())
+									{
+										UTextureMipDataProviderFactory* CustomMipDataProviderFactory = Cast<UMutableTextureMipDataProviderFactory>(UserData);
+										if (CustomMipDataProviderFactory)
+										{
+											bIsMutableTexture = true;
+										}
+									}
+
+									if ((SrcTex->GetPlatformData() != nullptr) &&
+										(SrcTex->GetPlatformData()->Mips.Num() > 0) &&
+										bIsMutableTexture)
+									{
+										FString ParameterSanitized = ParameterNames[i].GetPlainNameString();
+										RemoveRestrictedChars(ParameterSanitized);
+										ResourceName = ObjectName + "_" + MaterialName + "_" + ParameterSanitized;
+
+										if (!GetUniqueResourceName(SrcTex, ResourceName, ArrayCachedObject, ArrayCachedElement))
+										{
+											continue;
+										}
+
+										if (!ManageBakingAction(AssetPath, ResourceName))
+										{
+											return;
+										}
+
+										// Recover original name of the texture parameter value, now substituted by the generated Mutable texture
+										UTexture* OriginalTexture = nullptr;
+										UMaterialInstanceDynamic* InstDynamic = Cast<UMaterialInstanceDynamic>(Mesh->GetMaterials()[m].MaterialInterface);
+										if (InstDynamic != nullptr)
+										{
+											InstDynamic->Parent->GetTextureParameterValue(FName(*ParameterNames[i].GetPlainNameString()), OriginalTexture);
+										}
+
+										PackageName = FolderDlg->GetAssetPath() + FString("/") + ResourceName;
+										TMap<UObject*, UObject*> FakeReplacementMap;
+										UTexture2D* DupTex = BakeHelper_CreateAssetTexture(SrcTex, ResourceName, PackageName, OriginalTexture, true, FakeReplacementMap, BakingOverwritePermission);
+										ArrayCachedElement.Add(ResourceName);
+										ArrayCachedObject.Add(DupTex);
+										PackagesToSave.Add(DupTex->GetPackage());
+
+										if (OriginalTexture != nullptr)
+										{
+											TextureReplacementMaps[m].Add(i, DupTex);
+										}
+									}
+								}
+							}
+						}
+					}
+
+					// Duplicate non-Mutable material textures
+					for (int m = 0; m < Mesh->GetMaterials().Num(); ++m)
+					{
+						Inst = Cast<UMaterialInstance>(Mesh->GetMaterials()[m].MaterialInterface);
+						MaterialName = (Inst->Parent != nullptr) ? Inst->Parent->GetName() : "Material";
+						Material = Inst->Parent->GetMaterial();
+
+						if (Material != nullptr)
+						{
+							TArray<FName> ParameterNames = GetTextureParameterNames(Material);
+
+							for (int32 i = 0; i < ParameterNames.Num(); i++)
+							{
+								if (Inst->GetTextureParameterValue(ParameterNames[i], Texture))
+								{
+									FString ParameterSanitized = ParameterNames[i].GetPlainNameString();
+									RemoveRestrictedChars(ParameterSanitized);
+									ResourceName = ObjectName + "_" + MaterialName + "_" + ParameterSanitized;
+
+									if (ArrayCachedElement.Find(ResourceName) == INDEX_NONE)
+									{
+										if (!ManageBakingAction(AssetPath, ResourceName))
+										{
+											return;
+										}
+
+										PackageName = FolderDlg->GetAssetPath() + FString("/") + ResourceName;
+										TMap<UObject*, UObject*> FakeReplacementMap;
+										DuplicatedObject = BakeHelper_DuplicateAsset(Texture, ResourceName, PackageName, true, FakeReplacementMap, BakingOverwritePermission);
+										ArrayCachedElement.Add(ResourceName);
+										ArrayCachedObject.Add(DuplicatedObject);
+										PackagesToSave.Add(DuplicatedObject->GetPackage());
+
+										UTexture* DupTexture = Cast<UTexture>(DuplicatedObject);
+										TextureReplacementMaps[m].Add(i, DupTexture);
+									}
+								}
+							}
+						}
+					}
+
+
+					// Duplicate the materials used by each material instance so replacement map has proper information when duplicating the material instances
+					for (int m = 0; m < Mesh->GetMaterials().Num(); ++m)
+					{
+						Inst = Cast<UMaterialInstance>(Mesh->GetMaterials()[m].MaterialInterface);
+
+						if ((Inst != nullptr) && (Inst->GetMaterial() != nullptr))
+						{
+							Material = Inst->GetMaterial();
+
+							if (Material != nullptr)
+							{
+								ResourceName = ObjectName + "_Material_" + Material->GetName();
+
+								if (!GetUniqueResourceName(Material, ResourceName, ArrayCachedObject, ArrayCachedElement))
+								{
+									continue;
+								}
+
+								if (!ManageBakingAction(AssetPath, ResourceName))
+								{
+									return;
+								}
+
+								PackageName = FolderDlg->GetAssetPath() + FString("/") + ResourceName;
+								TMap<UObject*, UObject*> FakeReplacementMap;
+								DuplicatedObject = BakeHelper_DuplicateAsset(Material, ResourceName, PackageName, false, FakeReplacementMap, BakingOverwritePermission);
+								ArrayCachedElement.Add(ResourceName);
+								ArrayCachedObject.Add(DuplicatedObject);
+								ReplacementMap.Add(Inst, DuplicatedObject);
+								PackagesToSave.Add(DuplicatedObject->GetPackage());
+
+								if (UMaterial* DupMaterial = Cast<UMaterial>(DuplicatedObject))
+								{
+									TArray<FMaterialParameterInfo> parametersInfo;
+									TArray<FGuid> parametersGuids;
+
+									// copy scalar parameters
+									Inst->GetAllScalarParameterInfo(parametersInfo, parametersGuids);
+									for (const FMaterialParameterInfo& param : parametersInfo)
+									{
+										float value = 0.f;
+										if (Inst->GetScalarParameterValue(param, value))
+										{
+											DupMaterial->SetScalarParameterValueEditorOnly(param.Name, value);
+										}
+									}
+
+									// copy vector parameters
+									Inst->GetAllVectorParameterInfo(parametersInfo, parametersGuids);
+									for (const FMaterialParameterInfo& param : parametersInfo)
+									{
+										FLinearColor value;
+										if (Inst->GetVectorParameterValue(param, value))
+										{
+											DupMaterial->SetVectorParameterValueEditorOnly(param.Name, value);
+										}
+									}
+
+									// copy switch parameters
+									Inst->GetAllStaticSwitchParameterInfo(parametersInfo, parametersGuids);
+									for (int i = 0; i < parametersInfo.Num(); ++i)
+									{
+										bool value = false;
+										if (Inst->GetStaticSwitchParameterValue(parametersInfo[i].Name, value, parametersGuids[i]))
+										{
+											DupMaterial->SetStaticSwitchParameterValueEditorOnly(parametersInfo[i].Name, value, parametersGuids[i]);
+										}
+									}
+
+									// Replace Textures
+									TArray<FName> ParameterNames = GetTextureParameterNames(Material);
+									for (const TPair<int, UTexture*>& it : TextureReplacementMaps[m])
+									{
+										if (ParameterNames.IsValidIndex(it.Key))
+										{
+											DupMaterial->SetTextureParameterValueEditorOnly(ParameterNames[it.Key], it.Value);
+										}
+									}
+
+									// Fix potential errors compiling materials due to Sampler Types
+									for (TObjectPtr<UMaterialExpression> Expression : DupMaterial->GetExpressions())
+									{
+										if (UMaterialExpressionTextureBase* MatExpressionTexBase = Cast<UMaterialExpressionTextureBase>(Expression))
+										{
+											MatExpressionTexBase->AutoSetSampleType();
+										}
+									}
+
+									DuplicatedObject->PreEditChange(NULL);
+									DuplicatedObject->PostEditChange();
+								}
+							}
+						}
+					}
+				}
+				else
+				{
+					// Duplicate the material instances
+					for (int m = 0; m < Mesh->GetMaterials().Num(); ++m)
+					{
+						UMaterialInstance* Inst = Cast<UMaterialInstance>(Mesh->GetMaterials()[m].MaterialInterface);
+						FString MaterialName = (Inst->Parent != nullptr) ? Inst->Parent->GetName() : "Material";
+
+						// Material
+						FString MatObjName = ObjectName + "_" + MaterialName;
+
+						if (!GetUniqueResourceName(Inst, MatObjName, ArrayCachedObject, ArrayCachedElement))
+						{
+							continue;
+						}
+
+						if (!ManageBakingAction(AssetPath, MatObjName))
+						{
+							return;
+						}
+
+						FString MatPkgName = FolderDlg->GetAssetPath() + FString("/") + MatObjName;
+						UObject* DupMat = BakeHelper_DuplicateAsset(Inst, MatObjName, MatPkgName, false, ReplacementMap, BakingOverwritePermission);
+						ArrayCachedObject.Add(DupMat);
+						ArrayCachedElement.Add(MatObjName);
+						PackagesToSave.Add(DupMat->GetPackage());
+
+						// Duplicate generated textures
+						if (UMaterialInstanceDynamic* InstDynamic = Cast<UMaterialInstanceDynamic>(DupMat))
+						{
+							for (int t = 0; t < Inst->TextureParameterValues.Num(); ++t)
+							{
+								if (Inst->TextureParameterValues[t].ParameterValue)
+								{
+									UTexture2D* SrcTex = Cast<UTexture2D>(Inst->TextureParameterValues[t].ParameterValue);
+									FString ParameterSanitized = Inst->TextureParameterValues[t].ParameterInfo.Name.ToString();
+									RemoveRestrictedChars(ParameterSanitized);
+
+									FString TexObjName = ObjectName + "_" + MaterialName + "_" + ParameterSanitized;
+
+									if (!GetUniqueResourceName(SrcTex, TexObjName, ArrayCachedObject, ArrayCachedElement))
+									{
+										UTexture* PrevTexture = Cast<UTexture>(ArrayCachedObject[ArrayCachedElement.Find(TexObjName)]);
+										InstDynamic->SetTextureParameterValue(Inst->TextureParameterValues[t].ParameterInfo.Name, PrevTexture);
+										continue;
+									}
+
+									if (!ManageBakingAction(AssetPath, TexObjName))
+									{
+										return;
+									}
+
+									FString TexPkgName = FolderDlg->GetAssetPath() + FString("/") + TexObjName;
+									TMap<UObject*, UObject*> FakeReplacementMap;
+									UTexture2D* DupTex = BakeHelper_CreateAssetTexture(SrcTex, TexObjName, TexPkgName, nullptr, false, FakeReplacementMap, BakingOverwritePermission);
+									ArrayCachedObject.Add(DupTex);
+									ArrayCachedElement.Add(TexObjName);
+									PackagesToSave.Add(DupTex->GetPackage());
+
+									InstDynamic->SetTextureParameterValue(Inst->TextureParameterValues[t].ParameterInfo.Name, DupTex);
+								}
+							}
+						}
+					}
+				}
+
+				// Make sure source data is present in the mesh before we duplciate:
+				FUnrealBakeHelpers::BakeHelper_RegenerateImportedModel(Mesh);
+
+				// Skeletal Mesh
+				if (!ManageBakingAction(AssetPath, ObjectName))
+				{
+					return;
+				}
+				FString PkgName = FolderDlg->GetAssetPath() + FString("/") + ObjectName;
+				UObject* DupObject = BakeHelper_DuplicateAsset(Mesh, ObjectName, PkgName, false, ReplacementMap, BakingOverwritePermission);
+				ArrayCachedObject.Add(DupObject);
+				PackagesToSave.Add(DupObject->GetPackage());
+
+				Mesh->Build();
+
+				USkeletalMesh* SkeletalMesh = Cast<USkeletalMesh>(DupObject);
+				if (SkeletalMesh)
+				{
+					Helper_GetLODInfoArray(SkeletalMesh) = Helper_GetLODInfoArray(Mesh);
+
+					Helper_GetImportedModel(SkeletalMesh)->SkeletalMeshModelGUID = FGuid::NewGuid();
+
+					// Generate render data
+					SkeletalMesh->Build();
+				}
+
+				// Remove duplicated UObjects from Root (previously added to avoid objects from beeing GC in the middle of the bake process)
+				for (UObject* Obj : ArrayCachedObject)
+				{
+					Obj->RemoveFromRoot();
+				}
+			}
+
+			if (PackagesToSave.Num())
+			{
+				FEditorFileUtils::PromptForCheckoutAndSave(PackagesToSave, false, true);
+			}
+		}
+	}
+}
+
+
+bool FCustomizableObjectEditorViewportClient::GetUniqueResourceName(UObject* Resource, FString& ResourceName, TArray<UObject*>& InCachedResources, TArray<FString>& InCachedResourceNames)
+{
+	int32 FindResult = InCachedResourceNames.Find(ResourceName);
+	if (FindResult != INDEX_NONE)
+	{
+		if (Resource == InCachedResources[FindResult])
+		{
+			return false;
+		}
+
+		uint32 Count = 0;
+		while (FindResult != INDEX_NONE)
+		{
+			FindResult = InCachedResourceNames.Find(ResourceName + "_" + FString::FromInt(Count));
+			Count++;
+		}
+
+		ResourceName += "_" + FString::FromInt(--Count);
+	}
+
+	return true;
+}
+
+
+bool FCustomizableObjectEditorViewportClient::ManageBakingAction(const FString& Path, const FString& ObjName)
+{
+	FString PackagePath = Path + "/" + ObjName;
+	UPackage* ExistingPackage = FindPackage(NULL, *PackagePath);
+
+	if (!ExistingPackage)
+	{
+		FString PackageFilePath = PackagePath + "." + ObjName;
+
+		FString PackageFileName;
+		if (FPackageName::DoesPackageExist(PackageFilePath, &PackageFileName))
+		{
+			ExistingPackage = LoadPackage(nullptr, *PackageFileName, LOAD_EditorOnly);
+		}
+		else
+		{
+			// if package does not exists
+			BakingOverwritePermission = false;
+			return true;
+		}
+	}
+
+	if (ExistingPackage)
+	{
+		// Checking if the asset is open in an editor
+		TArray<IAssetEditorInstance*> ObjectEditors = GEditor->GetEditorSubsystem<UAssetEditorSubsystem>()->FindEditorsForAssetAndSubObjects(ExistingPackage);
+		if (ObjectEditors.Num())
+		{
+			for (IAssetEditorInstance* ObjectEditorInstance : ObjectEditors)
+			{
+				// Close the editors that contains this asset
+				if (!ObjectEditorInstance->CloseWindow())
+				{
+					FText Caption = LOCTEXT("OpenExisitngFile", "Open File");
+					FText Message = FText::Format(LOCTEXT("CantCloseAsset", "This Obejct \"{0}\" is open in an editor and can't be closed automatically. Please close the editor and try to bake it again"), FText::FromString(ObjName));
+
+					FMessageDialog::Open(EAppMsgType::Ok, Message, &Caption);
+
+					return false;
+				}
+			}
+		}
+
+		if (!BakingOverwritePermission)
+		{
+			FText Caption = LOCTEXT("Already existing baked files", "Already existing baked files");
+			FText Message = FText::Format(LOCTEXT("OverwriteBakedInstance", "Instance baked files already exist in selected destination \"{0}\", this action will overwrite them."), FText::AsCultureInvariant(Path));
+			
+			if (FMessageDialog::Open(EAppMsgType::OkCancel, Message, &Caption) == EAppReturnType::Cancel)
+			{
+				return false;
+			}
+
+			BakingOverwritePermission = true;
+		}
+
+		UObject* ExistingObject = StaticFindObject(UObject::StaticClass(), ExistingPackage, *ObjName);
+		if (ExistingObject)
+		{
+			ExistingPackage->FullyLoad();
+
+			TArray<UObject*> ObjectsToDelete;
+			ObjectsToDelete.Add(ExistingObject);
+
+			// Delete objects in the package with the same name as the one we want to create
+			const uint32 NumObjectsDeleted = ObjectTools::ForceDeleteObjects(ObjectsToDelete, false);
+
+			return NumObjectsDeleted == ObjectsToDelete.Num();
+		}
+	}
+
+	return true;
+}
+
+
+void FCustomizableObjectEditorViewportClient::StateChangeShowGeometryData()
+{
+	StateChangeShowGeometryDataFlag = !StateChangeShowGeometryDataFlag;
+	Invalidate();
+}
+
+
+void FCustomizableObjectEditorViewportClient::SetProjectorWidgetMode(UE::Widget::EWidgetMode InMode)
+{
+	WidgetMode = InMode;
+}
+
+
+const GizmoRTSProxy& FCustomizableObjectEditorViewportClient::GetGizmoProxy()
+{
+	return GizmoProxy;
+}
+
+
+bool FCustomizableObjectEditorViewportClient::GetGizmoIsProjectorParameterSelected()
+{
+	return GizmoProxy.IsProjectorParameterSelected();
+}
+
+
+bool FCustomizableObjectEditorViewportClient::GetIsManipulating()
+{
+	return bManipulating;
+}
+
+
+bool FCustomizableObjectEditorViewportClient::GetWidgetVisibility()
+{
+	return WidgetVisibility;
+}
+
+
+void FCustomizableObjectEditorViewportClient::UpdateGizmoDataToOrigin()
+{
+	GizmoProxy.UpdateOriginData();
+}
+
+
+void FCustomizableObjectEditorViewportClient::CopyTransformFromOriginData()
+{
+    GizmoProxy.CopyTransformFromOriginData();
+}
+
+
+bool FCustomizableObjectEditorViewportClient::AnyProjectorNodeSelected()
+{
+	return (GizmoProxy.HasAssignedData && GizmoProxy.AssignedDataIsFromNode);
+}
+
+
+void FCustomizableObjectEditorViewportClient::ShowInstanceGeometryInformation(FCanvas* InCanvas)
+{
+	UCustomizableObjectInstance* Instance = CustomizableObjectEditorPtr.Pin()->GetPreviewInstance();
+	float YOffset = 50.0f;
+	int32 ComponentIndex = 0;
+
+	// Show total number of triangles and vertices
+	for (TWeakObjectPtr<USkeletalMeshComponent> SkeletalMeshComponent : SkeletalMeshComponents)
+	{
+		if (SkeletalMeshComponent.IsValid() && SkeletalMeshComponent->GetSkeletalMesh())
+		{
+			const FSkeletalMeshRenderData* MeshRes = SkeletalMeshComponent->GetSkeletalMesh()->GetResourceForRendering();
+			int32 NumTriangles;
+			int32 NumVertices;
+			int32 NumLODLevel = MeshRes->LODRenderData.Num();
+
+			for (int32 i = 0; i < NumLODLevel; ++i)
+			{
+				NumTriangles = 0;
+				NumVertices = 0;
+				const FSkeletalMeshLODRenderData& lodModel = MeshRes->LODRenderData[i];
+				for (int32 j = 0; j < lodModel.RenderSections.Num(); ++j)
+				{
+					NumTriangles += lodModel.RenderSections[j].NumTriangles;
+					NumVertices += lodModel.RenderSections[j].NumVertices;
+				}
+
+				//draw a string showing what UV channel and LOD is being displayed
+				InCanvas->DrawShadowedString(
+					6.0f,
+					YOffset,
+					*FText::Format(NSLOCTEXT("UnrealEd", "UVOverlay_F", "Component {3} LOD {0} has {1} vertices and {2} triangles"),
+						FText::AsNumber(i), FText::AsNumber(NumVertices), FText::AsNumber(NumTriangles), FText::AsNumber(ComponentIndex)).ToString(),
+					GEngine->GetSmallFont(),
+					FLinearColor::White
+				);
+
+				YOffset += 20.0f;
+			}
+		}
+
+		YOffset += 40.0f;
+		ComponentIndex++;
+	}
+}
+
+
+void FCustomizableObjectEditorViewportClient::BuildMeanTimesBoxSize(const float MinValue, const float MaxValue, const float MeanValue, const float MaxWidth, const float Data, float& BoxSize, FLinearColor& BoxColor) const
+{
+	const float GreenLimit = MeanValue * 0.75f;
+	const float YellowLimit = MeanValue * 0.99;
+
+	if (Data < GreenLimit)
+	{
+		BoxColor = FLinearColor::Green;
+	}
+	else if (Data < YellowLimit)
+	{
+		BoxColor = FLinearColor::Yellow;
+	}
+	else
+	{
+		BoxColor = FLinearColor::Red;
+	}
+
+	BoxSize = (Data / MaxValue) * MaxWidth;
+}
+
+
+void FCustomizableObjectEditorViewportClient::BuildMeanTimesBoxSizes(float MaxWidth, TArray<float>& ArrayData, TArray<float>& ArrayBoxSize, TArray<FLinearColor>& ArrayBoxColor)
+{
+	const int32 NumElement = ArrayData.Num();
+	ArrayBoxSize.AddUninitialized(NumElement);
+	ArrayBoxColor.AddUninitialized(NumElement);
+
+	float MinValue = FLT_MAX;
+	float MaxValue = -1.0f * FLT_MAX;
+	float MeanValue = 0.0f;
+	for (int32 i = 0; i < ArrayData.Num(); ++i)
+	{
+		MeanValue += ArrayData[i];
+		MinValue = FMath::Min(MinValue, ArrayData[i]);
+		MaxValue = FMath::Max(MaxValue, ArrayData[i]);
+	}
+
+	float MaxMinRatio = 1.0f;
+	if (MinValue > 0.0f)
+	{
+		MaxMinRatio = MaxValue / MinValue;
+	}
+
+	MeanValue /= float(ArrayData.Num());
+	float GreenLimit = MeanValue + (MaxValue - MeanValue) * 0.1f;
+	float YellowLimit = MeanValue + (MaxValue - MeanValue) * 0.5f;
+
+	for (int32 i = 0; i < ArrayData.Num(); ++i)
+	{
+		if (ArrayData[i] < GreenLimit)
+		{
+			ArrayBoxColor[i] = FLinearColor::Green;
+		}
+		else if (ArrayData[i] < YellowLimit)
+		{
+			ArrayBoxColor[i] = FLinearColor::Yellow;
+		}
+		else
+		{
+			ArrayBoxColor[i] = FLinearColor::Red;
+		}
+
+		ArrayBoxSize[i] = (ArrayData[i] / MaxValue) * MaxWidth;
+	}
+}
+
+
+void FCustomizableObjectEditorViewportClient::SetCustomizableObject(UCustomizableObject* CustomizableObjectParameter)
+{
+	CustomizableObject = CustomizableObjectParameter;
+}
+
+
+void FCustomizableObjectEditorViewportClient::DrawShadowedString(FCanvas* Canvas, float StartX, float StartY, const FLinearColor& Color, float TextScale, FString String)
+{
+	UFont* StatFont = nullptr;
+
+	if (TextScale > 2.0f)
+	{
+		StatFont = GEngine->GetLargeFont();
+	}
+	else if (TextScale > 1.0f)
+	{
+		StatFont = GEngine->GetMediumFont();
+	}
+	else
+	{
+		StatFont = GEngine->GetSmallFont();
+	}
+
+	Canvas->DrawShadowedString(StartX, StartY, *String, StatFont, Color);
+}
+
+
+void FCustomizableObjectEditorViewportClient::AddReferencedObjects(FReferenceCollector& Collector)
+{
+	Collector.AddReferencedObject(ClipMorphMaterial);
+	Collector.AddReferencedObject(TransparentPlaneMaterialXY);
+}
+
+
+void FCustomizableObjectEditorViewportClient::SetAssetRegistryLoaded(bool Value)
+{
+	AssetRegistryLoaded = Value;
+}
+
+
+void FCustomizableObjectEditorViewportClient::SetAdvancedShowFlagsForScene(const bool bAdvancedShowFlags)
+{
+	if (bAdvancedShowFlags)
+	{
+		EngineShowFlags.EnableAdvancedFeatures();
+	}
+	else
+	{
+		EngineShowFlags.DisableAdvancedFeatures();
+	}
+}
+
+
+void FCustomizableObjectEditorViewportClient::OnAssetViewerSettingsChanged(const FName& InPropertyName)
+{
+	if (InPropertyName == GET_MEMBER_NAME_CHECKED(FPreviewSceneProfile, bPostProcessingEnabled) || InPropertyName == NAME_None)
+	{
+		UAssetViewerSettings* Settings = UAssetViewerSettings::Get();
+		int32 ProfileIndex = GetMutableDefault<UEditorPerProjectUserSettings>()->AssetViewerProfileIndex;
+		if (Settings->Profiles.IsValidIndex(ProfileIndex))
+		{
+			SetAdvancedShowFlagsForScene(Settings->Profiles[ProfileIndex].bPostProcessingEnabled);
+		}
+	}
+}
+
+
+void FCustomizableObjectEditorViewportClient::ProjectorParameterChanged(UCustomizableObjectNodeProjectorParameter* Node)
+{
+	GizmoProxy.ProjectorParameterChanged(Node);
+}
+
+
+void FCustomizableObjectEditorViewportClient::ProjectorParameterChanged(UCustomizableObjectNodeProjectorConstant* Node)
+{
+	GizmoProxy.ProjectorParameterChanged(Node);
+}
+
+
+void FCustomizableObjectEditorViewportClient::SetGizmoCallUpdateSkeletalMesh(bool Value)
+{
+	GizmoProxy.SetCallUpdateSkeletalMesh(Value);
+}
+
+
+FString FCustomizableObjectEditorViewportClient::GetGizmoProjectorParameterName()
+{
+	return GizmoProxy.ProjectorParameterName;
+}
+
+
+FString FCustomizableObjectEditorViewportClient::GetGizmoProjectorParameterNameWithIndex()
+{
+	return GizmoProxy.ProjectorParameterNameWithIndex;
+}
+
+
+bool FCustomizableObjectEditorViewportClient::GetGizmoHasAssignedData()
+{
+	return GizmoProxy.HasAssignedData;
+}
+
+
+bool FCustomizableObjectEditorViewportClient::GetGizmoAssignedDataIsFromNode()
+{
+	return GizmoProxy.AssignedDataIsFromNode;
+}
+
+
+bool FCustomizableObjectEditorViewportClient::IsNodeMeshClipMorphSelected()
+{
+	return bClipMorphVisible;
+}
+
+void FCustomizableObjectEditorViewportClient::DrawCylinderArc(FPrimitiveDrawInterface* PDI, const FMatrix& CylToWorld, const FVector& Base, const FVector& XAxis, const FVector& YAxis, const FVector& ZAxis, float Radius, float HalfHeight, uint32 Sides, const FMaterialRenderProxy* MaterialRenderProxy, uint8 DepthPriority, FColor Color, float MaxAngle)
+{
+	TArray<FDynamicMeshVertex> MeshVerts;
+	TArray<uint32> MeshIndices;
+
+	const float	AngleDelta = MaxAngle / (Sides - 1);
+	const float Offset = 0.5f * MaxAngle;
+
+	FVector2f TC = FVector2f(0.0f, 0.0f);
+	float TCStep = 1.0f / (Sides - 1);
+
+	FVector TopOffset = HalfHeight * ZAxis;
+	int32 BaseVertIndex = MeshVerts.Num();
+
+	//Compute vertices for base circle.
+	for (uint32 SideIndex = 0; SideIndex < Sides; SideIndex++)
+	{
+		const FVector Vertex = Base + (XAxis * FMath::Cos(AngleDelta * SideIndex - Offset) + YAxis * FMath::Sin(AngleDelta * SideIndex - Offset)) * Radius;
+		FVector Normal = Vertex - Base;
+		Normal.Normalize();
+
+		FDynamicMeshVertex MeshVertex;
+		MeshVertex.Position = FVector3f(Vertex - TopOffset);
+		MeshVertex.TextureCoordinate[0] = TC;
+		MeshVertex.SetTangents((FVector3f)-ZAxis, FVector3f((-ZAxis) ^ Normal), (FVector3f)Normal);
+		MeshVertex.Color = Color;
+		MeshVerts.Add(MeshVertex); //Add bottom vertex
+
+		TC.X += TCStep;
+	}
+
+	TC = FVector2f(0.0f, 1.0f);
+
+	//Compute vertices for the top circle
+	for (uint32 SideIndex = 0; SideIndex < Sides; SideIndex++)
+	{
+		const FVector Vertex = Base + (XAxis * FMath::Cos(AngleDelta * SideIndex - Offset) + YAxis * FMath::Sin(AngleDelta * SideIndex - Offset)) * Radius;
+		FVector Normal = Vertex - Base;
+		Normal.Normalize();
+
+		FDynamicMeshVertex MeshVertex;
+		MeshVertex.Position = FVector3f(Vertex + TopOffset);	// LWC_TODO: Precision Loss
+		MeshVertex.TextureCoordinate[0] = TC;
+		MeshVertex.SetTangents((FVector3f)-ZAxis, FVector3f((-ZAxis) ^ Normal), (FVector3f)Normal);
+		MeshVertex.Color = Color;
+		MeshVerts.Add(MeshVertex); //Add top vertex
+
+		TC.X += TCStep;
+	}
+
+	//Add sides.
+	for (uint32 SideIndex = 0; SideIndex < (Sides - 1); SideIndex++)
+	{
+		int32 V0 = BaseVertIndex + SideIndex;
+		int32 V1 = BaseVertIndex + ((SideIndex + 1) % Sides);
+		int32 V2 = V0 + Sides;
+		int32 V3 = V1 + Sides;
+
+		MeshIndices.Add(V0);
+		MeshIndices.Add(V2);
+		MeshIndices.Add(V1);
+
+		MeshIndices.Add(V2);
+		MeshIndices.Add(V3);
+		MeshIndices.Add(V1);
+	}
+
+	FDynamicMeshBuilder MeshBuilder(PDI->View->GetFeatureLevel());
+	MeshBuilder.AddVertices(MeshVerts);
+	MeshBuilder.AddTriangles(MeshIndices);
+
+	MeshBuilder.Draw(PDI, CylToWorld, MaterialRenderProxy, DepthPriority, 0.f);
+}
+
+
+bool FCustomizableObjectEditorViewportClient::GetFloorVisibility()
+{
+	FAdvancedPreviewScene* AdvancedScene = static_cast<FAdvancedPreviewScene*>(PreviewScene);
+	if (AdvancedScene != nullptr)
+	{
+		const UStaticMeshComponent* FloorMeshComponent = AdvancedScene->GetFloorMeshComponent();
+		if (FloorMeshComponent != nullptr)
+		{
+			return FloorMeshComponent->IsVisible();
+		}
+	}
+
+	return false;
+}
+
+
+void FCustomizableObjectEditorViewportClient::SetFloorVisibility(bool Value)
+{
+	FAdvancedPreviewScene* AdvancedScene = static_cast<FAdvancedPreviewScene*>(PreviewScene);
+	if (AdvancedScene != nullptr)
+	{
+		AdvancedScene->SetFloorVisibility(Value);
+	}
+}
+
+
+bool FCustomizableObjectEditorViewportClient::GetGridVisibility()
+{
+	return DrawHelper.bDrawGrid;
+}
+
+
+bool FCustomizableObjectEditorViewportClient::GetEnvironmentMeshVisibility()
+{
+	FCustomizableObjectPreviewScene* CustomizableObjectPreviewScene = static_cast<FCustomizableObjectPreviewScene*>(PreviewScene);
+	if (CustomizableObjectPreviewScene != nullptr)
+	{
+		return CustomizableObjectPreviewScene->GetSkyComponent()->IsVisible();
+	}
+
+	return false;
+}
+
+
+void FCustomizableObjectEditorViewportClient::SetEnvironmentMeshVisibility(uint32 Value)
+{
+	FCustomizableObjectPreviewScene* CustomizableObjectPreviewScene = static_cast<FCustomizableObjectPreviewScene*>(PreviewScene);
+	if (CustomizableObjectPreviewScene != nullptr)
+	{
+		CustomizableObjectPreviewScene->GetSkyComponent()->SetVisibility(Value == 1, true);
+	}
+
+	Invalidate();
+}
+
+
+TArray<FName> FCustomizableObjectEditorViewportClient::GetTextureParameterNames(UMaterial* Material)
+{
+	TArray<FGuid> Guids;
+	TArray<FName> ParameterNames;
+
+	TArray<FMaterialParameterInfo> OutParameterInfo;
+	Material->GetAllTextureParameterInfo(OutParameterInfo, Guids);
+
+	const int32 MaxIndex = OutParameterInfo.Num();
+	ParameterNames.SetNum(MaxIndex);
+
+	for (int32 i = 0; i < MaxIndex; i++)
+	{
+		ParameterNames[i] = OutParameterInfo[i].Name;
+	}
+
+	return ParameterNames;
+}
+
+bool FCustomizableObjectEditorViewportClient::IsOrbitalCameraActive() const
+{
+	return bActivateOrbitalCamera;
+}
+
+void FCustomizableObjectEditorViewportClient::SetCameraMode(bool Value)
+{
+	bActivateOrbitalCamera = Value;
+	UpdateCameraSetup();
+}
+
+/////////////////////////////////////////////////
+// select folder dialog \todo: move to its own file
+//////////////////////////////////////////////////
+void SMutableSelectFolderDlg::Construct(const FArguments& InArgs)
+{
+	AssetPath = FText::FromString(FPackageName::GetLongPackagePath(InArgs._DefaultAssetPath.ToString()));
+	FileName = InArgs._DefaultFileName;
+
+	AddAllMaterialTextures = false;
+
+	if (AssetPath.IsEmpty())
+	{
+		AssetPath = FText::FromString(TEXT("/Game"));
+	}
+
+	FPathPickerConfig PathPickerConfig;
+	PathPickerConfig.DefaultPath = AssetPath.ToString();
+	PathPickerConfig.OnPathSelected = FOnPathSelected::CreateSP(this, &SMutableSelectFolderDlg::OnPathChange);
+	PathPickerConfig.bAddDefaultPath = true;
+
+	FContentBrowserModule& ContentBrowserModule = FModuleManager::LoadModuleChecked<FContentBrowserModule>("ContentBrowser");
+
+	SWindow::Construct(SWindow::FArguments()
+		.Title(LOCTEXT("SMutableSelectFolderDlg_Title", "Select target folder for baked resources"))
+		.SupportsMinimize(false)
+		.SupportsMaximize(false)
+		//.SizingRule( ESizingRule::Autosized )
+		.ClientSize(FVector2D(450, 450))
+		[
+			SNew(SVerticalBox)
+
+			+ SVerticalBox::Slot() // Add user input block
+		.Padding(2)
+		[
+			SNew(SBorder)
+			.BorderImage(FAppStyle::GetBrush("ToolPanel.GroupBorder"))
+		[
+			SNew(SVerticalBox)
+
+			+ SVerticalBox::Slot()
+		.AutoHeight()
+		[
+			SNew(STextBlock)
+			.Text(LOCTEXT("SelectPath", "Select Path"))
+		.Font(FSlateFontInfo(FPaths::EngineContentDir() / TEXT("Slate/Fonts/Roboto-Regular.ttf"), 14))
+		]
+
+	+ SVerticalBox::Slot()
+		.FillHeight(1)
+		.Padding(3)
+		[
+			ContentBrowserModule.Get().CreatePathPicker(PathPickerConfig)
+		]
+
+	+ SVerticalBox::Slot()
+		.AutoHeight()
+		[
+			SNew(STextBlock)
+			.Text(LOCTEXT("FileName", "File Name"))
+			.Font(FSlateFontInfo(FPaths::EngineContentDir() / TEXT("Slate/Fonts/Roboto-Regular.ttf"), 14))
+		]
+
+	+ SVerticalBox::Slot()
+		.AutoHeight()
+		[
+			SNew(SEditableTextBox)
+			.Text(InArgs._DefaultFileName)
+			.OnTextCommitted(this, &SMutableSelectFolderDlg::OnNameChange)
+			.MinDesiredWidth(250)
+		]
+
+		]
+		]
+
+	+ SVerticalBox::Slot()
+		.AutoHeight()
+		[
+			SNew(SHorizontalBox)
+			+ SHorizontalBox::Slot()
+			.AutoWidth()
+			.HAlign(HAlign_Left)
+			[
+				SNew(STextBlock)
+				.Text(LOCTEXT("ExportAllUsedResources", "Export all used resources  "))
+				.Font(FSlateFontInfo(FPaths::EngineContentDir() / TEXT("Slate/Fonts/Roboto-Regular.ttf"), 14))
+				.ToolTipText(LOCTEXT("Export all used Resources", "All the materials and textures used by the object will be baked/stored in the target folder. Otherwise, only the assets that Mutable modifies will be baked/stored."))
+			]
+			+ SHorizontalBox::Slot()
+			.HAlign(HAlign_Left)
+			.AutoWidth()
+			[
+				SNew(SCheckBox)
+				.ToolTipText(LOCTEXT("ExportAllTextures", "Export all material's textures"))
+				.HAlign(HAlign_Right)
+				.IsChecked(AddAllMaterialTextures ? ECheckBoxState::Checked : ECheckBoxState::Unchecked)
+				.OnCheckStateChanged(this, &SMutableSelectFolderDlg::OnBoolParameterChanged)
+			]
+		]
+
+	+ SVerticalBox::Slot()
+		.AutoHeight()
+		.HAlign(HAlign_Right)
+		.Padding(5)
+		[
+			SNew(SUniformGridPanel)
+			.SlotPadding(FAppStyle::GetMargin("StandardDialog.SlotPadding"))
+		.MinDesiredSlotWidth(FAppStyle::GetFloat("StandardDialog.MinDesiredSlotWidth"))
+		.MinDesiredSlotHeight(FAppStyle::GetFloat("StandardDialog.MinDesiredSlotHeight"))
+		+ SUniformGridPanel::Slot(0, 0)
+		[
+			SNew(SButton)
+			.HAlign(HAlign_Center)
+		.ContentPadding(FAppStyle::GetMargin("StandardDialog.ContentPadding"))
+		.Text(LOCTEXT("OK", "OK"))
+		.OnClicked(this, &SMutableSelectFolderDlg::OnButtonClick, EAppReturnType::Ok)
+		]
+	+ SUniformGridPanel::Slot(1, 0)
+		[
+			SNew(SButton)
+			.HAlign(HAlign_Center)
+		.ContentPadding(FAppStyle::GetMargin("StandardDialog.ContentPadding"))
+		.Text(LOCTEXT("Cancel", "Cancel"))
+		.OnClicked(this, &SMutableSelectFolderDlg::OnButtonClick, EAppReturnType::Cancel)
+		]
+		]
+		]);
+}
+
+void SMutableSelectFolderDlg::OnPathChange(const FString& NewPath)
+{
+	AssetPath = FText::FromString(NewPath);
+}
+
+FReply SMutableSelectFolderDlg::OnButtonClick(EAppReturnType::Type ButtonID)
+{
+	UserResponse = ButtonID;
+
+	RequestDestroyWindow();
+
+	return FReply::Handled();
+}
+
+
+void SMutableSelectFolderDlg::OnNameChange(const FText& NewName, ETextCommit::Type CommitInfo)
+{
+	FileName = NewName;
+}
+
+
+void SMutableSelectFolderDlg::OnBoolParameterChanged(ECheckBoxState InCheckboxState)
+{
+	AddAllMaterialTextures = !AddAllMaterialTextures;
+}
+
+
+EAppReturnType::Type SMutableSelectFolderDlg::ShowModal()
+{
+	GEditor->EditorAddModalWindow(SharedThis(this));
+	return UserResponse;
+}
+
+FString SMutableSelectFolderDlg::GetAssetPath()
+{
+	return AssetPath.ToString();
+}
+
+
+FString SMutableSelectFolderDlg::GetFileName()
+{
+	return FileName.ToString();
+}
+
+
+bool SMutableSelectFolderDlg::GetAddAllMaterialTextures()
+{
+	return AddAllMaterialTextures;
+}
+
+#undef LOCTEXT_NAMESPACE 
