[mutable] Move Mutable plugin from the restricted folder to the experimental plugins folder.

diff --git a/Engine/Plugins/Experimental/Mutable/Source/CustomizableObjectEditor/Private/CustomizableObjectInstanceEditor.cpp b/Engine/Plugins/Experimental/Mutable/Source/CustomizableObjectEditor/Private/CustomizableObjectInstanceEditor.cpp
new file mode 100644
--- /dev/null
+++ b/Engine/Plugins/Experimental/Mutable/Source/CustomizableObjectEditor/Private/CustomizableObjectInstanceEditor.cpp
@@ -0,0 +1,928 @@
+// Copyright Epic Games, Inc. All Rights Reserved.
+
+#include "CustomizableObjectInstanceEditor.h"
+#include "CustomizableObjectEditorStyle.h"
+
+#include "CustomizableObjectEditor.h"
+#include "SCustomizableObjectEditorViewport.h"
+#include "CustomizableObjectEditorActions.h"
+#include "CustomizableObjectInstanceEditorActions.h"
+#include "SCustomizableObjectProperties.h"
+#include "CustomizableObject.h"
+#include "CustomizableObjectInstance.h"
+#include "CustomizableObjectSystem.h"
+#include "CustomizableObjectCompiler.h"
+#include "CustomizableObjectPreviewScene.h"
+#include "SCustomizableObjectSystemVersion.h"
+
+#include "AssetRegistry/AssetRegistryModule.h"
+#include "PreviewScene.h"
+#include "ScopedTransaction.h"
+#include "FileHelpers.h"
+#include "BusyCursor.h"
+#include "IDetailsView.h"
+
+#include "Slate/SceneViewport.h"
+#include "Widgets/Docking/SDockTab.h"
+#include "Components/SkeletalMeshComponent.h"
+#include "UnrealEditorPortabilityHelpers.h"
+
+#include "Editor.h"
+#include "Engine/Selection.h"
+#include "ContentBrowserModule.h"
+#include "IContentBrowserSingleton.h"
+#include "WorkspaceMenuStructureModule.h"
+#include "IDetailsView.h"
+#include "PropertyEditorModule.h"
+#include "WorkspaceMenuStructure.h"
+#include "Framework/Notifications/NotificationManager.h"
+#include "AssetRegistry/AssetRegistryModule.h"
+#include "AdvancedPreviewSceneModule.h"
+#include "Widgets/Notifications/SNotificationList.h"
+#include "CustomizableObjectBakeHelpers.h"
+#include "Animation/AnimSingleNodeInstance.h"
+
+#include "SCustomizableObjectEditorTextureAnalyzer.h"
+#include "CustomizableObjectEditorViewportClient.h"
+
+#include "Engine/Texture.h"
+#include "Widgets/Input/SButton.h"
+#include "Materials/MaterialInstance.h"
+
+#define LOCTEXT_NAMESPACE "CustomizableObjectInstanceEditor"
+
+
+DEFINE_LOG_CATEGORY_STATIC(LogCustomizableObjectInstanceEditor, Log, All);
+
+const FName FCustomizableObjectInstanceEditor::ViewportTabId( TEXT( "CustomizableObjectInstanceEditor_Viewport" ) );
+const FName FCustomizableObjectInstanceEditor::InstancePropertiesTabId( TEXT( "CustomizableObjectInstanceEditor_InstanceProperties" ) );
+const FName FCustomizableObjectInstanceEditor::AdvancedPreviewSettingsTabId(TEXT("CustomizableObjectEditor_AdvancedPreviewSettings"));
+const FName FCustomizableObjectInstanceEditor::TextureAnalyzerTabId(TEXT("CustomizableObjectInstanceEditor_TextureAnalyzer"));
+
+void UUpdateClassWrapperClass::DelegatedCallback(UCustomizableObjectInstance* Instance)
+{
+	Delegate.ExecuteIfBound();
+}
+
+void FCustomizableObjectInstanceEditor::RegisterTabSpawners(const TSharedRef<class FTabManager>& InTabManager)
+{
+	const IWorkspaceMenuStructure& MenuStructure = WorkspaceMenu::GetMenuStructure();
+
+	InTabManager->RegisterTabSpawner( ViewportTabId, FOnSpawnTab::CreateSP(this, &FCustomizableObjectInstanceEditor::SpawnTab_Viewport) )
+		.SetDisplayName( LOCTEXT( "ViewportTab", "Viewport" ) )
+		.SetGroup( MenuStructure.GetToolsCategory() );
+	
+	InTabManager->RegisterTabSpawner( InstancePropertiesTabId, FOnSpawnTab::CreateSP(this, &FCustomizableObjectInstanceEditor::SpawnTab_InstanceProperties) )
+		.SetDisplayName( LOCTEXT( "InstancePropertiesTab", "Instance Properties" ) )
+		.SetGroup( MenuStructure.GetToolsCategory() );
+
+	InTabManager->RegisterTabSpawner(AdvancedPreviewSettingsTabId, FOnSpawnTab::CreateSP(this, &FCustomizableObjectInstanceEditor::SpawnTab_AdvancedPreviewSettings))
+		.SetDisplayName(LOCTEXT("AdvancedPreviewSettingsTab", "Advanced Preview Settings"))
+		.SetGroup(MenuStructure.GetToolsCategory());
+
+	InTabManager->RegisterTabSpawner(TextureAnalyzerTabId, FOnSpawnTab::CreateSP(this, &FCustomizableObjectInstanceEditor::SpawnTab_TextureAnalyzer))
+		.SetDisplayName(LOCTEXT("InstanceTextureAnalyzer", "Texture Analyzer"));
+}
+
+void FCustomizableObjectInstanceEditor::UnregisterTabSpawners(const TSharedRef<class FTabManager>& InTabManager)
+{
+	InTabManager->UnregisterTabSpawner( ViewportTabId );
+	InTabManager->UnregisterTabSpawner( InstancePropertiesTabId );
+	InTabManager->UnregisterTabSpawner(AdvancedPreviewSettingsTabId);
+	InTabManager->UnregisterTabSpawner(TextureAnalyzerTabId);
+}
+
+
+FCustomizableObjectInstanceEditor::FCustomizableObjectInstanceEditor():
+	CreatePreviewInstanceAfterCOCompile(false)
+{
+	CustomizableObjectInstance = nullptr;
+	//PreviewCustomizableSkeletalComponent = nullptr;
+	PreviewStaticMeshComponent = nullptr;
+	//PreviewSkeletalMeshComponent = nullptr;
+	AssetRegistryLoaded = false;
+	UpdateSkeletalMeshAfterAssetLoaded = false;
+	ObjectToEditCopy = nullptr;
+
+	FAssetRegistryModule& AssetRegistryModule = FModuleManager::LoadModuleChecked<FAssetRegistryModule>(TEXT("AssetRegistry"));
+	if (AssetRegistryModule.Get().IsLoadingAssets())
+	{
+		AssetRegistryModule.Get().OnFilesLoaded().AddRaw(this, &FCustomizableObjectInstanceEditor::OnAssetRegistryLoadComplete);
+	}
+	else
+	{
+		AssetRegistryLoaded = true;
+	}
+
+	HelperCallback = nullptr;
+
+	PoseAsset = nullptr;
+}
+
+
+FCustomizableObjectInstanceEditor::~FCustomizableObjectInstanceEditor()
+{
+	if (HelperCallback) 
+	{
+		CustomizableObjectInstance->UpdatedDelegate.RemoveDynamic(HelperCallback, &UUpdateClassWrapperClass::DelegatedCallback);
+		HelperCallback->Delegate.Unbind();
+		HelperCallback = nullptr;
+	}
+
+	CustomizableObjectInstance = nullptr;
+
+	for (UCustomizableSkeletalComponent* PreviewCustomizableSkeletalComponent : PreviewCustomizableSkeletalComponents)
+	{
+		if (PreviewCustomizableSkeletalComponent)
+		{
+			PreviewCustomizableSkeletalComponent->DestroyComponent();
+		}
+	}
+
+	PreviewCustomizableSkeletalComponents.Reset();
+	PreviewStaticMeshComponent = nullptr;
+	PreviewSkeletalMeshComponents.Reset();
+
+	CustomizableInstanceDetailsView.Reset();
+	Viewport.Reset();
+
+	FCoreUObjectDelegates::OnObjectModified.Remove(OnObjectModifiedHandle);
+	Compiler->ForceFinishCompilation();
+
+	FAssetRegistryModule& AssetRegistryModule = FModuleManager::LoadModuleChecked<FAssetRegistryModule>(TEXT("AssetRegistry"));
+	AssetRegistryModule.Get().OnFilesLoaded().RemoveAll(this);
+
+	FCoreUObjectDelegates::OnObjectPropertyChanged.RemoveAll(this);
+}
+
+
+void FCustomizableObjectInstanceEditor::InitCustomizableObjectInstanceEditor( const EToolkitMode::Type Mode, const TSharedPtr< class IToolkitHost >& InitToolkitHost, UCustomizableObjectInstance* ObjectToEdit )
+{
+	// Register our commands. This will only register them if not previously registered
+	FCustomizableObjectInstanceEditorCommands::Register();
+	FCustomizableObjectEditorViewportCommands::Register();
+
+	BindCommands();
+
+	bool bUpdateFromSelection = false;
+	bool bIsLockable = false;
+	bool bAllowSearch = true;
+
+	FPropertyEditorModule& PropPlugin = FModuleManager::LoadModuleChecked<FPropertyEditorModule>("PropertyEditor");
+
+	FDetailsViewArgs DetailsViewArgs;
+	DetailsViewArgs.NameAreaSettings = FDetailsViewArgs::ENameAreaSettings::HideNameArea;
+	DetailsViewArgs.bAllowSearch = false;
+	//DetailsViewArgs.bShowActorLabel = false;
+	DetailsViewArgs.bShowObjectLabel = false;
+	CustomizableInstanceDetailsView = PropPlugin.CreateDetailView( DetailsViewArgs );
+	//CustomizableInstanceDetailsView->SetIsPropertyVisibleDelegate( FIsPropertyVisible::CreateStatic(&IsPropertyVisible) );
+
+	TextureAnalyzer = SNew(SCustomizableObjecEditorTextureAnalyzer).CustomizableObjectInstanceEditor(this).CustomizableObjectEditor(nullptr);
+
+	Compiler = MakeUnique<FCustomizableObjectCompiler>();
+	check(Compiler);
+	Viewport = SNew(SCustomizableObjectEditorViewportTabBody)
+		.CustomizableObjectEditor(SharedThis(this));
+
+	if (AssetRegistryLoaded)
+	{
+		if (ObjectToEdit)
+		{
+			if (UCustomizableObject* CustomizableObject = CustomizableObjectInstance->GetCustomizableObject();
+				CustomizableObject &&
+				!CustomizableObject->IsCompiled())
+			{
+				// Compile for the first time if necessary
+				CompileObject(CustomizableObject);
+				CreatePreviewInstanceAfterCOCompile = true;
+				ObjectToEditCopy = ObjectToEdit;
+			}
+		}
+
+		SetEditorMesh(ObjectToEdit);
+		if (!CreatePreviewInstanceAfterCOCompile)
+		{
+			CreatePreviewInstance();
+		}
+		
+		Viewport->SetAssetRegistryLoaded(true);
+	}
+	else
+	{
+		ObjectToEditCopy = ObjectToEdit;
+		UpdateSkeletalMeshAfterAssetLoaded = true;
+	}
+
+	FAdvancedPreviewSceneModule& AdvancedPreviewSceneModule = FModuleManager::LoadModuleChecked<FAdvancedPreviewSceneModule>("AdvancedPreviewScene");
+
+	TSharedPtr<FAdvancedPreviewScene> AdvancedPreviewScene = StaticCastSharedPtr<FAdvancedPreviewScene>(Viewport->GetPreviewScene());
+	AdvancedPreviewSettingsWidget = AdvancedPreviewSceneModule.CreateAdvancedPreviewSceneSettingsWidget(AdvancedPreviewScene.ToSharedRef());
+
+	const TSharedRef<FTabManager::FLayout> StandaloneDefaultLayout = FTabManager::NewLayout( "Standalone_CustomizableObjectInstanceEditor_Layout_v2.1" )
+	->AddArea
+	(
+		FTabManager::NewPrimaryArea()->SetOrientation(Orient_Horizontal)
+		->Split
+		(
+			FTabManager::NewStack()
+			->SetSizeCoefficient(0.8f)
+			->AddTab(ViewportTabId, ETabState::OpenedTab)
+		)
+		->Split
+		(
+			FTabManager::NewStack()
+			->SetSizeCoefficient(0.2f)
+			->AddTab(InstancePropertiesTabId, ETabState::OpenedTab)
+		)
+		->Split
+		(
+			FTabManager::NewStack()
+			->SetSizeCoefficient(0.2f)
+			->AddTab(AdvancedPreviewSettingsTabId, ETabState::OpenedTab)
+		)
+	);
+
+	const bool bCreateDefaultStandaloneMenu = true;
+	const bool bCreateDefaultToolbar = true;
+	FAssetEditorToolkit::InitAssetEditor( Mode, InitToolkitHost, CustomizableObjectInstanceEditorAppIdentifier, StandaloneDefaultLayout, bCreateDefaultStandaloneMenu, bCreateDefaultToolbar, ObjectToEdit );
+
+	ExtendToolbar();
+	RegenerateMenusAndToolbars();
+
+	// Clears selection highlight.
+	OnInstancePropertySelectionChanged(NULL);
+
+	OnObjectModifiedHandle = FCoreUObjectDelegates::OnObjectModified.AddRaw(this, &FCustomizableObjectInstanceEditor::OnObjectModified);
+
+	bOnlyRelevantParameters = true;
+	bOnlyRuntimeParameters = true;
+	
+	if (CustomizableObjectInstance)
+	{
+		bOnlyRelevantParameters = CustomizableObjectInstance->bShowOnlyRelevantParameters;
+		bOnlyRuntimeParameters = CustomizableObjectInstance->bShowOnlyRuntimeParameters;
+	}
+}
+
+
+FName FCustomizableObjectInstanceEditor::GetToolkitFName() const
+{
+	return FName("CustomizableObjectInstanceEditor");
+}
+
+
+FText FCustomizableObjectInstanceEditor::GetBaseToolkitName() const
+{
+	return LOCTEXT("ToolkitName", "Customizable Object Instance Editor");
+}
+
+
+void FCustomizableObjectInstanceEditor::CreatePreviewInstance()
+{
+	UCustomizableObject* CustomizableObject = CustomizableObjectInstance->GetCustomizableObject();
+
+	HelperCallback = NewObject<UUpdateClassWrapperClass>(GetTransientPackage());
+	HelperCallback->Delegate.BindSP(this, &FCustomizableObjectInstanceEditor::OnUpdatePreviewInstance);
+	CustomizableObjectInstance->UpdatedDelegate.AddDynamic(HelperCallback, &UUpdateClassWrapperClass::DelegatedCallback);
+
+	if (!CustomizableObject->IsCompiled() && !UCustomizableObjectSystem::GetInstance()->IsCompilationDisabled()) // In case the compilation failed
+	{
+		return;
+	}
+
+	PreviewStaticMeshComponent = nullptr;
+
+	// Getting the number of mesh components from the root node
+	int32 NumMeshComponents = 0;
+
+	if (CustomizableObject->Source)
+	{
+		TArray<UCustomizableObjectNodeObject*> RootNode;
+		CustomizableObject->Source->GetNodesOfClass< UCustomizableObjectNodeObject>(RootNode);
+
+		for (int32 i = 0; i < RootNode.Num(); ++i)
+		{
+			if (RootNode[i]->bIsBase)
+			{
+				NumMeshComponents = RootNode[i]->NumMeshComponents;
+				break;
+			}
+		}
+	}
+
+	PreviewSkeletalMeshComponents.AddZeroed(NumMeshComponents);
+	PreviewCustomizableSkeletalComponents.AddZeroed(NumMeshComponents);
+
+	if (CustomizableObject->ReferenceSkeletalMeshes.Num())
+	{
+		for (int32 ComponentIndex = 0; ComponentIndex < NumMeshComponents; ++ComponentIndex)
+		{
+			PreviewCustomizableSkeletalComponents[ComponentIndex] = NewObject<UCustomizableSkeletalComponent>(UCustomizableSkeletalComponent::StaticClass());
+
+			if (PreviewCustomizableSkeletalComponents[ComponentIndex])
+			{
+				PreviewCustomizableSkeletalComponents[ComponentIndex]->CustomizableObjectInstance = CustomizableObjectInstance;
+				PreviewCustomizableSkeletalComponents[ComponentIndex]->ComponentIndex = ComponentIndex;
+
+				PreviewSkeletalMeshComponents[ComponentIndex] = NewObject<UDebugSkelMeshComponent>(GetTransientPackage(), NAME_None, RF_Transient);
+
+				if (PreviewSkeletalMeshComponents[ComponentIndex])
+				{
+					PreviewSkeletalMeshComponents[ComponentIndex]->UseInGameBounds(true);
+					PreviewCustomizableSkeletalComponents[ComponentIndex]->AttachToComponent(PreviewSkeletalMeshComponents[ComponentIndex], FAttachmentTransformRules::KeepRelativeTransform);
+				}
+				else
+				{
+					return;
+				}
+			}
+			else
+			{
+				return;
+			}
+		}
+
+		CustomizableObjectInstance->UpdateSkeletalMeshAsync(true, true);
+		Viewport->SetPreviewComponents(PreviewSkeletalMeshComponents);
+	}
+}
+
+
+void FCustomizableObjectInstanceEditor::OnAssetRegistryLoadComplete()
+{
+	AssetRegistryLoaded = true;
+
+	Viewport->SetAssetRegistryLoaded(true);
+
+	if (UpdateSkeletalMeshAfterAssetLoaded)
+	{
+		if (ObjectToEditCopy && ObjectToEditCopy->GetCustomizableObject() && !ObjectToEditCopy->GetCustomizableObject()->IsCompiled())
+		{
+			// Compile for the first time if necessary
+			CompileObject(ObjectToEditCopy->GetCustomizableObject());
+			CreatePreviewInstanceAfterCOCompile = true;
+		}
+		else
+		{
+			SetEditorMesh(ObjectToEditCopy);
+			CreatePreviewInstance();
+		}
+		UpdateSkeletalMeshAfterAssetLoaded = false;
+	}
+}
+
+
+void FCustomizableObjectInstanceEditor::UpdatePreviewVisibility()
+{
+	for (UDebugSkelMeshComponent* PreviewSkeletalMeshComponent : PreviewSkeletalMeshComponents)
+	{
+		if (CustomizableObjectInstance->SkeletalMeshStatus != ESkeletalMeshState::UpdateError)
+		{
+			PreviewSkeletalMeshComponent->SetVisibility(true, true);
+		}
+		else
+		{
+			PreviewSkeletalMeshComponent->SetVisibility(false, true);
+		}
+	}
+}
+
+
+void FCustomizableObjectInstanceEditor::SaveAsset_Execute()
+{
+	if (CustomizableObjectInstance)
+	{
+		UPackage* Package = CustomizableObjectInstance->GetOutermost();
+
+		if (Package)
+		{
+			TArray<UPackage*> PackagesToSave;
+			PackagesToSave.Add(Package);
+
+			FEditorFileUtils::PromptForCheckoutAndSave(PackagesToSave, false, false);
+		}
+	}
+}
+
+
+bool FCustomizableObjectInstanceEditor::CanOpenOrShowParent()
+{
+	return CustomizableObjectInstance->GetCustomizableObject() != nullptr;
+}
+
+
+void FCustomizableObjectInstanceEditor::ShowParentInContentBrowser()
+{
+	FContentBrowserModule& ContentBrowserModule = FModuleManager::Get().LoadModuleChecked<FContentBrowserModule>(TEXT("ContentBrowser"));
+	ContentBrowserModule.Get().SyncBrowserToAssets(TArray<FAssetData>{FAssetData(CustomizableObjectInstance->GetCustomizableObject())});
+}
+
+
+void FCustomizableObjectInstanceEditor::OpenParentInEditor()
+{
+	GEditor->GetEditorSubsystem<UAssetEditorSubsystem>()->OpenEditorForAsset(CustomizableObjectInstance->GetCustomizableObject());
+}
+
+
+void FCustomizableObjectInstanceEditor::AddReferencedObjects( FReferenceCollector& Collector )
+{
+	Collector.AddReferencedObject( CustomizableObjectInstance );
+	
+	for (UCustomizableSkeletalComponent* PreviewCustomizableSkeletalComponent : PreviewCustomizableSkeletalComponents)
+	{
+		Collector.AddReferencedObject(PreviewCustomizableSkeletalComponent);
+	}
+
+	Collector.AddReferencedObject( PreviewStaticMeshComponent );
+
+	for (UDebugSkelMeshComponent* PreviewSkeletalMeshComponent : PreviewSkeletalMeshComponents)
+	{
+		Collector.AddReferencedObject(PreviewSkeletalMeshComponent);
+	}
+
+	Collector.AddReferencedObject( HelperCallback );
+}
+
+
+TSharedRef<SDockTab> FCustomizableObjectInstanceEditor::SpawnTab_Viewport( const FSpawnTabArgs& Args )
+{
+	check( Args.GetTabId() == ViewportTabId );
+
+	return SNew(SDockTab)
+		//.Icon(FCustomizableObjectEditorStyle::Get().GetBrush("CustomizableObjectEditor.Tabs.Preview"))
+		.Label( FText::FromString( GetTabPrefix() + LOCTEXT("CustomizableObjectViewport_TabTitle", "Instance Viewport").ToString() ) )
+		[
+			Viewport.ToSharedRef()
+		];
+}
+
+
+TSharedRef<SDockTab> FCustomizableObjectInstanceEditor::SpawnTab_InstanceProperties( const FSpawnTabArgs& Args )
+{
+	check( Args.GetTabId() == InstancePropertiesTabId );
+
+	return SNew(SDockTab)
+		//.Icon(FCustomizableObjectEditorStyle::Get().GetBrush("CustomizableObjectEditor.Tabs.CustomizableInstanceProperties"))
+		.Label( FText::FromString( GetTabPrefix() + LOCTEXT( "CustomizableInstanceProperties_TabTitle", "Instance Properties" ).ToString() ) )
+		[
+			CustomizableInstanceDetailsView.ToSharedRef()
+		];
+}
+
+
+TSharedRef<SDockTab> FCustomizableObjectInstanceEditor::SpawnTab_AdvancedPreviewSettings(const FSpawnTabArgs& Args)
+{
+	check(Args.GetTabId() == AdvancedPreviewSettingsTabId);
+	return SNew(SDockTab)
+		//.Icon(FCustomizableObjectEditorStyle::Get().GetBrush("CustomizableObjectEditor.Tabs.PreviewSettings"))
+		.Label(LOCTEXT("StaticMeshPreviewScene_TabTitle", "Preview Scene Settings"))
+		[
+			AdvancedPreviewSettingsWidget.ToSharedRef()
+		];
+}
+
+
+UCustomizableObjectInstance* FCustomizableObjectInstanceEditor::GetPreviewInstance()
+{
+	return CustomizableObjectInstance;
+}
+
+
+void FCustomizableObjectInstanceEditor::BindCommands()
+{
+	const FCustomizableObjectInstanceEditorCommands& Commands = FCustomizableObjectInstanceEditorCommands::Get();
+
+	const TSharedRef<FUICommandList>& UICommandList = GetToolkitCommands();
+
+	// Mesh menu
+	//UICommandList->MapAction(
+	//	Commands.Compile,
+	//	FExecuteAction::CreateSP( this, &FCustomizableObjectInstanceEditor::Compile ),
+	//	FCanExecuteAction(),
+	//	FIsActionChecked() );
+
+	
+	UICommandList->MapAction(
+		Commands.ShowParentCO,
+		FExecuteAction::CreateSP(this, &FCustomizableObjectInstanceEditor::ShowParentInContentBrowser),
+		FCanExecuteAction::CreateSP(this, &FCustomizableObjectInstanceEditor::CanOpenOrShowParent),
+		FIsActionChecked()
+	);
+
+	UICommandList->MapAction(
+		Commands.EditParentCO,
+		FExecuteAction::CreateSP(this, &FCustomizableObjectInstanceEditor::OpenParentInEditor),
+		FCanExecuteAction::CreateSP(this, &FCustomizableObjectInstanceEditor::CanOpenOrShowParent),
+		FIsActionChecked()
+	);
+
+	// Texture Analyzer
+	UICommandList->MapAction(
+		Commands.TextureAnalyzer,
+		FExecuteAction::CreateSP(this, &FCustomizableObjectInstanceEditor::OpenTextureAnalyzerTab),
+		FCanExecuteAction(),
+		FIsActionChecked());
+}
+
+
+void FCustomizableObjectInstanceEditor::ExtendToolbar()
+{
+	TSharedPtr<FUICommandList> CommandList = GetToolkitCommands();
+
+	struct Local
+	{
+		static void FillToolbar(FToolBarBuilder& ToolbarBuilder, FCustomizableObjectInstanceEditor* Editor, TSharedPtr<FUICommandList> CommandList)
+		{
+			ToolbarBuilder.BeginSection("Utilities");
+			ToolbarBuilder.AddToolBarButton(FCustomizableObjectInstanceEditorCommands::Get().ShowParentCO);
+			ToolbarBuilder.AddToolBarButton(FCustomizableObjectInstanceEditorCommands::Get().EditParentCO);
+			ToolbarBuilder.AddToolBarButton(FCustomizableObjectInstanceEditorCommands::Get().TextureAnalyzer);
+			ToolbarBuilder.EndSection();
+		}
+	};
+
+	TSharedPtr<FExtender> ToolbarExtender = MakeShareable(new FExtender);
+
+	ToolbarExtender->AddToolBarExtension(
+		"Asset",
+		EExtensionHook::After,
+		GetToolkitCommands(),
+		FToolBarExtensionDelegate::CreateStatic(&Local::FillToolbar, this, CommandList));
+
+	AddToolbarExtender(ToolbarExtender);
+
+	ICustomizableObjectEditorModule* CustomizableObjectEditorModule = &FModuleManager::LoadModuleChecked<ICustomizableObjectEditorModule>("CustomizableObjectEditor");
+	AddToolbarExtender(CustomizableObjectEditorModule->GetCustomizableObjectEditorToolBarExtensibilityManager()->GetAllExtenders());
+}
+
+
+FText FCustomizableObjectInstanceEditor::GetToolkitName() const
+{
+	FFormatNamedArguments Args;
+	Args.Add(TEXT("ObjectName"), FText::FromString(GetEditingObject()->GetName()));
+	Args.Add(TEXT("ToolkitName"), GetBaseToolkitName());
+	return FText::Format(LOCTEXT("AppLabelWithAssetName", "{ObjectName} - {ToolkitName}"), Args);
+}
+
+
+FString FCustomizableObjectInstanceEditor::GetWorldCentricTabPrefix() const
+{
+	return LOCTEXT( "WorldCentricTabPrefix", "CustomizableObjectInstance " ).ToString();
+}
+
+
+FLinearColor FCustomizableObjectInstanceEditor::GetWorldCentricTabColorScale() const
+{
+	return FLinearColor( 0.3f, 0.2f, 0.5f, 0.5f );
+}
+
+
+void FCustomizableObjectInstanceEditor::RefreshTool()
+{
+	RefreshViewport();
+}
+
+
+void FCustomizableObjectInstanceEditor::RefreshViewport()
+{
+	if (Viewport)
+	{
+		Viewport->GetViewportClient()->Invalidate();
+	}
+}
+
+
+void FCustomizableObjectInstanceEditor::SetPoseAsset(class UPoseAsset* PoseAssetParameter)
+{
+	PoseAsset = PoseAssetParameter;
+
+	if (PoseAsset != nullptr)
+	{
+		Viewport->SetAnimation(nullptr, EAnimationMode::AnimationBlueprint);
+
+		for (UDebugSkelMeshComponent* PreviewSkeletalMeshComponent : PreviewSkeletalMeshComponents)
+		{
+			PreviewSkeletalMeshComponent->SetAnimationMode(EAnimationMode::AnimationBlueprint);
+			PreviewSkeletalMeshComponent->InitAnim(false);
+			PreviewSkeletalMeshComponent->SetAnimation(PoseAsset);
+
+			UAnimSingleNodeInstance* SingleNodeInstance = Cast<UAnimSingleNodeInstance>(PreviewSkeletalMeshComponent->GetAnimInstance());
+			if (SingleNodeInstance)
+			{
+				TArray<FSmartName> ArrayPoseSmartNames = PoseAsset->GetPoseNames();
+				for (int32 i = 0; i < ArrayPoseSmartNames.Num(); ++i)
+				{
+					SingleNodeInstance->SetPreviewCurveOverride(ArrayPoseSmartNames[i].DisplayName, 1.0f, false);
+				}
+			}
+		}
+	}
+}
+
+
+void FCustomizableObjectInstanceEditor::OnObjectModified(UObject* Object)
+{
+	UCustomizableObjectInstance* Instance = Cast<UCustomizableObjectInstance>(Object);
+
+	if ((Instance != nullptr) && (Instance == CustomizableObjectInstance))
+	{
+		if (Instance->ProjectorLayerChange)
+		{
+			// Variable used in a non multilayer projector to not lose the focus of the gizmo
+			Instance->ProjectorLayerChange = false;
+			return;
+		}
+
+		if (Viewport->GetGizmoHasAssignedData())
+		{
+			Instance->LastSelectedProjectorParameter = Viewport->GetGizmoProjectorParameterName();
+			Instance->LastSelectedProjectorParameterWithIndex = Viewport->GetGizmoProjectorParameterNameWithIndex();
+		}
+
+		// Update projector parameter information (not projector parameter node information,
+		// this is the case where a projector parameter in the object details is selected)
+		const TArray<FCustomizableObjectProjectorParameterValue>& Parameters = CustomizableObjectInstance->GetProjectorParameters();
+		const int32 MaxIndex = Parameters.Num();
+		bool AnySelected = false;
+
+		for (int32 i = 0; (!AnySelected && (i < MaxIndex)); ++i)
+		{
+			int32 RangeIndex = Parameters[i].RangeValues.Num() > 0 ? 0 : -1;
+
+			for (; RangeIndex < Parameters[i].RangeValues.Num(); ++RangeIndex)
+			{
+				const FCustomizableObjectProjector& Value = (RangeIndex == -1) ? Parameters[i].Value : Parameters[i].RangeValues[RangeIndex];
+
+				FString ParameterNameWithIndex = Parameters[i].ParameterName;
+
+				if (RangeIndex != -1)
+				{
+					ParameterNameWithIndex += FString::Printf(TEXT("__%d"), RangeIndex);
+				}
+
+				if (CustomizableObjectInstance->GetProjectorState(Parameters[i].ParameterName, RangeIndex) == EProjectorState::Selected ||
+					(CustomizableObjectInstance->GetProjectorState(Parameters[i].ParameterName, RangeIndex) == EProjectorState::TypeChanged &&
+						ParameterNameWithIndex != Instance->LastSelectedProjectorParameterWithIndex))
+				{
+					Viewport->SetGizmoCallUpdateSkeletalMesh(false);
+					Viewport->SetProjectorVisibility(true, Parameters[i].ParameterName, ParameterNameWithIndex, RangeIndex, Value, i);
+					Viewport->SetGizmoCallUpdateSkeletalMesh(true);
+				AnySelected = true;
+			}
+				else if (CustomizableObjectInstance->GetProjectorState(Parameters[i].ParameterName, RangeIndex) == EProjectorState::TypeChanged &&
+					Parameters[i].ParameterName == Instance->LastSelectedProjectorParameter)
+					{
+					Viewport->SetProjectorType(true, Parameters[i].ParameterName, ParameterNameWithIndex, RangeIndex, Value, i);
+						AnySelected = true;
+					}
+				}
+
+		}
+
+		if (!AnySelected && !Viewport->GetIsManipulatingGizmo())
+		{
+			Viewport->SetGizmoCallUpdateSkeletalMesh(false);
+			Viewport->ResetProjectorVisibility(false);
+			Viewport->SetGizmoCallUpdateSkeletalMesh(true);
+			Instance->LastSelectedProjectorParameter = "";
+			Instance->LastSelectedProjectorParameterWithIndex = "";
+		}
+	}
+}
+
+
+bool FCustomizableObjectInstanceEditor::GetAssetRegistryLoaded()
+{
+	return AssetRegistryLoaded;
+}
+
+
+void FCustomizableObjectInstanceEditor::SetEditorMesh(UCustomizableObjectInstance* InCustomizableObjectInstance)
+{
+	CustomizableObjectInstance = InCustomizableObjectInstance;
+
+	// Set the details view.
+	TArray<UObject*> SelectedObjects;
+	SelectedObjects.Add(CustomizableObjectInstance);
+	CustomizableInstanceDetailsView->SetObjects(SelectedObjects);
+}
+
+
+void FCustomizableObjectInstanceEditor::OnChangeMesh()
+{
+	FEditorDelegates::LoadSelectedAssetsIfNeeded.Broadcast();
+	UCustomizableObjectInstance* SelectedMesh = GEditor->GetSelectedObjects()->GetTop<UCustomizableObjectInstance>();
+	if(SelectedMesh && SelectedMesh != CustomizableObjectInstance)
+	{
+		Helper_GetEditorSubsystem()->NotifyEditorClosed(this);
+		Helper_GetEditorSubsystem()->NotifyAssetOpened(SelectedMesh, this);
+
+		SetEditorMesh(SelectedMesh);
+//		LevelOfDetailSettings->MeshChanged();
+
+		// Refresh the tool so it will update it's LOD list.
+		//if(GenerateUniqueUVs.IsValid())
+		//{
+		//	GenerateUniqueUVs->RefreshTool();
+		//}
+	}
+}
+
+
+void FCustomizableObjectInstanceEditor::OnInstancePropertySelectionChanged(FProperty* InProperty)
+{
+	RefreshViewport();
+}
+
+
+void FCustomizableObjectInstanceEditor::OnUpdatePreviewInstance()
+{
+	RefreshViewport();
+
+	UpdatePreviewVisibility();
+
+	bool bNeedsToResetPreviewComponents = false;
+	for (UDebugSkelMeshComponent* PreviewSkeletalMeshComponent : PreviewSkeletalMeshComponents)
+	{
+		PreviewSkeletalMeshComponent->bComponentUseFixedSkelBounds = true; // First bounds computed would be using physics asset
+		PreviewSkeletalMeshComponent->UpdateBounds();
+
+		if (!PreviewSkeletalMeshComponent->IsRegistered())
+		{
+			bNeedsToResetPreviewComponents = true;
+			PreviewSkeletalMeshComponent->RegisterComponentWithWorld(PreviewSkeletalMeshComponent->GetWorld());
+		}
+	}
+
+	if(bNeedsToResetPreviewComponents)
+	{
+		Viewport->SetPreviewComponents(PreviewSkeletalMeshComponents);
+	}
+
+	if (!CustomizableObjectInstance) return;
+	
+	// If the instance is updated due to a change in a parameter projector, set again the LastSelectedProjectorParameter variable with the
+	// current projector the viewport gizmo has to continue having it as selected when rebuilding the parameter details widget.
+	// Otherwise, reset the value.
+	bool UpdateDetails = false;
+	if (Viewport->GetGizmoHasAssignedData())
+	{
+		CustomizableObjectInstance->LastSelectedProjectorParameter = Viewport->GetGizmoProjectorParameterName();
+		CustomizableObjectInstance->LastSelectedProjectorParameterWithIndex = Viewport->GetGizmoProjectorParameterNameWithIndex();
+		UpdateDetails = true;
+	}
+	else
+	{
+		CustomizableObjectInstance->LastSelectedProjectorParameter = "";
+		CustomizableObjectInstance->LastSelectedProjectorParameterWithIndex = "";
+	}
+
+	if ( !CustomizableObjectInstance->ProjectorUpdatedInViewport)
+	{
+		TArray<UObject*> instances;
+		instances.Add(CustomizableObjectInstance);
+		CustomizableInstanceDetailsView->SetObjects(instances, UpdateDetails);
+	}
+
+	if (TextureAnalyzer.IsValid())
+	{
+		TextureAnalyzer.Get()->RefreshTextureAnalyzerTable(CustomizableObjectInstance);
+	}
+	
+	for (int32 ComponentIndex = 0; ComponentIndex < CustomizableObjectInstance->SkeletalMeshes.Num(); ++ComponentIndex)
+	{
+		BakeHelper_RegenerateImportedModel(CustomizableObjectInstance->GetSkeletalMesh(ComponentIndex));
+	}
+}
+
+
+void FCustomizableObjectInstanceEditor::CompileObject(UCustomizableObject* Object)
+{
+	UE_LOG(LogMutable, Log, TEXT("PROFILE: -----------------------------------------------------------"));
+	UE_LOG(LogMutable, Log, TEXT("PROFILE: [ %16.8f ] FCustomizableObjectInstanceEditor::CompileObject start."), FPlatformTime::Seconds());
+
+	if (!AssetRegistryLoaded)
+	{
+		FNotificationInfo Info(NSLOCTEXT("CustomizableObject", "CustomizableObjectCompileTryLater", "Please wait until asset registry loads all assets"));
+		Info.bFireAndForget = true;
+		Info.bUseThrobber = true;
+		Info.FadeOutDuration = 1.0f;
+		Info.ExpireDuration = 2.0f;
+		FSlateNotificationManager::Get().AddNotification(Info);
+		return;
+	}
+
+	Viewport->UpdateGizmoDataToOrigin();
+
+	if (Object && Object->Source)
+	{
+		FCompilationOptions Options = Object->CompileOptions;
+		Options.bSilentCompilation = false;
+		Compiler->Compile(*Object, Options, true);
+	}
+
+	UE_LOG(LogMutable, Log, TEXT("PROFILE: [ %16.8f ] FCustomizableObjectInstanceEditor::CompileObject end."), FPlatformTime::Seconds());
+}
+
+
+bool FCustomizableObjectInstanceEditor::IsTickable(void) const
+{
+	return true;
+}
+
+
+void FCustomizableObjectInstanceEditor::Tick(float InDeltaTime)
+{
+	bool bUpdated = Compiler->Tick();
+
+	if (bUpdated)
+	{
+		if ((CustomizableObjectInstance == nullptr) && (ObjectToEditCopy != nullptr))
+		{
+			CustomizableObjectInstance = ObjectToEditCopy;
+		}
+
+		if (!CustomizableObjectInstance) return;
+
+		if (const UCustomizableObject* CustomizableObject = CustomizableObjectInstance->GetCustomizableObject();
+			CustomizableObject &&
+			CustomizableObject->Source)
+		{
+			if (PreviewCustomizableSkeletalComponents.Num() > 0)
+			{
+				CustomizableObjectInstance->UpdateSkeletalMeshAsync(true, true);
+			}
+			else
+			{
+				CreatePreviewInstance();
+
+				if (CreatePreviewInstanceAfterCOCompile)
+				{
+					CreatePreviewInstanceAfterCOCompile = false;
+				}
+			}
+		}
+
+		SetEditorMesh(CustomizableObjectInstance);
+	}
+
+	if (!CustomizableObjectInstance) return;
+
+	// TEMP CODE
+	if (CustomizableObjectInstance->TempUpdateGizmoInViewport)
+	{
+		CustomizableObjectInstance->TempUpdateGizmoInViewport = false;
+		CustomizableObjectInstance->AvoidResetProjectorVisibilityForNonNode = true;
+		Viewport->CopyTransformFromOriginData();
+		CustomizableObjectInstance->UpdateSkeletalMeshAsync(true);
+
+		EProjectorState::Type ProjectorState = CustomizableObjectInstance->GetProjectorState(CustomizableObjectInstance->TempProjectorParameterName, CustomizableObjectInstance->TempProjectorParameterRangeIndex);
+		if (ProjectorState != EProjectorState::Selected)
+		{
+			// ResetProjectorVisibilityForNonNode = false;
+			CustomizableObjectInstance->SetProjectorState(CustomizableObjectInstance->TempProjectorParameterName, CustomizableObjectInstance->TempProjectorParameterRangeIndex, EProjectorState::Selected);
+			FCoreUObjectDelegates::BroadcastOnObjectModified(CustomizableObjectInstance);
+		}
+	}
+
+	UE_LOG(LogMutable, Log, TEXT("PROFILE: [ %16.8f ] FCustomizableObjectInstanceEditor::CompileObject end."), FPlatformTime::Seconds());
+
+	// If we want to show the Relevant/Runtime parameters, we need to refresh the details view to make sure that the scroll bar appears		
+	if (bOnlyRelevantParameters != CustomizableObjectInstance->bShowOnlyRelevantParameters)
+	{
+		bOnlyRelevantParameters = CustomizableObjectInstance->bShowOnlyRelevantParameters;
+		CustomizableInstanceDetailsView->ForceRefresh();
+	}
+
+	if (bOnlyRuntimeParameters != CustomizableObjectInstance->bShowOnlyRuntimeParameters)
+	{
+		bOnlyRuntimeParameters = CustomizableObjectInstance->bShowOnlyRuntimeParameters;
+		CustomizableInstanceDetailsView->ForceRefresh();
+	}
+}
+
+
+TStatId FCustomizableObjectInstanceEditor::GetStatId() const
+{
+	RETURN_QUICK_DECLARE_CYCLE_STAT(FCustomizableObjectInstanceEditor, STATGROUP_Tickables);
+}
+
+void FCustomizableObjectInstanceEditor::OpenTextureAnalyzerTab()
+{
+	TabManager->TryInvokeTab(TextureAnalyzerTabId);
+}
+
+
+TSharedRef<SDockTab> FCustomizableObjectInstanceEditor::SpawnTab_TextureAnalyzer(const FSpawnTabArgs& Args)
+{
+	check(Args.GetTabId() == TextureAnalyzerTabId);
+	return SNew(SDockTab)
+		.Label(LOCTEXT("Texture Analyzer", "Texture Analyzer"))
+		[
+			TextureAnalyzer.ToSharedRef()
+		];
+}
+
+
+#undef LOCTEXT_NAMESPACE
diff --git a/Engine/Plugins/Experimental/Mutable/Source/CustomizableObjectEditor/Private/CustomizableObjectInstanceEditor.h b/Engine/Plugins/Experimental/Mutable/Source/CustomizableObjectEditor/Private/CustomizableObjectInstanceEditor.h
new file mode 100644
--- /dev/null
+++ b/Engine/Plugins/Experimental/Mutable/Source/CustomizableObjectEditor/Private/CustomizableObjectInstanceEditor.h
@@ -0,0 +1,219 @@
+// Copyright Epic Games, Inc. All Rights Reserved.
+
+#pragma once
+
+#include "ICustomizableObjectInstanceEditor.h"
+#include "Toolkits/AssetEditorToolkit.h"
+#include "GraphEditor.h"
+#include "TickableEditorObject.h"
+#include "Animation/DebugSkelMeshComponent.h"
+#include "CustomizableObjectInstanceEditor.generated.h"
+
+DECLARE_DELEGATE(FCreatePreviewInstanceFlagDelegate);
+
+/**
+* Wrapper UObject class for the UCustomizableObjectInstance::FObjectInstanceUpdatedDelegate dynamic multicast delegate
+*/
+UCLASS()
+class UUpdateClassWrapperClass : public UObject
+{
+public:
+	GENERATED_BODY()
+
+	/** Method to assign for the callback */
+	UFUNCTION()
+	void DelegatedCallback(UCustomizableObjectInstance* Instance);
+
+	FCreatePreviewInstanceFlagDelegate Delegate;
+};
+
+
+/**
+ * CustomizableObject Editor class
+ */
+class FCustomizableObjectInstanceEditor : 
+	public ICustomizableObjectInstanceEditor,
+	public FGCObject,
+	public FTickableEditorObject
+{
+public:
+	FCustomizableObjectInstanceEditor();
+	virtual ~FCustomizableObjectInstanceEditor() override;
+
+	virtual void RegisterTabSpawners(const TSharedRef<class FTabManager>& TabManager) override;
+	virtual void UnregisterTabSpawners(const TSharedRef<class FTabManager>& TabManager) override;
+
+	/**
+	 * Edits the specified object
+	 *
+	 * @param	Mode					Asset editing mode for this editor (standalone or world-centric)
+	 * @param	InitToolkitHost			When Mode is WorldCentric, this is the level editor instance to spawn this editor within
+	 * @param	ObjectToEdit			The object to edit
+	 */
+	void InitCustomizableObjectInstanceEditor( const EToolkitMode::Type Mode, const TSharedPtr< class IToolkitHost >& InitToolkitHost, UCustomizableObjectInstance* ObjectToEdit );
+
+	// FSerializableObject interface
+	virtual void AddReferencedObjects( FReferenceCollector& Collector ) override;
+	virtual FString GetReferencerName() const override
+	{
+		return TEXT("FCustomizableObjectInstanceEditor");
+	}
+	// End of FSerializableObject interface
+
+	/** IToolkit interface */
+	virtual FName GetToolkitFName() const override;
+	virtual FText GetBaseToolkitName() const override;
+	virtual FText GetToolkitName() const override;
+	virtual FString GetWorldCentricTabPrefix() const override;
+	virtual FLinearColor GetWorldCentricTabColorScale() const override;
+
+	/** ICustomizableObjectInstanceEditor interface */
+	UCustomizableObjectInstance* GetPreviewInstance() override;
+	virtual void RefreshTool() override;
+	virtual void RefreshViewport() override;
+	virtual void SetPoseAsset(class UPoseAsset* PoseAssetParameter) override;
+
+	/** Callback for the object modified event */
+	void OnObjectModified(UObject* Object);
+
+	/** Getter of AssetRegistryLoaded */
+	bool GetAssetRegistryLoaded();
+
+	/** Callback to notify the editor when the PreviewInstance has been updated */
+	void OnUpdatePreviewInstance();
+
+	/** Compile the customizable object. */
+	void CompileObject(UCustomizableObject* Object);
+
+	/** FTickableGameObject interface */
+	virtual bool IsTickable(void) const override;
+	virtual void Tick(float InDeltaTime) override;
+	virtual TStatId GetStatId() const override;
+
+private:
+
+	TSharedRef<SDockTab> SpawnTab_Viewport(const FSpawnTabArgs& Args);
+	TSharedRef<SDockTab> SpawnTab_InstanceProperties(const FSpawnTabArgs& Args);
+	TSharedRef<SDockTab> SpawnTab_AdvancedPreviewSettings(const FSpawnTabArgs& Args);
+	TSharedRef<SDockTab> SpawnTab_TextureAnalyzer(const FSpawnTabArgs& Args);
+
+private:
+
+	void CreatePreviewInstance();
+
+	/** Binds commands associated with the Static Mesh Editor. */
+	void BindCommands();
+		
+	/**
+	 *	Sets the editor's current mesh and refreshes various settings to correspond with the new data.
+	 *
+	 *	@param	InCustomizableObject		The static mesh to use for the editor.
+	 */
+	void SetEditorMesh(UCustomizableObjectInstance* InCustomizableObjectInstance);
+
+	/** Change the mesh the editor is viewing. */
+	void OnChangeMesh();
+
+	/** Callback when selection changes in the Property Tree. */
+	void OnInstancePropertySelectionChanged(FProperty* InProperty);
+
+	/** The instance notifies the editor of projector states changes through this. */
+	void OnProjectorStateChanged();
+
+	/** Callback for the asset registry initial load */
+	void OnAssetRegistryLoadComplete();
+
+	/** Updates the visibility of PreviewSkeletalMeshComponent */
+	void UpdatePreviewVisibility();
+
+	/** Save Customizable Object Instance open in the editor */
+	void SaveAsset_Execute() override;
+
+	/** Says if the customizable Object can be shown or be oppened in the editor */
+	bool CanOpenOrShowParent();
+
+	/** Show Customizable Object Intance's Parent In Content Browser */
+	void ShowParentInContentBrowser();
+
+	/** Open Customizable Object Intance's Parent In Editor */
+	void OpenParentInEditor();
+
+	/** Open the Texture Analyzer tab */
+	void OpenTextureAnalyzerTab();
+
+private:
+
+	/** The currently viewed object. */
+	class UCustomizableObjectInstance* CustomizableObjectInstance;
+	TArray<class UCustomizableSkeletalComponent*> PreviewCustomizableSkeletalComponents;
+	class UStaticMeshComponent* PreviewStaticMeshComponent;
+	TArray<class UDebugSkelMeshComponent*> PreviewSkeletalMeshComponents;
+
+	/** List of open tool panels; used to ensure only one exists at any one time */
+	TMap< FName, TWeakPtr<class SDockableTab> > SpawnedToolPanels;
+
+	/** Preview Viewport widget */
+	TSharedPtr<class SCustomizableObjectEditorViewportTabBody> Viewport;
+	TSharedPtr<class IDetailsView> CustomizableInstanceDetailsView;
+
+	/** Level of Details Settings widget. */
+	TSharedPtr< class SLevelOfDetailSettings> LevelOfDetailSettings;
+
+	/** Widget for displaying the available UV Channels. */
+	TSharedPtr< class STextComboBox > UVChannelCombo;
+
+	/** List of available UV Channels. */
+	TArray< TSharedPtr< FString > > UVChannels;
+
+	/** Widget for displaying the available LOD. */
+	TSharedPtr< class STextComboBox > LODLevelCombo;
+
+	/** List of LODs. */
+	TArray< TSharedPtr< FString > > LODLevels;
+
+	/**	The tab ids for all the tabs used */
+	static const FName ViewportTabId;
+	static const FName InstancePropertiesTabId;
+	static const FName AdvancedPreviewSettingsTabId;
+	static const FName TextureAnalyzerTabId;
+
+	/** Handle for the OnObjectModified event */
+	FDelegateHandle OnObjectModifiedHandle;
+
+	/** Flag to know when the asset registry initial loading has completed */
+	bool AssetRegistryLoaded;
+
+	/** Flag to know whether the lasts steps when making a new preview instance need to be done when asset registry completes asset loading */
+	bool UpdateSkeletalMeshAfterAssetLoaded;
+
+	/** Copy of variable needed when asset registry finishes loading assets to complete editor loading */
+	UCustomizableObjectInstance* ObjectToEditCopy;
+
+	/** UObject class to be able to use the update callback */
+	UUpdateClassWrapperClass* HelperCallback;
+
+	/** Scene preview settings widget */
+	TSharedPtr<SWidget> AdvancedPreviewSettingsWidget;
+
+	/** Object compiler */
+	TUniquePtr<class FCustomizableObjectCompiler> Compiler;
+
+	/** Used in InitCustomizableObjectInstanceEditor, in the case the CO being opened is not compiled. Since referencer assets are loaded asynchronously,
+	* this flag is used to complete editor initialization */
+	bool CreatePreviewInstanceAfterCOCompile;
+
+	/** Pose asset when doing drag and drop of an UPoseAsset to the viewport */
+	UPoseAsset* PoseAsset;
+
+	/** Texture Analyzer table widget which shows the information of the transient textures used in the customizable object instance */
+	TSharedPtr<class SCustomizableObjecEditorTextureAnalyzer> TextureAnalyzer;
+
+
+	/** Variables used to force the refresh of the details view widget. These are needed because sometimes the scrollbar of the window doesn't appear
+	    until we force the refresh */
+	bool bOnlyRuntimeParameters;
+	bool bOnlyRelevantParameters;
+
+	/** Adds the customizable Object Instance Editor commands to the default toolbar */
+	void ExtendToolbar();
+};
\ No newline at end of file
diff --git a/Engine/Plugins/Experimental/Mutable/Source/CustomizableObjectEditor/Private/CustomizableObjectInstanceEditorActions.cpp b/Engine/Plugins/Experimental/Mutable/Source/CustomizableObjectEditor/Private/CustomizableObjectInstanceEditorActions.cpp
new file mode 100644
--- /dev/null
+++ b/Engine/Plugins/Experimental/Mutable/Source/CustomizableObjectEditor/Private/CustomizableObjectInstanceEditorActions.cpp
@@ -0,0 +1,48 @@
+// Copyright Epic Games, Inc. All Rights Reserved.
+
+#include "CustomizableObjectInstanceEditorActions.h"
+#include "SCustomizableObjectEditorViewport.h"
+#include "CustomizableObjectEditorStyle.h"
+#include "Styling/AppStyle.h"
+
+
+#define LOCTEXT_NAMESPACE "CustomizableObjectEditor"
+
+
+FCustomizableObjectInstanceEditorCommands::FCustomizableObjectInstanceEditorCommands() 
+	: TCommands<FCustomizableObjectInstanceEditorCommands>
+(
+	"CustomizableObjectInstanceEditor", // Context name for fast lookup
+	NSLOCTEXT("Contexts", "CustomizableObjecInstancetEditor", "CustomizableObjectInstance Editor"), // Localized context name for displaying
+	NAME_None, // Parent
+	FCustomizableObjectEditorStyle::GetStyleSetName() // Icon Style Set
+	)
+{
+}
+
+
+void FCustomizableObjectInstanceEditorCommands::RegisterCommands()
+{
+	UI_COMMAND( SetDrawUVs, "UV", "Toggles display of the static mesh's UVs for the specified channel.", EUserInterfaceActionType::ToggleButton, FInputChord() );
+	UI_COMMAND( SetShowGrid, "Grid", "Displays the viewport grid.", EUserInterfaceActionType::ToggleButton, FInputChord());
+	UI_COMMAND( SetShowSky, "Sky", "Displays the viewport sky.", EUserInterfaceActionType::ToggleButton, FInputChord() );
+	UI_COMMAND( SetShowBounds, "Bounds", "Toggles display of the bounds of the static mesh.", EUserInterfaceActionType::ToggleButton, FInputChord() );
+	UI_COMMAND( SetShowCollision, "Collision", "Toggles display of the simplified collision mesh of the static mesh, if one has been assigned.", EUserInterfaceActionType::ToggleButton, FInputChord() );
+	UI_COMMAND( SetCameraLock, "Camera Lock", "Toggles viewport navigation between orbit and freely moving about.", EUserInterfaceActionType::ToggleButton, FInputChord() );
+	UI_COMMAND( SaveThumbnail, "Save Thumbnail", "Saves the viewpoint position and orientation in the Preview Pane for use as the thumbnail preview in the Content Browser.", EUserInterfaceActionType::Button, FInputChord() );
+
+	UI_COMMAND( SetShowNormals, "Normals", "Toggles display of vertex normals in the Preview Pane.", EUserInterfaceActionType::ToggleButton, FInputChord() );
+	UI_COMMAND( SetShowTangents, "Tangents", "Toggles display of vertex tangents in the Preview Pane.", EUserInterfaceActionType::ToggleButton, FInputChord() );
+	UI_COMMAND( SetShowBinormals, "Binormals", "Toggles display of vertex binormals (orthogonal vector to normal and tangent) in the Preview Pane.", EUserInterfaceActionType::ToggleButton, FInputChord() );
+	UI_COMMAND( SetShowPivot, "Show Pivot", "Display the pivot location of the static mesh.", EUserInterfaceActionType::ToggleButton, FInputChord() );
+
+	UI_COMMAND( BakeInstance, "Bake Instance", "Create baked unreal resources for the current instance.", EUserInterfaceActionType::Button, FInputChord());
+	UI_COMMAND( StateChangeTest, "State test", "Make a state change test, iterating through all possible states and runtime parameters for this instance", EUserInterfaceActionType::Button, FInputChord());
+
+	UI_COMMAND(ShowParentCO, "Show Parent", "Show the Parent of the Customizable Object in the Content Browser", EUserInterfaceActionType::Button, FInputChord());
+	UI_COMMAND(EditParentCO, "Edit Parent", "Open the Parent of the Customizable Object to Edit", EUserInterfaceActionType::Button, FInputChord());
+	UI_COMMAND(TextureAnalyzer, "Texture Memory Analyzer", "Open the Texture Analyzer window to check all the information of the textures created by Mutable.", EUserInterfaceActionType::Button, FInputChord());
+
+}
+
+#undef LOCTEXT_NAMESPACE
diff --git a/Engine/Plugins/Experimental/Mutable/Source/CustomizableObjectEditor/Private/CustomizableObjectInstanceFactory.cpp b/Engine/Plugins/Experimental/Mutable/Source/CustomizableObjectEditor/Private/CustomizableObjectInstanceFactory.cpp
new file mode 100644
--- /dev/null
+++ b/Engine/Plugins/Experimental/Mutable/Source/CustomizableObjectEditor/Private/CustomizableObjectInstanceFactory.cpp
@@ -0,0 +1,193 @@
+// Copyright Epic Games, Inc. All Rights Reserved.
+
+#include "CustomizableObjectInstanceFactory.h"
+#include "CustomizableObjectEditorModule.h"
+#include "CustomizableSkeletalMeshActor.h"
+#include "AssetRegistry/AssetRegistryModule.h"
+#include "Components/SkeletalMeshComponent.h"
+#include "CustomizableObjectInstance.h"
+
+
+#define LOCTEXT_NAMESPACE "CustomizableObjectInstanceFactory"
+
+
+UCustomizableObjectInstanceFactory::UCustomizableObjectInstanceFactory(const FObjectInitializer& ObjectInitializer)
+    : Super(ObjectInitializer)
+{
+    DisplayName = LOCTEXT("CustomizableObjectInstanceDisplayName", "Customizable Object Instance");
+    NewActorClass = ASkeletalMeshActor::StaticClass();
+    bUseSurfaceOrientation = true;
+}
+
+bool UCustomizableObjectInstanceFactory::CanCreateActorFrom(const FAssetData& AssetData, FText& OutErrorMsg)
+{
+    if (!AssetData.IsValid() ||
+        (!AssetData.GetClass()->IsChildOf(UCustomizableObjectInstance::StaticClass())))
+    {
+        OutErrorMsg = LOCTEXT("NoCOISeq", "A valid customizable object instance must be specified.");
+        return false;
+    }
+
+    FAssetRegistryModule& AssetRegistryModule = FModuleManager::LoadModuleChecked<FAssetRegistryModule>(TEXT("AssetRegistry"));
+    FAssetData CustomizableObjectInstanceData;
+
+    if (AssetData.GetClass()->IsChildOf(UCustomizableObjectInstance::StaticClass()))
+    {
+        if (UCustomizableObjectInstance* CustomizableObjectInstance = Cast<UCustomizableObjectInstance>(AssetData.GetAsset()))
+        {
+            if (USkeletalMesh* SkeletalMesh = CustomizableObjectInstance->GetSkeletalMesh())
+            {
+                return true;
+            }
+            else
+            {
+                if (UCustomizableObject* CustomizableObject = CustomizableObjectInstance->GetCustomizableObject())
+                {
+                    return true;
+                }
+                else
+                {
+                    OutErrorMsg = LOCTEXT("NoCustomizableObjectInstance", "The UCustomizableObjectInstance does not have a customizableObject.");
+                    return false;
+                }
+            }
+        }
+        else
+        {
+            OutErrorMsg = LOCTEXT("NoCustomizableObjectInstance", "The CustomizableObjectInstance is null.");
+        }
+    }
+
+    if (!CustomizableObjectInstanceData.IsValid())
+    {
+        OutErrorMsg = LOCTEXT("NoSkeletalMeshAss", "No valid skeletal mesh was found associated with the animation sequence.");
+        return false;
+    }
+
+    if (USkeletalMesh* SkeletalMeshCDO = Cast<USkeletalMesh>(AssetData.GetClass()->GetDefaultObject()))
+    {
+        if (SkeletalMeshCDO->HasCustomActorFactory())
+        {
+            return false;
+        }
+    }
+
+    return true;
+}
+
+USkeletalMesh* UCustomizableObjectInstanceFactory::GetSkeletalMeshFromAsset(UObject* Asset, int32 ComponentIndex) const
+{
+    USkeletalMesh* SkeletalMesh = Cast<USkeletalMesh>(Asset);
+    UCustomizableObjectInstance* CustomizableObjectInstance = Cast<UCustomizableObjectInstance>(Asset);
+
+    if (SkeletalMesh == nullptr && CustomizableObjectInstance != nullptr)
+    {
+        SkeletalMesh = CustomizableObjectInstance->GetSkeletalMesh(ComponentIndex);
+        if (SkeletalMesh == nullptr)
+        {
+            if (UCustomizableObject* CustomizableObject = CustomizableObjectInstance->GetCustomizableObject())
+            {
+                SkeletalMesh = CustomizableObject->GetRefSkeletalMesh(ComponentIndex);
+            }
+        }
+    }
+
+    //check(SkeletalMesh != NULL);
+    return SkeletalMesh;
+}
+
+
+void UCustomizableObjectInstanceFactory::PostSpawnActor(UObject* Asset, AActor* NewActor)
+{
+    UCustomizableObjectInstance* COInstance = Cast<UCustomizableObjectInstance>(Asset);
+    ACustomizableSkeletalMeshActor* NewCSMActor = CastChecked<ACustomizableSkeletalMeshActor>(NewActor);
+
+	if (NewCSMActor && COInstance)
+	{
+		int32 NumComponents = GetNumberOfComponents(COInstance);
+
+		for (int32 ComponentIndex = 0; ComponentIndex < NumComponents; ++ComponentIndex)
+		{
+			USkeletalMesh* SkeletalMesh = GetSkeletalMeshFromAsset(Asset, ComponentIndex);
+
+			if (ComponentIndex > 0 && ComponentIndex > NewCSMActor->GetNumComponents() - 1)
+			{
+				NewCSMActor->AttachNewComponent();
+			}
+
+			if (USkeletalMeshComponent* SkeletalMeshComp = NewCSMActor->GetSkeletalMeshComponentAt(ComponentIndex))
+			{
+				SkeletalMeshComp->UnregisterComponent();
+				SkeletalMeshComp->SkeletalMesh = SkeletalMesh;
+
+				if (ComponentIndex == 0 && NewCSMActor->GetWorld()->IsGameWorld())
+				{
+					NewCSMActor->ReplicatedMesh = SkeletalMesh;
+				}
+
+				if (UCustomizableSkeletalComponent* CustomSkeletalComp = NewCSMActor->GetCustomizableSkeletalComponent(ComponentIndex))
+				{
+					CustomSkeletalComp->UnregisterComponent();
+					CustomSkeletalComp->CustomizableObjectInstance = COInstance;
+					CustomSkeletalComp->ComponentIndex = ComponentIndex;
+					CustomSkeletalComp->SetSkeletalMesh(SkeletalMesh, false);
+					CustomSkeletalComp->UpdateSkeletalMeshAsync();
+					CustomSkeletalComp->RegisterComponent();
+				}
+
+				SkeletalMeshComp->RegisterComponent();
+			}			
+		}
+
+		UActorFactory::PostSpawnActor(COInstance, NewActor);
+	}
+}
+
+
+void UCustomizableObjectInstanceFactory::PostCreateBlueprint(UObject* Asset, AActor* CDO)
+{
+    //if (Asset != NULL && CDO != NULL)
+    //{
+    //	USkeletalMesh* SkeletalMesh = GetSkeletalMeshFromAsset(Asset);
+    //	UAnimBlueprint* AnimBlueprint = Cast<UAnimBlueprint>(Asset);
+
+    //	ASkeletalMeshActor* SkeletalMeshActor = CastChecked<ASkeletalMeshActor>(CDO);
+    //	SkeletalMeshActor->GetSkeletalMeshComponent()->SkeletalMesh = SkeletalMesh;
+    //	SkeletalMeshActor->GetSkeletalMeshComponent()->AnimClass = AnimBlueprint ? Cast<UAnimBlueprintGeneratedClass>(AnimBlueprint->GeneratedClass) : NULL;
+    //}
+}
+
+
+FQuat UCustomizableObjectInstanceFactory::AlignObjectToSurfaceNormal(const FVector& InSurfaceNormal, const FQuat& ActorRotation) const
+{
+    // Meshes align the Z (up) axis with the surface normal
+    return FindActorAlignmentRotation(ActorRotation, FVector(0.f, 0.f, 1.f), InSurfaceNormal);
+}
+
+
+int32 UCustomizableObjectInstanceFactory::GetNumberOfComponents(UCustomizableObjectInstance* COInstance)
+{
+	int32 NumMeshComponents = 0;
+
+	if (const UCustomizableObject* CustomizableObject = COInstance->GetCustomizableObject();
+		CustomizableObject &&
+		CustomizableObject->Source)
+	{
+		TArray<UCustomizableObjectNodeObject*> RootNodes;
+		CustomizableObject->Source->GetNodesOfClass<UCustomizableObjectNodeObject>(RootNodes);
+
+		for (int32 i = 0; i < RootNodes.Num(); ++i)
+		{
+			if (RootNodes[i]->bIsBase)
+			{
+				NumMeshComponents = RootNodes[i]->NumMeshComponents;
+				break;
+			}
+		}
+	}
+
+	return NumMeshComponents;
+}
+
+
+#undef LOCTEXT_NAMESPACE
diff --git a/Engine/Plugins/Experimental/Mutable/Source/CustomizableObjectEditor/Private/CustomizableObjectLayout.cpp b/Engine/Plugins/Experimental/Mutable/Source/CustomizableObjectEditor/Private/CustomizableObjectLayout.cpp
new file mode 100644
--- /dev/null
+++ b/Engine/Plugins/Experimental/Mutable/Source/CustomizableObjectEditor/Private/CustomizableObjectLayout.cpp
@@ -0,0 +1,166 @@
+// Copyright Epic Games, Inc. All Rights Reserved.
+
+#include "CustomizableObjectLayout.h"
+#include "CustomizableObjectCompiler.h"
+#include "ICustomizableObjectEditor.h"
+#include "GenerateMutableSource/GenerateMutableSource.h"
+#include "GraphTraversal.h"
+
+#include "Nodes/CustomizableObjectNode.h"
+#include "Nodes/CustomizableObjectNodeSkeletalMesh.h"
+#include "Nodes/CustomizableObjectNodeTable.h"
+#include "Nodes/CustomizableObjectNodeLayoutBlocks.h"
+
+#include "GenerateMutableSource/GenerateMutableSourceMesh.h"
+
+#include "MutableTools/Public/NodeLayout.h"
+
+#define LOCTEXT_NAMESPACE "CustomizableObjectEditor"
+
+UCustomizableObjectLayout::UCustomizableObjectLayout()
+{
+	GridSize = FIntPoint(4, 4);
+	MaxGridSize = FIntPoint(4, 4);
+
+	FCustomizableObjectLayoutBlock Block;
+	Block.Min = FIntPoint(0, 0);
+	Block.Max = FIntPoint(4, 4);
+	Block.Id = FGuid::NewGuid();
+	Block.Priority = 0;
+	Blocks.Add(Block);
+
+	PackingStrategy = ECustomizableObjectTextureLayoutPackingStrategy::Resizable;
+}
+
+
+void UCustomizableObjectLayout::SetLayout(UObject* InMesh, int32 LODIndex, int32 MatIndex, int32 UVIndex)
+{
+	Mesh = InMesh;
+	LOD = LODIndex;
+	Material = MatIndex;
+	UVChannel = UVIndex;
+}
+
+
+void UCustomizableObjectLayout::SetPackingStrategy(ECustomizableObjectTextureLayoutPackingStrategy Strategy)
+{
+	PackingStrategy = Strategy;
+}
+
+
+void UCustomizableObjectLayout::SetGridSize(FIntPoint Size)
+{
+	GridSize = Size;
+}
+
+
+void UCustomizableObjectLayout::SetMaxGridSize(FIntPoint Size)
+{
+	MaxGridSize = Size;
+}
+
+
+void UCustomizableObjectLayout::SetLayoutName(FString Name)
+{
+	LayoutName = Name;
+}
+
+
+void UCustomizableObjectLayout::GenerateBlocksFromUVs()
+{
+	mu::NodeLayoutBlocksPtr Layout = nullptr;
+
+	UCustomizableObjectNode* Node = Cast<UCustomizableObjectNode>(GetOuter());
+
+	if (Node && Mesh)
+	{
+		//Creating a GenerationContext
+		FCustomizableObjectCompiler* Compiler = new FCustomizableObjectCompiler();
+		UCustomizableObject* Object = Node->GetGraphEditor()->GetCustomizableObject();
+		FCompilationOptions Options = Object->CompileOptions;
+	
+		FMutableGraphGenerationContext GenerationContext(Object, Compiler, Options);
+	
+		//Transforming skeletalmesh to mutable mesh
+		mu::MeshPtr	MutableMesh = nullptr;
+		
+		if (USkeletalMesh* SkeletalMesh = Cast<USkeletalMesh>(Mesh))
+		{
+			GenerationContext.ComponentInfos.Add(SkeletalMesh);
+			MutableMesh = ConvertSkeletalMeshToMutable(SkeletalMesh, LOD, Material, GenerationContext, Node);
+		}
+		else if (UStaticMesh* StaticMesh = Cast<UStaticMesh>(Mesh))
+		{
+			MutableMesh = ConvertStaticMeshToMutable(StaticMesh, LOD, Material, GenerationContext, Node);
+		}
+	
+		if (MutableMesh)
+		{
+			//Generating blocks with the mutable mesh
+			Layout = mu::NodeLayoutBlocks::GenerateLayoutBlocks(MutableMesh, UVChannel, GridSize.X, GridSize.Y);
+		}
+	
+		delete Compiler;
+	
+		if (Layout)
+		{
+			Blocks.Empty();
+		
+			//Generating the layout blocks with the mutable layout		
+			for (int i = 0; i < Layout->GetBlockCount(); ++i)
+			{
+				int minX, minY, sizeX, sizeY;
+		
+				Layout->GetBlock(i, &minX, &minY, &sizeX, &sizeY);
+		
+				FCustomizableObjectLayoutBlock block;
+				block.Min = FIntPoint(minX, minY);
+				block.Max = FIntPoint(minX + sizeX, minY + sizeY);
+				block.Id = FGuid::NewGuid();
+				block.Priority = 0;
+				Blocks.Add(block);
+			}
+		
+			Node->PostEditChange();
+			Node->GetGraph()->MarkPackageDirty();
+		}
+	}
+}
+
+
+void UCustomizableObjectLayout::GetUVChannel(TArray<FVector2f>& UVs, int32 UVChannelIndex) const
+{
+	if (UObject* Node = GetOuter())
+	{
+		if (const UCustomizableObjectNodeLayoutBlocks* TypedNodeLayout = Cast<UCustomizableObjectNodeLayoutBlocks>(Node))
+		{
+			if (const UEdGraphPin* ConnectedPin = FollowOutputPin(*TypedNodeLayout->OutputPin()))
+			{
+				if (const UCustomizableObjectNodeMesh* MeshNode = Cast<UCustomizableObjectNodeMesh>(ConnectedPin->GetOwningNode()))
+				{
+					MeshNode->GetUVChannelForPin(ConnectedPin, UVs, UVChannelIndex);
+				}
+			}
+		}
+		else if (const UCustomizableObjectNodeTable* TypedNodeTable = Cast<UCustomizableObjectNodeTable>(Node))
+		{
+			TypedNodeTable->GetUVChannel(this, UVs);
+		}
+	}
+}
+
+
+int32 UCustomizableObjectLayout::FindBlock(const FGuid& InId) const
+{
+	for (int Index = 0; Index < Blocks.Num(); ++Index)
+	{
+		if (Blocks[Index].Id == InId)
+		{
+			return Index;
+		}
+	}
+
+	return -1;
+}
+
+#undef LOCTEXT_NAMESPACE
diff --git a/Engine/Plugins/Experimental/Mutable/Source/CustomizableObjectEditor/Private/CustomizableObjectNodeObjectGroupDetails.cpp b/Engine/Plugins/Experimental/Mutable/Source/CustomizableObjectEditor/Private/CustomizableObjectNodeObjectGroupDetails.cpp
new file mode 100644
--- /dev/null
+++ b/Engine/Plugins/Experimental/Mutable/Source/CustomizableObjectEditor/Private/CustomizableObjectNodeObjectGroupDetails.cpp
@@ -0,0 +1,106 @@
+// Copyright Epic Games, Inc. All Rights Reserved.
+
+#include "CustomizableObjectNodeObjectGroupDetails.h"
+#include "Nodes/CustomizableObjectNodeObjectGroup.h"
+#include "CustomizableObjectEditor.h"
+#include "CustomizableObjectEditorModule.h"
+#include "PropertyCustomizationHelpers.h"
+//#include "SPropertyEditorAsset.h"
+//#include "CustomizableObjectEditorUtilities.h"
+//
+#include "DetailLayoutBuilder.h"
+#include "DetailCategoryBuilder.h"
+//#include "LevelEditor.h"
+//#include "LevelEditorActions.h"
+//
+//#include "Widgets/Input/SButton.h"
+//#include "Widgets/Text/STextBlock.h"
+//
+//#include "Runtime/Launch/Resources/Version.h"
+//#include "Engine/SkeletalMesh.h"
+//#include "Widgets/Input/STextComboBox.h"
+//#include "Math/Vector4.h"
+//#include "Math/Vector.h"
+//
+//#include "Runtime/Launch/Resources/Version.h"
+//#if (ENGINE_MAJOR_VERSION == 4 && ENGINE_MINOR_VERSION > 22) || ENGINE_MAJOR_VERSION >= 5
+//#include "HAL/PlatformApplicationMisc.h"
+//#endif
+
+
+#define LOCTEXT_NAMESPACE "CustomizableObjectGroupDetails"
+
+
+TSharedRef<IDetailCustomization> FCustomizableObjectNodeObjectGroupDetails::MakeInstance()
+{
+	return MakeShareable(new FCustomizableObjectNodeObjectGroupDetails);
+}
+
+
+void FCustomizableObjectNodeObjectGroupDetails::CustomizeDetails(IDetailLayoutBuilder& DetailBuilder)
+{
+	const IDetailsView* DetailsView = DetailBuilder.GetDetailsView();
+	UCustomizableObjectNodeObjectGroup* NodeGroup = nullptr;
+	if (DetailsView->GetSelectedObjects().Num())
+	{
+		if (DetailsView->GetSelectedObjects()[0].Get()->IsA(UCustomizableObjectNodeObjectGroup::StaticClass()))
+		{
+			NodeGroup = Cast<UCustomizableObjectNodeObjectGroup>(DetailsView->GetSelectedObjects()[0].Get());
+		}
+	}
+
+	if (NodeGroup)
+	{
+		if (TSharedPtr<FCustomizableObjectEditor> GraphEditor = StaticCastSharedPtr<FCustomizableObjectEditor>(NodeGroup->GetGraphEditor()))
+		{
+			if (UCustomizableObject* NodeGroupCO = CastChecked<UCustomizableObject>(NodeGroup->GetCustomizableObjectGraph()->GetOuter()))
+			{
+				TArray<UCustomizableObject*> ChildNodes;
+				IDetailCategoryBuilder& BlocksCategory = DetailBuilder.EditCategory("Group Info");
+				GraphEditor->GetExternalChildObjects(NodeGroupCO, ChildNodes, false, EObjectFlags::RF_NoFlags);
+				for (UCustomizableObject* ChildNode : ChildNodes)
+				{
+					if (ChildNode)
+					{
+						TArray<UCustomizableObjectNodeObject*> ObjectNodes;
+						ChildNode->Source->GetNodesOfClass<UCustomizableObjectNodeObject>(ObjectNodes);
+						UCustomizableObjectNodeObject* groupParent = nullptr;
+
+						for (UCustomizableObjectNodeObject* ChildCONode : ObjectNodes)
+						{
+							if (ChildCONode->bIsBase)
+							{
+								groupParent = ChildCONode;
+								break;
+							}
+						}
+
+						if (!groupParent) continue;
+
+						FString GroupID = groupParent->Identifier.ToString();
+
+						FCustomizableObjectIdPair* DirectGroupChilds = NodeGroupCO->GroupNodeMap.Find(GroupID);
+						if (!DirectGroupChilds) continue;
+
+						if (NodeGroup->GroupName.Equals(DirectGroupChilds->CustomizableObjectGroupName))
+						{
+							BlocksCategory.AddCustomRow(LOCTEXT("FCustomizableObjectNodeObjectGroupDetails", "External Customizable Objects in this Group"))[
+								SNew(SObjectPropertyEntryBox)
+									.ObjectPath(ChildNode->GetPathName())
+									.AllowedClass(UCustomizableObject::StaticClass())
+									.AllowClear(false)
+									.DisplayUseSelected(false)
+									.DisplayBrowse(true)
+									.EnableContentPicker(false)
+									.DisplayThumbnail(true)
+							];
+						}
+					}
+				}
+			}
+		}
+	}
+}
+
+
+#undef LOCTEXT_NAMESPACE
\ No newline at end of file
diff --git a/Engine/Plugins/Experimental/Mutable/Source/CustomizableObjectEditor/Private/CustomizableObjectNodeObjectGroupDetails.h b/Engine/Plugins/Experimental/Mutable/Source/CustomizableObjectEditor/Private/CustomizableObjectNodeObjectGroupDetails.h
new file mode 100644
--- /dev/null
+++ b/Engine/Plugins/Experimental/Mutable/Source/CustomizableObjectEditor/Private/CustomizableObjectNodeObjectGroupDetails.h
@@ -0,0 +1,33 @@
+// Copyright Epic Games, Inc. All Rights Reserved.
+
+#pragma once
+
+#include "IDetailCustomization.h"
+//#include "Templates/SharedPointer.h"
+//#include "UObject/WeakObjectPtr.h"
+//#include "Input/Reply.h"
+//#include "PropertyHandle.h"
+
+
+class FCustomizableObjectNodeObjectGroupDetails : public IDetailCustomization
+{
+public:
+	/** Makes a new instance of this detail layout class for a specific detail view requesting it */
+	static TSharedRef<IDetailCustomization> MakeInstance();
+
+	/** ILayoutDetails interface */
+	virtual void CustomizeDetails( IDetailLayoutBuilder& DetailBuilder ) override;
+
+
+private:
+	//TArray< TSharedPtr<FString> > BoneComboOptions;
+
+	//class UCustomizableObjectNodeProjectorConstant* NodeConstant;
+	//class UCustomizableObjectNodeProjectorParameter* NodeParameter;
+	//FReply OnProjectorCopyPressed();
+	//FReply OnProjectorPastePressed();
+
+	//class USkeletalMesh* SkeletalMesh;
+
+	//void OnBoneComboBoxSelectionChanged(TSharedPtr<FString> Selection, ESelectInfo::Type SelectInfo, TSharedRef<IPropertyHandle> BoneProperty);
+};
\ No newline at end of file
