[mutable] Move Mutable plugin from the restricted folder to the experimental plugins folder.

diff --git a/Engine/Plugins/Experimental/Mutable/Source/CustomizableObjectEditor/Private/CustomizableObjectWidget.cpp b/Engine/Plugins/Experimental/Mutable/Source/CustomizableObjectEditor/Private/CustomizableObjectWidget.cpp
new file mode 100644
--- /dev/null
+++ b/Engine/Plugins/Experimental/Mutable/Source/CustomizableObjectEditor/Private/CustomizableObjectWidget.cpp
@@ -0,0 +1,1176 @@
+// Copyright Epic Games, Inc. All Rights Reserved.
+
+#include "CustomizableObjectWidget.h"
+#include "CustomizableObjectEditorModule.h"
+#include "DynamicMeshBuilder.h"
+#include "Materials/MaterialInstanceDynamic.h"
+#include "UnrealWidget.h"
+#include "UnrealEditorPortabilityHelpers.h"
+
+
+IMPLEMENT_HIT_PROXY(HCustomizableObjectWidgetAxis,HHitProxy);
+
+static const float AXIS_ARROW_RADIUS = 2.5f;
+static const float AXIS_CIRCLE_RADIUS	= 48.0f;
+static const float TRANSLATE_ROTATE_AXIS_CIRCLE_RADIUS	= 20.0f;
+static const float INNER_AXIS_CIRCLE_RADIUS = 48.0f;
+static const float OUTER_AXIS_CIRCLE_RADIUS = 56.0f;
+static const float ROTATION_TEXT_RADIUS = 75.0f;
+static const int32 AXIS_CIRCLE_SIDES		= 24;
+
+namespace WidgetDefs
+{
+	static const int32 GScreenAlignedSphereSides = 16;
+	static const float GScreenAlignedSphereRadius = 3.0f;
+}
+
+
+FCustomizableObjectWidget::FCustomizableObjectWidget(UCustomizableObjectInstance* InInstance, int InParamIndex)
+{
+	Instance = InInstance;
+	ParamIndex = InParamIndex;
+	WidgetMode = WM_Translate;
+
+	// Compute and store sample vertices for drawing the axis arrow heads
+	AxisColorX = FColor(255,0,0);
+	AxisColorY = FColor(0,255,0);
+	AxisColorZ = FColor(0,0,255);
+	PlaneColorXY = FColor(255,255,0);
+	ScreenSpaceColor  = FColor(196, 196, 196);
+	CurrentColor = FColor(255,255,0);
+
+	UMaterial* AxisMaterialBase = (UMaterial*)StaticLoadObject( UMaterial::StaticClass(),NULL,TEXT("/Engine/EditorMaterials/GizmoMaterial.GizmoMaterial") );
+
+	AxisMaterialX = UMaterialInstanceDynamic::Create( AxisMaterialBase, NULL );
+	AxisMaterialX->SetVectorParameterValue( "GizmoColor", AxisColorX );
+
+	AxisMaterialY = UMaterialInstanceDynamic::Create( AxisMaterialBase, NULL );
+	AxisMaterialY->SetVectorParameterValue( "GizmoColor", AxisColorY );
+
+	AxisMaterialZ = UMaterialInstanceDynamic::Create( AxisMaterialBase, NULL );
+	AxisMaterialZ->SetVectorParameterValue( "GizmoColor", AxisColorZ );
+
+	CurrentAxisMaterial = UMaterialInstanceDynamic::Create( AxisMaterialBase, NULL );
+	CurrentAxisMaterial->SetVectorParameterValue("GizmoColor", CurrentColor );
+
+	OpaquePlaneMaterialXY = UMaterialInstanceDynamic::Create( AxisMaterialBase, NULL );
+	OpaquePlaneMaterialXY->SetVectorParameterValue( "GizmoColor", FLinearColor::White );
+
+	TransparentPlaneMaterialXY = (UMaterial*)StaticLoadObject( UMaterial::StaticClass(),NULL,TEXT("/Engine/EditorMaterials/WidgetVertexColorMaterial.WidgetVertexColorMaterial"),NULL,LOAD_None,NULL );
+	
+	GridMaterial = (UMaterial*)StaticLoadObject( UMaterial::StaticClass(),NULL,TEXT("/Engine/EditorMaterials/WidgetGridVertexColorMaterial_Ma.WidgetGridVertexColorMaterial_Ma"),NULL,LOAD_None,NULL );
+	if (!GridMaterial)
+	{
+		GridMaterial = TransparentPlaneMaterialXY;
+	}
+
+	CurrentAxis = EAxisList::None;
+
+	CustomCoordSystem = FMatrix::Identity;
+
+	bAbsoluteTranslationInitialOffsetCached = false;
+	InitialTranslationOffset = FVector::ZeroVector;
+	InitialTranslationPosition = FVector(0, 0, 0);
+
+	bDragging = false;
+	bSnapEnabled = false;
+}
+
+
+void FCustomizableObjectWidget::SetWidgetMode( EWidgetMode Mode )
+{
+	WidgetMode = Mode;
+}
+
+
+void FCustomizableObjectWidget::SetTransform( FMatrix Transform )
+{
+	CustomCoordSystem = Transform;
+}
+
+
+void FCustomizableObjectWidget::SetScale( FVector2D InScale )
+{
+	ProjectorScale = InScale;
+}
+
+
+void FCustomizableObjectWidget::Render( const FSceneView* View,FPrimitiveDrawInterface* PDI )
+{
+	FVector Loc = FVector(0,0,0);
+	if(!View->ScreenToPixel(View->WorldToScreen(Loc),Origin))
+	{
+		Origin.X = Origin.Y = 0;
+	}
+
+	// Render projected image quad
+	{
+		FColor Color(200,200,100,50);
+		Render_Quad(PDI, CustomCoordSystem, TransparentPlaneMaterialXY, &Color, ProjectorScale*0.5f );
+	}
+
+
+	bool bDrawWidget = true;
+
+	switch( WidgetMode )
+	{
+		case WM_Translate:
+			Render_Translate( View, PDI, Loc, bDrawWidget );
+			break;
+
+		case WM_Rotate:
+			Render_Rotate( View, PDI, Loc, bDrawWidget );
+			break;
+
+		case WM_Scale:
+			Render_Scale( View, PDI, Loc, bDrawWidget );
+			break;
+
+		case WM_TranslateRotateZ:
+			Render_TranslateRotateZ( View, PDI, Loc, bDrawWidget );
+			break;
+
+		default:
+			break;
+	}
+};
+
+/**
+ * Draws an arrow head line for a specific axis.
+ */
+void FCustomizableObjectWidget::Render_Axis( const FSceneView* View, FPrimitiveDrawInterface* PDI, EAxisList::Type InAxis, FMatrix& InMatrix, UMaterialInterface* InMaterial, FVector2D& OutAxisEnd, float InScale, bool bDrawWidget, bool bCubeHead )
+{
+	FMatrix AxisRotation = FMatrix::Identity;
+	if( InAxis == EAxisList::Y )
+	{
+		AxisRotation = FRotationMatrix( FRotator(0,90.f,0) );
+	}
+	else if( InAxis == EAxisList::Z )
+	{
+		AxisRotation = FRotationMatrix( FRotator(90.f,0,0) );
+	}
+
+	FMatrix ArrowToWorld = AxisRotation * InMatrix;
+
+	FScaleMatrix Scale(InScale);
+	ArrowToWorld = Scale * ArrowToWorld;
+
+	if( bDrawWidget )
+	{
+		const bool bDisabled = false;
+		PDI->SetHitProxy( new HWidgetAxis( InAxis, bDisabled) );
+
+		const float HalfHeight = 35/2.0f;
+		const float CylinderRadius = 1.2f;
+		const FVector Offset( 0,0,HalfHeight );
+
+		const FMaterialRenderProxy* Proxy = Helper_GetMaterialProxy(InMaterial);
+
+		switch( InAxis )
+		{
+			case EAxisList::X:
+			{
+				DrawCylinder( PDI, Scale * FRotationMatrix( FRotator(-90,0.f,0) ) * InMatrix, Offset, FVector(1,0,0), FVector(0,1,0), FVector(0,0,1), CylinderRadius, HalfHeight, 16, Proxy, SDPG_Foreground );
+				break;
+			}
+			case EAxisList::Y:
+			{
+				DrawCylinder( PDI, Scale * FRotationMatrix( FRotator(0,0,90) ) * InMatrix, Offset, FVector(1,0,0), FVector(0,1,0), FVector(0,0,1), CylinderRadius, HalfHeight, 16, Proxy, SDPG_Foreground );
+				break;
+			}
+			case EAxisList::Z:
+			{
+				DrawCylinder( PDI, Scale * InMatrix, Offset, FVector(1,0,0), FVector(0,1,0), FVector(0,0,1), CylinderRadius, HalfHeight, 16, Proxy, SDPG_Foreground );
+				break;
+			}
+		}
+
+		if ( bCubeHead )
+		{
+			const FVector RootPos(38,0,0);
+
+			Render_Cube( PDI, FTranslationMatrix(RootPos)*ArrowToWorld, InMaterial, 4.0f );
+		}
+		else
+		{
+			const FVector RootPos(47,0,0);
+
+			const float Angle = FMath::DegreesToRadians( PI*5 );
+			DrawCone( PDI, FScaleMatrix(-13)*FTranslationMatrix(RootPos)*ArrowToWorld, Angle, Angle, 32, false, FColor::White, Proxy, SDPG_Foreground );
+		}
+	
+		PDI->SetHitProxy( NULL );
+	}
+
+	if(!View->ScreenToPixel(View->WorldToScreen(ArrowToWorld.TransformPosition(FVector(64,0,0))),OutAxisEnd))
+	{
+		OutAxisEnd.X = OutAxisEnd.Y = 0;
+	}
+}
+
+void FCustomizableObjectWidget::Render_Cube( FPrimitiveDrawInterface* PDI, const FMatrix& InMatrix, const UMaterialInterface* InMaterial, float InScale )
+{
+	const FMatrix CubeToWorld = FScaleMatrix(FVector(InScale,InScale,InScale)) * InMatrix;
+	DrawBox( PDI, CubeToWorld, FVector(1,1,1), Helper_GetMaterialProxy(InMaterial), SDPG_Foreground );
+}
+
+
+void FCustomizableObjectWidget::Render_Quad( FPrimitiveDrawInterface* PDI, const FMatrix& InMatrix, const UMaterialInterface* InMaterial, const FColor* InColor, FVector2D InScale )
+{
+	const FMatrix CubeToWorld = FScaleMatrix(FVector(InScale.X,InScale.Y,1)) * InMatrix;
+	FDynamicMeshBuilder MeshBuilder(PDI->View->GetFeatureLevel());
+
+	// Compute cube vertices.
+	FVector3f CubeVerts[4];
+	CubeVerts[0] = FVector3f( 1, -1, 0 );
+	CubeVerts[1] = FVector3f( 1, 1, 0 );
+	CubeVerts[2] = FVector3f( -1, -1, 0 );
+	CubeVerts[3] = FVector3f( -1, 1, 0 );
+
+	for(uint32 i = 0; i < 4; ++i)
+	{
+		MeshBuilder.AddVertex( CubeVerts[i], FVector2f::ZeroVector, FVector3f::ZeroVector, FVector3f::ZeroVector, FVector3f::ZeroVector, *InColor );
+	}
+
+	MeshBuilder.AddTriangle(0,1,2);
+	MeshBuilder.AddTriangle(2,1,3);
+
+	// Draw the arrow mesh.
+	MeshBuilder.Draw(PDI, CubeToWorld, Helper_GetMaterialProxy(InMaterial), SDPG_Foreground, 0.f);
+
+#undef CUBE_FACE
+}
+
+
+/**
+ * Draws the translation widget.
+ */
+void FCustomizableObjectWidget::Render_Translate( const FSceneView* View, FPrimitiveDrawInterface* PDI, const FVector& InLocation, bool bDrawWidget )
+{
+	float Scale = View->WorldToScreen( InLocation ).W * ( 4.0f / View->UnscaledViewRect.Width() / View->ViewMatrices.GetProjectionMatrix().M[0][0] );
+
+	// Figure out axis colors
+	const FLinearColor& XColor = ( CurrentAxis&EAxisList::X ? (FLinearColor)CurrentColor : AxisColorX );
+	const FLinearColor& YColor = ( CurrentAxis&EAxisList::Y ? (FLinearColor)CurrentColor : AxisColorY );
+	const FLinearColor& ZColor = ( CurrentAxis&EAxisList::Z ? (FLinearColor)CurrentColor : AxisColorZ );
+	FColor CurrentScreenColor = ( CurrentAxis & EAxisList::Screen ? CurrentColor : ScreenSpaceColor );
+
+	// Figure out axis materials
+
+	UMaterialInstanceDynamic* XMaterial = ( CurrentAxis&EAxisList::X ? CurrentAxisMaterial : AxisMaterialX );
+	UMaterialInstanceDynamic* YMaterial = ( CurrentAxis&EAxisList::Y ? CurrentAxisMaterial : AxisMaterialY );
+	UMaterialInstanceDynamic* ZMaterial = ( CurrentAxis&EAxisList::Z ? CurrentAxisMaterial : AxisMaterialZ );
+	UMaterialInstanceDynamic* XYZMaterial = ( CurrentAxis&EAxisList::Screen) ? CurrentAxisMaterial : OpaquePlaneMaterialXY;
+
+	
+	// Figure out axis matrices
+
+	FMatrix WidgetMatrix = CustomCoordSystem * FTranslationMatrix( InLocation );
+
+	const bool bIsPerspective = View->IsPerspectiveProjection();
+	const bool bIsOrthoXY = !bIsPerspective && FMath::Abs(View->ViewMatrices.GetViewMatrix().M[2][2]) > 0.0f;
+	const bool bIsOrthoXZ = !bIsPerspective && FMath::Abs(View->ViewMatrices.GetViewMatrix().M[1][2]) > 0.0f;
+	const bool bIsOrthoYZ = !bIsPerspective && FMath::Abs(View->ViewMatrices.GetViewMatrix().M[0][2]) > 0.0f;
+
+	// For local space widgets, we always want to draw all three axis, since they may not be aligned with
+	// the orthographic projection anyway.
+	const bool bIsLocalSpace = true;
+
+	const EAxisList::Type DrawAxis = EAxisList::All;
+
+	const bool bDisabled = false;
+
+	// Draw the axis lines with arrow heads
+	if( bIsPerspective || bIsLocalSpace || !bIsOrthoYZ )
+	{
+		Render_Axis( View, PDI, EAxisList::X, WidgetMatrix, XMaterial, XAxisEnd, Scale, bDrawWidget );
+	}
+
+	if(bIsPerspective || bIsLocalSpace || !bIsOrthoXZ )
+	{
+		Render_Axis( View, PDI, EAxisList::Y, WidgetMatrix, YMaterial, YAxisEnd, Scale, bDrawWidget );
+	}
+
+	if( bIsPerspective || bIsLocalSpace || !bIsOrthoXY )
+	{
+		Render_Axis( View, PDI, EAxisList::Z, WidgetMatrix, ZMaterial, ZAxisEnd, Scale, bDrawWidget );
+	}
+
+	
+	// Draw the grabbers
+	if( bDrawWidget )
+	{
+		/*
+		FVector CornerPos = FVector(7,0,7)*Scale;
+		FVector AxisSize = FVector(12,1.2,12)*Scale;
+		float CornerLength = 1.2f*Scale;
+
+		if( bIsPerspective || bIsLocalSpace || View->ViewMatrices.ViewMatrix.M[2][1] == 0.f )
+		{
+
+			if( (DrawAxis&EAxisList::XY) == EAxisList::XY )							// Top
+			{
+				UMaterialInstanceDynamic* XMaterial = ( (CurrentAxis&EAxisList::XY) == EAxisList::XY ? CurrentAxisMaterial : AxisMaterialX );
+				UMaterialInstanceDynamic* YMaterial = ( (CurrentAxis&EAxisList::XY) == EAxisList::XY? CurrentAxisMaterial : AxisMaterialY );
+
+				PDI->SetHitProxy( new HWidgetAxis(EAxisList::XY, bDisabled) );
+				{
+					DrawDualAxis( PDI, FTranslationMatrix( CornerPos )*FRotationMatrix( FRotator( 0, 0, 90 ) ) * WidgetMatrix, AxisSize, CornerLength, XMaterial->GetRenderProxy(false), YMaterial->GetRenderProxy(false) );
+
+					//PDI->DrawLine( WidgetMatrix.TransformPosition(FVector(16,0,0) * Scale), WidgetMatrix.TransformPosition(FVector(16,16,0) * Scale), XColor, SDPG_Foreground );
+					//PDI->DrawLine( WidgetMatrix.TransformPosition(FVector(16,16,0) * Scale), WidgetMatrix.TransformPosition(FVector(0,16,0) * Scale), YColor, SDPG_Foreground );
+				}
+				PDI->SetHitProxy( NULL );
+			}
+		}
+
+		if( bIsPerspective || bIsLocalSpace || View->ViewMatrices.ViewMatrix.M[1][2] == -1.f )		// Front
+		{
+			if( (DrawAxis&EAxisList::XZ) == EAxisList::XZ ) 
+			{
+				UMaterialInstanceDynamic* XMaterial = ( (CurrentAxis&EAxisList::XZ) == EAxisList::XZ ? CurrentAxisMaterial : AxisMaterialX );
+				UMaterialInstanceDynamic* ZMaterial = ( (CurrentAxis&EAxisList::XZ) == EAxisList::XZ ? CurrentAxisMaterial : AxisMaterialZ );
+
+				PDI->SetHitProxy( new HWidgetAxis(EAxisList::XZ, bDisabled) );
+				{
+
+					DrawDualAxis( PDI, FTranslationMatrix( CornerPos )*WidgetMatrix, AxisSize, CornerLength, XMaterial->GetRenderProxy(false), ZMaterial->GetRenderProxy(false) );
+
+					//PDI->DrawLine( WidgetMatrix.TransformPosition(FVector(16,0,0) * Scale), WidgetMatrix.TransformPosition(FVector(16,0,16) * Scale), XColor, SDPG_Foreground );
+					//PDI->DrawLine( WidgetMatrix.TransformPosition(FVector(16,0,16) * Scale), WidgetMatrix.TransformPosition(FVector(0,0,16) * Scale), ZColor, SDPG_Foreground );
+				}
+				PDI->SetHitProxy( NULL );
+			}
+		}
+
+		if( bIsPerspective || bIsLocalSpace || View->ViewMatrices.ViewMatrix.M[1][0] == 1.f )		// Side
+		{
+			if( (DrawAxis&EAxisList::YZ) == EAxisList::YZ ) 
+			{
+				UMaterialInstanceDynamic* YMaterial = ( (CurrentAxis&EAxisList::YZ) == EAxisList::YZ ? CurrentAxisMaterial : AxisMaterialY );
+				UMaterialInstanceDynamic* ZMaterial = ( (CurrentAxis&EAxisList::YZ) == EAxisList::YZ ? CurrentAxisMaterial : AxisMaterialZ );
+
+				PDI->SetHitProxy( new HWidgetAxis(EAxisList::YZ, bDisabled) );
+				{
+					DrawDualAxis( PDI, FTranslationMatrix( CornerPos ) *FRotationMatrix( FRotator( 0, 90, 0 ) ) * WidgetMatrix, AxisSize, CornerLength, YMaterial->GetRenderProxy(false), ZMaterial->GetRenderProxy( false ) );
+
+					//PDI->DrawLine( WidgetMatrix.TransformPosition(FVector(0,16,0) * Scale), WidgetMatrix.TransformPosition(FVector(0,16,16) * Scale), YColor, SDPG_Foreground );
+					//PDI->DrawLine( WidgetMatrix.TransformPosition(FVector(0,16,16) * Scale), WidgetMatrix.TransformPosition(FVector(0,0,16) * Scale), ZColor, SDPG_Foreground );
+				}
+				PDI->SetHitProxy( NULL );
+			}
+		}
+		*/
+	}
+
+	// Draw screen-space movement handle (circle)
+	if( bDrawWidget && ( DrawAxis & EAxisList::Screen ) && bIsPerspective )
+	{
+		PDI->SetHitProxy( new HWidgetAxis(EAxisList::Screen, bDisabled) );
+		const FVector CameraXAxis = View->ViewMatrices.GetViewMatrix().GetColumn(0);
+		const FVector CameraYAxis = View->ViewMatrices.GetViewMatrix().GetColumn(1);
+		const FVector CameraZAxis = View->ViewMatrices.GetViewMatrix().GetColumn(2);
+
+#if (ENGINE_MAJOR_VERSION == 4 && ENGINE_MINOR_VERSION >= 17) || ENGINE_MAJOR_VERSION >= 5
+		DrawSphere( PDI, InLocation, FRotator::ZeroRotator, FVector( 4.0f * Scale ), 10, 5, Helper_GetMaterialProxy(XYZMaterial), SDPG_Foreground );
+#else
+		DrawSphere(PDI, InLocation, FVector(4.0f * Scale), 10, 5, Helper_GetMaterialProxy(XYZMaterial), SDPG_Foreground);
+#endif
+
+		PDI->SetHitProxy( NULL );
+	}
+
+}
+
+/**
+ * Draws the rotation widget.
+ */
+void FCustomizableObjectWidget::Render_Rotate( const FSceneView* View,FPrimitiveDrawInterface* PDI, const FVector& InLocation, bool bDrawWidget )
+{
+	float Scale = View->WorldToScreen( InLocation ).W * ( 4.0f / View->UnscaledViewRect.Width() / View->ViewMatrices.GetProjectionMatrix().M[0][0] );
+
+	//get the axes 
+	FVector XAxis = CustomCoordSystem.TransformVector(FVector(-1, 0, 0));
+	FVector YAxis = CustomCoordSystem.TransformVector(FVector(0, -1, 0));
+	FVector ZAxis = CustomCoordSystem.TransformVector(FVector(0, 0, 1));
+
+	EAxisList::Type DrawAxis = EAxisList::All;
+
+	FMatrix XMatrix = FRotationMatrix( FRotator(0,90.f,0) ) * FTranslationMatrix( InLocation );
+
+	FVector DirectionToWidget = View->IsPerspectiveProjection() ? (InLocation - View->ViewMatrices.GetViewOrigin()) : -View->GetViewDirection();
+	DirectionToWidget.Normalize();
+
+	// Draw a circle for each axis
+	if (bDrawWidget || bDragging)
+	{
+		DrawRotationArc(View, PDI, EAxisList::X, InLocation, YAxis, ZAxis, DirectionToWidget, AxisColorX.ToFColor(true), Scale);
+		DrawRotationArc(View, PDI, EAxisList::Y, InLocation, ZAxis, XAxis, DirectionToWidget, AxisColorY.ToFColor(true), Scale);
+		DrawRotationArc(View, PDI, EAxisList::Z, InLocation, XAxis, YAxis, DirectionToWidget, AxisColorZ.ToFColor(true), Scale);
+	}
+
+	// Update Axis by projecting the axis vector to screenspace.
+	View->ScreenToPixel(View->WorldToScreen( XMatrix.TransformPosition( FVector(96,0,0) ) ), XAxisEnd);
+	// Update Axis by projecting the axis vector to screenspace.
+	View->ScreenToPixel(View->WorldToScreen( XMatrix.TransformPosition( FVector(0,96,0) ) ), YAxisEnd);
+	// Update Axis by projecting the axis vector to screenspace.
+	View->ScreenToPixel(View->WorldToScreen( XMatrix.TransformPosition( FVector(0,0,96) ) ), ZAxisEnd);
+}
+
+/**
+ * Draws the scaling widget.
+ */
+void FCustomizableObjectWidget::Render_Scale( const FSceneView* View,FPrimitiveDrawInterface* PDI, const FVector& InLocation, bool bDrawWidget )
+{
+	const float Scale(View->WorldToScreen( InLocation ).W * ( 4.0f / View->UnscaledViewRect.Width() / View->ViewMatrices.GetProjectionMatrix().M[0][0] ));
+
+	// Figure out axis colors
+	const FLinearColor& XColor = ( CurrentAxis&EAxisList::X ? (FLinearColor)CurrentColor : AxisColorX );
+	const FLinearColor& YColor = ( CurrentAxis&EAxisList::Y ? (FLinearColor)CurrentColor : AxisColorY );
+	const FLinearColor& ZColor = ( CurrentAxis&EAxisList::Z ? (FLinearColor)CurrentColor : AxisColorZ );
+	FColor CurrentScreenColor = ( CurrentAxis & EAxisList::Screen ? CurrentColor : ScreenSpaceColor );
+
+	// Figure out axis materials
+
+	UMaterialInstanceDynamic* XMaterial = ( CurrentAxis&EAxisList::X ? CurrentAxisMaterial : AxisMaterialX );
+	UMaterialInstanceDynamic* YMaterial = ( CurrentAxis&EAxisList::Y ? CurrentAxisMaterial : AxisMaterialY );
+	UMaterialInstanceDynamic* ZMaterial = ( CurrentAxis&EAxisList::Z ? CurrentAxisMaterial : AxisMaterialZ );
+	UMaterialInstanceDynamic* XYZMaterial = ( CurrentAxis&EAxisList::XYZ ? CurrentAxisMaterial : OpaquePlaneMaterialXY );
+
+	// Figure out axis matrices
+
+	FMatrix WidgetMatrix = CustomCoordSystem * FTranslationMatrix( InLocation );
+	// Determine viewport
+
+	const EAxisList::Type DrawAxis = EAxisList::All;
+	const bool bIsPerspective = View->IsPerspectiveProjection();
+	const bool bIsOrthoXY(!bIsPerspective && FMath::Abs(View->ViewMatrices.GetViewMatrix().M[2][2]) > 0.0f);
+	const bool bIsOrthoXZ(!bIsPerspective && FMath::Abs(View->ViewMatrices.GetViewMatrix().M[1][2]) > 0.0f);
+	const bool bIsOrthoYZ(!bIsPerspective && FMath::Abs(View->ViewMatrices.GetViewMatrix().M[0][2]) > 0.0f);
+
+	// Draw the axis lines with cube heads
+
+	if( !bIsOrthoYZ && DrawAxis&EAxisList::X )
+	{
+		Render_Axis( View, PDI, EAxisList::X, WidgetMatrix, XMaterial, XAxisEnd, Scale, bDrawWidget, true );
+	}
+
+	if( !bIsOrthoXZ && DrawAxis&EAxisList::Y )
+	{
+		Render_Axis( View, PDI, EAxisList::Y, WidgetMatrix, YMaterial, YAxisEnd, Scale, bDrawWidget, true );
+	}
+
+	if( !bIsOrthoXY &&  DrawAxis&EAxisList::Z )
+	{
+		//Render_Axis( View, PDI, EAxisList::Z, WidgetMatrix, ZMaterial, ZAxisEnd, Scale, bDrawWidget, true );
+	}
+
+	// Draw grabber handles and center cube
+	if ( bDrawWidget )
+	{
+		const bool bDisabled = false;
+
+		// Grabber handles
+		if( !bIsOrthoYZ && !bIsOrthoXZ && ((DrawAxis&(EAxisList::X|EAxisList::Y)) == (EAxisList::X|EAxisList::Y)) )
+		{
+			PDI->SetHitProxy( new HWidgetAxis(EAxisList::XY, bDisabled) );
+			{
+				PDI->DrawLine( WidgetMatrix.TransformPosition(FVector(24,0,0) * Scale), WidgetMatrix.TransformPosition(FVector(12,12,0) * Scale), XColor, SDPG_Foreground );
+				PDI->DrawLine( WidgetMatrix.TransformPosition(FVector(12,12,0) * Scale), WidgetMatrix.TransformPosition(FVector(0,24,0) * Scale), YColor, SDPG_Foreground );
+			}
+			PDI->SetHitProxy( NULL );
+		}
+
+		if( !bIsOrthoYZ && !bIsOrthoXY && ((DrawAxis&(EAxisList::X|EAxisList::Z)) == (EAxisList::X|EAxisList::Z)) )
+		{
+			PDI->SetHitProxy( new HWidgetAxis(EAxisList::XZ, bDisabled) );
+			{
+				PDI->DrawLine( WidgetMatrix.TransformPosition(FVector(24,0,0) * Scale), WidgetMatrix.TransformPosition(FVector(12,0,12) * Scale), XColor, SDPG_Foreground );
+				PDI->DrawLine( WidgetMatrix.TransformPosition(FVector(12,0,12) * Scale), WidgetMatrix.TransformPosition(FVector(0,0,24) * Scale), ZColor, SDPG_Foreground );
+			}
+			PDI->SetHitProxy( NULL );
+		}
+
+		if( !bIsOrthoXY && !bIsOrthoXZ && ((DrawAxis&(EAxisList::Y|EAxisList::Z)) == (EAxisList::Y|EAxisList::Z)) )
+		{
+			PDI->SetHitProxy( new HWidgetAxis(EAxisList::YZ, bDisabled) );
+			{
+				PDI->DrawLine( WidgetMatrix.TransformPosition(FVector(0,24,0) * Scale), WidgetMatrix.TransformPosition(FVector(0,12,12) * Scale), YColor, SDPG_Foreground );
+				PDI->DrawLine( WidgetMatrix.TransformPosition(FVector(0,12,12) * Scale), WidgetMatrix.TransformPosition(FVector(0,0,24) * Scale), ZColor, SDPG_Foreground );
+			}
+			PDI->SetHitProxy( NULL );
+		}
+
+		// Center cube
+		if( (DrawAxis&(EAxisList::XYZ)) == EAxisList::XYZ )
+		{
+			PDI->SetHitProxy( new HWidgetAxis(EAxisList::XYZ, bDisabled) );
+
+			Render_Cube(PDI, WidgetMatrix, XYZMaterial, Scale * 4 );
+
+			PDI->SetHitProxy( NULL );
+		}
+	}
+}
+
+
+/**
+* Draws the Translate & Rotate Z widget.
+*/
+
+void FCustomizableObjectWidget::Render_TranslateRotateZ( const FSceneView* View, FPrimitiveDrawInterface* PDI, const FVector& InLocation, bool bDrawWidget )
+{
+	float Scale = View->WorldToScreen( InLocation ).W * ( 4.0f / View->UnscaledViewRect.Width() / View->ViewMatrices.GetProjectionMatrix().M[0][0] );
+
+	// Figure out axis colors
+
+	FColor XYPlaneColor  = ( (CurrentAxis&EAxisList::XY) == EAxisList::XY ) ? CurrentColor : PlaneColorXY;
+	FColor ZRotateColor  = ( (CurrentAxis&EAxisList::ZRotation) == EAxisList::ZRotation ) ? CurrentColor : AxisColorZ.ToFColor(true);
+	FColor XColor        = ( (CurrentAxis&EAxisList::X) == EAxisList::X ) ? CurrentColor : AxisColorX.ToFColor(true);
+	FColor YColor        = ( (CurrentAxis&EAxisList::Y) == EAxisList::Y && CurrentAxis != EAxisList::ZRotation ) ? CurrentColor : AxisColorY.ToFColor(true);
+	FColor ZColor        = ( (CurrentAxis&EAxisList::Z) == EAxisList::Z ) ? CurrentColor : AxisColorZ.ToFColor(true);
+
+	// Figure out axis materials
+	UMaterialInstance* ZRotateMaterial = ( ( (CurrentAxis&EAxisList::ZRotation) == EAxisList::ZRotation ) ? CurrentAxisMaterial : AxisMaterialZ );
+	UMaterialInstance* XMaterial = ( CurrentAxis&EAxisList::X ? CurrentAxisMaterial : AxisMaterialX );
+	UMaterialInstance* YMaterial = ( ( CurrentAxis&EAxisList::Y && CurrentAxis != EAxisList::ZRotation ) ? CurrentAxisMaterial : AxisMaterialY );
+	UMaterialInstance* ZMaterial = ( CurrentAxis&EAxisList::Z ? CurrentAxisMaterial : AxisMaterialZ );
+
+	// Figure out axis matrices
+	FMatrix XMatrix = CustomCoordSystem * FTranslationMatrix( InLocation );
+	FMatrix YMatrix = FRotationMatrix( FRotator(0,90.f,0) ) * CustomCoordSystem * FTranslationMatrix( InLocation );
+	FMatrix ZMatrix = FRotationMatrix( FRotator(90.f,0,0) ) * CustomCoordSystem * FTranslationMatrix( InLocation );
+
+	const bool bIsPerspective = View->IsPerspectiveProjection();
+
+	// For local space widgets, we always want to draw all three axis, since they may not be aligned with
+	// the orthographic projection anyway.
+	bool bIsLocalSpace = false;//( GEditorModeTools().GetCoordSystem() == COORD_Local );
+
+	int32 DrawAxis = EAxisList::XYZ;//GEditorModeTools().GetWidgetAxisToDraw( GEditorModeTools().GetWidgetMode() );
+
+	// Draw the grabbers
+	if( bDrawWidget )
+	{
+		// Draw the axis lines with arrow heads
+		if( DrawAxis&EAxisList::X && (bIsPerspective || bIsLocalSpace || View->ViewMatrices.GetViewMatrix().M[0][2] != -1.f) )
+		{
+			Render_Axis( View, PDI, EAxisList::X, XMatrix, XMaterial, XAxisEnd, Scale, bDrawWidget );
+		}
+
+		if( DrawAxis&EAxisList::Y && (bIsPerspective || bIsLocalSpace || View->ViewMatrices.GetViewMatrix().M[1][2] != -1.f) )
+		{
+			Render_Axis( View, PDI, EAxisList::Y, YMatrix, YMaterial, YAxisEnd, Scale, bDrawWidget );
+		}
+
+		if( DrawAxis&EAxisList::Z && (bIsPerspective || bIsLocalSpace || View->ViewMatrices.GetViewMatrix().M[0][1] != 1.f) )
+		{
+			Render_Axis( View, PDI, EAxisList::Z, ZMatrix, ZMaterial, ZAxisEnd, Scale, bDrawWidget );
+		}
+
+		bool bDisabled = false;//GEditorModeTools().IsModeActive(EM_Default) && GEditor->HasLockedActors();
+
+		//ZRotation
+		/*
+		if( DrawAxis&EAxis::ZROTATION && (bIsPerspective || bIsLocalSpace || View->ViewMatrix.M[0][2] != -1.f) )
+		{
+			PDI->SetHitProxy( new HWidgetAxis(EAxis::ZROTATION, bDisabled) );
+			{
+				float ScaledRadius = TRANSLATE_ROTATE_EAxis::CIRCLE_RADIUS*Scale;
+				FVector XAxis = CustomCoordSystem.TransformPosition( FVector(1,0,0).RotateAngleAxis(GEditorModeTools().TranslateRotateXAxisAngle, FVector(0,0,1)) );
+				FVector YAxis = CustomCoordSystem.TransformPosition( FVector(0,1,0).RotateAngleAxis(GEditorModeTools().TranslateRotateXAxisAngle, FVector(0,0,1)) );
+				FVector BaseArrowPoint = InLocation + XAxis*ScaledRadius;
+				DrawFlatArrow(PDI, BaseArrowPoint, XAxis, YAxis, ZRotateColor, ScaledRadius, ScaledRadius*.5f, ZRotateMaterial->GetRenderProxy(false), SDPG_Foreground);
+			}
+			PDI->SetHitProxy( NULL );
+		}
+
+		//XY Plane
+		if( bIsPerspective || bIsLocalSpace || View->ViewMatrix.M[0][1] != 1.f )
+		{
+			if( (DrawAxis&(EAxis::XY)) == (EAxis::XY) ) 
+			{
+				PDI->SetHitProxy( new HWidgetAxis(EAxis::XY, bDisabled) );
+				{
+					DrawCircle( PDI, InLocation, CustomCoordSystem.TransformPosition( FVector(1,0,0) ), CustomCoordSystem.TransformPosition( FVector(0,1,0) ), XYPlaneColor, TRANSLATE_ROTATE_EAxis::CIRCLE_RADIUS*Scale, EAxis::CIRCLE_SIDES, SDPG_Foreground );
+					XYPlaneColor.A = ((CurrentAxis&EAxis::XY)==EAxis::XY) ? 0x3f : 0x0f;	//make the disc transparent
+					DrawDisc  ( PDI, InLocation, CustomCoordSystem.TransformPosition( FVector(1,0,0) ), CustomCoordSystem.TransformPosition( FVector(0,1,0) ), XYPlaneColor, TRANSLATE_ROTATE_EAxis::CIRCLE_RADIUS*Scale, EAxis::CIRCLE_SIDES, TransparentPlaneMaterialXY->GetRenderProxy(false), SDPG_Foreground );
+				}
+				PDI->SetHitProxy( NULL );
+			}
+		}
+		*/
+	}
+}
+
+
+/** Only some modes support Absolute Translation Movement */
+bool FCustomizableObjectWidget::AllowsAbsoluteTranslationMovement(void)
+{
+	EWidgetMode CurrentMode = WM_Translate;//GEditorModeTools().GetWidgetMode();
+	if ((CurrentMode == WM_Translate) || (CurrentMode == WM_TranslateRotateZ))
+	{
+		return true;
+	}
+	return false;
+}
+
+/** 
+ * Serializes the widget references so they don't get garbage collected.
+ *
+ * @param Ar	FArchive to serialize with
+ */
+void FCustomizableObjectWidget::AddReferencedObjects( FReferenceCollector& Collector )
+{
+	Collector.AddReferencedObject( AxisMaterialX );
+	Collector.AddReferencedObject( AxisMaterialY );
+	Collector.AddReferencedObject( AxisMaterialZ );
+	Collector.AddReferencedObject( OpaquePlaneMaterialXY );
+	Collector.AddReferencedObject( TransparentPlaneMaterialXY );
+	Collector.AddReferencedObject( GridMaterial );
+	Collector.AddReferencedObject( CurrentAxisMaterial );
+}
+
+#define CAMERA_LOCK_DAMPING_FACTOR .1f
+#define MAX_CAMERA_MOVEMENT_SPEED 512.0f
+/**
+ * Returns the Delta from the current position that the absolute movement system wants the object to be at
+ * @param InParams - Structure containing all the information needed for absolute movement
+ * @return - The requested delta from the current position
+ */
+FVector FCustomizableObjectWidget::GetAbsoluteTranslationDelta (const FAbsoluteMovementParams& InParams)
+{
+	FPlane MovementPlane(InParams.Position, InParams.PlaneNormal);
+	FVector ProposedEndofEyeVector = InParams.EyePos + InParams.PixelDir;
+
+	//default to not moving
+	FVector RequestedPosition = InParams.Position;
+
+	float DotProductWithPlaneNormal = InParams.PixelDir|InParams.PlaneNormal;
+	//check to make sure we're not co-planar
+	if (FMath::Abs(DotProductWithPlaneNormal) > DELTA)
+	{
+		//Get closest point on plane
+		RequestedPosition = FMath::LinePlaneIntersection(InParams.EyePos, ProposedEndofEyeVector, MovementPlane);
+	}
+
+	//drag is a delta position, so just update the different between the previous position and the new position
+	FVector DeltaPosition = RequestedPosition - InParams.Position;
+
+	//Retrieve the initial offset, passing in the current requested position and the current position
+	FVector InitialOffset = GetAbsoluteTranslationInitialOffset(RequestedPosition, InParams.Position);
+
+	//subtract off the initial offset (where the widget was clicked) to prevent popping
+	DeltaPosition -= InitialOffset;
+
+	//remove the component along the normal we want to mute
+	float MovementAlongMutedAxis = DeltaPosition|InParams.NormalToRemove;
+	FVector OutDrag = DeltaPosition - (InParams.NormalToRemove*MovementAlongMutedAxis);
+
+	if (InParams.bMovementLockedToCamera)
+	{
+		//DAMPEN ABSOLUTE MOVEMENT when the camera is locked to the object
+		OutDrag *= CAMERA_LOCK_DAMPING_FACTOR;
+		OutDrag.X = FMath::Clamp(OutDrag.X, -MAX_CAMERA_MOVEMENT_SPEED, MAX_CAMERA_MOVEMENT_SPEED);
+		OutDrag.Y = FMath::Clamp(OutDrag.Y, -MAX_CAMERA_MOVEMENT_SPEED, MAX_CAMERA_MOVEMENT_SPEED);
+		OutDrag.Z = FMath::Clamp(OutDrag.Z, -MAX_CAMERA_MOVEMENT_SPEED, MAX_CAMERA_MOVEMENT_SPEED);
+	}
+
+	//the they requested position snapping and we're not moving with the camera
+	if (InParams.bPositionSnapping && !InParams.bMovementLockedToCamera && bSnapEnabled)
+	{
+		FVector MovementAlongAxis = FVector(OutDrag|InParams.XAxis, OutDrag|InParams.YAxis, OutDrag|InParams.ZAxis);
+		//translation (either xy plane or z)
+		//GEditor->Snap( MovementAlongAxis, FVector(GEditor->GetGridSize(),GEditor->GetGridSize(),GEditor->GetGridSize()) );
+		OutDrag = MovementAlongAxis.X*InParams.XAxis + MovementAlongAxis.Y*InParams.YAxis + MovementAlongAxis.Z*InParams.ZAxis;
+	}
+
+	//get the distance from the original position to the new proposed position 
+	FVector DeltaFromStart = InParams.Position + OutDrag - InitialTranslationPosition;
+
+	//Get the vector from the eye to the proposed new position (to make sure it's not behind the camera
+	FVector EyeToNewPosition = (InParams.Position + OutDrag) - InParams.EyePos;
+	float BehindTheCameraDotProduct = EyeToNewPosition|InParams.CameraDir;
+
+	//clamp so we don't lose objects off the edge of the world, or the requested position is behind the camera
+	if ((DeltaFromStart.Size() > HALF_WORLD_MAX*.5) || ( BehindTheCameraDotProduct <= 0))
+	{
+		OutDrag = OutDrag.ZeroVector;
+	}
+	return OutDrag;
+}
+
+/**
+ * Returns the offset from the initial selection point
+ */
+FVector FCustomizableObjectWidget::GetAbsoluteTranslationInitialOffset(const FVector& InNewPosition, const FVector& InCurrentPosition)
+{
+	if (!bAbsoluteTranslationInitialOffsetCached)
+	{
+		bAbsoluteTranslationInitialOffsetCached = true;
+		InitialTranslationOffset = InNewPosition - InCurrentPosition;
+		InitialTranslationPosition = InCurrentPosition;
+	}
+	return InitialTranslationOffset;
+}
+
+
+
+/**
+ * Returns true if we're in Local Space editing mode or editing BSP (which uses the World axes anyway
+ */
+bool FCustomizableObjectWidget::IsRotationLocalSpace (void) const
+{
+	bool bIsLocalSpace = false;//( GEditorModeTools().GetCoordSystem() == COORD_Local );
+	//for bsp and things that don't have a "true" local space, they will always use world.  So do NOT invert.
+	if (bIsLocalSpace && CustomCoordSystem.Equals(FMatrix::Identity))
+	{
+		bIsLocalSpace = false;
+	}
+	return bIsLocalSpace;
+}
+
+/**
+ * Returns the angle in degrees representation of how far we have just rotated
+ */
+float FCustomizableObjectWidget::GetDeltaRotation (void) const
+{
+	bool bIsLocalSpace = IsRotationLocalSpace();
+	return (bIsLocalSpace ? -1 : 1)*30.0f;//GEditorModeTools().TotalDeltaRotation;
+}
+
+
+constexpr uint8 LargeInnerAlpha = 0x1f;
+constexpr uint8 SmallInnerAlpha = 0x0f;
+constexpr uint8 LargeOuterAlpha = 0x5f;
+constexpr uint8 SmallOuterAlpha = 0x0f;
+
+/**
+ * If actively dragging, draws a ring representing the potential rotation of the selected objects, snap ticks, and "delta" markers
+ * If not actively dragging, draws a quarter ring representing the closest quadrant to the camera
+ * @param View - Information about the scene/camera/etc
+ * @param PDI - Drawing interface
+ * @param InAxis - Enumeration of axis to rotate about
+ * @param InLocation - The Origin of the widget
+ * @param Axis0 - The Axis that describes a 0 degree rotation
+ * @param Axis1 - The Axis that describes a 90 degree rotation
+ * @param InDirectionToWidget - Direction from camera to the widget
+ * @param InColor - The color associated with the axis of rotation
+ * @param InScale - Multiplier to maintain a constant screen size for rendering the widget
+ */
+void FCustomizableObjectWidget::DrawRotationArc(const FSceneView* View, FPrimitiveDrawInterface* PDI, EAxisList::Type InAxis, const FVector& InLocation, const FVector& Axis0, const FVector& Axis1, const FVector& InDirectionToWidget, const FColor& InColor, const float InScale)
+{
+	bool bIsPerspective = View->IsPerspectiveProjection();
+	bool bIsOrtho = !bIsPerspective;
+	//if we're in an ortho viewport and the ring is perpendicular to the camera (both Axis0 & Axis1 are perpendicular)
+	bool bIsOrthoDrawingFullRing = bIsOrtho && (FMath::Abs(Axis0|InDirectionToWidget) < UE_KINDA_SMALL_NUMBER) && (FMath::Abs(Axis1|InDirectionToWidget) < UE_KINDA_SMALL_NUMBER);
+
+	FColor ArcColor = InColor;
+	ArcColor.A = LargeOuterAlpha;
+
+	if (bDragging || (bIsOrthoDrawingFullRing))
+	{
+		if ((CurrentAxis&InAxis) || (bIsOrthoDrawingFullRing))
+		{
+			float DeltaRotation = GetDeltaRotation();
+			float AbsRotation = FRotator::ClampAxis(fabs(DeltaRotation));
+			float AngleOfChangeRadians (AbsRotation * PI / 180.f);
+
+			//always draw clockwise, so if we're negative we need to flip the angle
+			float StartAngle = DeltaRotation < 0.0f ? -AngleOfChangeRadians : 0.0f;
+			float FilledAngle = AngleOfChangeRadians;
+
+			//the axis of rotation
+			FVector ZAxis = Axis0 ^ Axis1;
+
+			ArcColor.A = LargeOuterAlpha;
+			DrawPartialRotationArc(View, PDI, InAxis, InLocation,  Axis0, Axis1, StartAngle, StartAngle + FilledAngle, ArcColor, InScale, InDirectionToWidget);
+			ArcColor.A = SmallOuterAlpha;
+			DrawPartialRotationArc(View, PDI, InAxis, InLocation,  Axis0, Axis1, StartAngle + FilledAngle, StartAngle + 2*PI, ArcColor, InScale, InDirectionToWidget);
+
+			ArcColor = (CurrentAxis&InAxis) ? CurrentColor : ArcColor;
+			//Hallow Arrow
+			ArcColor.A = 0;
+			DrawStartStopMarker(PDI, InLocation, Axis0, Axis1, 0, ArcColor, InScale);
+			//Filled Arrow
+			ArcColor.A = LargeOuterAlpha;
+			DrawStartStopMarker(PDI, InLocation, Axis0, Axis1, DeltaRotation, ArcColor, InScale);
+
+			ArcColor.A = 255;
+
+
+			// Push out the snap marks so they dont z-fight with the background arcs
+			FVector SnapLocation = InLocation;
+			if( bIsOrtho)
+			{
+				SnapLocation += InDirectionToWidget * .5f;
+			}
+			else if( (InDirectionToWidget|ZAxis) <= 0.0f )
+			{
+				SnapLocation += ZAxis * .5f;
+			}
+			else
+			{
+				SnapLocation += ZAxis * -.5f;
+			}
+
+			
+			//if (GEditor->GetUserSettings().RotGridEnabled)
+			//{
+			//	float DeltaAngle = GEditor->GetUserSettings().RotGridSize.Yaw;
+			//	//every 22.5 degrees
+			//	float TickMarker = 22.5f;
+			//	for (float Angle = 0; Angle < 360.f; Angle+=DeltaAngle)
+			//	{ 
+			//		FVector GridAxis = Axis0.RotateAngleAxis(Angle, ZAxis);
+			//		float PercentSize = (FMath::Fmod(Angle, TickMarker)==0) ? .75f : .25f;
+			//		if (FMath::Fmod(Angle, 90.f) != 0)
+			//		{
+			//			DrawSnapMarker(PDI, SnapLocation,  GridAxis,  FVector::ZeroVector, ArcColor, InScale, 0.0f, PercentSize);
+			//		}
+			//	}
+			//}
+
+			//draw axis tick marks
+			FColor AxisColor = InColor;
+			//Rotate Colors to match Axis 0
+			Swap(AxisColor.R, AxisColor.G);
+			Swap(AxisColor.B, AxisColor.R);
+			AxisColor.A = (DeltaRotation == 0) ? MAX_uint8 : LargeOuterAlpha;
+			DrawSnapMarker(PDI, SnapLocation,  Axis0,  Axis1, AxisColor, InScale, .25f);
+			AxisColor.A = (DeltaRotation == 180.f) ? MAX_uint8 : LargeOuterAlpha;
+			DrawSnapMarker(PDI, SnapLocation, -Axis0, -Axis1, AxisColor, InScale, .25f);
+
+			//Rotate Colors to match Axis 1
+			Swap(AxisColor.R, AxisColor.G);
+			Swap(AxisColor.B, AxisColor.R);
+			AxisColor.A = (DeltaRotation == 90.f) ? MAX_uint8 : LargeOuterAlpha;
+			DrawSnapMarker(PDI, SnapLocation,  Axis1, -Axis0, AxisColor, InScale, .25f);
+			AxisColor.A = (DeltaRotation == 270.f) ? MAX_uint8 : LargeOuterAlpha;
+			DrawSnapMarker(PDI, SnapLocation, -Axis1,  Axis0, AxisColor, InScale, .25f);
+
+			if (bDragging)
+			{
+				float OffsetAngle = IsRotationLocalSpace() ? 0 : DeltaRotation;
+
+				//CacheRotationHUDText(View, PDI, InLocation, Axis0.RotateAngleAxis(OffsetAngle, ZAxis), Axis1.RotateAngleAxis(OffsetAngle, ZAxis), DeltaRotation, InScale);
+			}
+		}
+	}
+	else
+	{
+		//Reverse the axes based on camera view
+		FVector RenderAxis0 = ((Axis0|InDirectionToWidget) <= 0.0f) ? Axis0 : -Axis0;
+		FVector RenderAxis1 = ((Axis1|InDirectionToWidget) <= 0.0f) ? Axis1 : -Axis1;
+
+		DrawPartialRotationArc(View, PDI, InAxis, InLocation, RenderAxis0, RenderAxis1, 0, PI/2, ArcColor, InScale, InDirectionToWidget);
+	}
+}
+
+/**
+ * If actively dragging, draws a ring representing the potential rotation of the selected objects, snap ticks, and "delta" markers
+ * If not actively dragging, draws a quarter ring representing the closest quadrant to the camera
+ * @param View - Information about the scene/camera/etc
+ * @param PDI - Drawing interface
+ * @param InAxis - Enumeration of axis to rotate about
+ * @param InLocation - The Origin of the widget
+ * @param Axis0 - The Axis that describes a 0 degree rotation
+ * @param Axis1 - The Axis that describes a 90 degree rotation
+ * @param InStartAngle - The starting angle about (Axis0^Axis1) to render the arc, in radians
+ * @param InEndAngle - The ending angle about (Axis0^Axis1) to render the arc, in radians
+ * @param InColor - The color associated with the axis of rotation
+ * @param InScale - Multiplier to maintain a constant screen size for rendering the widget
+ */
+void FCustomizableObjectWidget::DrawPartialRotationArc(const FSceneView* View, FPrimitiveDrawInterface* PDI, EAxisList::Type InAxis, const FVector& InLocation, const FVector& Axis0, const FVector& Axis1, const float InStartAngle, const float InEndAngle, const FColor& InColor, const float InScale, const FVector& InDirectionToWidget )
+{
+	bool bIsPerspective = View->IsPerspectiveProjection();
+	PDI->SetHitProxy( new HWidgetAxis( InAxis ) );
+	{
+		FThickArcParams OuterArcParams(PDI, InLocation, TransparentPlaneMaterialXY, INNER_AXIS_CIRCLE_RADIUS*InScale, OUTER_AXIS_CIRCLE_RADIUS*InScale);
+		FColor OuterColor = ( CurrentAxis&InAxis ? CurrentColor : InColor );
+		//Pass through alpha
+		OuterColor.A = InColor.A;
+		DrawThickArc(OuterArcParams, Axis0, Axis1, InStartAngle, InEndAngle, OuterColor, InDirectionToWidget, !bIsPerspective );
+	}
+	PDI->SetHitProxy( NULL );
+
+	if (bIsPerspective)
+	{
+		FThickArcParams InnerArcParams(PDI, InLocation, GridMaterial, 0.0f, INNER_AXIS_CIRCLE_RADIUS*InScale);
+		FColor InnerColor = InColor;
+		//if something is selected and it's not this
+		InnerColor.A = ((CurrentAxis & InAxis) && !bDragging) ? LargeInnerAlpha : SmallInnerAlpha;
+		DrawThickArc(InnerArcParams, Axis0, Axis1, InStartAngle, InEndAngle, InnerColor, InDirectionToWidget, false );
+	}
+}
+
+/**
+ * Renders a portion of an arc for the rotation widget
+ * @param InParams - Material, Radii, etc
+ * @param InStartAxis - Start of the arc, in radians
+ * @param InEndAxis - End of the arc, in radians
+ * @param InColor - Color to use for the arc
+ */
+void FCustomizableObjectWidget::DrawThickArc (const FThickArcParams& InParams, const FVector& Axis0, const FVector& Axis1, const float InStartAngle, const float InEndAngle, const FColor& InColor, const FVector& InDirectionToWidget, bool bIsOrtho )
+{
+	if (InColor.A == 0)
+	{
+		return;
+	}
+	const int32 NumPoints = FMath::TruncToInt(AXIS_CIRCLE_SIDES * (InEndAngle-InStartAngle)/(PI/2)) + 1;
+
+	FColor TriangleColor = InColor;
+	FColor RingColor = InColor;
+	RingColor.A = MAX_uint8;
+
+	FVector ZAxis = Axis0 ^ Axis1;
+	FVector LastVertex;
+
+	FDynamicMeshBuilder MeshBuilder(InParams.PDI->View->GetFeatureLevel());
+
+	for (int32 RadiusIndex = 0; RadiusIndex < 2; ++RadiusIndex)
+	{
+		float Radius = (RadiusIndex == 0) ? InParams.OuterRadius : InParams.InnerRadius;
+		float TCRadius = Radius / (float) InParams.OuterRadius;
+		//Compute vertices for base circle.
+		for(int32 VertexIndex = 0;VertexIndex <= NumPoints;VertexIndex++)
+		{
+			float Percent = VertexIndex/(float)NumPoints;
+			float Angle = FMath::Lerp(InStartAngle, InEndAngle, Percent);
+			float AngleDeg = FRotator::ClampAxis(Angle * 180.f / PI);
+
+			FVector3f VertexDir = FVector3f(Axis0.RotateAngleAxis(AngleDeg, ZAxis));
+			VertexDir.Normalize();
+
+			float TCAngle = Percent*(PI/2);
+			FVector2f TC(TCRadius*FMath::Cos(Angle), TCRadius*FMath::Sin(Angle));
+
+			const FVector3f VertexPosition = FVector3f(InParams.Position) + VertexDir*Radius;
+			FVector3f Normal = VertexPosition - FVector3f(InParams.Position);
+			Normal.Normalize();
+
+			FDynamicMeshVertex MeshVertex;
+			MeshVertex.Position = FVector3f(VertexPosition);
+			MeshVertex.Color = TriangleColor;
+			MeshVertex.TextureCoordinate[0] = TC;
+
+			MeshVertex.SetTangents(
+				(FVector3f)-ZAxis,
+				FVector3f(-ZAxis) ^ Normal,
+				Normal
+				);
+
+			MeshBuilder.AddVertex(MeshVertex); //Add bottom vertex
+
+			// Push out the arc line borders so they dont z-fight with the mesh arcs
+			FVector StartLinePos = LastVertex;
+			FVector EndLinePos = FVector(VertexPosition);
+			if (VertexIndex != 0)
+			{
+				if( bIsOrtho )
+				{
+					StartLinePos += InDirectionToWidget * .5f;
+					EndLinePos += InDirectionToWidget * .5f;
+				}
+				else if( (InDirectionToWidget|ZAxis) <= 0.0f )
+				{
+					StartLinePos += ZAxis * .5f;
+					EndLinePos += ZAxis * .5f;
+				}
+				else
+				{
+					StartLinePos += ZAxis * -.5f;
+					EndLinePos += ZAxis * -.5f;
+				}
+
+				InParams.PDI->DrawLine(StartLinePos,EndLinePos,RingColor,SDPG_Foreground);
+			}
+			LastVertex = FVector(VertexPosition);
+		}
+	}
+
+	//Add top/bottom triangles, in the style of a fan.
+	int32 InnerVertexStartIndex = NumPoints + 1;
+	for(int32 VertexIndex = 0; VertexIndex < NumPoints; VertexIndex++)
+	{
+		MeshBuilder.AddTriangle(VertexIndex, VertexIndex+1, InnerVertexStartIndex+VertexIndex);
+		MeshBuilder.AddTriangle(VertexIndex+1, InnerVertexStartIndex+VertexIndex+1, InnerVertexStartIndex+VertexIndex);
+	}
+
+	MeshBuilder.Draw(InParams.PDI, FMatrix::Identity, Helper_GetMaterialProxy(InParams.Material),SDPG_Foreground,0.f);
+}
+
+/**
+ * Draws protractor like ticks where the rotation widget would snap too.
+ * Also, used to draw the wider axis tick marks
+ * @param PDI - Drawing interface
+ * @param InLocation - The Origin of the widget
+ * @param Axis0 - The Axis that describes a 0 degree rotation
+ * @param Axis1 - The Axis that describes a 90 degree rotation
+ * @param InAngle - The Angle to rotate about the axis of rotation, the vector (Axis0 ^ Axis1)
+ * @param InColor - The color to use for line/poly drawing
+ * @param InScale - Multiplier to maintain a constant screen size for rendering the widget
+ * @param InWidthPercent - The percent of the distance between the outer ring and inner ring to use for tangential thickness
+ * @param InPercentSize - The percent of the distance between the outer ring and inner ring to use for radial distance
+ */
+void FCustomizableObjectWidget::DrawSnapMarker(FPrimitiveDrawInterface* PDI, const FVector& InLocation, const FVector& Axis0, const FVector& Axis1, const FColor& InColor, const float InScale, const float InWidthPercent, const float InPercentSize)
+{
+	const float InnerDistance = (INNER_AXIS_CIRCLE_RADIUS*InScale);
+	const float OuterDistance = (OUTER_AXIS_CIRCLE_RADIUS*InScale);
+	const float MaxMarkerHeight = OuterDistance - InnerDistance;
+	const float MarkerWidth = MaxMarkerHeight*InWidthPercent;
+	const float MarkerHeight = MaxMarkerHeight*InPercentSize;
+
+	FVector Vertices[4];
+	Vertices[0] = InLocation + (OuterDistance)*Axis0 - (MarkerWidth*.5)*Axis1;
+	Vertices[1] = Vertices[0] + (MarkerWidth)*Axis1;
+	Vertices[2] = InLocation + (OuterDistance-MarkerHeight)*Axis0 - (MarkerWidth*.5)*Axis1;
+	Vertices[3] = Vertices[2] + (MarkerWidth)*Axis1;
+
+	//draw at least one line
+	PDI->DrawLine(Vertices[0], Vertices[2], InColor, SDPG_Foreground);
+
+	//if there should be thickness, draw the other lines
+	if (InWidthPercent > 0.0f)
+	{
+		PDI->DrawLine(Vertices[0], Vertices[1], InColor, SDPG_Foreground);
+		PDI->DrawLine(Vertices[1], Vertices[3], InColor, SDPG_Foreground);
+		PDI->DrawLine(Vertices[2], Vertices[3], InColor, SDPG_Foreground);
+
+		//fill in the box
+		FDynamicMeshBuilder MeshBuilder(PDI->View->GetFeatureLevel());
+
+		for(int32 VertexIndex = 0;VertexIndex < 4; VertexIndex++)
+		{
+			FDynamicMeshVertex MeshVertex;
+			MeshVertex.Position = (FVector3f)Vertices[VertexIndex];
+			MeshVertex.Color = InColor;
+			MeshVertex.TextureCoordinate[0] = FVector2f(0.0f, 0.0f);
+			MeshVertex.SetTangents(
+				(FVector3f)Axis0,
+				(FVector3f)Axis1,
+				(FVector3f)(Axis0 ^ Axis1)
+				);
+			MeshBuilder.AddVertex(MeshVertex); //Add bottom vertex
+		}
+
+		MeshBuilder.AddTriangle(0, 1, 2);
+		MeshBuilder.AddTriangle(1, 3, 2);
+		MeshBuilder.Draw(PDI, FMatrix::Identity, Helper_GetMaterialProxy(TransparentPlaneMaterialXY),SDPG_Foreground,0.f);
+	}
+}
+
+/**
+ * Draw Start/Stop Marker to show delta rotations along the arc of rotation
+ * @param PDI - Drawing interface
+ * @param InLocation - The Origin of the widget
+ * @param Axis0 - The Axis that describes a 0 degree rotation
+ * @param Axis1 - The Axis that describes a 90 degree rotation
+ * @param InAngle - The Angle to rotate about the axis of rotation, the vector (Axis0 ^ Axis1), units are degrees
+ * @param InColor - The color to use for line/poly drawing
+ * @param InScale - Multiplier to maintain a constant screen size for rendering the widget
+ */
+void FCustomizableObjectWidget::DrawStartStopMarker(FPrimitiveDrawInterface* PDI, const FVector& InLocation, const FVector& Axis0, const FVector& Axis1, const float InAngle, const FColor& InColor, const float InScale)
+{
+	const float ArrowHeightPercent = .8f;
+	const float InnerDistance = (INNER_AXIS_CIRCLE_RADIUS*InScale);
+	const float OuterDistance = (OUTER_AXIS_CIRCLE_RADIUS*InScale);
+	const float RingHeight = OuterDistance - InnerDistance;
+	const float ArrowHeight = RingHeight*ArrowHeightPercent;
+	const float ThirtyDegrees = PI / 6.0f;
+	const float HalfArrowidth = ArrowHeight*FMath::Tan(ThirtyDegrees);
+
+	FVector ZAxis = Axis0 ^ Axis1;
+	FVector RotatedAxis0 = Axis0.RotateAngleAxis(InAngle, ZAxis);
+	FVector RotatedAxis1 = Axis1.RotateAngleAxis(InAngle, ZAxis);
+
+	FVector Vertices[3];
+	Vertices[0] = InLocation + (OuterDistance)*RotatedAxis0;
+	Vertices[1] = Vertices[0] + (ArrowHeight)*RotatedAxis0 - HalfArrowidth*RotatedAxis1;
+	Vertices[2] = Vertices[1] + (2*HalfArrowidth)*RotatedAxis1;
+
+	PDI->DrawLine(Vertices[0], Vertices[1], InColor, SDPG_Foreground);
+	PDI->DrawLine(Vertices[1], Vertices[2], InColor, SDPG_Foreground);
+	PDI->DrawLine(Vertices[0], Vertices[2], InColor, SDPG_Foreground);
+
+	if (InColor.A > 0)
+	{
+		//fill in the box
+		FDynamicMeshBuilder MeshBuilder(PDI->View->GetFeatureLevel());
+
+		for(int32 VertexIndex = 0;VertexIndex < 3; VertexIndex++)
+		{
+			FDynamicMeshVertex MeshVertex;
+			MeshVertex.Position = (FVector3f)Vertices[VertexIndex];
+			MeshVertex.Color = InColor;
+			MeshVertex.TextureCoordinate[0] = FVector2f(0.0f, 0.0f);
+			MeshVertex.SetTangents(
+				(FVector3f)RotatedAxis0,
+				(FVector3f)RotatedAxis1,
+				(FVector3f)(RotatedAxis0 ^ RotatedAxis1)
+				);
+			MeshBuilder.AddVertex(MeshVertex); //Add bottom vertex
+		}
+
+		MeshBuilder.AddTriangle(0, 1, 2);
+		MeshBuilder.Draw(PDI, FMatrix::Identity, Helper_GetMaterialProxy(TransparentPlaneMaterialXY),SDPG_Foreground,0.f);
+	}
+}
+
+
+uint32 FCustomizableObjectWidget::GetDominantAxisIndex( const FVector& InDiff, ELevelViewportType ViewportType ) const
+{
+	uint32 DominantIndex = 0;
+	if( FMath::Abs(InDiff.X) < FMath::Abs(InDiff.Y) )
+	{
+		DominantIndex = 1;
+	}
+
+	//const int32 WidgetMode = EAxisList::XYZ;//GEditorModeTools().GetWidgetMode();
+
+	//switch(EAxisList::XYZ)
+	//{
+    //    case EAxisList::WM_Translate:
+			switch( ViewportType )
+			{
+				case LVT_OrthoXY:
+					if( CurrentAxis == EAxisList::X )
+					{
+						DominantIndex = 0;
+					}
+					else if( CurrentAxis == EAxisList::Y )
+					{
+						DominantIndex = 1;
+					}
+					break;
+				case LVT_OrthoXZ:
+					if( CurrentAxis == EAxisList::X )
+					{
+						DominantIndex = 0;
+					}
+					else if( CurrentAxis == EAxisList::Z )
+					{
+						DominantIndex = 1;
+					}
+					break;
+				case LVT_OrthoYZ:
+					if( CurrentAxis == EAxisList::Y )
+					{
+						DominantIndex = 0;
+					}
+					else if( CurrentAxis == EAxisList::Z )
+					{
+						DominantIndex = 1;
+					}
+					break;
+				default:
+					break;
+			}
+	//		break;
+	//	default:
+	//		break;
+	//}
+
+	return DominantIndex;
+}
+
diff --git a/Engine/Plugins/Experimental/Mutable/Source/CustomizableObjectEditor/Private/EdGraphSchema_CustomizableObject.cpp b/Engine/Plugins/Experimental/Mutable/Source/CustomizableObjectEditor/Private/EdGraphSchema_CustomizableObject.cpp
new file mode 100644
--- /dev/null
+++ b/Engine/Plugins/Experimental/Mutable/Source/CustomizableObjectEditor/Private/EdGraphSchema_CustomizableObject.cpp
@@ -0,0 +1,1114 @@
+// Copyright Epic Games, Inc. All Rights Reserved.
+
+#include "EdGraphSchema_CustomizableObject.h"
+
+#include "CustomizableObjectEditorModule.h"
+#include "CustomizableObjectEditor.h"
+
+#include "Nodes/CustomizableObjectNodeMeshGeometryOperation.h"
+#include "Nodes/CustomizableObjectNodeMeshReshape.h"
+
+#include "GraphEditorActions.h" 
+
+#include "Runtime/Launch/Resources/Version.h"
+
+#include "ToolMenu.h"
+#include "ToolMenuSection.h"
+
+#include "AssetRegistry/AssetData.h"
+#include "ScopedTransaction.h"
+#include "Framework/Commands/GenericCommands.h"
+#include "EdGraphNode_Comment.h"
+#include "Toolkits/ToolkitManager.h"
+
+#include "Engine/SkeletalMesh.h"
+#include "Engine/StaticMesh.h"
+
+#include "CustomizableObjectEditorNodeContextCommands.h"
+#include "Settings/EditorStyleSettings.h"
+
+
+#define LOCTEXT_NAMESPACE "CustomizableObjectSchema"
+
+
+#define SNAP_GRID (16) // @todo ensure this is the same as SNodePanel::GetSnapGridSize()
+
+namespace 
+{
+	// Maximum distance a drag can be off a node edge to require 'push off' from node
+	const int32 NodeDistance = 60;
+}
+
+
+UEdGraphNode* FCustomizableObjectSchemaAction_NewNode::PerformAction(class UEdGraph* ParentGraph, UEdGraphPin* FromPin, const FVector2D Location, bool bSelectNewNode)
+{
+	UEdGraphNode* ResultNode = nullptr;
+
+	// If there is a template, we actually use it
+	if (NodeTemplate)
+	{
+		const FScopedTransaction Transaction(LOCTEXT("AddNode", "Add Node"));
+		ParentGraph->Modify();
+
+		if (FromPin)
+		{
+			FromPin->Modify();
+		}
+
+		ResultNode = FCustomizableObjectSchemaAction_NewNode::CreateNode(ParentGraph, FromPin, Location, NodeTemplate);
+	}
+
+	return ResultNode;
+}
+
+
+UEdGraphNode* FCustomizableObjectSchemaAction_NewNode::PerformAction(class UEdGraph* ParentGraph, TArray<UEdGraphPin*>& FromPins, const FVector2D Location, bool bSelectNewNode) 
+{
+	UEdGraphNode* ResultNode = 0;
+
+	if (FromPins.Num() > 0)
+	{
+		ResultNode = PerformAction(ParentGraph, FromPins[0], Location,bSelectNewNode);
+
+		// Try autowiring the rest of the pins
+		for (int32 Index = 1; Index < FromPins.Num(); ++Index)
+		{
+			ResultNode->AutowireNewNode(FromPins[Index]);
+		}
+	}
+	else
+	{
+		ResultNode = PerformAction(ParentGraph, NULL, Location, bSelectNewNode);
+	}
+
+	return ResultNode;
+}
+
+
+UEdGraphNode* FCustomizableObjectSchemaAction_NewNode::CreateNode(class UEdGraph* ParentGraph, UEdGraphPin* FromPin, const FVector2D Location, class UEdGraphNode* InNodeTemplate)
+{
+	// UE Code from FSchemaAction_NewNode::CreateNode(...). Overlap calculations performed before AutowireNewNode(...).
+	
+	// Duplicate template node to create new node
+	UEdGraphNode* ResultNode = DuplicateObject<UEdGraphNode>(InNodeTemplate, ParentGraph);
+
+	ResultNode->SetFlags(RF_Transactional);
+
+	ParentGraph->AddNode(ResultNode, true);
+
+	ResultNode->CreateNewGuid();
+	ResultNode->PostPlacedNewNode();
+	if (UCustomizableObjectNode* TypedResultNode = Cast<UCustomizableObjectNode>(ResultNode))
+	{
+		TypedResultNode->BeginConstruct();
+		TypedResultNode->BackwardsCompatibleFixup(); // In theory not required but added to be consistent with PostLoad (called by DuplicateObject).
+		TypedResultNode->PostBackwardsCompatibleFixup();
+	}
+	ResultNode->ReconstructNode(); // Mutable node lifecycle always starts at ReconstructNode.
+
+	// For input pins, new node will generally overlap node being dragged off
+	// Work out if we want to visually push away from connected node
+	int32 XLocation = Location.X;
+	if (FromPin && FromPin->Direction == EGPD_Input)
+	{
+		UEdGraphNode* PinNode = FromPin->GetOwningNode();
+		const float XDelta = FMath::Abs(PinNode->NodePosX - Location.X);
+
+		if (XDelta < NodeDistance)
+		{
+			// Set location to edge of current node minus the max move distance
+			// to force node to push off from connect node enough to give selection handle
+			XLocation = PinNode->NodePosX - NodeDistance;
+		}
+	}
+
+	ResultNode->AutowireNewNode(FromPin);
+
+	ResultNode->NodePosX = XLocation;
+	ResultNode->NodePosY = Location.Y;
+	ResultNode->SnapToGrid(GetDefault<UEditorStyleSettings>()->GridSnapSize);
+
+	return ResultNode;
+}
+
+
+void FCustomizableObjectSchemaAction_NewNode::AddReferencedObjects( FReferenceCollector& Collector )
+{
+	FEdGraphSchemaAction::AddReferencedObjects( Collector );
+
+	// These don't get saved to disk, but we want to make sure the objects don't get GC'd while the action array is around
+	Collector.AddReferencedObject( NodeTemplate );
+}
+
+
+////////////////////////////////////////
+// FCustomizableObjectSchemaAction_Paste
+
+UEdGraphNode* FCustomizableObjectSchemaAction_Paste::PerformAction(class UEdGraph* ParentGraph, UEdGraphPin* FromPin, const FVector2D Location, bool bSelectNewNode/* = true*/)
+{
+	TSharedPtr<ICustomizableObjectEditor> CustomizableObjectEditor = UEdGraphSchema_CustomizableObject::GetCustomizableObjectEditor(ParentGraph);
+
+	if (CustomizableObjectEditor.IsValid() && CustomizableObjectEditor->CanPasteNodes())
+	{
+		CustomizableObjectEditor->PasteNodesHere(Location);
+	}
+	
+	return nullptr;
+}
+
+
+//////////////////////////////////////////////////////////////////////////
+// DO NOT change the values because it will break the external pin nodes!
+const FName UEdGraphSchema_CustomizableObject::PC_Object("object");
+const FName UEdGraphSchema_CustomizableObject::PC_Material("material");
+const FName UEdGraphSchema_CustomizableObject::PC_Mesh("mesh");
+const FName UEdGraphSchema_CustomizableObject::PC_Layout("layout");
+const FName UEdGraphSchema_CustomizableObject::PC_Image("image");
+const FName UEdGraphSchema_CustomizableObject::PC_Projector("projector");
+const FName UEdGraphSchema_CustomizableObject::PC_GroupProjector("groupProjector");
+const FName UEdGraphSchema_CustomizableObject::PC_Color("color");
+const FName UEdGraphSchema_CustomizableObject::PC_Float("float");
+const FName UEdGraphSchema_CustomizableObject::PC_Bool("bool");
+const FName UEdGraphSchema_CustomizableObject::PC_Enum("enum");
+const FName UEdGraphSchema_CustomizableObject::PC_Stack("stack");
+const FName UEdGraphSchema_CustomizableObject::PC_MaterialAsset("materialAsset");
+
+
+UEdGraphSchema_CustomizableObject::UEdGraphSchema_CustomizableObject()
+	: Super()
+{
+}
+
+
+TSharedPtr<FCustomizableObjectSchemaAction_NewNode> UEdGraphSchema_CustomizableObject::AddNewNodeAction(FGraphActionListBuilderBase& ContextMenuBuilder, const FString& Category, const FText& MenuDesc, const FText& Tooltip, const int32 Grouping, const FString& Keywords)
+{
+	TSharedPtr<FCustomizableObjectSchemaAction_NewNode> NewActionNode = TSharedPtr<FCustomizableObjectSchemaAction_NewNode>(new FCustomizableObjectSchemaAction_NewNode(Category, MenuDesc, Tooltip, Grouping, FText::FromString(Keywords)));
+
+	ContextMenuBuilder.AddAction( NewActionNode );
+
+	return NewActionNode;
+}
+
+
+TSharedPtr<FCustomizableObjectSchemaAction_NewNode> AddNewNodeAction(TArray< TSharedPtr<FEdGraphSchemaAction> >& OutTypes, const FString& Category, const FText& MenuDesc, const FText& Tooltip)
+{
+	return *(new (OutTypes) TSharedPtr<FCustomizableObjectSchemaAction_NewNode>(new FCustomizableObjectSchemaAction_NewNode(Category, MenuDesc, Tooltip, 0)));
+}
+
+
+namespace 
+{
+	bool PinRelevancyFilter(UCustomizableObjectNode* TemplateNode, FGraphContextMenuBuilder& ContextMenuBuilder)
+	{
+		const UEdGraphPin* FromPin = ContextMenuBuilder.FromPin;
+		if (!FromPin)
+		{
+			return true;
+		}
+
+		if (TemplateNode->ProvidesCustomPinRelevancyTest())
+		{
+			return TemplateNode->IsPinRelevant(ContextMenuBuilder.FromPin);
+		}
+
+		TemplateNode->BeginConstruct();
+		TemplateNode->ReconstructNode();
+		
+		for (const UEdGraphPin* Pin : TemplateNode->GetAllPins())
+		{
+			const UEdGraphPin* InputPin = nullptr;
+			const UEdGraphPin* OutputPin = nullptr;
+
+			if (!UEdGraphSchema_CustomizableObject::CategorizePinsByDirection(Pin, FromPin, InputPin, OutputPin))
+			{
+				continue;
+			}
+
+			const UCustomizableObjectNode* InputNode = Cast<UCustomizableObjectNode>(InputPin->GetOwningNode());
+			bool bOtherNodeIsBlocklisted = false;
+			bool bArePinsCompatible = false;
+			if (InputNode->CanConnect(InputPin, OutputPin, bOtherNodeIsBlocklisted, bArePinsCompatible))
+			{
+				return true;
+			}
+		}
+
+		return false;
+	}
+
+	template<class FilterFn>
+	void AddNewNodeActionFiltered(UCustomizableObjectNode* TemplateNode, FGraphContextMenuBuilder& ContextMenuBuilder, const FString& Category, const FText& MenuDesc, const int32 Grouping, FilterFn&& Filter)
+	{
+		if (!Filter(TemplateNode, ContextMenuBuilder))
+		{
+			return;
+		}
+
+		TSharedPtr<FCustomizableObjectSchemaAction_NewNode> Action = UEdGraphSchema_CustomizableObject::AddNewNodeAction(ContextMenuBuilder, Category, MenuDesc, FText(), Grouping);
+		Action->NodeTemplate = TemplateNode;
+	}
+
+	template<class FilterFn>
+	void AddNewNodeActionFiltered(UCustomizableObjectNode* TemplateNode, FGraphContextMenuBuilder& ContextMenuBuilder, const FString& Category, const int32 Grouping, FilterFn&& Filter)
+	{
+		AddNewNodeActionFiltered(TemplateNode, ContextMenuBuilder, Category, TemplateNode->GetNodeTitle(ENodeTitleType::ListView), Grouping, Forward<FilterFn>(Filter));
+	}
+
+	template<size_t N, class FilterFn>
+	void AddNewNodeCategoryActionsFiltered(UCustomizableObjectNode* (&TemplateNodes)[N], FGraphContextMenuBuilder& ContextMenuBuilder, const FString& Category, const int32 Grouping, FilterFn&& Filter)
+	{
+		for (size_t i = 0; i < N; ++i)
+		{
+			AddNewNodeActionFiltered(TemplateNodes[i], ContextMenuBuilder, Category, Grouping, Forward<FilterFn>(Filter));
+		}
+	}
+}
+
+
+void UEdGraphSchema_CustomizableObject::GetGraphContextActions(FGraphContextMenuBuilder& ContextMenuBuilder) const
+{
+	//const UCustomizableObjectGraph* CustomizableObjectGraph = CastChecked<UCustomizableObjectGraph>(ContextMenuBuilder.CurrentGraph);
+	//TSharedPtr<FCustomizableObjectSchemaAction_NewNode> Action;
+	int32 GeneralGrouping = 2;
+
+	const bool bDisableFilter = false;
+
+	// return true if Filter is passed.
+	const auto Filter = [bDisableFilter](UCustomizableObjectNode* TemplateNode, FGraphContextMenuBuilder& ContextMenuBuilder) -> bool
+	{
+		if (!ContextMenuBuilder.FromPin || bDisableFilter)
+		{
+			return true;
+		}
+
+		return PinRelevancyFilter(TemplateNode, ContextMenuBuilder);
+	};
+
+	{
+		UCustomizableObjectNode* Node = ContextMenuBuilder.CreateTemplateNode<UCustomizableObjectNodeObject>();
+		AddNewNodeActionFiltered(Node, ContextMenuBuilder, TEXT("Object"), LOCTEXT("Base_Group", "Base Object"), GeneralGrouping, 
+			[&Filter](UCustomizableObjectNode* TemplateNode, FGraphContextMenuBuilder& ContextMenuBuilder) -> bool
+			{
+				// Only let user add a base node if there isn't one in the graph
+				for (TObjectPtr<UEdGraphNode> AuxNode : ContextMenuBuilder.CurrentGraph->Nodes)
+				{
+					UCustomizableObjectNodeObject* CustomizableObjectNodeObject = Cast<UCustomizableObjectNodeObject>(AuxNode);
+
+					if (CustomizableObjectNodeObject && CustomizableObjectNodeObject->bIsBase)
+					{
+						return false;
+					}
+				}
+
+				return Filter(TemplateNode, ContextMenuBuilder);
+			}	
+		);
+	}
+
+	{
+		UCustomizableObjectNode* Node = NewObject<UCustomizableObjectNodeObjectGroup>(); //ContextMenuBuilder.CreateTemplateNode<UCustomizableObjectNodeObjectGroup>();
+		AddNewNodeActionFiltered(Node, ContextMenuBuilder, TEXT("Object"), LOCTEXT("Child_Group", "Object Group"), GeneralGrouping, Filter);
+	}
+	
+	{
+		UCustomizableObjectNode* Node = ContextMenuBuilder.CreateTemplateNode<UCustomizableObjectNodeObjectChild>();
+		AddNewNodeActionFiltered(Node, ContextMenuBuilder, TEXT("Object"), LOCTEXT("Child_Object", "Child Object"), GeneralGrouping, Filter);
+	}
+
+	{
+		UCustomizableObjectNode* ObjectTemplateNodes[]
+		{
+			ContextMenuBuilder.CreateTemplateNode<UCustomizableObjectNodeMaterial>(),
+			ContextMenuBuilder.CreateTemplateNode<UCustomizableObjectNodeExtendMaterial>(),
+			ContextMenuBuilder.CreateTemplateNode<UCustomizableObjectNodeRemoveMesh>(),
+			ContextMenuBuilder.CreateTemplateNode<UCustomizableObjectNodeRemoveMeshBlocks>(),
+			ContextMenuBuilder.CreateTemplateNode<UCustomizableObjectNodeEditMaterial>(),
+			ContextMenuBuilder.CreateTemplateNode<UCustomizableObjectNodeMaterialVariation>(),
+			ContextMenuBuilder.CreateTemplateNode<UCustomizableObjectNodeMorphMaterial>(),
+			ContextMenuBuilder.CreateTemplateNode<UCustomizableObjectNodeCopyMaterial>(),
+			ContextMenuBuilder.CreateTemplateNode<UCustomizableObjectNodeMeshClipMorph>(),
+			ContextMenuBuilder.CreateTemplateNode<UCustomizableObjectNodeMeshClipWithMesh>(),
+			ContextMenuBuilder.CreateTemplateNode<UCustomizableObjectNodeMeshClipDeform>(),
+			ContextMenuBuilder.CreateTemplateNode<UCustomizableObjectNodeTable>(),
+		};
+
+		AddNewNodeCategoryActionsFiltered(ObjectTemplateNodes, ContextMenuBuilder, TEXT("Object"), GeneralGrouping, Filter);
+	}
+
+	{
+		UCustomizableObjectNode* MeshTemplateNodes[]
+		{
+			ContextMenuBuilder.CreateTemplateNode<UCustomizableObjectNodeSkeletalMesh>(),
+			ContextMenuBuilder.CreateTemplateNode<UCustomizableObjectNodeStaticMesh>(),
+			ContextMenuBuilder.CreateTemplateNode<UCustomizableObjectNodeLayoutBlocks>(),
+			ContextMenuBuilder.CreateTemplateNode<UCustomizableObjectNodeMeshMorph>(),
+			ContextMenuBuilder.CreateTemplateNode<UCustomizableObjectNodeMeshSwitch>(),
+			ContextMenuBuilder.CreateTemplateNode<UCustomizableObjectNodeMeshVariation>(),
+			ContextMenuBuilder.CreateTemplateNode<UCustomizableObjectNodeMeshMorphStackDefinition>(),
+			ContextMenuBuilder.CreateTemplateNode<UCustomizableObjectNodeMeshMorphStackApplication>(),
+			ContextMenuBuilder.CreateTemplateNode<UCustomizableObjectNodeMeshGeometryOperation>(),
+			ContextMenuBuilder.CreateTemplateNode<UCustomizableObjectNodeMeshReshape>(),
+		};
+
+		AddNewNodeCategoryActionsFiltered(MeshTemplateNodes, ContextMenuBuilder, TEXT("Mesh"), GeneralGrouping, Filter);
+	}
+
+	{
+		UCustomizableObjectNode* PoseMeshTemplateNodes[]
+		{
+			ContextMenuBuilder.CreateTemplateNode<UCustomizableObjectNodeAnimationPose>(),
+		};
+
+		AddNewNodeCategoryActionsFiltered(PoseMeshTemplateNodes, ContextMenuBuilder, TEXT("PoseMesh"), GeneralGrouping, Filter);
+	}
+
+	{
+		UCustomizableObjectNode* TextureTemplateNodes[] =
+		{
+			ContextMenuBuilder.CreateTemplateNode<UCustomizableObjectNodeTexture>(),
+			ContextMenuBuilder.CreateTemplateNode<UCustomizableObjectNodeTextureBinarise>(),
+			ContextMenuBuilder.CreateTemplateNode<UCustomizableObjectNodeTextureInterpolate>(),
+			ContextMenuBuilder.CreateTemplateNode<UCustomizableObjectNodeTextureLayer>(),
+			ContextMenuBuilder.CreateTemplateNode<UCustomizableObjectNodeTextureSwitch>(),
+			ContextMenuBuilder.CreateTemplateNode<UCustomizableObjectNodeTextureVariation>(),
+			ContextMenuBuilder.CreateTemplateNode<UCustomizableObjectNodeTextureToChannels>(),
+			ContextMenuBuilder.CreateTemplateNode<UCustomizableObjectNodeTextureFromChannels>(),
+			ContextMenuBuilder.CreateTemplateNode<UCustomizableObjectNodeTextureFromColor>(),
+			ContextMenuBuilder.CreateTemplateNode<UCustomizableObjectNodeTextureProject>(),
+			ContextMenuBuilder.CreateTemplateNode<UCustomizableObjectNodeTextureParameter>(),
+			ContextMenuBuilder.CreateTemplateNode<UCustomizableObjectNodeTextureInvert>(),
+			ContextMenuBuilder.CreateTemplateNode<UCustomizableObjectNodeTextureColourMap>(),
+			ContextMenuBuilder.CreateTemplateNode<UCustomizableObjectNodeTextureTransform>(),
+		};
+
+		AddNewNodeCategoryActionsFiltered(TextureTemplateNodes, ContextMenuBuilder, TEXT("Texture"), GeneralGrouping, Filter);
+	}
+
+	{
+		UCustomizableObjectNode* ColorTemplateNodes[] =
+		{
+			ContextMenuBuilder.CreateTemplateNode<UCustomizableObjectNodeColorConstant>(),
+			ContextMenuBuilder.CreateTemplateNode<UCustomizableObjectNodeColorParameter>(),
+			ContextMenuBuilder.CreateTemplateNode<UCustomizableObjectNodeColorSwitch>(),
+			ContextMenuBuilder.CreateTemplateNode<UCustomizableObjectNodeColorVariation>(),
+			ContextMenuBuilder.CreateTemplateNode<UCustomizableObjectNodeTextureSample>(),
+			ContextMenuBuilder.CreateTemplateNode<UCustomizableObjectNodeColorArithmeticOp>(),
+			ContextMenuBuilder.CreateTemplateNode<UCustomizableObjectNodeColorFromFloats>(),
+		};
+
+		AddNewNodeCategoryActionsFiltered(ColorTemplateNodes, ContextMenuBuilder, TEXT("Color"), GeneralGrouping, Filter);
+	}
+
+	{
+		UCustomizableObjectNode* EnumTemplateNodes[] =
+		{
+			ContextMenuBuilder.CreateTemplateNode<UCustomizableObjectNodeEnumParameter>(),
+		};
+
+		AddNewNodeCategoryActionsFiltered(EnumTemplateNodes, ContextMenuBuilder, TEXT("Enum"), GeneralGrouping, Filter);
+	}
+
+	{
+		UCustomizableObjectNode* FloatTemplateNodes[] =
+		{
+			ContextMenuBuilder.CreateTemplateNode<UCustomizableObjectNodeFloatConstant>(),
+			ContextMenuBuilder.CreateTemplateNode<UCustomizableObjectNodeFloatParameter>(),
+			ContextMenuBuilder.CreateTemplateNode<UCustomizableObjectNodeFloatSwitch>(),
+			ContextMenuBuilder.CreateTemplateNode<UCustomizableObjectNodeFloatVariation>(),
+			ContextMenuBuilder.CreateTemplateNode<UCustomizableObjectNodeCurve>(),
+		};
+
+
+		AddNewNodeCategoryActionsFiltered(FloatTemplateNodes, ContextMenuBuilder, TEXT("Float"), GeneralGrouping, Filter);
+	}
+
+	{
+		UCustomizableObjectNode* ProjectorTemplateNodes[] =
+		{
+			ContextMenuBuilder.CreateTemplateNode<UCustomizableObjectNodeProjectorConstant>(),
+			ContextMenuBuilder.CreateTemplateNode<UCustomizableObjectNodeProjectorParameter>(),
+			ContextMenuBuilder.CreateTemplateNode<UCustomizableObjectNodeGroupProjectorParameter>(),
+		};
+
+		AddNewNodeCategoryActionsFiltered(ProjectorTemplateNodes, ContextMenuBuilder, TEXT("Projector"), GeneralGrouping, Filter);
+	}
+
+	{
+		// External Pin Nodes
+		const FName* PinTypes[] = { &PC_Material, &PC_Mesh, &PC_Image, &PC_Projector, &PC_GroupProjector, &PC_Color, &PC_Float, &PC_Bool, &PC_Enum, &PC_Stack };
+		
+		for (const FName* PinCategory : PinTypes)
+		{
+			UCustomizableObjectNodeExternalPin* CustomizableObjectNodeExternalPin = ContextMenuBuilder.CreateTemplateNode<UCustomizableObjectNodeExternalPin>();
+			CustomizableObjectNodeExternalPin->PinType = *PinCategory;
+			
+			AddNewNodeActionFiltered(CustomizableObjectNodeExternalPin, ContextMenuBuilder, TEXT("Import Pin"), GeneralGrouping, Filter);
+		}
+
+		for (const FName* PinCategory : PinTypes)
+		{
+			UCustomizableObjectNodeExposePin* CustomizableObjectNodeExposePin = ContextMenuBuilder.CreateTemplateNode<UCustomizableObjectNodeExposePin>();
+			CustomizableObjectNodeExposePin->PinType = *PinCategory;
+			
+			AddNewNodeActionFiltered(CustomizableObjectNodeExposePin, ContextMenuBuilder, TEXT("Export Pin"), GeneralGrouping, Filter);
+		}
+	}
+	
+	{
+		UEdGraphNode_Comment* Node = NewObject<UEdGraphNode_Comment>();
+		TSharedPtr<FCustomizableObjectSchemaAction_NewNode> Action = AddNewNodeAction(ContextMenuBuilder, FString(), Node->GetNodeTitle(ENodeTitleType::ListView), FText(), 1);
+		Action->NodeTemplate = Node;
+	}
+
+	// Add Paste here if appropriate
+	if (!ContextMenuBuilder.FromPin)
+	{
+		const FText PasteDesc = LOCTEXT("PasteDesc", "Paste Here");
+		const FText PasteToolTip = LOCTEXT("PasteToolTip", "Pastes copied items at this location.");
+		TSharedPtr<FCustomizableObjectSchemaAction_Paste> PasteAction(new FCustomizableObjectSchemaAction_Paste(FText::GetEmpty(), PasteDesc, PasteToolTip.ToString(), 0));
+		ContextMenuBuilder.AddAction(PasteAction);
+	}
+}
+
+
+const FPinConnectionResponse UEdGraphSchema_CustomizableObject::CanCreateConnection(const UEdGraphPin* PinA, const UEdGraphPin* PinB) const
+{
+	// Make sure the pins are not on the same node
+	if (PinA->GetOwningNode() == PinB->GetOwningNode())
+	{
+		return FPinConnectionResponse(CONNECT_RESPONSE_DISALLOW, TEXT("Both are on the same node"));
+	}
+
+	// Check both pins support connections
+	if(PinA->bNotConnectable || PinB->bNotConnectable)
+	{
+		return FPinConnectionResponse(CONNECT_RESPONSE_DISALLOW, TEXT("Pin doesn't support connections"));
+	}
+
+	// Compare the directions
+	const UEdGraphPin* InputPin = NULL;
+	const UEdGraphPin* OutputPin = NULL;
+
+	if (!CategorizePinsByDirection(PinA, PinB, /*out*/ InputPin, /*out*/ OutputPin))
+	{
+		return FPinConnectionResponse(CONNECT_RESPONSE_DISALLOW, TEXT("Directions are not compatible"));
+	}
+
+	// Type categories must match and the nodes need to be compatible with each other
+	bool bArePinsCompatible = false;
+	bool bIsOtherNodeBlocklisted = false;
+
+	UCustomizableObjectNode* InputNode = CastChecked<UCustomizableObjectNode>(InputPin->GetOwningNode());
+	if (!InputNode->CanConnect(InputPin, OutputPin, bIsOtherNodeBlocklisted,bArePinsCompatible))
+	{
+		if (!bArePinsCompatible)
+		{
+			return FPinConnectionResponse(CONNECT_RESPONSE_DISALLOW, TEXT("Types are not compatible"));
+		}
+		else if (bIsOtherNodeBlocklisted)
+		{
+			return FPinConnectionResponse(CONNECT_RESPONSE_DISALLOW, TEXT("Direct connections between these nodes are not allowed"));
+		}
+	}
+
+	// Some special nodes can only have one output
+	bool bBreakExistingDueToDataOutput = false;
+	if (UCustomizableObjectNode* n = Cast<UCustomizableObjectNode>(OutputPin->GetOwningNode()))
+	{
+		bBreakExistingDueToDataOutput = (OutputPin->LinkedTo.Num() > 0) && n->ShouldBreakExistingConnections(InputPin, OutputPin);
+	}
+
+	// See if we want to break existing connections (if its an input with an existing connection)
+#if (ENGINE_MAJOR_VERSION == 4 && ENGINE_MINOR_VERSION >= 17) || ENGINE_MAJOR_VERSION >= 5
+	const bool bBreakExistingDueToDataInput = (InputPin->LinkedTo.Num() > 0) && !InputPin->PinType.IsArray();
+#else
+	const bool bBreakExistingDueToDataInput = (InputPin->LinkedTo.Num() > 0) && !InputPin->PinType.bIsArray;
+#endif
+
+	if (bBreakExistingDueToDataOutput&&bBreakExistingDueToDataInput)
+	{
+		return FPinConnectionResponse(CONNECT_RESPONSE_BREAK_OTHERS_AB, TEXT("Replace connections at both ends"));
+	}
+	
+	if (bBreakExistingDueToDataInput)
+	{
+		const ECanCreateConnectionResponse ReplyBreakInputs = (PinA == InputPin) ? CONNECT_RESPONSE_BREAK_OTHERS_A : CONNECT_RESPONSE_BREAK_OTHERS_B;
+		return FPinConnectionResponse(ReplyBreakInputs, TEXT("Replace existing input connections"));
+	}
+
+	if (bBreakExistingDueToDataOutput)
+	{
+		const ECanCreateConnectionResponse ReplyBreakOutputs = (PinA == OutputPin) ? CONNECT_RESPONSE_BREAK_OTHERS_A : CONNECT_RESPONSE_BREAK_OTHERS_B;
+		return FPinConnectionResponse(ReplyBreakOutputs, TEXT("Replace existing output connections"));
+	}
+
+	return FPinConnectionResponse(CONNECT_RESPONSE_MAKE, TEXT(""));
+}
+
+
+FLinearColor UEdGraphSchema_CustomizableObject::GetPinTypeColor(const FName& TypeString) const
+{
+	if (TypeString == PC_Enum)
+	{
+		return FLinearColor(0.357667f, 0.500000f, 0.060000f, 1.000000f);
+	}
+	else if (TypeString == PC_Float)
+	{
+		return FLinearColor(0.357667f, 1.000000f, 0.060000f, 1.000000f);
+	}
+	else if (TypeString == PC_Color)
+	{
+		return FLinearColor(1.000000f, 0.591255f, 0.016512f, 1.000000f);
+	}
+	else if (TypeString == PC_Bool)
+	{
+		return FLinearColor(0.470000f, 0.0f, 0.000000f, 1.000000f);
+	}
+	else if (TypeString == PC_Projector)
+	{
+		return FLinearColor(1.000000f, 0.500000f, 1.000000f, 0.600000f);
+	}
+	else if (TypeString == PC_GroupProjector)
+	{
+		return FLinearColor(1.000000f, 0.172585f, 0.000000f, 1.000000f);
+	}
+	else if (TypeString == PC_Mesh)
+	{
+		return FLinearColor(0.100000f, 0.000000f, 0.500000f, 1.000000f);
+	}
+	else if (TypeString == PC_Layout)
+	{
+		return FLinearColor(0.500000f, 0.500000f, 0.100000f, 1.000000f);
+	}
+	else if (TypeString == PC_Image)
+	{
+		return FLinearColor(0.353393f, 0.454175f, 1.000000f, 1.000000f);
+	}
+	else if (TypeString == PC_Material)
+	{
+		return FLinearColor(0.000000f, 0.100000f, 0.600000f, 1.000000f);
+	}
+	else if (TypeString == PC_Object)
+	{
+		return FLinearColor(0.000000f, 0.400000f, 0.910000f, 1.000000f);
+	}
+	else if (TypeString == PC_Stack)
+	{
+		return FLinearColor(1.000000f, 0.000000f, 0.800000f, 1.000000f);
+	}
+	else if (TypeString == PC_MaterialAsset)
+	{
+		return FLinearColor(1.000000f, 1.000000f, 1.000000f, 1.000000f);
+	}
+
+	return FLinearColor(0.750000f, 0.600000f, 0.400000f, 1.000000f);
+}
+
+
+FLinearColor UEdGraphSchema_CustomizableObject::GetPinTypeColor(const FEdGraphPinType& PinType) const
+{
+#if (ENGINE_MAJOR_VERSION == 4 && ENGINE_MINOR_VERSION <= 18)
+	const FString& TypeString = PinType.PinCategory;
+	const FName& TypeName = FName(*TypeString);
+#else
+	const FName TypeName = PinType.PinCategory;
+#endif
+
+	return GetPinTypeColor(TypeName);
+}
+
+
+bool UEdGraphSchema_CustomizableObject::ShouldHidePinDefaultValue(UEdGraphPin* Pin) const
+{
+	check(Pin != NULL);
+
+	if (Pin->bDefaultValueIsIgnored)
+	{
+		return true;
+	}
+
+	return false;
+}
+
+
+#if (ENGINE_MAJOR_VERSION == 4 && ENGINE_MINOR_VERSION < 18)
+void UEdGraphSchema_CustomizableObject::GetContextMenuActions(const UEdGraph* CurrentGraph, const UEdGraphNode* InGraphNode, const UEdGraphPin* InGraphPin, FMenuBuilder* MenuBuilder, bool bIsDebugging) const
+{
+	check(CurrentGraph);	
+	//FKismetNodeContextMenuBuilder Context(CurrentGraph, InGraphNode, InGraphPin, MenuBuilder, bIsDebugging);
+
+	if (InGraphPin != NULL)
+	{
+		MenuBuilder->BeginSection("PinActionsMenuHeader");
+		if (!bIsDebugging)
+		{
+			// Break pin links
+			if (InGraphPin->LinkedTo.Num() > 1)
+			{
+				MenuBuilder->AddMenuEntry( FGraphEditorCommands::Get().BreakPinLinks );
+			}
+
+			// add sub menu for break link to
+			if (InGraphPin->LinkedTo.Num() > 0)
+			{
+				MenuBuilder->AddSubMenu(
+					LOCTEXT( "BreakLinkTo", "Break Link To..." ),
+					LOCTEXT("BreakSpecificLinks", "Break a specific link..." ),
+					FNewMenuDelegate::CreateUObject( (UEdGraphSchema_CustomizableObject*const)this, &UEdGraphSchema_CustomizableObject::GetBreakLinkToSubMenuActions, const_cast<UEdGraphPin*>(InGraphPin)));
+			}
+
+		}
+		MenuBuilder->EndSection();
+
+	}
+	else if (InGraphNode != NULL)
+	{
+		if (!bIsDebugging)
+		{
+			// Node contextual actions
+			MenuBuilder->BeginSection("NodeActionsMenuHeader");
+			MenuBuilder->AddMenuEntry( FGenericCommands::Get().Delete );
+			MenuBuilder->AddMenuEntry( FGenericCommands::Get().Cut );
+			MenuBuilder->AddMenuEntry( FGenericCommands::Get().Copy );
+			MenuBuilder->AddMenuEntry( FGenericCommands::Get().Duplicate );
+			MenuBuilder->AddMenuEntry( FGraphEditorCommands::Get().ReconstructNodes );
+			MenuBuilder->AddMenuEntry( FGraphEditorCommands::Get().BreakNodeLinks );
+			// In the case of a UCustomizableObjectNodeObjectGroup, add the option to refresh all Customizable Object Material Node nodes of all the children of this node
+			const UCustomizableObjectNodeObjectGroup* TypedNode = Cast<UCustomizableObjectNodeObjectGroup>(InGraphNode);
+			if (TypedNode != nullptr)
+			{
+				MenuBuilder->AddMenuEntry(FCustomizableObjectEditorNodeContextCommands::Get().RefreshMaterialNodesInAllChildren);
+			}
+			MenuBuilder->EndSection();
+		}
+
+		struct SCommentUtility
+		{
+			static void CreateComment(const UEdGraphSchema_CustomizableObject* Schema, UEdGraph* Graph)
+			{
+				if (Schema && Graph)
+				{
+					Schema->AddComment(Graph, NULL, FVector2D::ZeroVector, true);
+				}
+			}
+		};
+		
+		MenuBuilder->BeginSection("SchemaActionComment", LOCTEXT("MultiCommentHeader", "Comment Group"));
+		MenuBuilder->AddMenuEntry(LOCTEXT("MultiCommentDesc", "Create Comment from Selection"),
+			LOCTEXT("CommentToolTip", "Create a resizable comment box around selection."),
+			FSlateIcon(),
+			FUIAction(FExecuteAction::CreateStatic(SCommentUtility::CreateComment, this, const_cast<UEdGraph*>(CurrentGraph))));
+		MenuBuilder->EndSection();
+	}
+}
+
+#else
+void UEdGraphSchema_CustomizableObject::GetContextMenuActions(class UToolMenu* Menu, class UGraphNodeContextMenuContext* Context) const
+{
+	if (Context && !Context->Pin && Context->Node)
+	{
+		if (!Context->bIsDebugging)
+		{
+			// Node contextual actions
+			FToolMenuSection& Section = Menu->AddSection("EdGraphSchemaNodeActions", LOCTEXT("NodeActionsMenuHeader", "Node Actions"));
+			Section.AddMenuEntry(FGenericCommands::Get().Delete);
+			Section.AddMenuEntry(FGenericCommands::Get().Cut);
+			Section.AddMenuEntry(FGenericCommands::Get().Copy);
+			Section.AddMenuEntry(FGenericCommands::Get().Duplicate);
+			Section.AddMenuEntry(FGraphEditorCommands::Get().ReconstructNodes);
+			Section.AddMenuEntry(FGraphEditorCommands::Get().BreakNodeLinks);
+
+			// In the case of a UCustomizableObjectNodeObjectGroup, add the option to refresh all Customizable Object Material Node nodes of all the children of this node
+			const UCustomizableObjectNodeObjectGroup* TypedNode = Cast<UCustomizableObjectNodeObjectGroup>(Context->Node);
+			if (TypedNode != nullptr)
+			{
+				Section.AddMenuEntry( FCustomizableObjectEditorNodeContextCommands::Get().RefreshMaterialNodesInAllChildren );
+			}
+		}
+
+		struct SCommentUtility
+		{
+			static void CreateComment(const UEdGraphSchema_CustomizableObject* Schema, UEdGraph* Graph)
+			{
+				if (Schema && Graph)
+				{
+					Schema->AddComment(Graph, NULL, FVector2D::ZeroVector, true);
+				}
+			}
+		};
+
+		FToolMenuSection& section = Menu->AddSection("SchemaActionComment", LOCTEXT("MultiCommentHeader", "Comment Group"));
+		section.AddMenuEntry("MultiCommentDesc", LOCTEXT("MultiCommentDesc", "Create Comment from Selection"),
+			LOCTEXT("CommentToolTip", "Create a resizable comment box around selection."),
+			FSlateIcon(),
+#if ENGINE_MAJOR_VERSION >= 5
+			FUIAction(FExecuteAction::CreateStatic(SCommentUtility::CreateComment, this, const_cast<UEdGraph*>(ToRawPtr(Context->Graph)))));
+#else
+			FUIAction(FExecuteAction::CreateStatic(SCommentUtility::CreateComment, this, const_cast<UEdGraph*>(Context->Graph))));
+#endif
+	}
+}
+#endif
+
+
+void UEdGraphSchema_CustomizableObject::BreakPinLinks(UEdGraphPin& TargetPin, bool bSendsNodeNotification) const
+{
+	const FScopedTransaction Transaction(NSLOCTEXT("UnrealEd", "GraphEd_BreakPinLinks", "Break Pin Links"));
+	
+	Super::BreakPinLinks(TargetPin, bSendsNodeNotification);
+}
+
+void UEdGraphSchema_CustomizableObject::BreakSinglePinLink(UEdGraphPin* SourcePin, UEdGraphPin* TargetPin) const
+{
+	const FScopedTransaction Transaction(NSLOCTEXT("UnrealEd", "GraphEd_BreakSinglePinLink", "Break Pin Link"));
+
+	Super::BreakSinglePinLink(SourcePin, TargetPin);
+}
+
+
+TSharedPtr<ICustomizableObjectEditor> UEdGraphSchema_CustomizableObject::GetCustomizableObjectEditor(const class UEdGraph* ParentGraph)
+{
+	// Find the associated Editor
+	UCustomizableObject* CustomizableObject = Cast<UCustomizableObject>(ParentGraph->GetOuter());
+
+	TSharedPtr<ICustomizableObjectEditor> CustomizableObjectEditor;
+	if (CustomizableObject)
+	{
+		TSharedPtr< IToolkit > FoundAssetEditor = FToolkitManager::Get().FindEditorForAsset(CustomizableObject);
+		if (FoundAssetEditor.IsValid())
+		{
+			return StaticCastSharedPtr<ICustomizableObjectEditor>(FoundAssetEditor);
+		}
+	}
+
+	return nullptr;
+}
+
+UEdGraphNode* UEdGraphSchema_CustomizableObject::AddComment(class UEdGraph* ParentGraph, UEdGraphPin* FromPin, const FVector2D Location, bool bSelectNewNode/* = true*/) const
+{
+	TSharedPtr<ICustomizableObjectEditor> CustomizableObjectEditor = GetCustomizableObjectEditor(ParentGraph);
+
+	if(CustomizableObjectEditor.IsValid())
+	{
+		return CustomizableObjectEditor->CreateCommentBox(Location);
+	}
+
+	return nullptr;
+}
+
+
+void UEdGraphSchema_CustomizableObject::GetBreakLinkToSubMenuActions( class FMenuBuilder& MenuBuilder, UEdGraphPin* InGraphPin ) const
+{
+	// Make sure we have a unique name for every entry in the list
+	TMap< FString, uint32 > LinkTitleCount;
+
+	// Add all the links we could break from
+	for(TArray<class UEdGraphPin*>::TConstIterator Links(InGraphPin->LinkedTo); Links; ++Links)
+	{
+		FString Title = FString::Printf(TEXT("%s (%s)"), *(*Links)->GetOwningNode()->GetNodeTitle(ENodeTitleType::ListView).ToString(), *(*Links)->PinFriendlyName.ToString());
+		uint32 &Count = LinkTitleCount.FindOrAdd(Title);
+		FString Description;
+		if(Count == 0)
+		{
+			Description = FString::Printf(TEXT("Break link to %s"), *Title);
+		}
+		else
+		{
+			Description = FString::Printf(TEXT("Break link to %s (%d)"), *Title, Count);
+		}
+		++Count;
+		MenuBuilder.AddMenuEntry( 
+			FText::FromString(Description), 
+			FText::FromString(Description), 
+			FSlateIcon(),
+			FUIAction( FExecuteAction::CreateUObject((UEdGraphSchema_CustomizableObject*const)this, &UEdGraphSchema_CustomizableObject::BreakSinglePinLink, const_cast< UEdGraphPin* >(InGraphPin), *Links) ) 
+			);
+	}
+}
+
+
+void UEdGraphSchema_CustomizableObject::DroppedAssetsOnGraph(const TArray<FAssetData>& Assets, const FVector2D& GraphPosition, UEdGraph* Graph) const
+{
+	// To prevent overlapping when multiple assets are dropped at the same time on the graph
+	const int PixelOffset = 20;
+	int CurrentOffset = 0;
+	
+	for (FAssetData Asset : Assets)
+	{
+		// If it is not a valid asset to be spawned then just skip it
+		ESpawnableObjectType ObjectType;
+		if (!IsSpawnableAsset(Asset,ObjectType))
+		{
+			continue;
+		}
+
+		// At this point we know we are working with an asset we can spawn as a mutable node.
+		
+		UObject* Object = Asset.GetAsset();
+		UEdGraphNode* GraphNode = nullptr;
+
+		// Depending on the UObjectType spawn one or another mutable node
+		switch (ObjectType)
+		{
+		case ESpawnableObjectType::UTexture2D:
+			{
+				UTexture2D* Texture = Cast<UTexture2D>(Object);
+				UCustomizableObjectNodeTexture* Node = NewObject<UCustomizableObjectNodeTexture>(Graph);
+				Node->Texture = Texture;
+				GraphNode = Node;
+				break;
+			}
+		case ESpawnableObjectType::USkeletalMesh: 
+			{
+				USkeletalMesh* SkeletalMesh = Cast<USkeletalMesh>(Object);
+				UCustomizableObjectNodeSkeletalMesh* Node = NewObject<UCustomizableObjectNodeSkeletalMesh>(Graph);
+				Node->SkeletalMesh = SkeletalMesh;
+				GraphNode = Node;
+				break;
+			}
+		case ESpawnableObjectType::UStaticMesh:
+			{
+				UStaticMesh* Mesh = Cast<UStaticMesh>(Object);
+				UCustomizableObjectNodeStaticMesh* Node = NewObject<UCustomizableObjectNodeStaticMesh>(Graph);
+				Node->StaticMesh = Mesh;
+				GraphNode = Node;
+				break;
+			}
+		case ESpawnableObjectType::UMaterialInterface:
+			{
+				UMaterialInterface* Material = Cast<UMaterialInterface>(Object);
+				UCustomizableObjectNodeMaterial* Node = NewObject<UCustomizableObjectNodeMaterial>(Graph);
+				Node->Material = Material;
+				GraphNode = Node;
+				break;
+			}
+		// Error : A new compatible type set on UEdGraphSchema_CustomizableObject::IsSpawnableAsset is not providing a valid ESpawnableObjectType value
+		case ESpawnableObjectType::None:
+		// Error : a switch entry is missing for a ESpawnableObjectType value
+		default:
+			{
+				UE_LOG(LogTemp,Error,TEXT("Unable to create new mutable node for target asset : Invalid ESpawnableObjectType value. "));
+				checkNoEntry();
+				break;
+			}
+		}
+
+		// A node must have been spawned at this point.
+		if (GraphNode)
+		{
+			// A new node has been instanced, add it to the graph
+			GraphNode->CreateNewGuid();
+			GraphNode->PostPlacedNewNode();
+			GraphNode->AllocateDefaultPins();
+			GraphNode->NodePosX = GraphPosition.X + CurrentOffset;
+			GraphNode->NodePosY = GraphPosition.Y + CurrentOffset;
+			Graph->AddNode(GraphNode, true);
+			CurrentOffset += PixelOffset;	
+		}
+		else
+		{
+			UE_LOG(LogTemp,Error,TEXT("Unable to add null node to graph. "))
+		}
+
+	}
+}
+
+bool UEdGraphSchema_CustomizableObject::IsSpawnableAsset(const FAssetData& InAsset, ESpawnableObjectType& OutObjectType) const
+{
+	UObject* Object = InAsset.GetAsset();
+
+	// Type used to know what kind of UObject this asset is
+	OutObjectType = ESpawnableObjectType::None;
+
+	// Check if the provided object can be casted to any of the UObject types we can spawn as CO Nodes
+	if (Cast<UTexture2D>(Object))
+	{
+		OutObjectType = ESpawnableObjectType::UTexture2D;
+		return true;
+	}
+	else if (Cast<USkeletalMesh>(Object))
+	{
+		OutObjectType = ESpawnableObjectType::USkeletalMesh;
+		return true;
+	}
+	else if (Cast<UStaticMesh>(Object))
+	{
+		OutObjectType = ESpawnableObjectType::UStaticMesh;
+		return true;
+	}
+	else if (Cast<UMaterialInterface>(Object))
+	{
+		OutObjectType = ESpawnableObjectType::UMaterialInterface;
+		return true;
+	}
+	// Add more compatible types here, sync it up with ESpawnableObjectType
+	else
+	{
+		// Non spawnable object
+		return false;
+	}
+	
+}
+
+
+void UEdGraphSchema_CustomizableObject::GetAssetsGraphHoverMessage(const TArray<FAssetData>& Assets,
+	const UEdGraph* HoverGraph, FString& OutTooltipText, bool& OutOkIcon) const
+{
+	Super::GetAssetsGraphHoverMessage(Assets, HoverGraph, OutTooltipText, OutOkIcon);
+
+	// Accept entry by default
+	OutOkIcon = true;
+	uint32 AmountOfIncompatibleAssets = 0;
+
+	// Iterate over the assets
+	for (const FAssetData& Asset : Assets)
+	{
+		// On first fail abort the consequent checks and tell the user
+		ESpawnableObjectType ObjectType;
+		if (!IsSpawnableAsset(Asset, ObjectType))
+		{
+			AmountOfIncompatibleAssets++;
+			OutOkIcon = false;
+
+			// Stop checking once we know that more than one asset is not compatible, the UI output will be the same
+			if (AmountOfIncompatibleAssets > 1)
+			{
+				break;
+			}
+		}
+	}
+
+	// Output debug message depending on the quantity of incompatible objects
+	if (!OutOkIcon)
+	{
+		if (Assets.Num() == 1)
+		{
+			OutTooltipText = FString("Incompatible asset selected : No node can be created for this type of asset.");
+		}
+		else if (Assets.Num() > 1 )
+		{
+			if (Assets.Num() == AmountOfIncompatibleAssets)
+			{
+				OutTooltipText = FString("Incompatible assets selected : No node can be created for any of the selected assets.");
+			}
+			else 
+			{
+				OutTooltipText = FString("Incompatible asset selected : Some assets will not be placed as nodes on the graph.");
+			}
+		}
+	}
+}
+
+bool UEdGraphSchema_CustomizableObject::TryCreateConnection(UEdGraphPin* PinA, UEdGraphPin* PinB) const
+{
+	bool Result = Super::TryCreateConnection(PinA, PinB);
+
+	if (!PinA || PinA->bWasTrashed || !PinB || PinB->bWasTrashed)
+	{
+		return Result;
+	}
+	
+	UEdGraphPin* InputPin;
+	UEdGraphPin* OutputPin;
+	if (!CategorizePinsByDirection(PinA, PinB, InputPin, OutputPin))
+	{
+		return Result;
+	}
+
+	if (UCustomizableObjectNode* Node = Cast<UCustomizableObjectNode>(InputPin->GetOwningNode()))
+	{
+		Node->BreakExistingConnectionsPostConnection(InputPin, OutputPin);
+	}
+
+	if (UCustomizableObjectNode* Node = Cast<UCustomizableObjectNode>(OutputPin->GetOwningNode()))
+	{
+		Node->BreakExistingConnectionsPostConnection(InputPin, OutputPin);
+	}
+
+	return Result;
+}
+
+
+FText UEdGraphSchema_CustomizableObject::GetPinCategoryName(const FName& PinCategory)
+{
+	if (PinCategory == UEdGraphSchema_CustomizableObject::PC_Object)
+	{
+		return LOCTEXT("Object_Pin_Category", "Object");
+	}
+	else if (PinCategory == UEdGraphSchema_CustomizableObject::PC_Material)
+	{
+		return LOCTEXT("Material_Pin_Category", "Material");
+	}
+	else if (PinCategory == UEdGraphSchema_CustomizableObject::PC_Mesh)
+	{
+		return LOCTEXT("Mesh_Pin_Category", "Mesh");
+	}
+	else if (PinCategory == UEdGraphSchema_CustomizableObject::PC_Layout)
+	{
+		return LOCTEXT("Layout_Pin_Category", "Layout");
+	}
+	else if (PinCategory == UEdGraphSchema_CustomizableObject::PC_Image)
+	{
+		return LOCTEXT("Image_Pin_Category", "Texture");
+	}
+	else if (PinCategory == UEdGraphSchema_CustomizableObject::PC_Projector)
+	{
+		return LOCTEXT("Projector_Pin_Category", "Projector");
+	}
+	else if (PinCategory == UEdGraphSchema_CustomizableObject::PC_GroupProjector)
+	{
+		return LOCTEXT("Group_Projector_Pin_Category", "Group Projector");
+	}
+	else if (PinCategory == UEdGraphSchema_CustomizableObject::PC_Color)
+	{
+		return LOCTEXT("Color_Pin_Category", "Color");
+	}
+	else if (PinCategory == UEdGraphSchema_CustomizableObject::PC_Float)
+	{
+		return LOCTEXT("Float_Pin_Category", "Float");
+	}
+	else if (PinCategory == UEdGraphSchema_CustomizableObject::PC_Bool)
+	{
+		return LOCTEXT("Bool_Pin_Category", "Bool");
+	}
+	else if (PinCategory == UEdGraphSchema_CustomizableObject::PC_Enum)
+	{
+		return LOCTEXT("Enum_Pin_Category", "Enum");
+	}
+	else if (PinCategory == UEdGraphSchema_CustomizableObject::PC_Stack)
+	{
+		return LOCTEXT("Stack_Pin_Category", "Stack");
+	}
+	else if (PinCategory == UEdGraphSchema_CustomizableObject::PC_MaterialAsset)
+	{
+		return LOCTEXT("Material_Asset_Pin_Category", "materialAsset");
+	}
+	else
+	{
+		check(false); // Unknown pin category. Add the unknown category to the "switch".
+		return FText();
+	}
+}
+
+
+const FName& UEdGraphSchema_CustomizableObject::GetPinCategory(const EMaterialParameterType Type)
+{
+	switch(Type)
+	{
+	case EMaterialParameterType::Texture:
+		return UEdGraphSchema_CustomizableObject::PC_Image;
+
+	case EMaterialParameterType::Vector:
+		return UEdGraphSchema_CustomizableObject::PC_Color;
+
+	case EMaterialParameterType::Scalar:
+		return UEdGraphSchema_CustomizableObject::PC_Float;
+
+	default:
+		check(false); // Type not contemplated.
+		return PC_Image; // Fake result.
+	}
+}
+
+
+#undef LOCTEXT_NAMESPACE
