[mutable] Move Mutable plugin from the restricted folder to the experimental plugins folder.

diff --git a/Engine/Plugins/Experimental/Mutable/Source/CustomizableObjectEditor/Private/GenerateMutableSource/GenerateMutableSource.cpp b/Engine/Plugins/Experimental/Mutable/Source/CustomizableObjectEditor/Private/GenerateMutableSource/GenerateMutableSource.cpp
new file mode 100644
--- /dev/null
+++ b/Engine/Plugins/Experimental/Mutable/Source/CustomizableObjectEditor/Private/GenerateMutableSource/GenerateMutableSource.cpp
@@ -0,0 +1,1374 @@
+// Copyright Epic Games, Inc. All Rights Reserved.
+
+#include "GenerateMutableSource.h"
+
+#include "GenerateMutableSource/GenerateMutableSourceSurface.h"
+#include "GenerateMutableSource/GenerateMutableSourceGroupProjector.h"
+#include "GenerateMutableSource/GenerateMutableSourceModifier.h"
+
+#include "Templates/UnrealTemplate.h"
+#include "Templates/Decay.h"
+#include "GenericPlatform/GenericPlatformMath.h"
+#include "StaticMeshResources.h"
+#include "Interfaces/ITargetPlatform.h"
+#include "Engine/SkeletalMesh.h"
+#include "Engine/StaticMesh.h"
+#include "Engine/AssetManager.h"
+#include "Engine/TextureLODSettings.h"
+#include "Animation/Skeleton.h"
+#include "Animation/MorphTarget.h"
+#include "PhysicsEngine/PhysicsAsset.h"
+#include "ClothingAsset.h"
+#include "MeshUtilities.h"
+#include "EdGraphToken.h"
+#include "AssetRegistry/AssetRegistryModule.h"
+#include "Animation/AnimInstance.h"
+#include "UnrealEditorPortabilityHelpers.h"
+#include "UnrealPortabilityHelpers.h"
+#include "PlatformInfo.h"
+#include "Engine/SkeletalMeshSocket.h"
+
+#include "CustomizableObjectCompiler.h"
+#include "UnrealToMutableTextureConversionUtils.h"
+#include "UnrealConversionUtils.h"
+#include "CustomizableObjectEditorModule.h"
+#include "CustomizableObjectSystem.h"
+#include "CustomizableObject.h"
+#include "GraphTraversal.h"
+#include "MutableUtils.h"
+
+#include "MutableTools/Public/Node.h"
+#include "MutableTools/Public/NodeImageBinarise.h"
+#include "MutableTools/Public/NodeMeshConstant.h"
+#include "MutableTools/Public/NodeMeshApplyPose.h"
+#include "MutableTools/Public/NodeObjectNew.h"
+#include "MutableTools/Public/NodeObjectGroup.h"
+#include "MutableTools/Public/NodeLOD.h"
+#include "MutableTools/Public/NodeComponentNew.h"
+#include "MutableTools/Public/NodeComponentEdit.h"
+#include "MutableTools/Public/NodeSurfaceNew.h"
+#include "MutableTools/Public/NodeSurfaceEdit.h"
+#include "MutableTools/Public/Compiler.h"
+#include "MutableTools/Public/Streams.h"
+#include "MutableTools/Public/NodeModifierMeshClipMorphPlane.h"
+#include "MutableTools/Public/NodeModifierMeshClipWithMesh.h"
+
+#include "MutableRuntime/Public/Mesh.h"
+#include "MutableRuntime/Public/Skeleton.h"
+#include "MutableRuntime/Public/Image.h"
+
+#define LOCTEXT_NAMESPACE "CustomizableObjectEditor"
+
+
+// Forward declarations 
+bool AffectsCurrentComponent(const UEdGraphPin* Pin, FMutableGraphGenerationContext& GenerationContext);
+
+
+uint32 GetTypeHash(const FGraphCycleKey& Key)
+{
+	return HashCombine(GetTypeHash(&Key.Pin), GetTypeHash(Key.Id));
+}
+
+
+FGraphCycle::FGraphCycle(const FGraphCycleKey&& Key, FMutableGraphGenerationContext& Context) :
+	Key(Key),
+	Context(Context)
+{
+}
+
+
+FGraphCycle::~FGraphCycle()
+{
+	Context.VisitedPins.Remove(Key);
+}
+
+
+bool FGraphCycle::FoundCycle() const
+{
+	const UCustomizableObjectNode& Node = *Cast<UCustomizableObjectNode>(Key.Pin.GetOwningNode());
+
+	if (const UCustomizableObject** Result = Context.VisitedPins.Find(Key))
+	{
+		Context.Compiler->CompilerLog(LOCTEXT("CycleFoundNode", "Cycle detected."), &Node, EMessageSeverity::Error, true);
+		Context.CustomizableObjectWithCycle = *Result;
+		return true;
+	}
+	else
+	{
+		Context.VisitedPins.Add(Key, Node.GetGraph()->GetTypedOuter<UCustomizableObject>());
+		return false;	
+	}
+}
+
+
+/** Warn if the node has more outputs than it is meant to have. */
+void CheckNumOutputs(const UEdGraphPin& Pin, const FMutableGraphGenerationContext& GenerationContext)
+{
+	if (const UCustomizableObjectNode* Typed = Cast<UCustomizableObjectNode>(Pin.GetOwningNode()))
+	{
+		if (Typed->IsSingleOutputNode())
+		{
+			int numOutLinks = 0;
+			for (UEdGraphPin* NodePin : Typed->GetAllNonOrphanPins())
+			{
+				if (NodePin->Direction == EGPD_Output)
+				{
+					numOutLinks += NodePin->LinkedTo.Num();
+				}
+			}
+
+			if (numOutLinks > 1)
+			{
+				GenerationContext.Compiler->CompilerLog(LOCTEXT("MultipleOutgoing", "The node has several outgoing connections, but it should be limited to 1."), CastChecked<UCustomizableObjectNode>(Pin.GetOwningNode()));
+			}
+		}
+	}
+}
+
+
+FMutableGraphGenerationContext::FMutableGraphGenerationContext(UCustomizableObject* InObject, FCustomizableObjectCompiler* InCompiler, const FCompilationOptions& InOptions)
+	: Object(InObject), Compiler(InCompiler), Options(InOptions)
+{
+	// Default flags for mesh generation nodes.
+	MeshGenerationFlags.Push(0);
+}
+
+
+mu::MeshPtr FMutableGraphGenerationContext::FindGeneratedMesh( const FGeneratedMeshData::FKey& Key )
+{
+	for (const FGeneratedMeshData& d : GeneratedMeshes)
+	{
+		if (d.Key == Key )
+		{
+			return d.Generated;
+		}
+	}
+
+	return nullptr;
+}
+
+
+/** Adds to ParameterNamesMap the node Node to the array of elements with name Name */
+void FMutableGraphGenerationContext::AddParameterNameUnique(const UCustomizableObjectNode* Node, FString Name)
+{
+	if (TArray<const UCustomizableObjectNode*>* ArrayResult = ParameterNamesMap.Find(Name))
+	{
+		ArrayResult->AddUnique(Node);
+	}
+	else
+	{
+		TArray<const UCustomizableObjectNode*> ArrayTemp;
+		ArrayTemp.Add(Node);
+		ParameterNamesMap.Add(Name, ArrayTemp);
+	}
+}
+
+
+const FGuid FMutableGraphGenerationContext::GetNodeIdUnique(const UCustomizableObjectNode* Node)
+{
+	TArray<const UCustomizableObjectNode*>* ArrayResult = NodeIdsMap.Find(Node->NodeGuid);
+
+	if (ArrayResult == nullptr)
+	{
+		TArray<const UCustomizableObjectNode*> ArrayTemp;
+		ArrayTemp.Add(Node);
+		NodeIdsMap.Add(Node->NodeGuid, ArrayTemp);
+		return Node->NodeGuid;
+	}
+
+	ArrayResult->Add(Node);
+	return FGuid::NewGuid();
+}
+
+
+void FMutableGraphGenerationContext::GenerateClippingCOInternalTags()
+{
+	int32 i = 0;
+	int32 j;
+
+	// The following array ArrayCachedTaggedMaterialNodes and related functionality was removed because it prevented two clip nodes from
+	// clipping the same material node
+	// In the case more than one CO clips a particular CO, the associated mu::NodeSurfaceNewPtr nodes with that particular CO only need to be tagged once
+	//TArray<FGuid> ArrayCachedTaggedMaterialNodes;
+
+	for (auto It = MapClipMeshNodeToMutableClipMeshNodeArray.CreateConstIterator(); It; ++It)
+	{
+		FString TagName = FString::Printf(TEXT("ClippingTag_%d"), i);
+		i++;
+
+		for (j = 0; j < It->Value.Num(); ++j)
+		{
+			It->Value[j]->AddTag(TCHAR_TO_ANSI(*TagName));
+		}
+
+		UCustomizableObjectNodeMeshClipWithMesh* CustomizableObjectNodeMeshClipWithMesh = It->Key;
+
+		// Tag also those nodes in the material side of the clipping functionality (since a CO can have several materials with the same LOD,
+		// all of them are considered, so CustomizableObjectNodeMeshClipWithMesh->ArrayMaterialNodeToClipWithID contains the Guids of all
+		// the material nodes with the same name selected for being affected by the clipping operation
+		TArray<FGuid>& ArrayMaterialNodeToClipWithID = CustomizableObjectNodeMeshClipWithMesh->ArrayMaterialNodeToClipWithID;
+
+		for (j = 0; j < ArrayMaterialNodeToClipWithID.Num(); ++j)
+		{
+			// The following lines were commented for the same reasons as ArrayCachedTaggedMaterialNodes above
+			//if (ArrayCachedTaggedMaterialNodes.Find(CustomizableObjectNodeMeshClipWithMesh->ArrayMaterialNodeToClipWithID[j]) != INDEX_NONE)
+			//{
+			//    continue;
+			//}
+
+			bool FoundEntryInMap = false;
+			for (auto It2 = MapMaterialNodeToMutableSurfaceNodeArray.CreateConstIterator(); (It2 && !FoundEntryInMap); ++It2)
+			{
+				if (It2->Key->NodeGuid == ArrayMaterialNodeToClipWithID[j])
+				{
+					FoundEntryInMap = true;
+
+					for (int32 k = 0; k < It2->Value.Num(); ++k)
+					{
+						It2->Value[k]->AddTag(TCHAR_TO_ANSI(*TagName));
+					}
+				}
+			}
+
+			// The following lines were commented for the same reasons as ArrayCachedTaggedMaterialNodes above
+			//ArrayCachedTaggedMaterialNodes.Add(CustomizableObjectNodeMeshClipWithMesh->ArrayMaterialNodeToClipWithID[j]);
+		}
+	}
+}
+
+
+void FMutableGraphGenerationContext::CheckPhysicsAssetInSkeletalMesh(const USkeletalMesh* SkeletalMesh)
+{
+	if (SkeletalMesh && SkeletalMesh->GetPhysicsAsset() && !DiscartedPhysicsAssetMap.Find(SkeletalMesh->GetPhysicsAsset()))
+	{
+		for (TObjectPtr<USkeletalBodySetup> Bodies : SkeletalMesh->GetPhysicsAsset()->SkeletalBodySetups)
+		{
+			if (Bodies->BoneName != NAME_None && SkeletalMesh->GetRefSkeleton().FindBoneIndex(Bodies->BoneName) == INDEX_NONE)
+			{
+				DiscartedPhysicsAssetMap.Add(SkeletalMesh->GetPhysicsAsset(), DiscartedPhysicsAssetMap.Num());
+				break;
+			}
+		}
+	}
+}
+
+
+FMutableComponentInfo& FMutableGraphGenerationContext::GetCurrentComponentInfo()
+{
+	check(ComponentInfos.IsValidIndex(CurrentMeshComponent));
+	return ComponentInfos[CurrentMeshComponent];
+}
+
+
+FGeneratedKey::FGeneratedKey(void* InFunctionAddress, const UEdGraphPin& InPin, const UCustomizableObjectNode& Node, FMutableGraphGenerationContext& GenerationContext, const bool UseMesh)
+{
+	FunctionAddress = InFunctionAddress;
+	Pin = &InPin;
+	LOD = Node.IsAffectedByLOD() ? GenerationContext.CurrentLOD : 0;
+
+	if (UseMesh)
+	{
+		Flags = GenerationContext.MeshGenerationFlags.Last();
+		MeshMorphStack = GenerationContext.MeshMorphStack;
+	}
+}
+
+
+bool FGeneratedKey::operator==(const FGeneratedKey& Other) const
+{
+	return FunctionAddress == Other.FunctionAddress &&
+		Pin == Other.Pin &&
+		LOD == Other.LOD &&
+		Flags == Other.Flags &&
+		MeshMorphStack == Other.MeshMorphStack;
+}
+
+
+uint32 GetTypeHash(const FGeneratedKey& Key)
+{
+	uint32 Hash = GetTypeHash(Key.FunctionAddress);
+	Hash = HashCombine(Hash, GetTypeHash(Key.Pin));
+	Hash = HashCombine(Hash, GetTypeHash(Key.LOD));
+	Hash = HashCombine(Hash, GetTypeHash(Key.Flags));
+	//Hash = HashCombine(Hash, GetTypeHash(Key.MeshMorphStack)); // Does not support array
+	
+	return Hash;
+}
+
+
+void FPinDataValue::Append(const FPinDataValue& From)
+{
+	MeshesData.Append(From.MeshesData);
+}
+
+
+FPinDataValue* FPinData::Find(const UEdGraphPin* Pin)
+{
+	return Data.Find(Pin);
+}
+
+
+FPinDataValue& FPinData::GetCurrent()
+{
+	return Data[PinStack.Last()];
+}
+
+
+void FPinData::Pop()
+{
+	const uint32 Num = PinStack.Num();
+
+	check(Num >= 1); // Pop called without a previous push
+
+	if (Num >= 2)
+	{
+		Data[PinStack[Num - 2]].Append(Data[PinStack[Num - 1]]);
+	}
+
+	PinStack.Pop();
+}
+
+
+void FPinData::Push(const UEdGraphPin* Pin)
+{
+	PinStack.Push(Pin);
+	Data.FindOrAdd(Pin);
+}
+
+
+FGraphCycleKey::FGraphCycleKey(const UEdGraphPin& Pin, uint32 Id) :
+	Pin(Pin),
+	Id(Id)
+{
+}
+
+
+bool FGraphCycleKey::operator==(const FGraphCycleKey& Other) const
+{
+	return &Pin == &Other.Pin && Id == Other.Id;
+}
+
+
+FScopedPinData::FScopedPinData(FMutableGraphGenerationContext& Context, const UEdGraphPin* Pin)
+	: Context(Context)
+{
+	Context.PinData.Push(Pin);
+}
+
+
+FScopedPinData::~FScopedPinData()
+{
+	Context.PinData.Pop();
+}
+
+
+UTexture2D* FindReferenceImage(const UEdGraphPin* Pin, FMutableGraphGenerationContext& GenerationContext)
+{
+	check(Pin)
+	RETURN_ON_CYCLE(*Pin, GenerationContext)
+
+	CheckNumOutputs(*Pin, GenerationContext);
+	
+	const UEdGraphSchema_CustomizableObject* Schema = GetDefault<UEdGraphSchema_CustomizableObject>();
+	const UCustomizableObjectNode* Node = CastChecked<UCustomizableObjectNode>(Pin->GetOwningNode());
+
+	UTexture2D* Result = nullptr;
+
+	if (const UCustomizableObjectNodeTexture* TypedNodeTex = Cast<UCustomizableObjectNodeTexture>(Node))
+	{
+		Result = TypedNodeTex->Texture;
+	}
+
+	else if (const UCustomizableObjectNodeTextureParameter* ParamNodeTex = Cast<UCustomizableObjectNodeTextureParameter>(Node))
+	{
+		Result = ParamNodeTex->DefaultValue;
+	}
+
+	else if (const UCustomizableObjectNodeMesh* TypedNodeMesh = Cast<UCustomizableObjectNodeMesh>(Node))
+	{
+		Result = TypedNodeMesh->FindTextureForPin(Pin);
+	}
+
+	else if (const UCustomizableObjectNodeTextureInterpolate* TypedNodeInterp = Cast<UCustomizableObjectNodeTextureInterpolate>(Node))
+	{
+		for (int LayerIndex = 0; !Result && LayerIndex < TypedNodeInterp->GetNumTargets(); ++LayerIndex)
+		{
+			if (const UEdGraphPin* ConnectedPin = FollowInputPin(*TypedNodeInterp->Targets(LayerIndex)))
+			{
+				Result = FindReferenceImage(ConnectedPin, GenerationContext);
+			}
+		}
+	}
+
+	else if (const UCustomizableObjectNodeTextureLayer* TypedNodeLayer = Cast<UCustomizableObjectNodeTextureLayer>(Node))
+	{
+		if (const UEdGraphPin* ConnectedPin = FollowInputPin(*TypedNodeLayer->BasePin()))
+		{
+			Result = FindReferenceImage(ConnectedPin, GenerationContext);
+		}
+
+		for (int LayerIndex = 0; !Result && LayerIndex < TypedNodeLayer->GetNumLayers(); ++LayerIndex)
+		{
+			if (const UEdGraphPin* ConnectedPin = FollowInputPin(*TypedNodeLayer->LayerPin(LayerIndex)))
+			{
+				if (ConnectedPin->PinType.PinCategory == Helper_GetPinCategory(Schema->PC_Image))
+				{
+					Result = FindReferenceImage(ConnectedPin, GenerationContext);
+				}
+			}
+		}
+	}
+
+	else if (const UCustomizableObjectNodeTextureSwitch* TypedNodeSwitch = Cast<UCustomizableObjectNodeTextureSwitch>(Node))
+	{
+		for (int SelectorIndex = 0; !Result && SelectorIndex < TypedNodeSwitch->GetNumElements(); ++SelectorIndex)
+		{
+			if (const UEdGraphPin* ConnectedPin = FollowInputPin(*TypedNodeSwitch->GetElementPin(SelectorIndex)))
+			{
+				Result = FindReferenceImage(ConnectedPin, GenerationContext);
+			}
+		}
+	}
+
+	else if (const UCustomizableObjectNodeTextureVariation* TypedNodeVariation = Cast<UCustomizableObjectNodeTextureVariation>(Node))
+	{
+		if (const UEdGraphPin* ConnectedPin = FollowInputPin(*TypedNodeVariation->DefaultPin()))
+		{
+			Result = FindReferenceImage(ConnectedPin, GenerationContext);
+		}
+
+		for (int SelectorIndex = 0; !Result && SelectorIndex < TypedNodeVariation->GetNumVariations(); ++SelectorIndex)
+		{
+			if (const UEdGraphPin* ConnectedPin = FollowInputPin(*TypedNodeVariation->VariationPin(SelectorIndex)))
+			{
+				Result = FindReferenceImage(ConnectedPin, GenerationContext);
+			}
+		}
+	}
+
+	else if (const UCustomizableObjectNodeTextureFromChannels* TypedNodeFrom = Cast<UCustomizableObjectNodeTextureFromChannels>(Node))
+	{
+		mu::NodeImagePtr RNode, GNode, BNode, ANode;
+
+		if (const UEdGraphPin* ConnectedPin = FollowInputPin(*TypedNodeFrom->RPin()))
+		{
+			Result = FindReferenceImage(ConnectedPin, GenerationContext);
+		}
+		if (!Result)
+		{
+			if (const UEdGraphPin* ConnectedPin = FollowInputPin(*TypedNodeFrom->GPin()))
+			{
+				Result = FindReferenceImage(ConnectedPin, GenerationContext);
+			}
+		}
+		if (!Result)
+		{
+			if (const UEdGraphPin* ConnectedPin = FollowInputPin(*TypedNodeFrom->BPin()))
+			{
+				Result = FindReferenceImage(ConnectedPin, GenerationContext);
+			}
+			
+		}
+		if (!Result)
+		{
+			if (const UEdGraphPin* ConnectedPin = FollowInputPin(*TypedNodeFrom->APin()))
+			{
+				Result = FindReferenceImage(ConnectedPin, GenerationContext);
+			}
+		}
+	}
+
+	else if (const UCustomizableObjectNodeTextureToChannels* TypedNodeTo = Cast<UCustomizableObjectNodeTextureToChannels>(Node))
+	{
+		if (const UEdGraphPin* ConnectedPin = FollowInputPin(*TypedNodeTo->InputPin()))
+		{
+			Result = FindReferenceImage(ConnectedPin, GenerationContext);
+		}
+	}
+
+	else if (const UCustomizableObjectNodeTextureProject* TypedNodeProj = Cast<UCustomizableObjectNodeTextureProject>(Node))
+	{
+		int TexIndex = -1;// TypedNodeProj->OutputPins.Find((UEdGraphPin*)Pin);
+		for (int32 i = 0; i < TypedNodeProj->GetNumOutputs(); ++i)
+		{
+			if (TypedNodeProj->OutputPins(i) == Pin)
+			{
+				TexIndex = i;
+			}
+		}
+
+		check(TexIndex >= 0 && TexIndex < TypedNodeProj->GetNumTextures());
+
+		if (const UEdGraphPin* ConnectedPin = FollowInputPin(*TypedNodeProj->TexturePins(TexIndex)))
+		{
+			Result = FindReferenceImage(ConnectedPin, GenerationContext);
+		}
+	}
+
+	else if (const UCustomizableObjectNodeTextureBinarise* TypedNodeBin = Cast<UCustomizableObjectNodeTextureBinarise>(Node))
+	{
+		if (const UEdGraphPin* ConnectedPin = FollowInputPin(*TypedNodeBin->GetBaseImagePin()))
+		{
+			Result = FindReferenceImage(ConnectedPin, GenerationContext);
+		}
+	}
+
+	else if (const UCustomizableObjectNodeTextureInvert* TypedNodeInv = Cast<UCustomizableObjectNodeTextureInvert>(Node))
+	{
+		if (const UEdGraphPin* ConnectedPin = FollowInputPin(*TypedNodeInv->GetBaseImagePin()))
+		{
+			Result = FindReferenceImage(ConnectedPin, GenerationContext);
+		}
+	}
+
+	else if (const UCustomizableObjectNodeTextureColourMap* TypedNodeColourMap = Cast<UCustomizableObjectNodeTextureColourMap>(Node))
+	{
+		if (const UEdGraphPin* ConnectedPin = FollowInputPin(*TypedNodeColourMap->GetBasePin()))
+		{
+			Result = FindReferenceImage(ConnectedPin, GenerationContext);
+		}
+	}
+
+	else if (const UCustomizableObjectNodeTextureTransform* TypedNodeTransform = Cast<UCustomizableObjectNodeTextureTransform>(Node))
+	{
+		if ( UEdGraphPin* BaseImagePin = FollowInputPin(*TypedNodeTransform->GetBaseImagePin()) )
+		{
+			Result = FindReferenceImage(BaseImagePin, GenerationContext);
+		}
+	}
+
+	else if (const UCustomizableObjectNodeTable* TypedNodeTable = Cast<UCustomizableObjectNodeTable>(Node))
+	{
+		if (Pin->PinType.PinCategory == Helper_GetPinCategory(Schema->PC_MaterialAsset))
+		{
+			Result = TypedNodeTable->FindReferenceTextureParameter(Pin, GenerationContext.CurrentMaterialTableParameter);
+		}
+		else
+		{
+			Result = TypedNodeTable->GetColumnDefaultAssetByType<UTexture2D>(Pin);
+		}
+	}
+	
+	return Result;
+}
+
+
+mu::NodeMeshApplyPosePtr CreateNodeMeshApplyPose(mu::NodeMeshPtr InputMeshNode, UCustomizableObject * CustomizableObject, TArray<FString> ArrayBoneName, TArray<FTransform> ArrayTransform)
+{
+	check(ArrayBoneName.Num() == ArrayTransform.Num());
+
+	mu::MeshPtr MutableMesh = new mu::Mesh();
+	mu::NodeMeshConstantPtr PoseNodeMesh = new mu::NodeMeshConstant;
+	PoseNodeMesh->SetValue(MutableMesh);
+
+	mu::SkeletonPtr MutableSkeleton = new mu::Skeleton;
+	MutableMesh->SetSkeleton(MutableSkeleton);
+	MutableMesh->SetBonePoseCount(ArrayBoneName.Num());
+	MutableSkeleton->SetBoneCount(ArrayBoneName.Num());
+
+	for (int32 i = 0; i < ArrayBoneName.Num(); ++i)
+	{
+		MutableSkeleton->SetBoneName(i, TCHAR_TO_ANSI(*ArrayBoneName[i]));
+		MutableMesh->SetBonePose(i, TCHAR_TO_ANSI(*ArrayBoneName[i]), (FTransform3f)ArrayTransform[i], true);
+	}
+
+	mu::NodeMeshApplyPosePtr NodeMeshApplyPose = new mu::NodeMeshApplyPose;
+	NodeMeshApplyPose->SetBase(InputMeshNode);
+	NodeMeshApplyPose->SetPose(PoseNodeMesh);
+
+	return NodeMeshApplyPose;
+}
+
+
+// Convert a CustomizableObject Source Graph into a mutable source graph  
+mu::NodeObjectPtr GenerateMutableSource(const UEdGraphPin * Pin, FMutableGraphGenerationContext & GenerationContext, bool bPartialCompilation)
+{
+	check(Pin)
+	RETURN_ON_CYCLE(*Pin, GenerationContext)
+
+	CheckNumOutputs(*Pin, GenerationContext);
+	
+	UCustomizableObjectNode* Node = CastChecked<UCustomizableObjectNode>(Pin->GetOwningNode());
+	
+	const FGeneratedKey Key(reinterpret_cast<void*>(&GenerateMutableSource), *Pin, *Node, GenerationContext, true);
+	if (const FGeneratedData* Generated = GenerationContext.Generated.Find(Key))
+	{
+		return static_cast<mu::NodeObject*>(Generated->Node.get());
+	}
+
+	mu::NodeObjectPtr Result;
+	
+	if (const UCustomizableObjectNodeObject* TypedNodeObj = Cast<UCustomizableObjectNodeObject>(Node))
+	{
+		// Add it to the current processing hierarchy if we have none
+		// \TODO: or special flag starting a separate component in the future.
+		bool bStartedComponent = false;
+		if (!GenerationContext.ComponentNewNode.Num())
+		{
+			GenerationContext.ComponentNewNode.Push(FMutableGraphGenerationContext::ObjectParent());
+			bStartedComponent = true;
+		}
+
+		mu::NodeObjectNewPtr ObjectNode = new mu::NodeObjectNew();
+		Result = ObjectNode;
+
+		ObjectNode->SetName(TCHAR_TO_ANSI(*TypedNodeObj->ObjectName));
+		ObjectNode->SetUid(TCHAR_TO_ANSI(*TypedNodeObj->NodeGuid.ToString()));
+
+
+		// LOD
+		int NumLODs = TypedNodeObj->GetNumLODPins();
+
+		// Fill the basic LOD Settings
+		if (!GenerationContext.NumLODsInRoot)
+		{
+			check(!GenerationContext.ComponentInfos.IsEmpty());
+
+			// NumLODsInRoot
+			int32 MaxRefMeshLODs = 1;
+			for (int32 MeshIndex = 0; MeshIndex < GenerationContext.ComponentInfos.Num(); ++MeshIndex)
+			{
+				const USkeletalMesh* RefSkeletalMesh = GenerationContext.ComponentInfos[MeshIndex].RefSkeletalMesh;
+				if (RefSkeletalMesh && RefSkeletalMesh->GetLODNum() > MaxRefMeshLODs)
+				{
+					MaxRefMeshLODs = RefSkeletalMesh->GetLODNum();
+				}
+			}
+
+			if (MaxRefMeshLODs < NumLODs)
+			{
+				FString Msg = FString::Printf(TEXT("The object has %d LODs but the reference mesh only %d. Resulting objects will have %d LODs."),
+					NumLODs, MaxRefMeshLODs, MaxRefMeshLODs);
+				GenerationContext.Compiler->CompilerLog(FText::FromString(Msg), Node, EMessageSeverity::Warning);
+				GenerationContext.NumLODsInRoot = MaxRefMeshLODs;
+			}
+			else
+			{
+				GenerationContext.NumLODsInRoot = NumLODs;
+			}
+		
+			USkeletalMesh* RefSkeletalMesh = GenerationContext.ComponentInfos[0].RefSkeletalMesh;
+			check(RefSkeletalMesh);
+
+			const FMutableLODSettings& LODSettings = GenerationContext.Object->LODSettings;
+
+			// Find the MinLOD available for the target platform
+			if (RefSkeletalMesh->IsMinLodQualityLevelEnable())
+			{
+				GenerationContext.FirstLODAvailable = LODSettings.MinQualityLevelLOD.GetValueForPlatform(GenerationContext.Options.TargetPlatform);
+			}
+			else
+			{
+				GenerationContext.FirstLODAvailable = LODSettings.MinLOD.GetValueForPlatform(*GenerationContext.Options.TargetPlatform->IniPlatformName());
+			}
+
+			GenerationContext.FirstLODAvailable = FMath::Clamp(GenerationContext.FirstLODAvailable, 0, GenerationContext.NumLODsInRoot - 1);
+
+			// Find the streaming settings for the target platform
+			if (LODSettings.bOverrideLODStreamingSettings)
+			{
+				GenerationContext.bEnableLODStreaming = LODSettings.bEnableLODStreaming.GetValueForPlatform(*GenerationContext.Options.TargetPlatform->IniPlatformName());
+				GenerationContext.NumMaxLODsToStream = LODSettings.NumMaxStreamedLODs.GetValueForPlatform(*GenerationContext.Options.TargetPlatform->IniPlatformName());
+			}
+			else
+			{
+				for (int32 MeshIndex = 0; MeshIndex < GenerationContext.ComponentInfos.Num(); ++MeshIndex)
+				{
+					RefSkeletalMesh = GenerationContext.ComponentInfos[MeshIndex].RefSkeletalMesh;
+					check(RefSkeletalMesh);
+
+					GenerationContext.bEnableLODStreaming = GenerationContext.bEnableLODStreaming &&
+						RefSkeletalMesh->GetEnableLODStreaming(GenerationContext.Options.TargetPlatform);
+
+					GenerationContext.NumMaxLODsToStream = FMath::Min(GenerationContext.NumMaxLODsToStream,
+						RefSkeletalMesh->GetMaxNumStreamedLODs(GenerationContext.Options.TargetPlatform));
+				}
+			}
+		}
+
+
+		// States
+		int NumStates = TypedNodeObj->States.Num();
+		ObjectNode->SetStateCount(NumStates);
+
+		// In a partial compilation we will filter the states of the root object
+		bool bFilterStates = true;
+
+		if (bPartialCompilation)
+		{
+			if (!TypedNodeObj->ParentObject)
+			{
+				bFilterStates = false;
+			}
+		}
+
+		for (int StateIndex = 0; StateIndex < NumStates && bFilterStates; ++StateIndex)
+		{
+			const FCustomizableObjectState& State = TypedNodeObj->States[StateIndex];
+			ObjectNode->SetStateName(StateIndex, TCHAR_TO_ANSI(*State.Name));
+			for (int ParamIndex = 0; ParamIndex < State.RuntimeParameters.Num(); ++ParamIndex)
+			{
+				ObjectNode->AddStateParam(StateIndex, TCHAR_TO_ANSI(*State.RuntimeParameters[ParamIndex]));
+			}
+			ObjectNode->SetStateProperties(StateIndex, State.bDontCompressRuntimeTextures, State.bBuildOnlyFirstLOD, GenerationContext.FirstLODAvailable);
+
+			// UI Data
+			FParameterUIData ParameterUIData(State.Name, State.StateUIMetadata, EMutableParameterType::None);
+			ParameterUIData.bDontCompressRuntimeTextures = State.bDontCompressRuntimeTextures;
+			ParameterUIData.ForcedParameterValues = State.ForcedParameterValues;
+
+			GenerationContext.StateUIDataMap.Add(State.Name, ParameterUIData);
+		}
+
+		// Update the current automatic LOD policy
+		ECustomizableObjectAutomaticLODStrategy LastLODStrategy = GenerationContext.CurrentAutoLODStrategy;
+		if (TypedNodeObj->AutoLODStrategy != ECustomizableObjectAutomaticLODStrategy::Inherited)
+		{
+			GenerationContext.CurrentAutoLODStrategy = TypedNodeObj->AutoLODStrategy;
+		}
+
+		// UI Data
+		GenerationContext.ParameterUIDataMap.Add(TypedNodeObj->ObjectName, FParameterUIData(
+			TypedNodeObj->ObjectName,
+			TypedNodeObj->ParamUIMetadata,
+			EMutableParameterType::Int));
+
+		// Mesh components per instance
+		if (!GenerationContext.NumMeshComponentsInRoot)
+		{
+			GenerationContext.NumMeshComponentsInRoot = TypedNodeObj->NumMeshComponents;
+		}
+
+		while (GenerationContext.ComponentNewNode.Last().ComponentsPerLOD.Num() < GenerationContext.NumLODsInRoot)
+		{
+			TArray<mu::NodeComponentNewPtr> AuxArray;
+			AuxArray.AddZeroed(GenerationContext.NumMeshComponentsInRoot);
+			GenerationContext.ComponentNewNode.Last().ComponentsPerLOD.Add(AuxArray);
+		}
+
+		int32 LastDefinedLOD = -1;
+		ObjectNode->SetLODCount(GenerationContext.NumLODsInRoot);
+		for (int32 LODIndex = 0; LODIndex < GenerationContext.NumLODsInRoot; ++LODIndex)
+		{
+			GenerationContext.CurrentLOD = LODIndex;
+
+			mu::NodeLODPtr LODNode = new mu::NodeLOD();
+			ObjectNode->SetLOD(LODIndex, LODNode);
+			LODNode->SetMessageContext(Node);
+			LODNode->SetComponentCount(GenerationContext.NumMeshComponentsInRoot);
+
+			for (int32 MeshComponentIndex = 0; MeshComponentIndex < GenerationContext.NumMeshComponentsInRoot; ++MeshComponentIndex)
+			{
+				GenerationContext.CurrentMeshComponent = MeshComponentIndex;
+				mu::NodeComponentPtr ComponentNode;
+
+				if (!GenerationContext.ComponentNewNode.Last().ComponentsPerLOD[LODIndex][MeshComponentIndex])
+				{
+					mu::NodeComponentNewPtr ComponentNewNode = new mu::NodeComponentNew();
+					ComponentNewNode->SetMessageContext(Node);
+					ComponentNewNode->SetId(MeshComponentIndex);
+					GenerationContext.ComponentNewNode.Last().ComponentsPerLOD[LODIndex][MeshComponentIndex] = ComponentNewNode;
+					ComponentNode = ComponentNewNode;
+				}
+				else
+				{
+					mu::NodeComponentEditPtr ComponentEditNode = new mu::NodeComponentEdit();
+					ComponentEditNode->SetMessageContext(Node);
+					ComponentEditNode->SetParent(GenerationContext.ComponentNewNode.Last().ComponentsPerLOD[LODIndex][MeshComponentIndex].get());
+					ComponentNode = ComponentEditNode;
+				}
+
+				LODNode->SetComponent(MeshComponentIndex, ComponentNode);
+
+				// Is the LOD defined in this object?
+				int32 LODToGenerate = -1;
+				if (LODIndex < NumLODs)
+				{
+					// If it is manual LODs, we don't want to show anything for this node in this LOD
+					if (GenerationContext.CurrentAutoLODStrategy == ECustomizableObjectAutomaticLODStrategy::AutomaticFromMesh
+						&& LastDefinedLOD != INDEX_NONE)
+					{
+						LODToGenerate = LastDefinedLOD;
+					}
+					else
+					{
+						// In automatic LODs, we follow the last LOD
+						LastDefinedLOD = LODIndex;
+						LODToGenerate = LODIndex;
+					}
+				}
+				else
+				{
+					// This LOD is not defined in this object, so we assume the last LOD has to be used instead
+					LODToGenerate = LastDefinedLOD;
+				}
+
+				if (GenerationContext.CurrentLOD < GenerationContext.FirstLODAvailable)
+				{
+					continue;
+				}
+
+				if (LODToGenerate >= 0)
+				{
+					TArray<UEdGraphPin*> ConnectedLODPins = FollowInputPinArray(*TypedNodeObj->LODPin(LODToGenerate));
+					int32 LODConnections = ConnectedLODPins.Num();
+					int32 NumMaterials = 0;
+					int32 NumModifiers = 0;
+					ComponentNode->SetSurfaceCount(NumMaterials);
+					for (int32 MatIndex = 0; MatIndex < LODConnections; ++MatIndex)
+					{
+						const UEdGraphPin* ChildNodePin = ConnectedLODPins[MatIndex];
+						if (!AffectsCurrentComponent(ChildNodePin, GenerationContext))
+						{
+							continue;
+						}
+
+						if (Cast<UCustomizableObjectNodeModifierBase>(ChildNodePin->GetOwningNode()))
+						{
+							mu::NodeModifierPtr MatNode = GenerateMutableSourceModifier(ChildNodePin, GenerationContext);
+							LODNode->SetModifierCount(NumModifiers + 1);
+							LODNode->SetModifier(NumModifiers, MatNode.get());
+							LODNode->SetMessageContext(Node);
+							++NumModifiers;
+						}
+						else
+						{
+							FMutableGraphSurfaceGenerationData DummySurfaceData;
+							mu::NodeSurfacePtr MatNode = GenerateMutableSourceSurface(ChildNodePin, GenerationContext, DummySurfaceData);
+							ComponentNode->SetSurfaceCount(NumMaterials + 1);
+							ComponentNode->SetSurface(NumMaterials, MatNode.get());
+							ComponentNode->SetMessageContext(Node);
+							++NumMaterials;
+						}
+					}
+				}
+			}
+		}
+
+		// Children
+		TArray<UEdGraphPin*> ConnectedChildrenPins = FollowInputPinArray(*TypedNodeObj->ChildrenPin());
+		ObjectNode->SetChildCount(ConnectedChildrenPins.Num());
+		for (int ChildIndex = 0; ChildIndex < ConnectedChildrenPins.Num(); ++ChildIndex)
+		{
+			mu::NodeObjectPtr ChildNode = GenerateMutableSource(ConnectedChildrenPins[ChildIndex], GenerationContext, bPartialCompilation);
+			ObjectNode->SetChild(ChildIndex, ChildNode.get());
+		}
+
+		// Remove from the current processing hierarchy
+		if (bStartedComponent)
+		{
+			GenerationContext.ComponentNewNode.Pop();
+		}
+
+		GenerationContext.CurrentAutoLODStrategy = LastLODStrategy;
+	}
+
+	else if (const UCustomizableObjectNodeObjectGroup* TypedNodeGroup = Cast<UCustomizableObjectNodeObjectGroup>(Node))
+	{
+		mu::NodeObjectGroupPtr GroupNode = new mu::NodeObjectGroup();
+		Result = GroupNode;
+
+		GenerationContext.AddParameterNameUnique(TypedNodeGroup, TypedNodeGroup->GroupName);
+		GroupNode->SetName(TCHAR_TO_ANSI(*TypedNodeGroup->GroupName));
+		GroupNode->SetUid(TCHAR_TO_ANSI(*TypedNodeGroup->NodeGuid.ToString()));
+
+		// Get all group projectors and put them in the generation context so that they are available to the child material nodes of this group node
+		uint32 NumProjectorCountBeforeNode = GenerationContext.ProjectorGroupMap.Num();
+		UEdGraphPin* ProjectorsPin = TypedNodeGroup->GroupProjectorsPin();
+		
+		if (ProjectorsPin && TypedNodeGroup)
+		{
+			for (const UEdGraphPin* GroupProjectorNodePin : FollowInputPinArray(*ProjectorsPin))
+			{
+				GenerateMutableSourceGroupProjector(GroupProjectorNodePin, GenerationContext, TypedNodeGroup);
+			}
+		}
+
+		mu::NodeObjectGroup::CHILD_SELECTION Type = mu::NodeObjectGroup::CS_ALWAYS_ALL;
+		switch (TypedNodeGroup->GroupType)
+		{
+		case ECustomizableObjectGroupType::COGT_ALL: Type = mu::NodeObjectGroup::CS_ALWAYS_ALL; break;
+		case ECustomizableObjectGroupType::COGT_TOGGLE: Type = mu::NodeObjectGroup::CS_TOGGLE_EACH; break;
+		case ECustomizableObjectGroupType::COGT_ONE: Type = mu::NodeObjectGroup::CS_ALWAYS_ONE; break;
+		case ECustomizableObjectGroupType::COGT_ONE_OR_NONE: Type = mu::NodeObjectGroup::CS_ONE_OR_NONE; break;
+		default:
+			GenerationContext.Compiler->CompilerLog(LOCTEXT("UnsupportedGroupType", "Object Group Type not supported. Setting to 'ALL'."), Node);
+			break;
+		}
+		GroupNode->SetSelectionType(Type);
+
+		// External children
+		TArray<UCustomizableObjectNodeObject*> ExternalChildNodes;
+		GenerationContext.GroupIdToExternalNodeMap.MultiFind(TypedNodeGroup->NodeGuid, ExternalChildNodes);
+		GenerationContext.GuidToParamNameMap.Add(TypedNodeGroup->NodeGuid, TypedNodeGroup->GroupName);
+
+		// Children
+		const TArray<UEdGraphPin*> ConnectedChildrenPins = FollowInputPinArray(*TypedNodeGroup->ObjectsPin());
+		const int32 NumChildren = ConnectedChildrenPins.Num();
+		const int32 TotalNumChildren = NumChildren + ExternalChildNodes.Num();
+
+		GroupNode->SetChildCount(TotalNumChildren);
+		int32 ChildIndex = 0;
+
+		// UI data
+		FParameterUIData ParameterUIData(
+			TypedNodeGroup->GroupName,
+			TypedNodeGroup->ParamUIMetadata,
+			EMutableParameterType::Int);
+
+		ParameterUIData.IntegerParameterGroupType = TypedNodeGroup->GroupType;
+
+		// In the case of a partial compilation, make sure at least one child is connected so that the param is no optimized
+		bool bAtLeastOneConnected = false;
+
+		for (; ChildIndex < NumChildren; ++ChildIndex)
+		{
+			bool bLastChildNode = (ChildIndex == NumChildren - 1) && (ExternalChildNodes.Num() == 0);
+			bool bConnectAtLeastTheLastChild = bLastChildNode && !bAtLeastOneConnected;
+
+			UCustomizableObjectNodeObject* CustomizableObjectNodeObject = Cast<UCustomizableObjectNodeObject>(ConnectedChildrenPins[ChildIndex]->GetOwningNode());
+
+			FString* SelectedOptionName = GenerationContext.ParamNamesToSelectedOptions.Find(TypedNodeGroup->GroupName); // If the param is in the map restrict to only the selected option
+
+			if (bConnectAtLeastTheLastChild || !SelectedOptionName || *SelectedOptionName == CustomizableObjectNodeObject->ObjectName)
+			{
+				bAtLeastOneConnected = true;
+
+				mu::NodeObjectPtr ChildNode = GenerateMutableSource(ConnectedChildrenPins[ChildIndex], GenerationContext, bPartialCompilation);
+				GroupNode->SetChild(ChildIndex, ChildNode.get());
+
+				if (CustomizableObjectNodeObject)
+				{
+					FString LeftSplit = CustomizableObjectNodeObject->GetPathName();
+					LeftSplit.Split(".", &LeftSplit, nullptr);
+					GenerationContext.CustomizableObjectPathMap.Add(CustomizableObjectNodeObject->Identifier.ToString(), LeftSplit);
+					GenerationContext.GroupNodeMap.Add(CustomizableObjectNodeObject->Identifier.ToString(), FCustomizableObjectIdPair(TypedNodeGroup->GroupName, ChildNode->GetName()));
+					ParameterUIData.ArrayIntegerParameterOption.Add(FIntegerParameterUIData(
+						CustomizableObjectNodeObject->ObjectName,
+						CustomizableObjectNodeObject->ParamUIMetadata));
+
+					if (TypedNodeGroup->GroupType == ECustomizableObjectGroupType::COGT_TOGGLE)
+					{
+						GenerationContext.AddParameterNameUnique(CustomizableObjectNodeObject, CustomizableObjectNodeObject->ObjectName);
+					}
+				}
+			}
+			else
+			{
+				mu::NodeObjectPtr ChildNode = new mu::NodeObjectNew;
+				ChildNode->SetName(TCHAR_TO_ANSI(*CustomizableObjectNodeObject->ObjectName));
+				GroupNode->SetChild(ChildIndex, ChildNode.get());
+			}
+		}
+
+		const bool bCollapseUnderParent = TypedNodeGroup->ParamUIMetadata.ExtraInformation.Find(FString("CollapseUnderParent")) != nullptr;
+		constexpr bool bHideWhenNotSelected = true; //TypedNodeGroup->ParamUIMetadata.ExtraInformation.Find(FString("HideWhenNotSelected"));
+		
+		if (bCollapseUnderParent || bHideWhenNotSelected)
+		{
+			if (const UEdGraphPin* ConnectedPin = FollowOutputPin(*Pin))
+			{
+				UCustomizableObjectNodeObject* NodeObject = Cast<UCustomizableObjectNodeObject>(ConnectedPin->GetOwningNode());
+
+				if (NodeObject && NodeObject->bIsBase)
+				{
+					const FGuid* ParentId = GenerationContext.GroupIdToExternalNodeMap.FindKey(NodeObject);
+
+					if (ParentId)
+					{
+						FString* ParentParamName = GenerationContext.GuidToParamNameMap.Find(*ParentId);
+
+						if (ParentParamName)
+						{
+							ParameterUIData.ParamUIMetadata.ExtraInformation.Add(FString("__ParentParamName"), *ParentParamName);
+
+							if (bHideWhenNotSelected)
+							{
+								ParameterUIData.ParamUIMetadata.ExtraInformation.Add(FString("__DisplayWhenParentValueEquals"), NodeObject->ObjectName);
+							}
+
+							if (bCollapseUnderParent)
+							{
+								ParameterUIData.ParamUIMetadata.ExtraInformation.Add(FString("CollapseUnderParent"));
+
+								FParameterUIData ParentParameterUIData;
+								ParentParameterUIData.ParamUIMetadata.ExtraInformation.Add(FString("__HasCollapsibleChildren"));
+								GenerationContext.ParameterUIDataMap.Add(*ParentParamName, ParentParameterUIData);
+							}
+						}
+					}
+				}
+			}
+		}
+
+		// Build external objects that reference this object as parent
+		const int32 NumExternalChildren = FMath::Max(0, TotalNumChildren - NumChildren);
+		for (int32 ExternalChildIndex = 0; ExternalChildIndex < NumExternalChildren; ++ExternalChildIndex)
+		{
+			const UCustomizableObjectNodeObject* ExternalChildNode = ExternalChildNodes[ExternalChildIndex];
+			bool bLastExternalChildNode = ExternalChildIndex == ExternalChildNodes.Num() - 1;
+			bool bConnectAtLeastTheLastChild = bLastExternalChildNode && !bAtLeastOneConnected;
+
+			UCustomizableObjectNodeObject* CustomizableObjectNodeObject = Cast<UCustomizableObjectNodeObject>(ExternalChildNode->OutputPin()->GetOwningNode());
+
+			FString* SelectedOptionName = GenerationContext.ParamNamesToSelectedOptions.Find(TypedNodeGroup->GroupName); // If the param is in the map restrict to only the selected option
+
+			if (bConnectAtLeastTheLastChild || !SelectedOptionName || *SelectedOptionName == CustomizableObjectNodeObject->ObjectName)
+			{
+				bAtLeastOneConnected = true;
+
+				mu::NodeObjectPtr ChildNode = GenerateMutableSource(ExternalChildNode->OutputPin(), GenerationContext, bPartialCompilation);
+				GroupNode->SetChild(ChildIndex, ChildNode.get());
+
+				if (CustomizableObjectNodeObject)
+				{
+					FString LeftSplit = ExternalChildNode->GetPathName();
+					LeftSplit.Split(".", &LeftSplit, nullptr);
+					GenerationContext.CustomizableObjectPathMap.Add(CustomizableObjectNodeObject->Identifier.ToString(), LeftSplit);
+					GenerationContext.GroupNodeMap.Add(CustomizableObjectNodeObject->Identifier.ToString(), FCustomizableObjectIdPair(TypedNodeGroup->GroupName, ChildNode->GetName()));
+					ParameterUIData.ArrayIntegerParameterOption.Add(FIntegerParameterUIData(
+						CustomizableObjectNodeObject->ObjectName,
+						CustomizableObjectNodeObject->ParamUIMetadata));
+
+					if (CustomizableObjectNodeObject->ObjectName.IsEmpty())
+					{
+						GenerationContext.NoNameNodeObjectArray.AddUnique(CustomizableObjectNodeObject);
+					}
+
+					if (TypedNodeGroup->GroupType == ECustomizableObjectGroupType::COGT_TOGGLE)
+					{
+						GenerationContext.AddParameterNameUnique(CustomizableObjectNodeObject, CustomizableObjectNodeObject->ObjectName);
+					}
+				}
+			}
+			else
+			{
+				mu::NodeObjectPtr ChildNode = new mu::NodeObjectNew;
+				ChildNode->SetName(TCHAR_TO_ANSI(*CustomizableObjectNodeObject->ObjectName));
+				GroupNode->SetChild(ChildIndex, ChildNode.get());
+			}
+
+			ChildIndex++;
+		}
+
+		const FParameterUIData* ChildFilledUIData = GenerationContext.ParameterUIDataMap.Find(TypedNodeGroup->GroupName);
+		if (ChildFilledUIData && ChildFilledUIData->ParamUIMetadata.ExtraInformation.Find(FString("__HasCollapsibleChildren")))
+		{
+			// Some child param filled the HasCollapsibleChildren UI info, refill it so it's not lost
+			ParameterUIData.ParamUIMetadata.ExtraInformation.Add(FString("__HasCollapsibleChildren"));
+		}
+
+		if (TypedNodeGroup->GroupType == ECustomizableObjectGroupType::COGT_TOGGLE)
+		{
+			for (const FIntegerParameterUIData& BooleanParam : ParameterUIData.ArrayIntegerParameterOption)
+			{
+				FParameterUIData ParameterUIDataBoolean(
+					BooleanParam.Name,
+					BooleanParam.ParamUIMetadata,
+					EMutableParameterType::Bool);
+
+				GenerationContext.ParameterUIDataMap.Add(BooleanParam.Name, ParameterUIDataBoolean);
+			}
+		}
+		else
+		{
+			GenerationContext.ParameterUIDataMap.Add(TypedNodeGroup->GroupName, ParameterUIData);
+		}
+
+		// Remove the projectors from this node
+		GenerationContext.ProjectorGroupMap.Remove(TypedNodeGroup);
+
+		check(NumProjectorCountBeforeNode == GenerationContext.ProjectorGroupMap.Num());
+	}
+	else
+	{
+		GenerationContext.Compiler->CompilerLog(LOCTEXT("UnimplementedNode", "Node type not implemented yet."), Node);
+	}
+
+	GenerationContext.Generated.Add(Key, FGeneratedData(Node, Result));
+	GenerationContext.GeneratedNodes.Add(Node);
+
+	return Result;
+}
+
+bool AffectsCurrentComponent(const UEdGraphPin* Pin, FMutableGraphGenerationContext& GenerationContext)
+{
+	check(Pin);
+	RETURN_ON_CYCLE(*Pin, GenerationContext)
+	
+	int32 ComponentIndex = INDEX_NONE;
+
+	const UCustomizableObjectNode* Node = CastChecked<UCustomizableObjectNode>(Pin->GetOwningNode());
+	if (const UCustomizableObjectNodeMaterialVariation* TypedNodeVar = Cast<UCustomizableObjectNodeMaterialVariation>(Node))
+	{
+		bool bAffectsCurrentComponent = false;
+
+		for (const UEdGraphPin* ConnectedPin : FollowInputPinArray(*TypedNodeVar->DefaultPin()))
+		{
+			if (!AffectsCurrentComponent(ConnectedPin, GenerationContext))
+			{
+				if (bAffectsCurrentComponent)
+				{
+					FString Msg = FString::Printf(TEXT("Error! One or more materials nodes linked to a material variation node have different component index"));
+					GenerationContext.Compiler->CompilerLog(FText::FromString(Msg), Node, EMessageSeverity::Error);
+				}
+
+				return false;
+			}
+		
+			bAffectsCurrentComponent = true;
+		}
+
+		for (int VariationIndex = 0; VariationIndex < TypedNodeVar->Variations.Num(); ++VariationIndex)
+		{
+			if (UEdGraphPin* VariationPin = TypedNodeVar->VariationPin(VariationIndex))
+			{
+				for (const UEdGraphPin* ConnectedPin : FollowInputPinArray(*VariationPin))
+				{
+					if (!AffectsCurrentComponent(ConnectedPin, GenerationContext))
+					{
+						if (bAffectsCurrentComponent)
+						{
+							FString Msg = FString::Printf(TEXT("Error! One or more materials nodes linked to a material variation node have different component index"));
+							GenerationContext.Compiler->CompilerLog(FText::FromString(Msg), Node, EMessageSeverity::Error);
+						}
+
+						return false;
+					}
+
+					bAffectsCurrentComponent = true;
+				}
+			}
+		}
+
+		return bAffectsCurrentComponent;
+	}
+	else if (const UCustomizableObjectNodeMaterial* TypedNodeMat = Cast<UCustomizableObjectNodeMaterial>(Node))
+	{
+		ComponentIndex = TypedNodeMat->MeshComponentIndex;
+	}
+	else if (const UCustomizableObjectNodeExtendMaterial* TypedNodeExt = Cast<UCustomizableObjectNodeExtendMaterial>(Node))
+	{
+		UCustomizableObjectNodeMaterial* OriginalParentMaterialNode = TypedNodeExt->GetParentMaterialNode();
+		ComponentIndex = OriginalParentMaterialNode ? OriginalParentMaterialNode->MeshComponentIndex : GenerationContext.CurrentMeshComponent;
+	}
+	else if (const UCustomizableObjectNodeEditMaterialBase* TypedNodeEdit = Cast<UCustomizableObjectNodeEditMaterialBase>(Node))
+	{
+		UCustomizableObjectNodeMaterial* OriginalParentMaterialNode = TypedNodeEdit->GetParentMaterialNode();
+		ComponentIndex = OriginalParentMaterialNode ? OriginalParentMaterialNode->MeshComponentIndex : GenerationContext.CurrentMeshComponent;
+	}
+	else if (const UCustomizableObjectNodeModifierBase* TypedNodeModifier = Cast<UCustomizableObjectNodeModifierBase>(Node))
+	{
+		// Modifiers affect all components with the exception of the UCustomizableObjectNodeMeshClipMorph
+		if (const UCustomizableObjectNodeMeshClipMorph* TypedNodeMeshClipMorph = Cast<UCustomizableObjectNodeMeshClipMorph>(Node))
+		{
+			ComponentIndex = TypedNodeMeshClipMorph->ReferenceSkeletonIndex;
+		}
+		else
+		{
+			return true;
+		}
+	}
+	else
+	{
+		unimplemented()
+	}
+
+	return ComponentIndex == GenerationContext.CurrentMeshComponent;
+}
+
+
+FString GenerateAnimationInstanceTag(const FString& AnimInstance, int32 SlotIndex)
+{
+	return FString("__AnimBP:") + FString::Printf(TEXT("%d_Slot_"), SlotIndex) + AnimInstance;
+}
+
+
+FString GenerateGameplayTag(const FString& GameplayTag)
+{
+	return FString("__AnimBPTag:") + GameplayTag;
+}
+
+
+void PopulateReferenceSkeletalMeshesData(FMutableGraphGenerationContext& GenerationContext)
+{
+	const FString PlatformName = GenerationContext.Options.TargetPlatform->IniPlatformName();
+	
+	const uint32 LODCount = GenerationContext.NumLODsInRoot;
+	const uint32 ComponentCount = GenerationContext.NumMeshComponentsInRoot;
+
+	GenerationContext.ReferenceSkeletalMeshesData.AddDefaulted(ComponentCount);
+	for(uint32 ComponentIndex = 0; ComponentIndex < ComponentCount; ++ComponentIndex)
+	{
+		const USkeletalMesh* RefSkeletalMesh = GenerationContext.ComponentInfos[ComponentIndex].RefSkeletalMesh;
+		check(RefSkeletalMesh);
+
+		// Gather LODData, this may include per LOD settings such as render data config or LODDataInfoArray
+		FMutableRefSkeletalMeshData& Data = GenerationContext.ReferenceSkeletalMeshesData[ComponentIndex];
+		Data.LODData.AddDefaulted(GenerationContext.NumLODsInRoot);
+		
+		const uint32 RefSkeletalMeshLODCount = RefSkeletalMesh->GetLODNum();
+		const TArray<FSkeletalMeshLODInfo>& LODInfos = RefSkeletalMesh->GetLODInfoArray();
+
+		for(uint32 LODIndex = 0; LODIndex < LODCount; ++LODIndex)
+		{
+			FMutableRefLODData& LODData = Data.LODData[LODIndex];
+			if(LODIndex < RefSkeletalMeshLODCount)
+			{
+				// Copy LOD info data from the reference skeletal mesh
+				const FSkeletalMeshLODInfo& LODInfo = LODInfos[LODIndex];
+				LODData.LODInfo.ScreenSize = LODInfo.ScreenSize.GetValueForPlatform(*PlatformName);
+				LODData.LODInfo.LODHysteresis = LODInfo.LODHysteresis;
+				LODData.LODInfo.bSupportUniformlyDistributedSampling = LODInfo.bSupportUniformlyDistributedSampling;
+				LODData.LODInfo.bAllowCPUAccess = LODInfo.bAllowCPUAccess;
+
+				// Copy Render data settings from the reference skeletal mesh
+				const FSkeletalMeshLODRenderData& ReferenceLODModel = Helper_GetLODData(RefSkeletalMesh)[LODIndex];
+				LODData.RenderData.bIsLODOptional = ReferenceLODModel.bIsLODOptional;
+				LODData.RenderData.bStreamedDataInlined = ReferenceLODModel.bStreamedDataInlined;
+			}
+			else
+			{
+				LODData.LODInfo.ScreenSize = 0.3f / (LODIndex + 1);
+				LODData.LODInfo.LODHysteresis = 0.02f;				
+			}
+		}
+
+		// Gather SkeletalMesh Sockets;
+		const TArray<USkeletalMeshSocket*>& RefSkeletonSockets = RefSkeletalMesh->GetMeshOnlySocketList();
+		const uint32 SocketCount = RefSkeletonSockets.Num();
+
+		Data.Sockets.AddDefaulted(SocketCount);
+		for(uint32 SocketIndex = 0 ; SocketIndex < SocketCount; ++SocketIndex)
+		{
+			const USkeletalMeshSocket* RefSocket = RefSkeletonSockets[SocketIndex];
+			check(RefSocket);
+
+			FMutableRefSocket& Socket = Data.Sockets[SocketIndex];
+			Socket.SocketName = RefSocket->SocketName;
+			Socket.BoneName = RefSocket->BoneName;
+			Socket.RelativeLocation = RefSocket->RelativeLocation;
+			Socket.RelativeRotation = RefSocket->RelativeRotation;
+			Socket.RelativeScale = RefSocket->RelativeScale;
+		}
+
+		// TODO: Generate Bounds?
+		// Gather Bounds
+		Data.Bounds = RefSkeletalMesh->GetBounds();
+		
+		// Additional Settings
+		Data.Settings.bEnablePerPolyCollision = RefSkeletalMesh->GetEnablePerPolyCollision();
+
+		// Skeleton
+		if(const USkeleton* Skeleton = RefSkeletalMesh->GetSkeleton())
+		{
+			Data.Skeleton = RefSkeletalMesh->GetSkeleton();
+			GenerationContext.ReferencedSkeletons.AddUnique(Skeleton);
+		}
+		
+		// Physics Asset
+		if (const UPhysicsAsset* PhysicsAsset = RefSkeletalMesh->GetPhysicsAsset())
+		{
+			Data.PhysicsAsset = PhysicsAsset;
+			GenerationContext.PhysicsAssetMap.FindOrAdd(Data.PhysicsAsset.ToString(), Data.PhysicsAsset);
+		}
+
+		// Post ProcessAnimInstance
+		if(const TSubclassOf<UAnimInstance> PostProcessAnimInstance = RefSkeletalMesh->GetPostProcessAnimBlueprint())
+		{
+			Data.PostProcessAnimInst = PostProcessAnimInstance;
+			GenerationContext.AnimBPAssetsMap.FindOrAdd(Data.PostProcessAnimInst.ToString(), Data.PostProcessAnimInst);
+		}
+		
+		// Shadow Physics Asset
+		if (const UPhysicsAsset* PhysicsAsset = RefSkeletalMesh->GetPhysicsAsset())
+		{
+			Data.PhysicsAsset = PhysicsAsset;
+			GenerationContext.PhysicsAssetMap.FindOrAdd(Data.PhysicsAsset.ToString(), Data.PhysicsAsset);
+		}
+	}
+}
+
+
+int32 ComputeLODBias(const FMutableGraphGenerationContext& GenerationContext, const UTexture2D* ReferenceTexture, int32 MaxTextureSize,
+	const UCustomizableObjectNodeMaterial* MaterialNode, const int32 ImageIndex)
+{
+	int32 LODBias = 0;
+
+	// We used to calculate the lod bias directly from the group like this:
+	//int LODBias = 0;
+	//if (LODSettings.TextureLODGroups.IsValidIndex(ReferenceTexture->LODGroup))
+	//{
+	//	const FTextureLODGroup& Group = LODSettings.GetTextureLODGroup(ReferenceTexture->LODGroup);
+	//	LODBias = Group.LODBias;
+	//}
+
+	// ...but now it seems to be more complicated and we may need this:
+	// This is not 100% correct either because it makes assumptions about the final texture size
+	// that may not be correct, but if the reference texture is really representative of the average 
+	// case, then it is as good as we can do.
+	if (ReferenceTexture)
+	{
+		const UTextureLODSettings& LODSettings = GenerationContext.Options.TargetPlatform->GetTextureLODSettings();
+
+		LODBias = LODSettings.CalculateLODBias(ReferenceTexture->Source.GetSizeX(), ReferenceTexture->Source.GetSizeY(), MaxTextureSize,
+			ReferenceTexture->LODGroup, ReferenceTexture->LODBias, false, ReferenceTexture->MipGenSettings, ReferenceTexture->IsCurrentlyVirtualTextured());
+	}
+
+	// Reduce data size and compilation time for server platforms.
+	// \todo: review this hack
+	bool bIsServer = GenerationContext.Options.TargetPlatform->GetTargetPlatformInfo().PlatformType == EBuildTargetType::Server;
+	if (bIsServer)
+	{
+		LODBias += 4;
+	}
+
+	// Increment the LOD bias per each LOD if we are using automatic LODs
+	if (GenerationContext.CurrentLOD > 0
+		&&
+		GenerationContext.CurrentAutoLODStrategy == ECustomizableObjectAutomaticLODStrategy::AutomaticFromMesh)
+	{
+		// Only if the texture actually uses a layout. Otherwise it could be a special texture we shouldn't scale.
+		if (MaterialNode->GetImageUVLayout(ImageIndex) >= 0)
+		{
+			// \todo: make it an object property to be tweaked
+			int MipsToSkipPerLOD = 1;
+			LODBias += MipsToSkipPerLOD * GenerationContext.CurrentLOD;
+		}
+	}
+
+	if (ReferenceTexture)
+	{
+		UE_LOG(LogMutable, Verbose, TEXT("Compiling texture with reference [%s] will have LOD Bias %d."), *ReferenceTexture->GetName(), LODBias);
+	}
+	else
+	{
+		UE_LOG(LogMutable, Verbose, TEXT("Compiling texture without reference will have LOD Bias %d."), LODBias);
+	}
+
+	return LODBias;
+}
+
+
+int32 GetMaxTextureSize(const UTexture2D* ReferenceTexture, const FMutableGraphGenerationContext& GenerationContext)
+{
+	if (ReferenceTexture)
+	{
+		// Setting the maximum texture size
+		const UTextureLODSettings& LODSettings = GenerationContext.Options.TargetPlatform->GetTextureLODSettings();
+		FTextureLODGroup TextureGroupSettings = LODSettings.GetTextureLODGroup(ReferenceTexture->LODGroup);
+
+		if (TextureGroupSettings.MaxLODSize > 0)
+		{
+			return ReferenceTexture->MaxTextureSize == 0 ? TextureGroupSettings.MaxLODSize : FMath::Min(TextureGroupSettings.MaxLODSize, ReferenceTexture->MaxTextureSize);
+		}
+		else
+		{
+			return ReferenceTexture->MaxTextureSize;
+		}
+	}
+
+	return 0;
+}
+
+#undef LOCTEXT_NAMESPACE
diff --git a/Engine/Plugins/Experimental/Mutable/Source/CustomizableObjectEditor/Private/GenerateMutableSource/GenerateMutableSource.h b/Engine/Plugins/Experimental/Mutable/Source/CustomizableObjectEditor/Private/GenerateMutableSource/GenerateMutableSource.h
new file mode 100644
--- /dev/null
+++ b/Engine/Plugins/Experimental/Mutable/Source/CustomizableObjectEditor/Private/GenerateMutableSource/GenerateMutableSource.h
@@ -0,0 +1,720 @@
+// Copyright Epic Games, Inc. All Rights Reserved.
+
+#pragma once
+
+#include "CustomizableObject.h"
+#include "NodeMeshApplyPose.h"
+#include "Nodes/CustomizableObjectNodeObject.h"
+#include "Nodes/CustomizableObjectNodeMaterial.h"
+#include "Nodes/CustomizableObjectNodeMesh.h"
+#include "Nodes/CustomizableObjectNodeMeshMorph.h"
+
+#include "Animation/Skeleton.h"
+#include "Engine/SkeletalMesh.h"
+
+#include "MutableTools/Public/NodeObject.h"
+#include "MutableTools/Public/NodeImageConstant.h"
+#include "MutableTools/Public/NodeSurfaceNew.h"
+#include "MutableTools/Public/NodeComponentNew.h"
+#include "MutableTools/Public/NodeModifierMeshClipWithMesh.h"
+#include "MutableTools/Public/NodeProjector.h"
+#include "MutableTools/Public/NodeScalarParameter.h"
+#include "MutableTools/Public/NodeScalarEnumParameter.h"
+#include "MutableRuntime/Public/Mesh.h"
+#include "MutableTools/Public/Table.h"
+
+struct FMutableGraphGenerationContext;
+
+
+struct FGeneratedImageProperties
+{
+	/** Name in the Material. */
+	FString TextureParameterName;
+
+	TEnumAsByte<TextureCompressionSettings> CompressionSettings = TC_Default;
+
+	TEnumAsByte<TextureFilter> Filter = TF_Bilinear;
+
+	uint32 SRGB = 0;
+
+	uint32 bFlipGreenChannel = 0;
+
+	int32 LODBias = 0;
+
+	TEnumAsByte<TextureMipGenSettings> MipGenSettings = TMGS_SimpleAverage;
+
+	int32 MaxTextureSize = 0;
+
+	TEnumAsByte<TextureGroup> LODGroup = TEnumAsByte<TextureGroup>(TMGS_FromTextureGroup);
+
+	TEnumAsByte<TextureAddress> AddressX = TA_Clamp;
+	TEnumAsByte<TextureAddress> AddressY = TA_Clamp;
+};
+
+
+// Flags that can influence the mesh conversion
+enum class EMutableMeshConversionFlags : uint32
+{
+	// 
+	None = 0,
+	// Ignore the skeleton and skinning
+	IgnoreSkinning = 1 << 0,
+
+	// Ignore Physics assets
+	IgnorePhysics = 1 << 1
+};
+
+
+// Key for the data stored for each processed unreal graph node.
+class FGeneratedKey
+{
+	friend uint32 GetTypeHash(const FGeneratedKey& Key);
+
+public:
+	FGeneratedKey(void* InFunctionAddress, const UEdGraphPin& InPin, const UCustomizableObjectNode& Node, FMutableGraphGenerationContext& GenerationContext, const bool UseMesh = false);
+	
+	bool operator==(const FGeneratedKey& Other) const;
+
+private:
+	/** Used to differentiate pins being cached from different functions (e.g. a PC_Color pin cached from GenerateMutableSourceImage and GenerateMutableSourceColor). */
+	void* FunctionAddress;
+	
+	const UEdGraphPin* Pin;
+	
+	int32 LOD;
+
+	/** Flag used to generate this mesh. Bit mask of EMutableMeshConversionFlags */
+	uint32 Flags = 0;
+
+	/** Active morphs at the time of mesh generation. */
+	TArray<UCustomizableObjectNodeMeshMorph*> MeshMorphStack;
+};
+
+
+uint32 GetTypeHash(const FGeneratedKey& Key);
+
+
+struct FGeneratedImageKey
+{
+	FGeneratedImageKey(const UEdGraphPin* InPin)
+	{
+		Pin = InPin;
+	}
+
+	bool operator==(const FGeneratedImageKey& Other) const
+	{
+		return Pin == Other.Pin;
+	}
+
+	const UEdGraphPin* Pin;
+};
+
+/** Structure used to store the results of the recursive GenerateMutableSourceSurface calls. */
+struct FMutableGraphSurfaceGenerationData
+{
+	/** Pointer to the NodeMaterial which a pin is connected (directly or indirectly). Used to find the real NodeMaterial on the NodeCopyMaterial case. 
+	 * The NodeMaterial may be behind a export/import nodes, or already generated
+	 */
+	const UCustomizableObjectNodeMaterial* NodeMaterial = nullptr;
+	
+	FGeneratedImageProperties ImageProperties;
+};
+
+// Structure storing results to propagate up when generating mutable mesh node expressions.
+struct FMutableGraphMeshGenerationData
+{
+	// Did we find any mesh with vertex colours in the expression?
+	bool bHasVertexColors = false;
+
+	// Maximum number of texture channels found in the expression.
+	int NumTexCoordChannels = 0;
+
+	// Maximum number of bones per vertex found in the expression.
+	int MaxNumBonesPerVertex = 0;
+
+	// Maximum size of the vertex bone index type in the expression.
+	int MaxBoneIndexTypeSizeBytes = 0;
+
+	int32 MaxNumTriangles = 0;
+	int32 MinNumTriangles = TNumericLimits<int32>::Max();
+
+	// Combine another generated data looking for the most general case.
+	void Combine(const FMutableGraphMeshGenerationData& other)
+	{
+		bHasVertexColors = bHasVertexColors || other.bHasVertexColors;
+		NumTexCoordChannels = FMath::Max(other.NumTexCoordChannels, NumTexCoordChannels);
+		MaxNumBonesPerVertex = FMath::Max(other.MaxNumBonesPerVertex, MaxNumBonesPerVertex);
+		MaxBoneIndexTypeSizeBytes = FMath::Max(other.MaxBoneIndexTypeSizeBytes, MaxBoneIndexTypeSizeBytes);
+		MaxNumTriangles = FMath::Max(other.MaxNumTriangles, MaxNumTriangles);
+		MinNumTriangles = FMath::Min(other.MinNumTriangles, MinNumTriangles);
+	}
+};
+
+
+// Data stored for each processed unreal graph node, stored in the cache.
+struct FGeneratedData
+{
+	FGeneratedData(const UEdGraphNode* InSource, mu::NodePtr InNode,
+		const FMutableGraphMeshGenerationData* InMeshData = nullptr)
+	{
+		Source = InSource;
+		Node = InNode;
+		if (InMeshData)
+		{
+			meshData = *InMeshData;
+		}
+	}
+
+	const UEdGraphNode* Source;
+	mu::NodePtr Node;
+
+	// Used for mesh nodes only
+	FMutableGraphMeshGenerationData meshData;
+};
+
+
+inline uint32 GetTypeHash(const FGeneratedImageKey& Key)
+{
+	uint32 GuidHash = GetTypeHash(Key.Pin->PinId);
+
+	return GuidHash;
+}
+
+
+/** Struct describing the conversion task for each Unreal texture that needs to be converted to Mutable format */
+struct FTextureUnrealToMutableTask
+{
+	/** Default constructor */
+	FTextureUnrealToMutableTask(mu::NodeImageConstantPtr ImageNodeParameter, UTexture2D* TextureParameter, const UCustomizableObjectNode* NodeParameter, bool bIsNormalCompositeParameter = false)
+		: ImageNode(ImageNodeParameter)
+		, Texture(TextureParameter)
+		, Node(NodeParameter)
+		, bIsNormalComposite(bIsNormalCompositeParameter)
+	{}
+
+	/** Table Node constructor */
+	FTextureUnrealToMutableTask(mu::TablePtr TableParameter, UTexture2D* TextureParameter, const UCustomizableObjectNode* NodeParameter, int32 ColumnIndex, int32 RowIndex, bool bIsNormalCompositeParameter = false)
+		: TableNode(TableParameter)
+		, Texture(TextureParameter)
+		, Node(NodeParameter)
+		, TableColumn(ColumnIndex)
+		, TableRow(RowIndex)
+		, bIsNormalComposite(bIsNormalCompositeParameter)
+	{}
+
+	mu::NodeImageConstantPtr ImageNode;
+	mu::TablePtr TableNode;
+	UTexture2D* Texture;
+	const UCustomizableObjectNode* Node;
+	int32 TableColumn;
+	int32 TableRow;
+	bool bIsNormalComposite;	
+};
+
+
+struct FPoseBoneData
+{
+	TArray<FString> ArrayBoneName;
+	TArray<FTransform> ArrayTransform;
+};
+
+
+struct FGroupProjectorTempData
+{
+	class UCustomizableObjectNodeGroupProjectorParameter* CustomizableObjectNodeGroupProjectorParameter;
+	mu::NodeProjectorParameterPtr NodeProjectorParameterPtr;
+	mu::NodeImagePtr NodeImagePtr;
+	mu::NodeRangePtr NodeRange;
+	mu::NodeScalarParameterPtr NodeOpacityParameter;
+
+	mu::NodeScalarEnumParameterPtr PoseOptionsParameter;
+	TArray<FPoseBoneData> PoseBoneDataArray;
+
+	bool bAlternateResStateNameWarningDisplayed = false; // Used to display this warning only once
+};
+
+
+struct FGroupNodeIdsTempData
+{
+	FGroupNodeIdsTempData(FGuid OldGuid, FGuid NewGuid = FGuid()) :
+		OldGroupNodeId(OldGuid),
+		NewGroupNodeId(NewGuid)
+	{
+
+	}
+
+	FGuid OldGroupNodeId;
+	FGuid NewGroupNodeId;
+
+	bool operator==(const FGroupNodeIdsTempData& Other) const
+	{
+		return OldGroupNodeId == Other.OldGroupNodeId;
+	}
+};
+
+struct FGroupProjectorImageInfo
+{
+	mu::NodeImagePtr ImageNode;
+	mu::NodeImagePtr ImageResizeNode;
+	mu::NodeSurfaceNewPtr SurfNode;
+	UCustomizableObjectNodeMaterial* TypedNodeMat;
+	FString TextureName;
+	FString RealTextureName;
+	FString AlternateResStateName;
+	float AlternateProjectionResolutionFactor;
+	bool bIsAlternateResolutionResized = false;
+	int32 UVLayout = 0;
+
+	FGroupProjectorImageInfo(mu::NodeImagePtr InImageNode, const FString& InTextureName, const FString& InRealTextureName, UCustomizableObjectNodeMaterial* InTypedNodeMat,
+		float InAlternateProjectionResolutionFactor, const FString& InAlternateResStateName, mu::NodeSurfaceNewPtr InSurfNode, int32 InUVLayout)
+		: TypedNodeMat(InTypedNodeMat), TextureName(InTextureName), RealTextureName(InRealTextureName),
+		AlternateResStateName(InAlternateResStateName), AlternateProjectionResolutionFactor(InAlternateProjectionResolutionFactor), 
+		UVLayout(InUVLayout)
+	{
+		ImageNode = InImageNode;
+		SurfNode = InSurfNode;
+	}
+
+	static FString GenerateId(const UCustomizableObjectNodeMaterialBase* TypedNodeMat, int32 ImageIndex)
+	{
+		return TypedNodeMat->GetOutermost()->GetPathName() + TypedNodeMat->NodeGuid.ToString() + FString("-") + FString::FromInt(ImageIndex);
+	}
+};
+
+
+/** Struct that defines a mesh (Mesh + LOD + MaterialIndex) */
+struct FMeshData
+{
+	/** Mesh which it may contain multiple LODs and Materials. */
+	const UObject* Mesh;
+
+	/** Specific LOD of the mesh. */
+	int LOD;
+
+	/** Specific index of the mesh material (MaterialIndex is an alias of Section). */
+	int MaterialIndex;
+
+	/** Node where the mesh is defined. Not a UCustomizableObjectNodeMesh due to Table nodes */
+	const UCustomizableObjectNode* Node;
+
+	bool operator==(const FMeshData& Other) const
+	{
+		return Mesh == Other.Mesh &&
+			LOD == Other.LOD &&
+			MaterialIndex == Other.MaterialIndex &&
+			Node == Other.Node;
+	}
+};
+
+inline uint32 GetTypeHash(const FMeshData& Key)
+{
+	uint32 MeshHash = GetTypeHash(Key.Mesh);
+	uint32 NodeHash = GetTypeHash(Key.Node->GetUniqueID());
+	
+	return HashCombine(HashCombine(HashCombine(MeshHash, Key.LOD), Key.MaterialIndex), NodeHash);
+}
+
+
+/** Struct used to store info specific to each component during compilation */
+struct FMutableComponentInfo
+{
+	FMutableComponentInfo(USkeletalMesh* InRefSkeletalMesh)
+	{
+		if(!InRefSkeletalMesh || !InRefSkeletalMesh->GetSkeleton())
+		{
+			return;
+		}
+		
+		RefSkeletalMesh = InRefSkeletalMesh;
+		RefSkeleton = RefSkeletalMesh->GetSkeleton();
+			
+		const int32 NumBones = RefSkeleton->GetReferenceSkeleton().GetRawBoneNum();
+		BoneNamesToPathHash.Reserve(NumBones);
+
+		const TArray<FMeshBoneInfo>& Bones = RefSkeleton->GetReferenceSkeleton().GetRawRefBoneInfo();
+
+		for (int32 BoneIndex = 0; BoneIndex < NumBones; ++BoneIndex)
+		{
+			const FMeshBoneInfo& Bone = Bones[BoneIndex];
+
+			// Retrieve parent bone name and respective hash, root-bone is assumed to have a parent hash of 0
+			const FName ParentName = Bone.ParentIndex != INDEX_NONE ? Bones[Bone.ParentIndex].Name : NAME_None;
+			const uint32 ParentHash = Bone.ParentIndex != INDEX_NONE ? GetTypeHash(ParentName) : 0;
+
+			// Look-up the path-hash from root to the parent bone
+			const uint32* ParentPath = BoneNamesToPathHash.Find(ParentName);
+			const uint32 ParentPathHash = ParentPath ? *ParentPath : 0;
+
+			// Append parent hash to path to give full path hash to current bone
+			const uint32 BonePathHash = HashCombine(ParentPathHash, ParentHash);
+
+			// Add path hash to current bone
+			BoneNamesToPathHash.Add(Bone.Name, BonePathHash);
+		}
+	}
+
+	// Each component must have a reference SkeletalMesh with a valid Skeleton
+	USkeletalMesh* RefSkeletalMesh = nullptr;
+	USkeleton* RefSkeleton = nullptr;
+
+	// Map to check skeleton compatibility
+	TMap<USkeleton*, bool> SkeletonCompatibility;
+	
+	// Hierarchy hash from parent-bone to root bone, used to check if additional skeletons are compatible with
+	// the RefSkeleton
+	TMap<FName, uint32> BoneNamesToPathHash;
+};
+
+
+/** Struct of data which is behind a given output pin. 
+ *
+ * Each newly added field has should also be defined on the Append() function body.
+ */
+struct FPinDataValue
+{
+	/** Set of all meshes behind a given output pin. */
+	TSet<FMeshData> MeshesData;
+
+	/** Add all data from other PinData. */
+	void Append(const FPinDataValue& From);
+};
+
+/** Eases the managment and access of data behind a given pin.
+ * 
+ * We understand as "behind a pin" as all pins which are connected to a given pin directly and indirectly. 
+ * 
+ * Given the the following pins:
+ * C -> B -> A
+ *      D ---^
+ * 
+ * FPinData will propagate the added data in the following way:
+ * C (data added on C) -> B (data added on C, B) -> A (data added on C, B, A, D)
+ *                           D (data added on D) ---^
+ * 
+ * FPinDataValue defines which data is going to be saved behind pins.
+ *
+ * When traversing the graph:
+ * - Each time an output pin is explored recursively, call the Push() function.
+ * - When exiting the recursion, call the Pop() function.
+ * - Add data to be seen behind all the stacked pins using the GetCurrent() function.
+ * - To avoid missing any returning branch, use the SCOPE_PIN_DATA wrapper.
+ * 
+ * The all data from a behind a pin can be queried using the Find() method.
+ */
+class FPinData // DEPRECATED. DO NOT USE!
+{
+public:
+	/** Get the PinData from a given pin. */
+	FPinDataValue* Find(const UEdGraphPin* Pin);
+
+	/** Get the PinData of the current pin. The top pin in the stack. */
+	FPinDataValue& GetCurrent();
+
+	/** Pop the current pin from the stack and append all its PinData with to previous pin in the stack. */
+	void Pop();
+
+	/** Push a pin to the stack and create its PinData. */
+	void Push(const UEdGraphPin* Pin);
+
+private:
+	/** Data for each pin. */
+	TMap<const UEdGraphPin*, FPinDataValue> Data;
+
+	/** Stack of pins. Used to know which pin the data belongs to during the graph traversal recursion. */
+	TArray<const UEdGraphPin*> PinStack;
+};
+
+/** Graph cycle key.
+ *
+ * Pin is not enough since we can call multiple recursive functions with the same pin.
+ * Each function has to have an unique identifier.
+ */
+struct FGraphCycleKey
+{
+	friend uint32 GetTypeHash(const FGraphCycleKey& Key);
+
+	FGraphCycleKey(const UEdGraphPin& Pin, uint32 Id);
+
+	bool operator==(const FGraphCycleKey& Other) const;
+	
+	/** Valid pin. */
+	const UEdGraphPin& Pin;
+
+	/** Unique id. */
+	uint32 Id;
+};
+
+/** Graph Cycle scope.
+ *
+ * Detect a cycle during the graph traversal.
+ */
+class FGraphCycle
+{
+public:
+	explicit FGraphCycle(const FGraphCycleKey&& Key, FMutableGraphGenerationContext &Context);
+	~FGraphCycle();
+
+	/** Return true if there is a cycle. */
+	bool FoundCycle() const;
+	
+private:
+	/** Graph traversal key. */
+	FGraphCycleKey Key;
+
+	/** Generation context. */
+	FMutableGraphGenerationContext& Context;
+};
+
+/** Return the default value if there is a cycle. */
+#define RETURN_ON_CYCLE(Pin, GenerationContext) \
+	FGraphCycle GraphCycle(FGraphCycleKey(Pin, __COUNTER__), GenerationContext); \
+	if (GraphCycle.FoundCycle()) \
+	{ \
+		return {}; \
+	} \
+
+struct FMutableGraphGenerationContext
+{
+	FMutableGraphGenerationContext(UCustomizableObject* CustomizableObject, class FCustomizableObjectCompiler* InCompiler, const FCompilationOptions& InOptions);
+	UCustomizableObject* Object = nullptr;
+
+	// Non-owned reference to the compiler object
+	FCustomizableObjectCompiler* Compiler = nullptr;
+
+	// Compilation options, including target platform
+	FCompilationOptions Options;
+
+	// Cache of generated pins per LOD
+	TMap<FGeneratedKey, FGeneratedData> Generated;
+
+	/** Set of all generated nodes. */
+	TSet<UCustomizableObjectNode*> GeneratedNodes;
+
+	// Cache of generated Node Tables
+	TMap<FString, mu::TablePtr> GeneratedTables;
+
+	// Cache of generated images, because sometimes they are reused by LOD, we use this as a second
+	// level cache
+	TMap<FGeneratedImageKey, mu::NodeImageConstantPtr> GeneratedImages;
+
+    // Global morph selection overrides.
+    TArray<FRealTimeMorphSelectionOverride> RealTimeMorphTargetsOverrides;
+
+	// Mutable meshes already build for source UStaticMesh or USkeletalMesh.
+	struct FGeneratedMeshData
+	{
+		struct FKey
+		{
+			/** Source mesh data. */
+			const UObject* Mesh = nullptr;
+			int LOD = 0;
+			int MaterialIndex = 0;
+
+			/** Flag used to generate this mesh. Bit mask of EMutableMeshConversionFlags */
+			uint32 Flags = 0;
+
+			bool operator==( const FKey& k ) const
+			{
+				return Mesh == k.Mesh && LOD == k.LOD && MaterialIndex == k.MaterialIndex 
+					&& Flags == k.Flags;
+			}
+		};
+
+		FKey Key;
+
+		/** Generated mesh. */
+		mu::MeshPtr Generated;
+	};
+	TArray<FGeneratedMeshData> GeneratedMeshes;
+
+	// Stack of mesh generation flags. The last one is the currently valid.
+	// The value is a bit mask of EMutableMeshConversionFlags
+	TArray<uint32> MeshGenerationFlags;
+
+	/** Find a mesh if already generated for a given source and flags. */
+	mu::MeshPtr FindGeneratedMesh(const FGeneratedMeshData::FKey& Key);
+
+	/** Adds to ParameterNamesMap the node Node to the array of elements with name Name */
+	void AddParameterNameUnique(const UCustomizableObjectNode* Node, FString Name);
+
+	// Check if the Id of the node Node already exists, if it's new adds it to NodeIds array, otherwise, returns new Id
+	const FGuid GetNodeIdUnique(const UCustomizableObjectNode* Node);
+
+	/** Generates new tags for the UCustomizableObjectNodeMeshClipWithMesh nodes which have assigned a CO in the
+	* UCustomizableObjectNodeMeshClipWithMesh::CustomizableObjectToClipWith field
+	* @return nothing */
+	void GenerateClippingCOInternalTags();
+
+	/** Check if the PhysicsAsset of a given SkeletalMesh has any SkeletalBodySetup with BoneNames not present in the
+	* InSkeletalMesh's RefSkeleton, if so, adds the PhysicsAsset to the DiscartedPhysicsAssetMap to display a warning later on */
+	void CheckPhysicsAssetInSkeletalMesh(const USkeletalMesh* InSkeletalMesh);
+
+	/** Get the reference skeletal mesh associated to the current mesh component being generated */
+	FMutableComponentInfo& GetCurrentComponentInfo();
+
+
+	TArray<FMutableComponentInfo> ComponentInfos;
+	TArray<FMutableRefSkeletalMeshData> ReferenceSkeletalMeshesData;
+	
+	TArray<UMaterialInterface*> ReferencedMaterials;
+	TArray<FName> ReferencedMaterialSlotNames;
+	TArray<FGeneratedImageProperties> ImageProperties;
+	TMap<FString, TArray<const UCustomizableObjectNode*>> ParameterNamesMap;
+	TArray<const UCustomizableObjectNode*> NoNameNodeObjectArray;
+	TMap<FString, FCustomizableObjectIdPair> GroupNodeMap;
+	TMap<FString, FString> CustomizableObjectPathMap;
+	TMap<FString, FParameterUIData> ParameterUIDataMap;
+	TMap<FString, FParameterUIData> StateUIDataMap;
+	TMultiMap<const UCustomizableObjectNodeObjectGroup*, FGroupProjectorTempData> ProjectorGroupMap;
+	TMap<UPhysicsAsset*, uint32> DiscartedPhysicsAssetMap;
+
+	TArray<const USkeleton*> ReferencedSkeletons;
+
+	// Used to aviod Nodes with duplicated ids
+	TMap<FGuid, TArray<const UCustomizableObjectNode*>> NodeIdsMap;
+	TMultiMap<const UCustomizableObject*, FGroupNodeIdsTempData> DuplicatedGroupNodeIds;
+
+	// For a given material node (the key is node package path + node uid + image index in node) stores images generated for the same node at a higher quality LOD to reuse that image node
+	TMap<FString, FGroupProjectorImageInfo> GroupProjectorLODCache;
+
+	// Data structures used for clipping feature by assigning a Customizable Object
+	// Map with pairs (Unreal Mutable material node Guid, array with its corresponding Mutable surface node) to add tags from clipping nodes
+	TMap<class UCustomizableObjectNodeMaterial*, TArray<mu::NodeSurfaceNewPtr>> MapMaterialNodeToMutableSurfaceNodeArray;
+
+	// Map with pairs (Unreal Mutable clip mesh node, array with its corresponding Mutable mesh modifier nodes) to add tags
+	TMap<class UCustomizableObjectNodeMeshClipWithMesh*, TArray<mu::NodeModifierMeshClipWithMeshPtr>> MapClipMeshNodeToMutableClipMeshNodeArray;
+
+	// Data used for MorphTarget reconstruction.
+	TArray<FMorphTargetInfo> ContributingMorphTargetsInfo;
+	TArray<FMorphTargetVertexData> MorphTargetReconstructionData;
+
+	// Data used for Clothing reconstruction.
+	TArray<FCustomizableObjectMeshToMeshVertData> ClothMeshToMeshVertData;
+	TArray<FCustomizableObjectClothingAssetData> ContributingClothingAssetsData;
+
+	// Hierarchy of current ComponentNew nodes, each stored for every LOD
+	struct ObjectParent
+	{
+		TArray<TArray<mu::NodeComponentNewPtr>> ComponentsPerLOD;
+	};
+	TArray< ObjectParent > ComponentNewNode;
+
+	int32 CurrentLOD = 0;
+	int32 NumLODsInRoot = 0;
+	int32 CurrentMeshComponent = 0;
+	int32 NumMeshComponentsInRoot = 0;
+	int32 CurrentTextureLODBias = 0;
+
+	int32 FirstLODAvailable = MAX_MESH_LOD_COUNT;
+	int32 NumMaxLODsToStream = MAX_MESH_LOD_COUNT;
+
+	bool bEnableLODStreaming = true;
+
+	// Based on the last object visited.
+	ECustomizableObjectAutomaticLODStrategy CurrentAutoLODStrategy = ECustomizableObjectAutomaticLODStrategy::Manual;
+
+	// Stores external graph root nodes to be added to the specified group nodes
+	TMultiMap<FGuid, UCustomizableObjectNodeObject*> GroupIdToExternalNodeMap;
+
+	// Easily retrieve a parameter name from its node guid
+	TMap<FGuid, FString> GuidToParamNameMap;
+
+	// Graph cycle detection
+	/** Visited nodes during the DAC recursion traversal.
+	 * It acts like stack, pushing pins when recursively exploring a new pin an popping it when exiting the recursion. */
+	TMap<FGraphCycleKey, const UCustomizableObject*> VisitedPins;
+	const UCustomizableObject* CustomizableObjectWithCycle = nullptr;
+
+	// Texture management flag, captured from the real object root node
+	bool bDisableTextureLayoutManagementFlag = false;
+
+	// Set of all the guids of all the child CustomizableObjects in the compilation
+	TSet<FGuid> CustomizableObjectGuidsInCompilation;
+
+	/** Array with the conversion tasks for each Unreal texture that needs to be converted to Mutable format for a particular Customizable Object.
+	* Multiple tasks for the same texture may be added here when referenced from different LODs or nodes, however when precessing the array it'll
+	* make sure to avoid converting the same texture more than once.
+	*/
+	TArray<FTextureUnrealToMutableTask> ArrayTextureUnrealToMutableTask;
+
+	/** Stores the physics assets gathered from the SkeletalMesh nodes during compilation, to be used in mesh generation in-game */
+	TMap<FString, TSoftObjectPtr<UPhysicsAsset>> PhysicsAssetMap;
+
+	/** Stores the anim BP assets gathered from the SkeletalMesh nodes during compilation, to be used in mesh generation in-game */
+	TMap<FString, TSoftClassPtr<UAnimInstance>> AnimBPAssetsMap;
+
+	// Stores the textures that will be used to mask-out areas in the projection. The cache isn't used for rendering, but for coverage testing
+	TMap<FString, FString> MaskOutMaterialCache; // Maps a UMaterial's asset path to a UTexture's asset path
+	TMap<FString, FMaskOutTexture> MaskOutTextureCache; // Maps a UTexture's asset path to the cached mask-out texture data
+
+	// Stores the only option of an Int Param that should be compiled in a partial compilation
+	TMap<FString, FString> ParamNamesToSelectedOptions;
+
+	TArray<const UEdGraphNode*> LimitedParameters;
+	int32 ParameterLimitationCount = 0;
+
+	/** Data which is behind a given output pin. 
+	 *
+	 * We only consider output pins. Do not use input pins as keys. See FPinData class definition. 
+	 */
+	FPinData PinData; // DEPRECATED. DO NOT USE!
+
+	// Stores all morphs to apply them directly to a skeletal mesh node
+	TArray<UCustomizableObjectNodeMeshMorph*> MeshMorphStack;
+
+	// Current material parameter name to find the corresponding column in a mutable table
+	FString CurrentMaterialTableParameter;
+
+	// Current material parameter id to find the corresponding column in a mutable table
+	FString CurrentMaterialTableParameterId;
+
+	// Stores the parameters generated in the node tables
+	TMap<const class UCustomizableObjectNodeTable*, TArray<FGuid>> GeneratedParametersInTables;
+};
+
+/** Pin Data scope wrapper. Pops the pin data on scope exit. */
+class FScopedPinData
+{
+public:
+	explicit FScopedPinData(FMutableGraphGenerationContext& Context, const UEdGraphPin* Pin);
+	~FScopedPinData();
+
+private:
+	FMutableGraphGenerationContext& Context;
+};
+
+#define SCOPED_PIN_DATA(Context, Pin) \
+	FScopedPinData ScopedPinData(Context, Pin);
+
+
+//
+mu::NodeObjectPtr GenerateMutableSource(const class UEdGraphPin* Pin, FMutableGraphGenerationContext& GenerationContext, bool bPartialCompilation);
+
+
+/** Populate an array with all the information related to the reference skeletal meshes we might need in-game to generate instances */
+void PopulateReferenceSkeletalMeshesData(FMutableGraphGenerationContext& GenerationContext);
+
+
+void CheckNumOutputs(const UEdGraphPin& Pin, const FMutableGraphGenerationContext& GenerationContext);
+
+
+// TODO GMT Remove generation context dependency and move to GraphTraversal.
+UTexture2D* FindReferenceImage(const UEdGraphPin* Pin, FMutableGraphGenerationContext& GenerationContext);
+
+
+mu::NodeMeshApplyPosePtr CreateNodeMeshApplyPose(mu::NodeMeshPtr InputMeshNode, UCustomizableObject * CustomizableObject, TArray<FString> ArrayBoneName, TArray<FTransform> ArrayTransform);
+
+
+// Generates the tag for an animation instance
+FString GenerateAnimationInstanceTag(const FString& AnimInstance, int32 SlotIndex);
+
+
+FString GenerateGameplayTag(const FString& GameplayTag);
+
+// Computes the LOD bias for a texture given the current mesh LOD and automatic LOD settings, the reference texture settings
+// and whether it's being built for a server or not
+int32 ComputeLODBias(const FMutableGraphGenerationContext& GenerationContext, const UTexture2D* ReferenceTexture, int32 MaxTextureSize,
+	const UCustomizableObjectNodeMaterial* MaterialNode, const int32 ImageIndex);
+
+
+int32 GetMaxTextureSize(const UTexture2D* ReferenceTexture, const FMutableGraphGenerationContext& GenerationContext);
diff --git a/Engine/Plugins/Experimental/Mutable/Source/CustomizableObjectEditor/Private/GenerateMutableSource/GenerateMutableSourceColor.cpp b/Engine/Plugins/Experimental/Mutable/Source/CustomizableObjectEditor/Private/GenerateMutableSource/GenerateMutableSourceColor.cpp
new file mode 100644
--- /dev/null
+++ b/Engine/Plugins/Experimental/Mutable/Source/CustomizableObjectEditor/Private/GenerateMutableSource/GenerateMutableSourceColor.cpp
@@ -0,0 +1,338 @@
+// Copyright Epic Games, Inc. All Rights Reserved.
+
+#include "GenerateMutableSourceColor.h"
+
+#include "GenerateMutableSource/GenerateMutableSourceFloat.h"
+#include "GenerateMutableSource/GenerateMutableSourceImage.h"
+#include "GenerateMutableSource/GenerateMutableSourceTable.h"
+
+#include "Templates/UnrealTemplate.h"
+#include "GenericPlatform/GenericPlatformMath.h"
+#include "Engine/AssetManager.h"
+#include "Engine/TextureLODSettings.h"
+#include "ClothingAsset.h"
+#include "MeshUtilities.h"
+#include "EdGraphToken.h"
+#include "AssetRegistry/AssetRegistryModule.h"
+#include "UnrealEditorPortabilityHelpers.h"
+#include "PlatformInfo.h"
+
+#include "CustomizableObjectCompiler.h"
+#include "UnrealToMutableTextureConversionUtils.h"
+#include "UnrealConversionUtils.h"
+#include "CustomizableObjectEditorModule.h"
+#include "CustomizableObject.h"
+#include "GraphTraversal.h"
+
+#include "MutableTools/Public/Node.h"
+#include "MutableTools/Public/NodeScalarConstant.h"
+#include "MutableTools/Public/NodeScalarEnumParameter.h"
+#include "MutableTools/Public/NodeColourConstant.h"
+#include "MutableTools/Public/NodeColourParameter.h"
+#include "MutableTools/Public/NodeColourSwitch.h"
+#include "MutableTools/Public/NodeColourSampleImage.h"
+#include "MutableTools/Public/NodeColourArithmeticOperation.h"
+#include "MutableTools/Public/NodeColourFromScalars.h"
+#include "MutableTools/Public/NodeColourVariation.h"
+#include "MutableTools/Public/NodeColourTable.h"
+#include "MutableTools/Public/NodeImageBinarise.h"
+#include "MutableTools/Public/NodeSurfaceNew.h"
+#include "MutableTools/Public/Compiler.h"
+#include "MutableTools/Public/Streams.h"
+
+
+#define LOCTEXT_NAMESPACE "CustomizableObjectEditor"
+
+
+mu::NodeColourPtr GenerateMutableSourceColor(const UEdGraphPin* Pin, FMutableGraphGenerationContext& GenerationContext)
+{
+	check(Pin)
+	RETURN_ON_CYCLE(*Pin, GenerationContext)
+
+	CheckNumOutputs(*Pin, GenerationContext);
+	
+	const UEdGraphSchema_CustomizableObject* Schema = GetDefault<UEdGraphSchema_CustomizableObject>();
+
+	UCustomizableObjectNode* Node = CastChecked<UCustomizableObjectNode>(Pin->GetOwningNode());
+
+	const FGeneratedKey Key(reinterpret_cast<void*>(&GenerateMutableSourceColor), *Pin, *Node, GenerationContext);
+	if (const FGeneratedData* Generated = GenerationContext.Generated.Find(Key))
+	{
+		return static_cast<mu::NodeColour*>(Generated->Node.get());
+	}
+
+	mu::NodeColourPtr Result;
+
+	if (const UCustomizableObjectNodeColorConstant* TypedNodeColorConst = Cast<UCustomizableObjectNodeColorConstant>(Node))
+	{
+		mu::NodeColourConstantPtr ColorNode = new mu::NodeColourConstant();
+		Result = ColorNode;
+
+		ColorNode->SetValue(TypedNodeColorConst->Value.R, TypedNodeColorConst->Value.G, TypedNodeColorConst->Value.B);
+	}
+
+	else if (const UCustomizableObjectNodeColorParameter* TypedNodeColorParam = Cast<UCustomizableObjectNodeColorParameter>(Node))
+	{
+		mu::NodeColourParameterPtr ColorNode = new mu::NodeColourParameter();
+		Result = ColorNode;
+
+		GenerationContext.AddParameterNameUnique(Node, TypedNodeColorParam->ParameterName);
+
+		ColorNode->SetName(TCHAR_TO_ANSI(*TypedNodeColorParam->ParameterName));
+		ColorNode->SetUid(TCHAR_TO_ANSI(*GenerationContext.GetNodeIdUnique(Node).ToString()));
+		ColorNode->SetDefaultValue(TypedNodeColorParam->DefaultValue.R, TypedNodeColorParam->DefaultValue.G, TypedNodeColorParam->DefaultValue.B);
+
+		GenerationContext.ParameterUIDataMap.Add(TypedNodeColorParam->ParameterName, FParameterUIData(
+			TypedNodeColorParam->ParameterName,
+			TypedNodeColorParam->ParamUIMetadata,
+			EMutableParameterType::Color));
+	}
+
+	else if (const UCustomizableObjectNodeColorSwitch* TypedNodeColorSwitch = Cast<UCustomizableObjectNodeColorSwitch>(Node))
+	{
+		Result = [&]()
+		{
+			if (const int32 NumParameters = FollowInputPinArray(*TypedNodeColorSwitch->SwitchParameter()).Num();
+				NumParameters != 1)
+			{
+				const FText Message = NumParameters
+					? LOCTEXT("NoEnumParamInSwitch", "Switch nodes must have an enum switch parameter. Please connect an enum and refesh the switch node.")
+					: LOCTEXT("InvalidEnumInSwitch", "Switch nodes must have a single enum with all the options inside. Please remove all the enums but one and refresh the switch node.");
+
+				GenerationContext.Compiler->CompilerLog(Message, Node);
+				return Result;
+			}
+
+			const UEdGraphPin* EnumPin = FollowInputPin(*TypedNodeColorSwitch->SwitchParameter());
+			mu::NodeScalarPtr SwitchParam = EnumPin ? GenerateMutableSourceFloat(EnumPin, GenerationContext) : nullptr;
+
+			// Switch Param not generated
+			if (!SwitchParam)
+			{
+				// Warn about a failure.
+				if (EnumPin)
+				{
+					const FText Message = LOCTEXT("FailedToGenerateSwitchParam", "Could not generate switch enum parameter. Please refesh the switch node and connect an enum.");
+					GenerationContext.Compiler->CompilerLog(Message, Node);
+				}
+
+				return Result;
+			}
+
+			if (SwitchParam->GetType() != mu::NodeScalarEnumParameter::GetStaticType())
+			{
+				const FText Message = LOCTEXT("WrongSwitchParamType", "Switch parameter of incorrect type.");
+				GenerationContext.Compiler->CompilerLog(Message, Node);
+
+				return Result;
+			}
+
+			const int32 NumSwitchOptions = TypedNodeColorSwitch->GetNumElements();
+
+			mu::NodeScalarEnumParameter* EnumParameter = static_cast<mu::NodeScalarEnumParameter*>(SwitchParam.get());
+			if (NumSwitchOptions != EnumParameter->GetValueCount())
+			{
+				const FText Message = LOCTEXT("MismatchedSwitch", "Switch enum and switch node have different number of options. Please refresh the switch node to make sure the outcomes are labeled properly.");
+				GenerationContext.Compiler->CompilerLog(Message, Node);
+			}
+
+			mu::NodeColourSwitchPtr SwitchNode = new mu::NodeColourSwitch;
+			SwitchNode->SetParameter(SwitchParam);
+			SwitchNode->SetOptionCount(NumSwitchOptions);
+
+			for (int SelectorIndex = 0; SelectorIndex < NumSwitchOptions; ++SelectorIndex)
+			{
+				const UEdGraphPin* const ColorPin = TypedNodeColorSwitch->GetElementPin(SelectorIndex);
+				if (const UEdGraphPin* ConnectedPin = FollowInputPin(*ColorPin))
+				{
+					SwitchNode->SetOption(SelectorIndex, GenerateMutableSourceColor(ConnectedPin, GenerationContext));
+				}
+			}
+
+			return static_cast<mu::NodeColourPtr>(SwitchNode);
+		}(); // invoke lambda;
+	}
+
+	else if (const UCustomizableObjectNodeTextureSample* TypedNodeTexSample = Cast<UCustomizableObjectNodeTextureSample>(Node))
+	{
+		mu::NodeColourSampleImagePtr ColorNode = new mu::NodeColourSampleImage();
+		Result = ColorNode;
+
+		if (const UEdGraphPin* ConnectedPin = FollowInputPin(*TypedNodeTexSample->TexturePin()))
+		{
+			mu::NodeImagePtr TextureNode = GenerateMutableSourceImage(ConnectedPin, GenerationContext, 0.f);
+			ColorNode->SetImage(TextureNode);
+		}
+
+		if (const UEdGraphPin* ConnectedPin = FollowInputPin(*TypedNodeTexSample->XPin()))
+		{
+			mu::NodeScalarPtr XNode = GenerateMutableSourceFloat(ConnectedPin, GenerationContext);
+			ColorNode->SetX(XNode);
+		}
+
+		if (const UEdGraphPin* ConnectedPin = FollowInputPin(*TypedNodeTexSample->YPin()))
+		{
+			mu::NodeScalarPtr YNode = GenerateMutableSourceFloat(ConnectedPin, GenerationContext);
+			ColorNode->SetY(YNode);
+		}
+	}
+
+	else if (const UCustomizableObjectNodeColorArithmeticOp* TypedNodeColorArith = Cast<UCustomizableObjectNodeColorArithmeticOp>(Node))
+	{
+		mu::NodeColourArithmeticOperationPtr OpNode = new mu::NodeColourArithmeticOperation();
+		Result = OpNode;
+
+		if (const UEdGraphPin* ConnectedPin = FollowInputPin(*TypedNodeColorArith->XPin()))
+		{
+			mu::NodeColourPtr XNode = GenerateMutableSourceColor(ConnectedPin, GenerationContext);
+			OpNode->SetA(XNode);
+		}
+
+		if (const UEdGraphPin* ConnectedPin = FollowInputPin(*TypedNodeColorArith->YPin()))
+		{
+			mu::NodeColourPtr YNode = GenerateMutableSourceColor(ConnectedPin, GenerationContext);
+			OpNode->SetB(YNode);
+		}
+
+		switch (TypedNodeColorArith->Operation)
+		{
+		case EColorArithmeticOperation::E_Add:
+			OpNode->SetOperation(mu::NodeColourArithmeticOperation::OPERATION::AO_ADD);
+			break;
+
+		case EColorArithmeticOperation::E_Sub:
+			OpNode->SetOperation(mu::NodeColourArithmeticOperation::OPERATION::AO_SUBTRACT);
+			break;
+
+		case EColorArithmeticOperation::E_Mul:
+			OpNode->SetOperation(mu::NodeColourArithmeticOperation::OPERATION::AO_MULTIPLY);
+			break;
+
+		case EColorArithmeticOperation::E_Div:
+			OpNode->SetOperation(mu::NodeColourArithmeticOperation::OPERATION::AO_DIVIDE);
+			break;
+		}
+	}
+
+	else if (const UCustomizableObjectNodeColorFromFloats* TypedNodeFrom = Cast<UCustomizableObjectNodeColorFromFloats>(Node))
+	{
+		mu::NodeColourFromScalarsPtr OpNode = new mu::NodeColourFromScalars();
+		Result = OpNode;
+
+		if (const UEdGraphPin* ConnectedPin = FollowInputPin(*TypedNodeFrom->RPin()))
+		{
+			mu::NodeScalarPtr FloatNode = GenerateMutableSourceFloat(ConnectedPin, GenerationContext);
+			OpNode->SetX(FloatNode);
+		}
+
+		if (const UEdGraphPin* ConnectedPin = FollowInputPin(*TypedNodeFrom->GPin()))
+		{
+			mu::NodeScalarPtr FloatNode = GenerateMutableSourceFloat(ConnectedPin, GenerationContext);
+			OpNode->SetY(FloatNode);
+		}
+
+		if (const UEdGraphPin* ConnectedPin = FollowInputPin(*TypedNodeFrom->BPin()))
+		{
+			mu::NodeScalarPtr FloatNode = GenerateMutableSourceFloat(ConnectedPin, GenerationContext);
+			OpNode->SetZ(FloatNode);
+		}
+
+		if (const UEdGraphPin* ConnectedPin = FollowInputPin(*TypedNodeFrom->APin()))
+		{
+			mu::NodeScalarPtr FloatNode = GenerateMutableSourceFloat(ConnectedPin, GenerationContext);
+			OpNode->SetW(FloatNode);
+		}
+	}
+
+	else if (const UCustomizableObjectNodeColorVariation* TypedNodeColorVar = Cast<const UCustomizableObjectNodeColorVariation>(Node))
+	{
+		mu::NodeColourVariationPtr ColorNode = new mu::NodeColourVariation();
+		Result = ColorNode;
+
+		if (const UEdGraphPin* ConnectedPin = FollowInputPin(*TypedNodeColorVar->DefaultPin()))
+		{
+			mu::NodeColourPtr ChildNode = GenerateMutableSourceColor(ConnectedPin, GenerationContext);
+			if (ChildNode)
+			{
+				ColorNode->SetDefaultColour(ChildNode.get());
+			}
+			else
+			{
+				GenerationContext.Compiler->CompilerLog(LOCTEXT("ColorFailed", "Color generation failed."), Node);
+			}
+		}
+		else
+		{
+			GenerationContext.Compiler->CompilerLog(LOCTEXT("ColorVarMissingDef", "Color variation node requires a default value."), Node);
+		}
+
+		ColorNode->SetVariationCount(TypedNodeColorVar->Variations.Num());
+		for (int VariationIndex = 0; VariationIndex < TypedNodeColorVar->Variations.Num(); ++VariationIndex)
+		{
+			const UEdGraphPin* VariationPin = TypedNodeColorVar->VariationPin(VariationIndex);
+			if (!VariationPin) continue;
+
+			ColorNode->SetVariationTag(VariationIndex, TCHAR_TO_ANSI(*TypedNodeColorVar->Variations[VariationIndex].Tag));
+			if (const UEdGraphPin* ConnectedPin = FollowInputPin(*VariationPin))
+			{
+				mu::NodeColourPtr ChildNode = GenerateMutableSourceColor(ConnectedPin, GenerationContext);
+				ColorNode->SetVariationColour(VariationIndex, ChildNode.get());
+			}
+		}
+	}
+
+	else if (const UCustomizableObjectNodeTable* TypedNodeTable = Cast<UCustomizableObjectNodeTable>(Node))
+	{
+		mu::NodeColourTablePtr ColorTableNode = new mu::NodeColourTable();
+		Result = ColorTableNode;
+
+		mu::TablePtr Table;
+		if (TypedNodeTable->Table)
+		{
+			FString ColumnName = Pin->PinFriendlyName.ToString();
+
+			if (Pin->PinType.PinCategory == Schema->PC_MaterialAsset)
+			{
+				ColumnName = GenerationContext.CurrentMaterialTableParameterId;
+			}
+
+			Table = GenerateMutableSourceTable(TypedNodeTable->Table->GetName(), Pin, GenerationContext);
+
+			ColorTableNode->SetTable(Table);
+			ColorTableNode->SetColumn(TCHAR_TO_ANSI(*ColumnName));
+			ColorTableNode->SetParameterName(TCHAR_TO_ANSI(*TypedNodeTable->ParameterName));
+
+			GenerationContext.AddParameterNameUnique(Node, TypedNodeTable->ParameterName);
+
+			if (Table->FindColumn(TCHAR_TO_ANSI(*ColumnName)) == -1)
+			{
+				FString Msg = FString::Printf(TEXT("Couldn't find pin column with name %s"), *ColumnName);
+				GenerationContext.Compiler->CompilerLog(FText::FromString(Msg), Node);
+			}
+
+		}
+		else
+		{
+			Table = new mu::Table();
+			ColorTableNode->SetTable(Table);
+
+			GenerationContext.Compiler->CompilerLog(LOCTEXT("ColorTableError", "Couldn't find the data table of the node."), Node);
+		}
+	}
+
+	else
+	{
+		GenerationContext.Compiler->CompilerLog(LOCTEXT("UnimplementedNode", "Node type not implemented yet."), Node);
+	}
+
+	FGeneratedData CacheData = FGeneratedData(Node, Result);
+	GenerationContext.Generated.Add(Key, CacheData);
+	GenerationContext.GeneratedNodes.Add(Node);
+	
+	if (Result)
+	{
+		Result->SetMessageContext(Node);
+	}
+
+	return Result;
+}
diff --git a/Engine/Plugins/Experimental/Mutable/Source/CustomizableObjectEditor/Private/GenerateMutableSource/GenerateMutableSourceColor.h b/Engine/Plugins/Experimental/Mutable/Source/CustomizableObjectEditor/Private/GenerateMutableSource/GenerateMutableSourceColor.h
new file mode 100644
--- /dev/null
+++ b/Engine/Plugins/Experimental/Mutable/Source/CustomizableObjectEditor/Private/GenerateMutableSource/GenerateMutableSourceColor.h
@@ -0,0 +1,10 @@
+// Copyright Epic Games, Inc. All Rights Reserved.
+
+#pragma once
+
+#include "GenerateMutableSource/GenerateMutableSource.h"
+
+#include "MutableTools/Public/NodeColour.h"
+
+/** Convert a CustomizableObject Source Graph into a mutable source graph. */
+mu::NodeColourPtr GenerateMutableSourceColor(const UEdGraphPin* Pin, FMutableGraphGenerationContext& GenerationContext);
