[mutable] Move Mutable plugin from the restricted folder to the experimental plugins folder.

diff --git a/Engine/Plugins/Experimental/Mutable/Source/CustomizableObjectEditor/Private/GenerateMutableSource/GenerateMutableSourceFloat.cpp b/Engine/Plugins/Experimental/Mutable/Source/CustomizableObjectEditor/Private/GenerateMutableSource/GenerateMutableSourceFloat.cpp
new file mode 100644
--- /dev/null
+++ b/Engine/Plugins/Experimental/Mutable/Source/CustomizableObjectEditor/Private/GenerateMutableSource/GenerateMutableSourceFloat.cpp
@@ -0,0 +1,391 @@
+// Copyright Epic Games, Inc. All Rights Reserved.
+
+#include "GenerateMutableSourceFloat.h"
+
+#include "GenerateMutableSource/GenerateMutableSourceFloat.h"
+#include "GenerateMutableSource/GenerateMutableSourceColor.h"
+#include "GenerateMutableSource/GenerateMutableSourceProjector.h"
+#include "GenerateMutableSource/GenerateMutableSourceSurface.h"
+#include "GenerateMutableSource/GenerateMutableSourceGroupProjector.h"
+#include "GenerateMutableSource/GenerateMutableSourceImage.h"
+#include "GenerateMutableSource/GenerateMutableSourceMesh.h"
+#include "GenerateMutableSource/GenerateMutableSourceLayout.h"
+#include "GenerateMutableSource/GenerateMutableSourceTable.h"
+
+#include "Templates/UnrealTemplate.h"
+#include "GenericPlatform/GenericPlatformMath.h"
+#include "Engine/AssetManager.h"
+#include "Engine/TextureLODSettings.h"
+#include "Curves/CurveLinearColor.h"
+#include "Curves/CurveVector.h"
+#include "Curves/CurveFloat.h"
+#include "ClothingAsset.h"
+#include "MeshUtilities.h"
+#include "EdGraphToken.h"
+#include "AssetRegistry/AssetRegistryModule.h"
+#include "UnrealEditorPortabilityHelpers.h"
+#include "PlatformInfo.h"
+
+#include "CustomizableObjectCompiler.h"
+#include "CustomizableObjectEditorModule.h"
+#include "CustomizableObject.h"
+#include "GraphTraversal.h"
+
+#include "MutableTools/Public/Node.h"
+#include "MutableTools/Public/NodeScalarConstant.h"
+#include "MutableTools/Public/NodeScalarParameter.h"
+#include "MutableTools/Public/NodeScalarSwitch.h"
+#include "MutableTools/Public/NodeScalarEnumParameter.h"
+#include "MutableTools/Public/NodeScalarCurve.h"
+#include "MutableTools/Public/NodeScalarVariation.h"
+#include "MutableTools/Public/NodeScalarTable.h"
+#include "MutableTools/Public/NodeImageBinarise.h"
+#include "MutableTools/Public/NodeSurfaceNew.h"
+#include "MutableTools/Public/Compiler.h"
+#include "MutableTools/Public/Streams.h"
+
+#include "MutableRuntime/Public/Skeleton.h"
+#include "MutableRuntime/Public/Image.h"
+
+#define LOCTEXT_NAMESPACE "CustomizableObjectEditor"
+
+
+void SetCurveNodeValues(mu::NodeScalarCurvePtr& CurveNode, const FRichCurve& Curve)
+{
+	if (Curve.HasAnyData())
+	{
+		CurveNode->SetDefaultValue(Curve.GetDefaultValue());
+		CurveNode->SetKeyFrameCount(Curve.GetNumKeys());
+		int32 KeyNum = 0;
+
+		for (auto Itr = Curve.GetKeyIterator(); Itr; ++Itr)
+		{
+			CurveNode->SetKeyFrame(KeyNum,
+				Itr->Time,
+				Itr->Value,
+				Itr->ArriveTangent,
+				Itr->ArriveTangentWeight,
+				Itr->LeaveTangent,
+				Itr->LeaveTangentWeight,
+				Itr->InterpMode,
+				Itr->TangentMode,
+				Itr->TangentWeightMode);
+
+			KeyNum++;
+		}
+	}
+	else
+	{
+		CurveNode->SetDefaultValue(0.0f); // FRichCurve default DefaultValue is MAX_flt
+	}
+}
+
+
+mu::NodeScalarPtr GenerateMutableSourceFloat(const UEdGraphPin* Pin, FMutableGraphGenerationContext& GenerationContext)
+{
+	check(Pin)
+	RETURN_ON_CYCLE(*Pin, GenerationContext)
+
+	CheckNumOutputs(*Pin, GenerationContext);
+	
+	const UEdGraphSchema_CustomizableObject* Schema = GetDefault<UEdGraphSchema_CustomizableObject>();
+
+	UCustomizableObjectNode* Node = CastChecked<UCustomizableObjectNode>(Pin->GetOwningNode());
+
+	const FGeneratedKey Key(reinterpret_cast<void*>(&GenerateMutableSourceFloat), *Pin, *Node, GenerationContext);
+	if (const FGeneratedData* Generated = GenerationContext.Generated.Find(Key))
+	{
+		return static_cast<mu::NodeScalar*>(Generated->Node.get());
+	}
+
+	mu::NodeScalarPtr Result;
+	
+	if (const UCustomizableObjectNodeFloatConstant* FloatConstantNode = Cast<UCustomizableObjectNodeFloatConstant>(Node))
+	{
+		mu::NodeScalarConstantPtr ScalarNode = new mu::NodeScalarConstant();
+		Result = ScalarNode;
+
+		ScalarNode->SetValue(FloatConstantNode->Value);
+	}
+
+	else if (const UCustomizableObjectNodeFloatParameter* FloatParameterNode = Cast<UCustomizableObjectNodeFloatParameter>(Node))
+	{
+		mu::NodeScalarParameterPtr ScalarNode = new mu::NodeScalarParameter();
+		Result = ScalarNode;
+
+		GenerationContext.AddParameterNameUnique(Node, FloatParameterNode->ParameterName);
+
+		ScalarNode->SetName(TCHAR_TO_ANSI(*FloatParameterNode->ParameterName));
+		ScalarNode->SetUid(TCHAR_TO_ANSI(*GenerationContext.GetNodeIdUnique(Node).ToString()));
+		ScalarNode->SetDefaultValue(FloatParameterNode->DefaultValue);
+
+		const int32 NumDescriptionImage = FloatParameterNode->GetNumDescriptionImage();
+		ScalarNode->SetDescImageCount(NumDescriptionImage);
+
+		for (int d = 0; d < NumDescriptionImage; ++d)
+		{
+			const UEdGraphPin* DescriptionImagePin = FloatParameterNode->GetDescriptionImagePin(d);
+			if (!DescriptionImagePin)
+			{
+				GenerationContext.Compiler->CompilerLog(LOCTEXT("WrongFloatParameterDescriptionImagePin", "Float parameter in inconsistent state, please refresh the node."), FloatParameterNode);
+			}
+			else if (const UEdGraphPin* ConnectedPin = FollowInputPin(*DescriptionImagePin))
+			{
+				mu::NodeImagePtr TextureNode = GenerateMutableSourceImage(ConnectedPin, GenerationContext, 0.f);
+				ScalarNode->SetDescImage(d, TextureNode);
+			}
+		}
+
+		GenerationContext.ParameterUIDataMap.Add(FloatParameterNode->ParameterName, FParameterUIData(
+			FloatParameterNode->ParameterName,
+			FloatParameterNode->ParamUIMetadata,
+			EMutableParameterType::Float));
+	}
+
+	else if (const UCustomizableObjectNodeEnumParameter* EnumParamNode = Cast<UCustomizableObjectNodeEnumParameter>(Node))
+	{
+		mu::NodeScalarEnumParameterPtr EnumParameterNode = new mu::NodeScalarEnumParameter;
+
+		const int32 NumSelectors = EnumParamNode->Values.Num();
+
+		int32 DefaultValue = FMath::Clamp(EnumParamNode->DefaultIndex, 0, NumSelectors - 1);
+
+		GenerationContext.AddParameterNameUnique(Node, EnumParamNode->ParameterName);
+
+		EnumParameterNode->SetName(TCHAR_TO_ANSI(*EnumParamNode->ParameterName));
+		EnumParameterNode->SetUid(TCHAR_TO_ANSI(*GenerationContext.GetNodeIdUnique(Node).ToString()));
+		EnumParameterNode->SetValueCount(NumSelectors);
+		EnumParameterNode->SetDefaultValueIndex(DefaultValue);
+
+		FParameterUIData ParameterUIData(EnumParamNode->ParameterName, EnumParamNode->ParamUIMetadata, EMutableParameterType::Int);
+		ParameterUIData.IntegerParameterGroupType = ECustomizableObjectGroupType::COGT_ONE;
+
+		for (int SelectorIndex = 0; SelectorIndex < NumSelectors; ++SelectorIndex)
+		{
+			EnumParameterNode->SetValue(SelectorIndex, (float)SelectorIndex, TCHAR_TO_ANSI(*EnumParamNode->Values[SelectorIndex].Name));
+
+			ParameterUIData.ArrayIntegerParameterOption.Add(FIntegerParameterUIData(
+				EnumParamNode->Values[SelectorIndex].Name,
+				EnumParamNode->Values[SelectorIndex].ParamUIMetadata));
+		}
+
+		Result = EnumParameterNode;
+
+		GenerationContext.ParameterUIDataMap.Add(EnumParamNode->ParameterName, ParameterUIData);
+	}
+
+	else if (const UCustomizableObjectNodeFloatSwitch* TypedNodeFloatSwitch = Cast<UCustomizableObjectNodeFloatSwitch>(Node))
+	{
+		// Using a lambda so control flow is easier to manage.
+		Result = [&]()
+		{
+			const UEdGraphPin* SwitchParameter = TypedNodeFloatSwitch->SwitchParameter();
+
+			// Check Switch Parameter arity preconditions.
+			if (const int32 NumParameters = FollowInputPinArray(*SwitchParameter).Num();
+				NumParameters != 1)
+			{
+				const FText Message = NumParameters
+					? LOCTEXT("NoEnumParamInSwitch", "Switch nodes must have an enum switch parameter. Please connect an enum and refesh the switch node.")
+					: LOCTEXT("InvalidEnumInSwitch", "Switch nodes must have a single enum with all the options inside. Please remove all the enums but one and refresh the switch node.");
+
+				GenerationContext.Compiler->CompilerLog(Message, Node);
+				return Result;
+			}
+
+			const UEdGraphPin* EnumPin = FollowInputPin(*SwitchParameter);
+			mu::NodeScalarPtr SwitchParam = GenerateMutableSourceFloat(EnumPin, GenerationContext);
+
+			// Switch Param not generated
+			if (!SwitchParam)
+			{
+				// Warn about a failure.
+				if (EnumPin)
+				{
+					const FText Message = LOCTEXT("FailedToGenerateSwitchParam", "Could not generate switch enum parameter. Please refesh the switch node and connect an enum.");
+					GenerationContext.Compiler->CompilerLog(Message, Node);
+				}
+
+				return Result;
+			}
+
+			if (SwitchParam->GetType() != mu::NodeScalarEnumParameter::GetStaticType())
+			{
+				const FText Message = LOCTEXT("WrongSwitchParamType", "Switch parameter of incorrect type.");
+				GenerationContext.Compiler->CompilerLog(Message, Node);
+
+				return Result;
+			}
+
+			const int32 NumSwitchOptions = TypedNodeFloatSwitch->GetNumElements();
+
+			mu::NodeScalarEnumParameter* EnumParameter = static_cast<mu::NodeScalarEnumParameter*>(SwitchParam.get());
+			if (NumSwitchOptions != EnumParameter->GetValueCount())
+			{
+				const FText Message = LOCTEXT("MismatchedSwitch", "Switch enum and switch node have different number of options. Please refresh the switch node to make sure the outcomes are labeled properly.");
+				GenerationContext.Compiler->CompilerLog(Message, Node);
+			}
+
+			mu::NodeScalarSwitchPtr SwitchNode = new mu::NodeScalarSwitch;
+			SwitchNode->SetParameter(SwitchParam);
+			SwitchNode->SetOptionCount(NumSwitchOptions);
+
+			for (int SelectorIndex = 0; SelectorIndex < NumSwitchOptions; ++SelectorIndex)
+			{
+				if (const UEdGraphPin* const FloatPin = TypedNodeFloatSwitch->GetElementPin(SelectorIndex))
+				{
+					if (const UEdGraphPin* ConnectedPin = FollowInputPin(*FloatPin))
+					{
+						SwitchNode->SetOption(SelectorIndex, GenerateMutableSourceFloat(ConnectedPin, GenerationContext));
+					}
+				}
+			}
+
+			Result = SwitchNode;
+			return Result;
+		}(); // invoke lambda.
+	}
+
+	else if (const UCustomizableObjectNodeCurve* TypedNodeCurve = Cast<UCustomizableObjectNodeCurve>(Node))
+	{
+		mu::NodeScalarCurvePtr CurveNode = new mu::NodeScalarCurve();
+		Result = CurveNode;
+
+		if (const UEdGraphPin* ConnectedPin = FollowInputPin(*TypedNodeCurve->InputPin()))
+		{
+			mu::NodeScalarPtr ScalarNode = GenerateMutableSourceFloat(ConnectedPin, GenerationContext);
+			CurveNode->SetT(ScalarNode);
+		}
+
+		if (UCurveBase* CurveAsset = TypedNodeCurve->CurveAsset)
+		{
+			int32 PinIndex = -1;
+
+			for (int32 i = 0; i < TypedNodeCurve->GetNumCurvePins(); ++i)
+			{
+				if (TypedNodeCurve->CurvePins(i) == Pin)
+				{
+					PinIndex = i;
+					break;
+				}
+			}
+
+			if (const UCurveLinearColor* const CurveColor = Cast<UCurveLinearColor>(CurveAsset))
+			{
+				if (PinIndex >= 0 && PinIndex <= 3)
+				{
+					const FRichCurve& Curve = CurveColor->FloatCurves[PinIndex];
+					SetCurveNodeValues(CurveNode, Curve);
+				}
+			}
+			else if (const UCurveVector* const CurveVector = Cast<UCurveVector>(CurveAsset))
+			{
+				if (PinIndex >= 0 && PinIndex <= 2)
+				{
+					const FRichCurve& Curve = CurveVector->FloatCurves[PinIndex];
+					SetCurveNodeValues(CurveNode, Curve);
+				}
+			}
+			else if (const UCurveFloat* const CurveFloat = Cast<UCurveFloat>(CurveAsset))
+			{
+				if (PinIndex == 0)
+				{
+					const FRichCurve& Curve = CurveFloat->FloatCurve;
+					SetCurveNodeValues(CurveNode, Curve);
+				}
+			}
+		}
+	}
+
+	else if (const UCustomizableObjectNodeFloatVariation* TypedNodeFloatVar = Cast<const UCustomizableObjectNodeFloatVariation>(Node))
+	{
+		mu::NodeScalarVariationPtr FloatNode = new mu::NodeScalarVariation();
+		Result = FloatNode;
+
+		if (const UEdGraphPin* ConnectedPin = FollowInputPin(*TypedNodeFloatVar->DefaultPin()))
+		{
+			mu::NodeScalarPtr ChildNode = GenerateMutableSourceFloat(ConnectedPin, GenerationContext);
+			if (ChildNode)
+			{
+				FloatNode->SetDefaultScalar(ChildNode.get());
+			}
+			else
+			{
+				GenerationContext.Compiler->CompilerLog(LOCTEXT("FloatFailed", "Float generation failed."), Node);
+			}
+		}
+		else
+		{
+			GenerationContext.Compiler->CompilerLog(LOCTEXT("FloatVarMissingDef", "Float variation node requires a default value."), Node);
+		}
+
+		FloatNode->SetVariationCount(TypedNodeFloatVar->Variations.Num());
+		for (int VariationIndex = 0; VariationIndex < TypedNodeFloatVar->Variations.Num(); ++VariationIndex)
+		{
+			UEdGraphPin* VariationPin = TypedNodeFloatVar->VariationPin(VariationIndex);
+			if (!VariationPin) continue;
+
+			FloatNode->SetVariationTag(VariationIndex, TCHAR_TO_ANSI(*TypedNodeFloatVar->Variations[VariationIndex].Tag));
+			if (const UEdGraphPin* ConnectedPin = FollowInputPin(*VariationPin))
+			{
+				mu::NodeScalarPtr ChildNode = GenerateMutableSourceFloat(ConnectedPin, GenerationContext);
+				FloatNode->SetVariationScalar(VariationIndex, ChildNode.get());
+			}
+		}
+	}
+
+	else if (const UCustomizableObjectNodeTable* TypedNodeTable = Cast<UCustomizableObjectNodeTable>(Node))
+	{
+		mu::NodeScalarTablePtr ScalarTableNode = new mu::NodeScalarTable();
+		Result = ScalarTableNode;
+
+		mu::TablePtr Table;
+		if (TypedNodeTable->Table)
+		{
+			FString ColumnName = Pin->PinFriendlyName.ToString();
+
+			if (Pin->PinType.PinCategory == Schema->PC_MaterialAsset)
+			{
+				ColumnName = GenerationContext.CurrentMaterialTableParameterId;
+			}
+			
+			Table = GenerateMutableSourceTable(TypedNodeTable->Table->GetName(), Pin, GenerationContext);
+			
+			ScalarTableNode->SetTable(Table);
+			ScalarTableNode->SetColumn(TCHAR_TO_ANSI(*ColumnName));
+			ScalarTableNode->SetParameterName(TCHAR_TO_ANSI(*TypedNodeTable->ParameterName));
+
+			GenerationContext.AddParameterNameUnique(Node, TypedNodeTable->ParameterName);
+			
+			if (Table->FindColumn(TCHAR_TO_ANSI(*ColumnName)) == -1)
+			{
+				Table = new mu::Table();
+
+				FString Msg = FString::Printf(TEXT("Couldn't find pin column with name %s"), *ColumnName);
+				GenerationContext.Compiler->CompilerLog(FText::FromString(Msg), Node);
+			}
+		}
+		else
+		{
+			Table = new mu::Table();
+			ScalarTableNode->SetTable(Table);
+
+			GenerationContext.Compiler->CompilerLog(LOCTEXT("ScalarTableError", "Couldn't find the data table of the node."), Node);
+		}
+	}
+
+	else
+	{
+		GenerationContext.Compiler->CompilerLog(LOCTEXT("UnimplementedNode", "Node type not implemented yet."), Node);
+	}
+
+	GenerationContext.Generated.Add(Key, FGeneratedData(Node, Result));
+	GenerationContext.GeneratedNodes.Add(Node);
+
+	if (Result)
+	{
+		Result->SetMessageContext(Node);
+	}
+
+	return Result;
+}
diff --git a/Engine/Plugins/Experimental/Mutable/Source/CustomizableObjectEditor/Private/GenerateMutableSource/GenerateMutableSourceFloat.h b/Engine/Plugins/Experimental/Mutable/Source/CustomizableObjectEditor/Private/GenerateMutableSource/GenerateMutableSourceFloat.h
new file mode 100644
--- /dev/null
+++ b/Engine/Plugins/Experimental/Mutable/Source/CustomizableObjectEditor/Private/GenerateMutableSource/GenerateMutableSourceFloat.h
@@ -0,0 +1,9 @@
+// Copyright Epic Games, Inc. All Rights Reserved.
+
+#pragma once
+
+#include "GenerateMutableSource/GenerateMutableSource.h"
+
+
+/** Convert a CustomizableObject Source Graph into a mutable source graph. */
+mu::NodeScalarPtr GenerateMutableSourceFloat(const UEdGraphPin* Pin, FMutableGraphGenerationContext& GenerationContext);
diff --git a/Engine/Plugins/Experimental/Mutable/Source/CustomizableObjectEditor/Private/GenerateMutableSource/GenerateMutableSourceGroupProjector.cpp b/Engine/Plugins/Experimental/Mutable/Source/CustomizableObjectEditor/Private/GenerateMutableSource/GenerateMutableSourceGroupProjector.cpp
new file mode 100644
--- /dev/null
+++ b/Engine/Plugins/Experimental/Mutable/Source/CustomizableObjectEditor/Private/GenerateMutableSource/GenerateMutableSourceGroupProjector.cpp
@@ -0,0 +1,618 @@
+// Copyright Epic Games, Inc. All Rights Reserved.
+
+#include "GenerateMutableSourceGroupProjector.h"
+
+#include "GenerateMutableSource/GenerateMutableSourceFloat.h"
+#include "GenerateMutableSource/GenerateMutableSourceProjector.h"
+
+#include "Templates/UnrealTemplate.h"
+#include "GenericPlatform/GenericPlatformMath.h"
+#include "Interfaces/ITargetPlatform.h"
+#include "Engine/AssetManager.h"
+#include "Engine/TextureLODSettings.h"
+#include "ClothingAsset.h"
+#include "MeshUtilities.h"
+#include "EdGraphToken.h"
+#include "AssetRegistry/AssetRegistryModule.h"
+#include "Materials/MaterialInstance.h"
+#include "UnrealEditorPortabilityHelpers.h"
+#include "PlatformInfo.h"
+
+#include "CustomizableObjectCompiler.h"
+#include "UnrealToMutableTextureConversionUtils.h"
+#include "UnrealConversionUtils.h"
+#include "CustomizableObjectEditorModule.h"
+#include "CustomizableObject.h"
+#include "GraphTraversal.h"
+#include "MultilayerProjector.h"
+
+#include "MutableTools/Public/NodeScalarConstant.h"
+#include "MutableTools/Public/NodeScalarParameter.h"
+#include "MutableTools/Public/NodeScalarEnumParameter.h"
+#include "MutableTools/Public/NodeColourConstant.h"
+#include "MutableTools/Public/NodeColourFromScalars.h"
+#include "MutableTools/Public/NodeProjector.h"
+#include "MutableTools/Public/NodeImageBinarise.h"
+#include "MutableTools/Public/NodeImageConstant.h"
+#include "MutableTools/Public/NodeImageSwitch.h"
+#include "MutableTools/Public/NodeImageLayerColour.h"
+#include "MutableTools/Public/NodeImageSwizzle.h"
+#include "MutableTools/Public/NodeImageProject.h"
+#include "MutableTools/Public/NodeImagePlainColour.h"
+#include "MutableTools/Public/NodeImageMultiLayer.h"
+#include "MutableTools/Public/NodeImageResize.h"
+#include "MutableTools/Public/NodeMeshConstant.h"
+#include "MutableTools/Public/NodeMeshSwitch.h"
+#include "MutableTools/Public/NodeMeshApplyPose.h"
+#include "MutableTools/Public/Compiler.h"
+#include "MutableTools/Public/Streams.h"
+#include "MutableTools/Public/NodeRangeFromScalar.h"
+
+#define LOCTEXT_NAMESPACE "CustomizableObjectEditor"
+
+
+mu::NodeImagePtr GenerateMutableGroupProjection(const int32 NodeLOD, const int32 ImageIndex, mu::NodeMeshPtr MeshNode, FMutableGraphGenerationContext& GenerationContext,
+	const UCustomizableObjectNodeMaterialBase* NodeMaterialBase, bool& bShareProjectionTexturesBetweenLODs, bool& bIsGroupProjectorImage,
+	UTexture2D*& GroupProjectionReferenceTexture, TMap<FString, float>& TextureNameToProjectionResFactor, FString& AlternateResStateName,
+	UCustomizableObjectNodeMaterial* ParentMaterial)
+{
+	if (!GenerationContext.ProjectorGroupMap.Num() || !MeshNode.get())
+	{
+		return mu::NodeImagePtr();
+	}
+	
+	const UCustomizableObjectNodeMaterial* TypedNodeMat = Cast<UCustomizableObjectNodeMaterial>(NodeMaterialBase);
+	check(!TypedNodeMat || !ParentMaterial); // Logical implication
+
+	const UCustomizableObjectNodeExtendMaterial* TypedNodeExt = Cast<UCustomizableObjectNodeExtendMaterial>(NodeMaterialBase);
+	check(!TypedNodeExt || ParentMaterial); // Logical implication
+
+	check(TypedNodeMat || TypedNodeExt);
+	
+	TArray<mu::NodeImageProjectPtr> ImageNodes;
+	TArray<FGroupProjectorTempData> ImageNodes_ProjectorTempData;
+
+	int32 TextureSize = 512;
+
+	for (TPair<const UCustomizableObjectNodeObjectGroup*, FGroupProjectorTempData>& Elem : GenerationContext.ProjectorGroupMap)
+	{
+		FGroupProjectorTempData& ProjectorTempData = Elem.Value;
+
+		const int32& DropProjectionTextureAtLOD = ProjectorTempData.CustomizableObjectNodeGroupProjectorParameter->DropProjectionTextureAtLOD;
+		if (DropProjectionTextureAtLOD >= 0 && NodeLOD >= DropProjectionTextureAtLOD)
+		{
+			continue;
+		}
+
+		bShareProjectionTexturesBetweenLODs |= ProjectorTempData.CustomizableObjectNodeGroupProjectorParameter->bShareProjectionTexturesBetweenLODs;
+
+		if (!GroupProjectionReferenceTexture)
+		{
+			GroupProjectionReferenceTexture = ProjectorTempData.CustomizableObjectNodeGroupProjectorParameter->ReferenceTexture;
+		}
+
+		const bool bProjectToImage = [&]
+		{
+			if (TypedNodeMat)
+			{
+				const FString ParameterName = TypedNodeMat->GetParameterName(EMaterialParameterType::Texture, ImageIndex).ToString();
+				return ParameterName == ProjectorTempData.CustomizableObjectNodeGroupProjectorParameter->MaterialChannelNameToConnect;
+			}
+			else
+			{
+				const FGuid ImageId = ParentMaterial->GetParameterId(EMaterialParameterType::Texture, ImageIndex);
+				const FString ImagePinName = TypedNodeExt->GetUsedImagePin(ImageId)->PinName.ToString();
+				return ImagePinName.StartsWith(ProjectorTempData.CustomizableObjectNodeGroupProjectorParameter->MaterialChannelNameToConnect);
+			}
+		}();
+		
+		if (bProjectToImage)
+		{
+			const bool bWarningReplacedImage = [&]
+			{
+				if (TypedNodeMat)
+				{
+					return TypedNodeMat->IsImageMutableMode(ImageIndex);
+				}
+				else
+				{
+					const FGuid ImageId = ParentMaterial->GetParameterId(EMaterialParameterType::Texture, ImageIndex);
+					return TypedNodeExt->UsesImage(ImageId);
+				}
+			}();
+			
+			if (bWarningReplacedImage)
+			{
+				const FString ImageName = [&]
+				{
+					if (TypedNodeMat)
+					{
+						return TypedNodeMat->GetParameterName(EMaterialParameterType::Texture, ImageIndex).ToString();
+					}
+					else
+					{
+						const FGuid ImageId = ParentMaterial->GetParameterId(EMaterialParameterType::Texture, ImageIndex);
+						return TypedNodeExt->GetUsedImagePin(ImageId)->PinName.ToString();
+					}
+				}();
+				
+				FString msg = FString::Printf(TEXT("Material image [%s] is connected to an image but will be replaced by a Group Projector."), *ImageName);
+				GenerationContext.Compiler->CompilerLog(FText::FromString(msg), TypedNodeMat);
+				continue;
+			}
+			
+			mu::NodeImageProjectPtr ImageNode = new mu::NodeImageProject();
+			bIsGroupProjectorImage = true;
+			ImageNode->SetLayout(ProjectorTempData.CustomizableObjectNodeGroupProjectorParameter->UVLayout);
+
+			{
+				if (TypedNodeMat &&
+					!ProjectorTempData.CustomizableObjectNodeGroupProjectorParameter->AlternateProjectionResolutionStateName.IsEmpty()
+					&& ProjectorTempData.CustomizableObjectNodeGroupProjectorParameter->AlternateProjectionResolutionFactor > 0)
+				{
+					TextureNameToProjectionResFactor.Add(
+						TypedNodeMat->GetParameterName(EMaterialParameterType::Texture, ImageIndex).ToString(),
+						ProjectorTempData.CustomizableObjectNodeGroupProjectorParameter->AlternateProjectionResolutionFactor);
+
+					if (!AlternateResStateName.IsEmpty() &&
+						AlternateResStateName != ProjectorTempData.CustomizableObjectNodeGroupProjectorParameter->AlternateProjectionResolutionStateName &&
+						!ProjectorTempData.bAlternateResStateNameWarningDisplayed)
+					{
+						FString msg = FString::Printf(TEXT("All 'Alternate Projection Resolution State Name' properties in Group Projector Parameter nodes connected to same Group node must have the same value or be blank. Only the value of the last connected node will be used."));
+						GenerationContext.Compiler->CompilerLog(FText::FromString(msg), ProjectorTempData.CustomizableObjectNodeGroupProjectorParameter);
+						ProjectorTempData.bAlternateResStateNameWarningDisplayed = true;
+					}
+
+					AlternateResStateName = ProjectorTempData.CustomizableObjectNodeGroupProjectorParameter->AlternateProjectionResolutionStateName;
+				}
+			}
+
+			{
+				mu::NodeScalarConstantPtr ScalarNode = new mu::NodeScalarConstant;
+				ScalarNode->SetValue(120.f);
+				ImageNode->SetAngleFadeStart(ScalarNode);
+			}
+
+			{
+				mu::NodeScalarConstantPtr ScalarNode = new mu::NodeScalarConstant;
+				ScalarNode->SetValue(150.f);
+				ImageNode->SetAngleFadeEnd(ScalarNode);
+			}
+
+			mu::NodeMeshSwitchPtr MeshSwitchNode = new mu::NodeMeshSwitch;
+			MeshSwitchNode->SetParameter(ProjectorTempData.PoseOptionsParameter);
+			MeshSwitchNode->SetOptionCount(ProjectorTempData.CustomizableObjectNodeGroupProjectorParameter->OptionPoses.Num() + 1);
+			MeshSwitchNode->SetOption(0, MeshNode);
+
+			for (int32 SelectorIndex = 0; SelectorIndex < ProjectorTempData.CustomizableObjectNodeGroupProjectorParameter->OptionPoses.Num(); ++SelectorIndex)
+			{
+				if (ProjectorTempData.CustomizableObjectNodeGroupProjectorParameter->OptionPoses[SelectorIndex].OptionPose)
+				{
+					mu::NodeMeshApplyPosePtr NodeMeshApplyPose = CreateNodeMeshApplyPose(MeshNode, GenerationContext.Object,
+						ProjectorTempData.PoseBoneDataArray[SelectorIndex].ArrayBoneName, ProjectorTempData.PoseBoneDataArray[SelectorIndex].ArrayTransform);
+
+					if (!NodeMeshApplyPose)
+					{
+						FString msg = FString::Printf(TEXT("Couldn't get bone transform information from a Pose Asset."));
+						GenerationContext.Compiler->CompilerLog(FText::FromString(msg), NodeMaterialBase);
+					}
+
+					MeshSwitchNode->SetOption(SelectorIndex + 1, NodeMeshApplyPose);
+				}
+				else
+				{
+					MeshSwitchNode->SetOption(SelectorIndex + 1, MeshNode);
+				}
+			}
+
+			ImageNode->SetMesh(MeshSwitchNode);
+			ImageNode->SetProjector(ProjectorTempData.NodeProjectorParameterPtr);
+
+			mu::NodeImagePtr ImageToProject = ProjectorTempData.NodeImagePtr;
+			{
+				// Apply LODBias to the images to project
+
+				int32 LODBias = ComputeLODBias(GenerationContext, GroupProjectionReferenceTexture, GroupProjectionReferenceTexture ? GroupProjectionReferenceTexture->MaxTextureSize : 0, TypedNodeMat, ImageIndex);
+
+				if (LODBias > 0)
+				{
+					mu::NodeImageResizePtr ResizeImage = new mu::NodeImageResize();
+					ResizeImage->SetBase(ImageToProject.get());
+					ResizeImage->SetRelative(true);
+					float factor = FMath::Pow(0.5f, LODBias);
+					ResizeImage->SetSize(factor, factor);
+					ImageToProject = ResizeImage;
+				}
+
+			}
+
+			ImageNode->SetImage(ImageToProject);
+
+			TextureSize = [&]
+			{
+				if (const int32 Size = ProjectorTempData.CustomizableObjectNodeGroupProjectorParameter->ProjectionTextureSize;
+					Size <= 0 || !FMath::IsPowerOfTwo(Size))
+				{
+					// \todo: closest power of 2 bigger than the set value?
+					return 512;
+				}
+				else
+				{
+					return Size;
+				}
+			}();
+			ImageNode->SetImageSize( FUintVector2(TextureSize, TextureSize) );
+
+			ImageNodes.Add(ImageNode);
+			ImageNodes_ProjectorTempData.Add(ProjectorTempData);
+
+			// Generate the projection mask-out texture cache
+			const FString& MaskName = ProjectorTempData.CustomizableObjectNodeGroupProjectorParameter->MaskedOutAreaMaterialChannelName;
+
+			if (!MaskName.IsEmpty())
+			{
+				const UMaterialInterface* MaskMaterial = TypedNodeMat ? TypedNodeMat->Material : ParentMaterial->Material;
+				if (MaskMaterial)
+				{
+					UTexture* MaskTexture = nullptr;
+
+					if (MaskMaterial->GetTextureParameterValue(FName(*MaskName), MaskTexture))
+					{
+						UTexture2D* MaskTexture2D = Cast<UTexture2D>(MaskTexture);
+
+						if (MaskTexture2D && MaskTexture2D->Source.GetNumMips() > 0)
+						{
+							TArray64<uint8> TempData;
+							MaskTexture2D->Source.GetMipData(TempData, 0, 0, 0);
+							uint32 TotalTexels = MaskTexture2D->Source.GetSizeX() * MaskTexture2D->Source.GetSizeY();
+
+							ETextureSourceFormat PixelFormat = MaskTexture2D->Source.GetFormat();
+							int32 BytesPerPixel = MaskTexture2D->Source.GetBytesPerPixel();
+							ensure((PixelFormat == TSF_BGRA8 && BytesPerPixel == 4) || BytesPerPixel == 1);
+
+							{
+								uint32 DataSize = TotalTexels * MaskTexture2D->Source.GetBytesPerPixel();
+								FString MaskTexture2DPath = MaskTexture2D->GetPathName();
+
+								if (!GenerationContext.MaskOutTextureCache.Find(MaskTexture2DPath) && DataSize > 1024) // Discard default or too small textures to have any mask detail
+								{
+									FMaskOutTexture& CachedTexture = GenerationContext.MaskOutTextureCache.Add(MaskTexture2DPath);
+									
+									CachedTexture.SetTextureSize(MaskTexture2D->Source.GetSizeX(), MaskTexture2D->Source.GetSizeY());
+
+									for (size_t p = 0; p < DataSize / BytesPerPixel; ++p)
+									{
+										if (BytesPerPixel == 4)
+										{
+											CachedTexture.GetTexelReference(p) = TempData[p * 4 + 3] > 0; // Copy alpha channel of PF_R8G8B8A8
+										}
+										else if (BytesPerPixel == 1)
+										{
+											CachedTexture.GetTexelReference(p) = TempData[p] > 0;
+										}
+										else
+										{
+											check(false);
+										}
+									}
+								}
+
+								GenerationContext.MaskOutMaterialCache.Add(MaskMaterial->GetPathName(), MaskTexture2DPath);
+							}
+						}
+					}
+				}
+			}
+		}
+	}
+
+	if (ImageNodes.Num() == 0)
+	{
+		return mu::NodeImagePtr();
+	}
+	
+	mu::NodeColourConstantPtr ZeroColorNode = new mu::NodeColourConstant();
+	ZeroColorNode->SetValue(0.f, 0.f, 0.f);
+
+	if (TextureSize <= 0 || !FMath::IsPowerOfTwo(TextureSize))
+	{
+		TextureSize = 512;
+	}
+
+	mu::NodeImagePlainColourPtr ZeroPlainColourNode = new mu::NodeImagePlainColour;
+	ZeroPlainColourNode->SetSize(TextureSize, TextureSize);
+	ZeroPlainColourNode->SetColour(ZeroColorNode);
+
+	mu::NodeImageSwizzlePtr ZeroChannelNode = new mu::NodeImageSwizzle;
+	ZeroChannelNode->SetFormat(mu::EImageFormat::IF_L_UBYTE);
+	ZeroChannelNode->SetSource(0, ZeroPlainColourNode);
+	ZeroChannelNode->SetSourceChannel(0, 2); // Just take a zeroed channel for the base alpha
+
+	mu::NodeScalarConstantPtr OneConstantNode = new mu::NodeScalarConstant;
+	OneConstantNode->SetValue(1.f);
+
+	mu::NodeImagePtr ResultAlpha = ZeroChannelNode;
+	mu::NodeImagePtr ResultImage = ZeroPlainColourNode;
+
+	for (int32 i = 0; i < ImageNodes.Num(); ++i)
+	{
+		if (i > 0) // Resize the projection texture if necessary after the first iteration
+		{
+			int32 NewTextureSize = ImageNodes_ProjectorTempData.IsValidIndex(i) ? ImageNodes_ProjectorTempData[i].CustomizableObjectNodeGroupProjectorParameter->ProjectionTextureSize : TextureSize;
+
+			if (NewTextureSize <= 0 || !FMath::IsPowerOfTwo(NewTextureSize))
+			{
+				NewTextureSize = TextureSize;
+			}
+
+			if (NewTextureSize != TextureSize)
+			{
+				TextureSize = NewTextureSize;
+
+				ZeroPlainColourNode = new mu::NodeImagePlainColour;
+				ZeroPlainColourNode->SetSize(TextureSize, TextureSize);
+				ZeroPlainColourNode->SetColour(ZeroColorNode);
+			}
+		}
+
+		mu::NodeImageSwizzlePtr ImageNodesAlphaChannelNode = new mu::NodeImageSwizzle;
+		ImageNodesAlphaChannelNode->SetFormat(mu::EImageFormat::IF_L_UBYTE);
+		ImageNodesAlphaChannelNode->SetSource(0, ImageNodes[i]);
+		ImageNodesAlphaChannelNode->SetSourceChannel(0, 3);
+
+		mu::NodeColourFromScalarsPtr ColourFromScalars = new mu::NodeColourFromScalars;
+		ColourFromScalars->SetX(ImageNodes_ProjectorTempData[i].NodeOpacityParameter);
+		ColourFromScalars->SetY(ImageNodes_ProjectorTempData[i].NodeOpacityParameter);
+		ColourFromScalars->SetZ(ImageNodes_ProjectorTempData[i].NodeOpacityParameter);
+		ColourFromScalars->SetW(OneConstantNode);
+
+		mu::NodeImageLayerColourPtr OpacityMultiLayerNode = new mu::NodeImageLayerColour;
+		OpacityMultiLayerNode->SetBlendType(mu::EBlendType::BT_MULTIPLY);
+		OpacityMultiLayerNode->SetColour(ColourFromScalars);
+		OpacityMultiLayerNode->SetBase(ImageNodesAlphaChannelNode);
+		//OpacityMultiLayerNode->SetMask(OneChannelNode); // No mask needed
+
+		mu::NodeImageSwizzlePtr MultiplySwizzleNode = new mu::NodeImageSwizzle;
+		MultiplySwizzleNode->SetFormat(mu::EImageFormat::IF_L_UBYTE);
+		MultiplySwizzleNode->SetSource(0, OpacityMultiLayerNode);
+		MultiplySwizzleNode->SetSourceChannel(0, 0);
+
+		mu::NodeImageMultiLayerPtr BaseAlphaMultiLayerNode = new mu::NodeImageMultiLayer;
+		BaseAlphaMultiLayerNode->SetRange(ImageNodes_ProjectorTempData[i].NodeRange);
+		BaseAlphaMultiLayerNode->SetBlendType(mu::EBlendType::BT_ALPHA_OVERLAY);
+		BaseAlphaMultiLayerNode->SetBase(ResultAlpha);
+		BaseAlphaMultiLayerNode->SetBlended(MultiplySwizzleNode);
+		//BaseAlphaMultiLayerNode->SetMask(MultiplySwizzleNode); // No mask needed
+		ResultAlpha = BaseAlphaMultiLayerNode;
+
+		mu::NodeImageMultiLayerPtr BaseMultiLayerNode = new mu::NodeImageMultiLayer;
+		BaseMultiLayerNode->SetRange(ImageNodes_ProjectorTempData[i].NodeRange);
+		BaseMultiLayerNode->SetBlendType(mu::EBlendType::BT_BLEND);
+		BaseMultiLayerNode->SetBase(ResultImage);
+		BaseMultiLayerNode->SetBlended(ImageNodes[i]);
+		BaseMultiLayerNode->SetMask(MultiplySwizzleNode);
+		ResultImage = BaseMultiLayerNode;
+	}
+
+	mu::NodeImageSwizzlePtr SwizzleNodeR = new mu::NodeImageSwizzle;
+	SwizzleNodeR->SetFormat(mu::EImageFormat::IF_L_UBYTE);
+	SwizzleNodeR->SetSource(0, ResultImage);
+	SwizzleNodeR->SetSourceChannel(0, 0);
+
+	mu::NodeImageSwizzlePtr SwizzleNodeG = new mu::NodeImageSwizzle;
+	SwizzleNodeG->SetFormat(mu::EImageFormat::IF_L_UBYTE);
+	SwizzleNodeG->SetSource(0, ResultImage);
+	SwizzleNodeG->SetSourceChannel(0, 1);
+
+	mu::NodeImageSwizzlePtr SwizzleNodeB = new mu::NodeImageSwizzle;
+	SwizzleNodeB->SetFormat(mu::EImageFormat::IF_L_UBYTE);
+	SwizzleNodeB->SetSource(0, ResultImage);
+	SwizzleNodeB->SetSourceChannel(0, 2);
+
+	mu::NodeImageSwizzlePtr FinalSwizzleNode = new mu::NodeImageSwizzle;
+	FinalSwizzleNode->SetFormat(mu::EImageFormat::IF_RGBA_UBYTE);
+	FinalSwizzleNode->SetSource(0, SwizzleNodeR);
+	FinalSwizzleNode->SetSourceChannel(0, 0);
+	FinalSwizzleNode->SetSource(1, SwizzleNodeG);
+	FinalSwizzleNode->SetSourceChannel(1, 0);
+	FinalSwizzleNode->SetSource(2, SwizzleNodeB);
+	FinalSwizzleNode->SetSourceChannel(2, 0);
+	FinalSwizzleNode->SetSource(3, ResultAlpha);
+	FinalSwizzleNode->SetSourceChannel(3, 0);
+
+	return FinalSwizzleNode;
+}
+
+
+bool GenerateMutableSourceGroupProjector(const UEdGraphPin* Pin, FMutableGraphGenerationContext& GenerationContext, const UCustomizableObjectNodeObjectGroup* originalGroup)
+{
+	check(Pin)
+	check(originalGroup);
+	RETURN_ON_CYCLE(*Pin, GenerationContext)
+
+	CheckNumOutputs(*Pin, GenerationContext);
+	
+	UCustomizableObjectNode* Node = CastChecked<UCustomizableObjectNode>(Pin->GetOwningNode());
+
+	const FGeneratedKey Key(reinterpret_cast<void*>(&GenerateMutableSourceGroupProjector), *Pin, *Node, GenerationContext, true);
+	if (const FGeneratedData* Generated = GenerationContext.Generated.Find(Key))
+	{
+		return true;
+	}
+
+	mu::NodeProjectorPtr Result;
+	
+	if (UCustomizableObjectNodeGroupProjectorParameter* ProjParamNode = Cast<UCustomizableObjectNodeGroupProjectorParameter>(Node))
+	{
+		FGroupProjectorTempData GroupProjectorTempData;
+		// The static cast works here because it's already known to be a mu::NodeProjectorParameter* because of the UE5 Cast in the previous line
+		GroupProjectorTempData.NodeProjectorParameterPtr = static_cast<mu::NodeProjectorParameter*>(GenerateMutableSourceProjector(Pin, GenerationContext).get());
+
+		if (GroupProjectorTempData.NodeProjectorParameterPtr)
+		{
+			// Use the projector parameter uid + suffix to identify parameters derived from this node
+			FString ProjectorParamUid = GroupProjectorTempData.NodeProjectorParameterPtr->GetUid();
+
+			// Add to UCustomizableObjectNodeGroupProjectorParameter::OptionImages those textures that are present in
+			// UCustomizableObjectNodeGroupProjectorParameter::OptionImagesDataTable avoiding any repeated element
+			TArray<FGroupProjectorParameterImage> ArrayOptionImage = ProjParamNode->GetFinalOptionImagesNoRepeat();
+
+			if ((ProjParamNode->OptionImagesDataTable != nullptr) &&
+				(ProjParamNode->DataTableTextureColumnName.ToString().IsEmpty() || (ProjParamNode->DataTableTextureColumnName.ToString() == "None")))
+			{
+				FString msg = FString::Printf(TEXT("The group projection node has a table assigned to the Option Images Data Table property, but no column to read textures is specified at the Data Table Texture Column Name property."));
+				GenerationContext.Compiler->CompilerLog(FText::FromString(msg), ProjParamNode, EMessageSeverity::Error, true);
+			}
+
+			GroupProjectorTempData.CustomizableObjectNodeGroupProjectorParameter = ProjParamNode;
+
+			mu::NodeScalarParameterPtr NodeScalarParam = new mu::NodeScalarParameter;
+			FString NodeScalarParamName = ProjParamNode->ParameterName + FMultilayerProjector::NUM_LAYERS_PARAMETER_POSTFIX;
+			NodeScalarParam->SetName(TCHAR_TO_ANSI(*(NodeScalarParamName)));
+			NodeScalarParam->SetUid(TCHAR_TO_ANSI(*(ProjectorParamUid + FString("_NL"))));
+			GenerationContext.AddParameterNameUnique(originalGroup, NodeScalarParamName);
+			NodeScalarParam->SetDefaultValue(0.f);
+
+			GenerationContext.ParameterUIDataMap.Add(NodeScalarParamName, FParameterUIData(
+				NodeScalarParamName,
+				ProjParamNode->ParamUIMetadata,
+				EMutableParameterType::Int));
+
+			mu::Ptr<mu::NodeRangeFromScalar> NodeRangeFromScalar = new mu::NodeRangeFromScalar;
+			NodeRangeFromScalar->SetSize(NodeScalarParam);
+			GroupProjectorTempData.NodeRange = NodeRangeFromScalar;
+			GroupProjectorTempData.NodeProjectorParameterPtr->SetRangeCount(1);
+			GroupProjectorTempData.NodeProjectorParameterPtr->SetRange(0, NodeRangeFromScalar);
+
+			mu::NodeScalarEnumParameterPtr EnumParameterNode = new mu::NodeScalarEnumParameter;
+			FString NodeEnumParamName = ProjParamNode->ParameterName + FMultilayerProjector::IMAGE_PARAMETER_POSTFIX;
+			EnumParameterNode->SetName(TCHAR_TO_ANSI(*(NodeEnumParamName)));
+			EnumParameterNode->SetUid(TCHAR_TO_ANSI(*(ProjectorParamUid + FString("_SI"))));
+			GenerationContext.AddParameterNameUnique(originalGroup, NodeEnumParamName);
+			EnumParameterNode->SetValueCount(ArrayOptionImage.Num());
+			EnumParameterNode->SetDefaultValueIndex(0);
+			EnumParameterNode->SetRangeCount(1);
+			EnumParameterNode->SetRange(0, NodeRangeFromScalar);
+
+
+			FParameterUIData ParameterUIData(NodeEnumParamName, ProjParamNode->ParamUIMetadata, EMutableParameterType::Int);
+			ParameterUIData.IntegerParameterGroupType = ECustomizableObjectGroupType::COGT_ONE;
+			ParameterUIData.ParamUIMetadata.ExtraInformation.Add(FString("UseThumbnails"));
+
+			mu::NodeScalarEnumParameterPtr PoseEnumParameterNode = new mu::NodeScalarEnumParameter;
+			FString PoseNodeEnumParamName = ProjParamNode->ParameterName + FMultilayerProjector::POSE_PARAMETER_POSTFIX;
+			PoseEnumParameterNode->SetName(TCHAR_TO_ANSI(*(PoseNodeEnumParamName)));
+			PoseEnumParameterNode->SetUid(TCHAR_TO_ANSI(*(ProjectorParamUid + FString("_SP"))));
+			GenerationContext.AddParameterNameUnique(originalGroup, PoseNodeEnumParamName);
+			PoseEnumParameterNode->SetValueCount(ProjParamNode->OptionPoses.Num() + 1);
+			PoseEnumParameterNode->SetDefaultValueIndex(0);
+			GroupProjectorTempData.PoseOptionsParameter = PoseEnumParameterNode;
+
+			GenerationContext.ParameterUIDataMap.Add(PoseNodeEnumParamName, FParameterUIData(
+				PoseNodeEnumParamName,
+				ProjParamNode->ParamUIMetadata,
+				EMutableParameterType::Int));
+
+			mu::NodeScalarParameterPtr OpacityParameterNode = new mu::NodeScalarParameter;
+			FString OpacityParameterNodeName = ProjParamNode->ParameterName + FMultilayerProjector::OPACITY_PARAMETER_POSTFIX;
+			OpacityParameterNode->SetName(TCHAR_TO_ANSI(*(OpacityParameterNodeName)));
+			OpacityParameterNode->SetUid(TCHAR_TO_ANSI(*(ProjectorParamUid + FString("_O"))));
+			GenerationContext.AddParameterNameUnique(originalGroup, OpacityParameterNodeName);
+			OpacityParameterNode->SetDefaultValue(0.75f);
+			OpacityParameterNode->SetRangeCount(1);
+			OpacityParameterNode->SetRange(0, NodeRangeFromScalar);
+			FMutableParamUIMetadata OpacityMetadata = ProjParamNode->ParamUIMetadata;
+			OpacityMetadata.ObjectFriendlyName = FString("Transparency");
+			GroupProjectorTempData.NodeOpacityParameter = OpacityParameterNode;
+
+			GenerationContext.ParameterUIDataMap.Add(OpacityParameterNodeName, FParameterUIData(
+				OpacityParameterNodeName,
+				OpacityMetadata,
+				EMutableParameterType::Float));
+
+			if (ArrayOptionImage.Num() == 0)
+			{
+				FString msg = FString::Printf(TEXT("The group projection node must have at least one option image connected to a texture or at least one valid element in Option Images Data Table."));
+				GenerationContext.Compiler->CompilerLog(FText::FromString(msg), ProjParamNode, EMessageSeverity::Error, true);
+				return false;
+			}
+
+			for (int ImageIndex = 0; ImageIndex < ArrayOptionImage.Num(); ++ImageIndex)
+			{
+				EnumParameterNode->SetValue(ImageIndex, (float)ImageIndex, TCHAR_TO_ANSI(*ArrayOptionImage[ImageIndex].OptionName));
+
+				FMutableParamUIMetadata optionMetadata = ParameterUIData.ParamUIMetadata;
+				optionMetadata.UIThumbnail = ArrayOptionImage[ImageIndex].OptionImage;
+				ParameterUIData.ArrayIntegerParameterOption.Add(FIntegerParameterUIData(
+					ArrayOptionImage[ImageIndex].OptionName,
+					optionMetadata));
+			}
+
+			GenerationContext.ParameterUIDataMap.Add(NodeEnumParamName, ParameterUIData);
+
+			PoseEnumParameterNode->SetValue(0, 0.f, "Default pose");
+
+			for (int PoseIndex = 0; PoseIndex < ProjParamNode->OptionPoses.Num(); ++PoseIndex)
+			{
+				PoseEnumParameterNode->SetValue(PoseIndex + 1, (float)PoseIndex + 1.f, TCHAR_TO_ANSI(*ProjParamNode->OptionPoses[PoseIndex].PoseName));
+
+				TArray<FString> ArrayBoneName;
+				TArray<FTransform> ArrayTransform;
+				UPoseAsset* PoseAsset = ProjParamNode->OptionPoses[PoseIndex].OptionPose;
+				if (PoseAsset == nullptr) // Check if the slot has a selected pose. Could be left empty by the user
+				{
+					FString msg = FString::Printf(TEXT("The group projection node must have a pose assigned on each Option Poses element."));
+					GenerationContext.Compiler->CompilerLog(FText::FromString(msg), ProjParamNode, EMessageSeverity::Error, true);
+					return false;
+				}
+
+				check(GroupProjectorTempData.PoseBoneDataArray.Num() == PoseIndex);
+				GroupProjectorTempData.PoseBoneDataArray.AddDefaulted(1);
+				UCustomizableObjectNodeAnimationPose::StaticRetrievePoseInformation(PoseAsset, GenerationContext.GetCurrentComponentInfo().RefSkeletalMesh,
+					GroupProjectorTempData.PoseBoneDataArray[PoseIndex].ArrayBoneName, GroupProjectorTempData.PoseBoneDataArray[PoseIndex].ArrayTransform);
+			}
+
+			mu::NodeImageSwitchPtr SwitchNode = new mu::NodeImageSwitch;
+			SwitchNode->SetParameter(EnumParameterNode);
+			SwitchNode->SetOptionCount(ArrayOptionImage.Num());
+
+			bool bFoundUnlinkedPin = false;
+
+			for (int SelectorIndex = 0; SelectorIndex < ArrayOptionImage.Num(); ++SelectorIndex)
+			{
+				if (ArrayOptionImage[SelectorIndex].OptionImage)
+				{
+					mu::NodeImageConstantPtr ImageNode = new mu::NodeImageConstant();
+					GenerationContext.ArrayTextureUnrealToMutableTask.Add(FTextureUnrealToMutableTask(ImageNode, ArrayOptionImage[SelectorIndex].OptionImage, ProjParamNode));
+					SwitchNode->SetOption(SelectorIndex, ImageNode);
+				}
+				else
+				{
+					//SwitchNode->SetOption(SelectorIndex, nullptr);
+					bFoundUnlinkedPin = true;
+				}
+			}
+
+			if (bFoundUnlinkedPin)
+			{
+				FString msg = FString::Printf(TEXT("The group projection node must have all the option images connected to a texture. Please set a texture for all the options."));
+				GenerationContext.Compiler->CompilerLog(FText::FromString(msg), ProjParamNode);
+			}
+
+			GroupProjectorTempData.NodeImagePtr = SwitchNode;
+
+			GenerationContext.ProjectorGroupMap.Add(originalGroup, GroupProjectorTempData);
+		}
+		else
+		{
+			GenerationContext.Compiler->CompilerLog(FText::FromString(TEXT("Error getting group projection properties.")), ProjParamNode, EMessageSeverity::Error, true);
+			return false;
+		}
+	}
+	
+
+	GenerationContext.Generated.Add(Key, FGeneratedData(Node, Result));
+	GenerationContext.GeneratedNodes.Add(Node);
+
+	return true;
+}
diff --git a/Engine/Plugins/Experimental/Mutable/Source/CustomizableObjectEditor/Private/GenerateMutableSource/GenerateMutableSourceGroupProjector.h b/Engine/Plugins/Experimental/Mutable/Source/CustomizableObjectEditor/Private/GenerateMutableSource/GenerateMutableSourceGroupProjector.h
new file mode 100644
--- /dev/null
+++ b/Engine/Plugins/Experimental/Mutable/Source/CustomizableObjectEditor/Private/GenerateMutableSource/GenerateMutableSourceGroupProjector.h
@@ -0,0 +1,15 @@
+// Copyright Epic Games, Inc. All Rights Reserved.
+
+#pragma once
+
+#include "GenerateMutableSource/GenerateMutableSource.h"
+
+
+mu::NodeImagePtr GenerateMutableGroupProjection(const int32 NodeLOD, const int32 ImageIndex, mu::NodeMeshPtr MeshNode, FMutableGraphGenerationContext& GenerationContext,
+	const UCustomizableObjectNodeMaterialBase* NodeMaterialBase, bool& bShareProjectionTexturesBetweenLODs, bool& bIsGroupProjectorImage,
+	UTexture2D*& GroupProjectionReferenceTexture, TMap<FString, float>& TextureNameToProjectionResFactor, FString& AlternateResStateName,
+	UCustomizableObjectNodeMaterial* ParentMaterial);
+
+
+/** Convert a CustomizableObject Source Graph into a mutable source graph. */
+bool GenerateMutableSourceGroupProjector(const UEdGraphPin* Pin, FMutableGraphGenerationContext& GenerationContext, const UCustomizableObjectNodeObjectGroup* originalGroup);
diff --git a/Engine/Plugins/Experimental/Mutable/Source/CustomizableObjectEditor/Private/GenerateMutableSource/GenerateMutableSourceImage.cpp b/Engine/Plugins/Experimental/Mutable/Source/CustomizableObjectEditor/Private/GenerateMutableSource/GenerateMutableSourceImage.cpp
new file mode 100644
--- /dev/null
+++ b/Engine/Plugins/Experimental/Mutable/Source/CustomizableObjectEditor/Private/GenerateMutableSource/GenerateMutableSourceImage.cpp
@@ -0,0 +1,820 @@
+// Copyright Epic Games, Inc. All Rights Reserved.
+
+#include "GenerateMutableSourceImage.h"
+
+#include "GenerateMutableSource/GenerateMutableSourceFloat.h"
+#include "GenerateMutableSource/GenerateMutableSourceColor.h"
+#include "GenerateMutableSource/GenerateMutableSourceProjector.h"
+#include "GenerateMutableSource/GenerateMutableSourceMesh.h"
+#include "GenerateMutableSource/GenerateMutableSourceTable.h"
+
+#include "Templates/UnrealTemplate.h"
+#include "GenericPlatform/GenericPlatformMath.h"
+#include "Engine/AssetManager.h"
+#include "Engine/TextureLODSettings.h"
+#include "ClothingAsset.h"
+#include "MeshUtilities.h"
+#include "EdGraphToken.h"
+#include "AssetRegistry/AssetRegistryModule.h"
+#include "UnrealEditorPortabilityHelpers.h"
+#include "PlatformInfo.h"
+
+#include "CustomizableObjectCompiler.h"
+#include "UnrealToMutableTextureConversionUtils.h"
+#include "UnrealConversionUtils.h"
+#include "CustomizableObjectEditorModule.h"
+#include "CustomizableObject.h"
+#include "GraphTraversal.h"
+
+#include "MutableTools/Public/Node.h"
+#include "MutableTools/Public/NodeScalarConstant.h"
+#include "MutableTools/Public/NodeScalarEnumParameter.h"
+#include "MutableTools/Public/NodeProjector.h"
+#include "MutableTools/Public/NodeImageBinarise.h"
+#include "MutableTools/Public/NodeImageConstant.h"
+#include "MutableTools/Public/NodeImageParameter.h"
+#include "MutableTools/Public/NodeImageSwitch.h"
+#include "MutableTools/Public/NodeImageInterpolate.h"
+#include "MutableTools/Public/NodeImageLayer.h"
+#include "MutableTools/Public/NodeImageLayerColour.h"
+#include "MutableTools/Public/NodeImageSwizzle.h"
+#include "MutableTools/Public/NodeImageProject.h"
+#include "MutableTools/Public/NodeImagePlainColour.h"
+#include "MutableTools/Public/NodeImageFormat.h"
+#include "MutableTools/Public/NodeImageResize.h"
+#include "MutableTools/Public/NodeImageInvert.h"
+#include "MutableTools/Public/NodeImageColourMap.h"
+#include "MutableTools/Public/NodeImageVariation.h"
+#include "MutableTools/Public/NodeImageTable.h"
+#include "MutableTools/Public/NodeImageTransform.h"
+#include "MutableTools/Public/NodeLayout.h"
+#include "MutableTools/Public/Compiler.h"
+
+#define LOCTEXT_NAMESPACE "CustomizableObjectEditor"
+
+
+mu::ImagePtr ConvertTextureUnrealToMutable(UTexture2D* Texture, const UCustomizableObjectNode* Node, FCustomizableObjectCompiler* Compiler, bool bIsNormalComposite)
+{     
+	mu::ImagePtr MutableImage;
+	EUnrealToMutableConversionError Error = EUnrealToMutableConversionError::Unknown;
+
+	Tie(MutableImage, Error) = ConvertTextureUnrealToMutable(Texture, bIsNormalComposite);
+
+	switch(Error)
+	{
+	case EUnrealToMutableConversionError::Success: break;
+	case EUnrealToMutableConversionError::UnsupportedFormat:
+		{
+			Compiler->CompilerLog(
+					LOCTEXT("UnsupportedImageFormat", 
+							"Image format not supported."), 
+					Node);
+			break;
+		}
+	case EUnrealToMutableConversionError::CompositeImageDimensionMissmatch:
+		{
+			Compiler->CompilerLog(
+					LOCTEXT("CompositeImageDimensionMissmatch", 
+							"Composite image dimension missmatch."), 
+					Node);
+			break;
+		}
+	case EUnrealToMutableConversionError::CompositeUnsupportedFormat:
+		{
+			Compiler->CompilerLog(
+					LOCTEXT("CompositeUnsupportedFormat", 
+							"Composite image format not supported."), 
+					Node);
+			break;
+		}
+	default:
+		{
+			Compiler->CompilerLog(
+					LOCTEXT("ImageConversionUnknownError", 
+							"Image conversion unknown error."), 
+					Node);
+			break;
+		}
+	}
+
+	return MutableImage;
+}
+
+
+mu::NodeImagePtr ResizeToMaxTextureSize(float MaxTextureSize, const UTexture2D* BaseTexture, mu::NodeImageConstantPtr ImageNode)
+{
+	// To scale when above MaxTextureSize if defined
+	if (MaxTextureSize > 0 && BaseTexture
+		&& (BaseTexture->GetImportedSize().X > MaxTextureSize
+			|| BaseTexture->GetImportedSize().Y > MaxTextureSize))
+	{
+		mu::NodeImageResizePtr ResizeImage = new mu::NodeImageResize();
+		ResizeImage->SetBase(ImageNode.get());
+		ResizeImage->SetRelative(false);
+		float Factor = FMath::Min(
+			MaxTextureSize / (float)(BaseTexture->GetImportedSize().X),
+			MaxTextureSize / (float)(BaseTexture->GetImportedSize().Y));
+		ResizeImage->SetSize(BaseTexture->GetImportedSize().X * Factor, BaseTexture->GetImportedSize().Y * Factor);
+		return ResizeImage;
+	}
+	// TODO: MaxTextureSize = 0 indicates platform-specific maximum size, we should scale for it
+
+	return ImageNode;
+}
+
+
+mu::NodeImagePtr GenerateMutableSourceImage(const UEdGraphPin* Pin, FMutableGraphGenerationContext& GenerationContext, float MaxTextureSize)
+{
+	check(Pin)
+	RETURN_ON_CYCLE(*Pin, GenerationContext)
+
+	CheckNumOutputs(*Pin, GenerationContext);
+	
+	const UEdGraphSchema_CustomizableObject* Schema = GetDefault<UEdGraphSchema_CustomizableObject>();
+
+	UCustomizableObjectNode* Node = CastChecked<UCustomizableObjectNode>(Pin->GetOwningNode());
+
+	const FGeneratedKey Key(reinterpret_cast<void*>(&GenerateMutableSourceImage), *Pin, *Node, GenerationContext, true);
+	if (const FGeneratedData* Generated = GenerationContext.Generated.Find(Key))
+	{
+		return static_cast<mu::NodeImage*>(Generated->Node.get());
+	}
+
+	mu::NodeImagePtr Result;
+	
+	if (const UCustomizableObjectNodeTexture* TypedNodeTex = Cast<UCustomizableObjectNodeTexture>(Node))
+	{
+		UTexture2D* BaseTexture = TypedNodeTex->Texture;
+		if (BaseTexture)
+		{
+			// Check the specific image cache
+			FGeneratedImageKey imageKey = FGeneratedImageKey(Pin);
+			mu::NodeImageConstantPtr ImageNode;
+			mu::NodeImageConstantPtr* Cached = GenerationContext.GeneratedImages.Find(imageKey);
+			if (Cached)
+			{
+				ImageNode = *Cached;
+			}
+			else
+			{
+				ImageNode = new mu::NodeImageConstant();
+				GenerationContext.ArrayTextureUnrealToMutableTask.Add(FTextureUnrealToMutableTask(ImageNode, BaseTexture, Node));
+				GenerationContext.GeneratedImages.Add(imageKey, ImageNode);
+			}
+
+			Result = ResizeToMaxTextureSize(MaxTextureSize, BaseTexture, ImageNode);
+
+		}
+		else
+		{
+			GenerationContext.Compiler->CompilerLog(LOCTEXT("MissingImage", "Missing image in texture node."), Node);
+		}
+	}
+
+	else if (const UCustomizableObjectNodeTextureParameter* TypedNodeParam = Cast<UCustomizableObjectNodeTextureParameter>(Node))
+	{
+		mu::NodeImageParameterPtr TextureNode = new mu::NodeImageParameter();
+
+		GenerationContext.AddParameterNameUnique(Node, TypedNodeParam->ParameterName);
+
+		TextureNode->SetName(TCHAR_TO_ANSI(*TypedNodeParam->ParameterName));
+		TextureNode->SetUid(TCHAR_TO_ANSI(*GenerationContext.GetNodeIdUnique(Node).ToString()));
+
+		// TODO: Set a default value for the texture parameter?
+		//ColorNode->SetDefaultValue(TypedNodeParam->DefaultValue.R, TypedNodeParam->DefaultValue.G, TypedNodeParam->DefaultValue.B);
+
+		GenerationContext.ParameterUIDataMap.Add(TypedNodeParam->ParameterName, FParameterUIData(
+			TypedNodeParam->ParameterName,
+			TypedNodeParam->ParamUIMetadata,
+			EMutableParameterType::Texture));
+
+		// Force the same format that the default texture if any.
+		mu::NodeImageFormatPtr FormatNode = new mu::NodeImageFormat();
+		// \TODO: Take if from default?
+		//if (TypedNodeParam->DefaultValue)
+		//{
+		//	FormatNode->SetFormat(mu::EImageFormat::IF_RGBA_UBYTE);
+		//}
+		//else
+		{
+			// Force an "easy format" on the texture.
+			FormatNode->SetFormat(mu::EImageFormat::IF_RGBA_UBYTE);
+		}
+		FormatNode->SetSource(TextureNode);
+
+		mu::NodeImageResizePtr ResizeNode = new mu::NodeImageResize();
+		ResizeNode->SetBase(FormatNode);
+		ResizeNode->SetRelative(false);
+
+		if (TypedNodeParam->DefaultValue)
+		{
+			ResizeNode->SetSize(TypedNodeParam->DefaultValue->GetResource()->GetSizeX(), TypedNodeParam->DefaultValue->GetResource()->GetSizeY() );
+		}
+		else
+		{
+			// \TODO: Let the user specify this in the node?
+			ResizeNode->SetSize(1024,1024);
+		}
+
+		Result = ResizeNode;
+	}
+
+	else if (const UCustomizableObjectNodeMesh* TypedNodeMesh = Cast<UCustomizableObjectNodeMesh>(Node))
+	{
+		mu::NodeImageConstantPtr ImageNode = new mu::NodeImageConstant();
+		Result = ImageNode;
+
+		UTexture2D* Texture = TypedNodeMesh->FindTextureForPin(Pin);
+
+		if (Texture)
+		{
+			GenerationContext.ArrayTextureUnrealToMutableTask.Add(FTextureUnrealToMutableTask(ImageNode, Texture, Node));
+		}
+
+		Result = ResizeToMaxTextureSize(MaxTextureSize, Texture, ImageNode);
+	}
+
+	else if (const UCustomizableObjectNodeTextureInterpolate* TypedNodeInterp = Cast<UCustomizableObjectNodeTextureInterpolate>(Node))
+	{
+		mu::NodeImageInterpolatePtr ImageNode = new mu::NodeImageInterpolate();
+		Result = ImageNode;
+
+		if (const UEdGraphPin* ConnectedPin = FollowInputPin(*TypedNodeInterp->FactorPin()))
+		{
+			mu::NodeScalarPtr FactorNode = GenerateMutableSourceFloat(ConnectedPin, GenerationContext);
+			ImageNode->SetFactor(FactorNode);
+		}
+
+		int currentTarget = 0;
+		for (int LayerIndex = 0; LayerIndex < TypedNodeInterp->GetNumTargets(); ++LayerIndex)
+		{
+			if (const UEdGraphPin* ConnectedPin = FollowInputPin(*TypedNodeInterp->Targets(LayerIndex)))
+			{
+				mu::NodeImagePtr TargetNode = GenerateMutableSourceImage(ConnectedPin, GenerationContext, MaxTextureSize);
+
+				if (TargetNode)
+				{
+					ImageNode->SetTargetCount(currentTarget + 1);
+					ImageNode->SetTarget(currentTarget, TargetNode);
+					++currentTarget;
+				}
+			}
+		}
+	}
+
+	else if (const UCustomizableObjectNodeTextureLayer* TypedNodeLayer = Cast<UCustomizableObjectNodeTextureLayer>(Node))
+	{
+		if (const UEdGraphPin* ConnectedPin = FollowInputPin(*TypedNodeLayer->BasePin()))
+		{
+			Result = GenerateMutableSourceImage(ConnectedPin, GenerationContext, MaxTextureSize);
+		}
+
+		for (int LayerIndex = 0; LayerIndex < TypedNodeLayer->GetNumLayers(); ++LayerIndex)
+		{
+			if (const UEdGraphPin* OtherPin = FollowInputPin(*TypedNodeLayer->LayerPin(LayerIndex)))
+			{
+				mu::NodeImagePtr MaskNode = nullptr;
+				if (const UEdGraphPin* ConnectedPin = FollowInputPin(*TypedNodeLayer->MaskPin(LayerIndex)))
+				{
+					MaskNode = GenerateMutableSourceImage(ConnectedPin, GenerationContext, MaxTextureSize);
+				}
+
+				mu::EBlendType Type = mu::EBlendType::BT_BLEND;
+				switch (TypedNodeLayer->Layers[LayerIndex].Effect)
+				{
+				case COTLE_MODULATE: Type = mu::EBlendType::BT_BLEND; break;
+				case COTLE_MULTIPLY: Type = mu::EBlendType::BT_MULTIPLY; break;
+				case COTLE_SOFTLIGHT: Type = mu::EBlendType::BT_SOFTLIGHT; break;
+				case COTLE_HARDLIGHT: Type = mu::EBlendType::BT_HARDLIGHT; break;
+				case COTLE_DODGE: Type = mu::EBlendType::BT_DODGE; break;
+				case COTLE_BURN: Type = mu::EBlendType::BT_BURN; break;
+				case COTLE_SCREEN: Type = mu::EBlendType::BT_SCREEN; break;
+				case COTLE_OVERLAY: Type = mu::EBlendType::BT_OVERLAY; break;
+				case COTLE_ALPHA_OVERLAY: Type = mu::EBlendType::BT_ALPHA_OVERLAY; break;
+				case COTLE_NORMAL_COMBINE: Type = mu::EBlendType::BT_NORMAL_COMBINE; break;
+				default:
+					GenerationContext.Compiler->CompilerLog(LOCTEXT("UnsupportedImageEffect", "Texture layer effect not supported. Setting to 'Blend'."), Node);
+					break;
+				}
+
+				if (Type == mu::EBlendType::BT_BLEND && !MaskNode)
+				{
+					GenerationContext.Compiler->CompilerLog(LOCTEXT("ModulateWithoutMask", "Texture layer effect uses Modulate without a mask. It will replace everything below it!"), Node);
+				}
+				
+				if (OtherPin->PinType.PinCategory == Helper_GetPinCategory(Schema->PC_Image))
+				{
+					mu::NodeImagePtr BlendNode = GenerateMutableSourceImage(OtherPin, GenerationContext, MaxTextureSize);
+
+					mu::NodeImageLayerPtr LayerNode = new mu::NodeImageLayer;
+					LayerNode->SetBlendType(Type);
+					LayerNode->SetBase(Result);
+					LayerNode->SetBlended(BlendNode);
+					LayerNode->SetMask(MaskNode);
+					Result = LayerNode;
+				}
+
+				else if (OtherPin->PinType.PinCategory == Helper_GetPinCategory(Schema->PC_Color))
+				{
+					mu::NodeColourPtr ColorNode = GenerateMutableSourceColor(OtherPin, GenerationContext);
+
+					mu::NodeImageLayerColourPtr LayerNode = new mu::NodeImageLayerColour;
+					LayerNode->SetBlendType(Type);
+					LayerNode->SetBase(Result);
+					LayerNode->SetColour(ColorNode);
+					LayerNode->SetMask(MaskNode);
+					Result = LayerNode;
+				}
+
+				// We need it here because we create multiple nodes.
+				Result->SetMessageContext(Node);
+			}
+		}
+	}
+
+	else if (const UCustomizableObjectNodeTextureSwitch* TypedNodeTextureSwitch = Cast<UCustomizableObjectNodeTextureSwitch>(Node))
+	{
+		Result = [&]()
+		{
+			const UEdGraphPin* SwitchParameter = TypedNodeTextureSwitch->SwitchParameter();
+
+			// Check Switch Parameter arity preconditions.
+			if (const UEdGraphPin* ConnectedPin = FollowInputPin(*SwitchParameter))
+			{
+				mu::NodeScalarPtr SwitchParam = GenerateMutableSourceFloat(ConnectedPin, GenerationContext);
+				// Switch Param not generated
+				if (!SwitchParam)
+				{
+					const FText Message = LOCTEXT("FailedToGenerateSwitchParam", "Could not generate switch enum parameter. Please refesh the switch node and connect an enum.");
+					GenerationContext.Compiler->CompilerLog(Message, Node);
+					
+					return Result;
+				}
+
+				if (SwitchParam->GetType() != mu::NodeScalarEnumParameter::GetStaticType())
+				{
+					const FText Message = LOCTEXT("WrongSwitchParamType", "Switch parameter of incorrect type.");
+					GenerationContext.Compiler->CompilerLog(Message, Node);
+
+					return Result;
+				}
+
+				const int32 NumSwitchOptions = TypedNodeTextureSwitch->GetNumElements();
+
+				mu::NodeScalarEnumParameter* EnumParameter = static_cast<mu::NodeScalarEnumParameter*>(SwitchParam.get());
+				if (NumSwitchOptions != EnumParameter->GetValueCount())
+				{
+					const FText Message = LOCTEXT("MismatchedSwitch", "Switch enum and switch node have different number of options. Please refresh the switch node to make sure the outcomes are labeled properly.");
+					GenerationContext.Compiler->CompilerLog(Message, Node);
+				}
+
+				mu::NodeImageSwitchPtr SwitchNode = new mu::NodeImageSwitch;
+				SwitchNode->SetParameter(SwitchParam);
+				SwitchNode->SetOptionCount(NumSwitchOptions);
+
+				for (int SelectorIndex = 0; SelectorIndex < NumSwitchOptions; ++SelectorIndex)
+				{
+					if (const UEdGraphPin* TexturePin = FollowInputPin(*TypedNodeTextureSwitch->GetElementPin(SelectorIndex)))
+					{
+						SwitchNode->SetOption(SelectorIndex, GenerateMutableSourceImage(TexturePin, GenerationContext, MaxTextureSize));
+					}
+					else
+					{
+						const FText Message = LOCTEXT("MissingTexture", "Unable to generate texture switch node. Required connection not found.");
+						GenerationContext.Compiler->CompilerLog(Message, Node);
+						return Result;
+					}
+				}
+
+				Result = SwitchNode;
+				return Result;
+			}
+			else
+			{
+				GenerationContext.Compiler->CompilerLog(LOCTEXT("NoEnumParamInSwitch", "Switch nodes must have an enum switch parameter. Please connect an enum and refesh the switch node."), Node);
+				return Result;
+			}
+		}(); // invoke lambda.
+	}
+
+	else if (const UCustomizableObjectNodeTextureVariation* TypedNodeImageVar = Cast<const UCustomizableObjectNodeTextureVariation>(Node))
+	{
+		mu::NodeImageVariationPtr TextureNode = new mu::NodeImageVariation();
+		Result = TextureNode;
+
+		if (const UEdGraphPin* ConnectedPin = FollowInputPin(*TypedNodeImageVar->DefaultPin()))
+		{
+			mu::NodeImagePtr ChildNode = GenerateMutableSourceImage(ConnectedPin, GenerationContext, MaxTextureSize);
+			if (ChildNode)
+			{
+				TextureNode->SetDefaultImage(ChildNode.get());
+			}
+			else
+			{
+				GenerationContext.Compiler->CompilerLog(LOCTEXT("TextureFailed", "Texture generation failed."), Node);
+			}
+		}
+		else
+		{
+			GenerationContext.Compiler->CompilerLog(LOCTEXT("TextureVarMissingDef", "Texture variation node requires a default value."), Node);
+		}
+
+		TextureNode->SetVariationCount(TypedNodeImageVar->Variations.Num());
+		for (int VariationIndex = 0; VariationIndex < TypedNodeImageVar->Variations.Num(); ++VariationIndex)
+		{
+			const UEdGraphPin* VariationPin = TypedNodeImageVar->VariationPin(VariationIndex);
+			if (!VariationPin) continue;
+
+			TextureNode->SetVariationTag(VariationIndex, TCHAR_TO_ANSI(*TypedNodeImageVar->Variations[VariationIndex].Tag));
+			if (const UEdGraphPin* ConnectedPin = FollowInputPin(*VariationPin))
+			{
+				mu::NodeImagePtr ChildNode = GenerateMutableSourceImage(ConnectedPin, GenerationContext, MaxTextureSize);
+				TextureNode->SetVariationImage(VariationIndex, ChildNode.get());
+			}
+		}
+	}
+
+	else if (const UCustomizableObjectNodeTextureFromColor* TypedNodeFromColor = Cast<UCustomizableObjectNodeTextureFromColor>(Node))
+	{
+		mu::NodeColourPtr color;
+
+		if (const UEdGraphPin* ConnectedPin = FollowInputPin(*TypedNodeFromColor->ColorPin()))
+		{
+			color = GenerateMutableSourceColor(ConnectedPin, GenerationContext);
+		}
+
+		mu::NodeImagePlainColourPtr ImageFromColour = new mu::NodeImagePlainColour;
+		Result = ImageFromColour;
+
+		if (color)
+		{
+			ImageFromColour->SetColour(color);
+		}
+	}
+
+	else if (const UCustomizableObjectNodeTextureFromChannels* TypedNodeFrom = Cast<UCustomizableObjectNodeTextureFromChannels>(Node))
+	{
+		mu::NodeImagePtr RNode, GNode, BNode, ANode;
+		bool RGB = false;
+
+		if (const UEdGraphPin* ConnectedPin = FollowInputPin(*TypedNodeFrom->RPin()))
+		{
+			RNode = GenerateMutableSourceImage(ConnectedPin, GenerationContext, MaxTextureSize);
+			RGB = true;
+		}
+		if (const UEdGraphPin* ConnectedPin = FollowInputPin(*TypedNodeFrom->GPin()))
+		{
+			GNode = GenerateMutableSourceImage(ConnectedPin, GenerationContext, MaxTextureSize);
+			RGB = true;
+		}
+		if (const UEdGraphPin* ConnectedPin = FollowInputPin(*TypedNodeFrom->BPin()))
+		{
+			BNode = GenerateMutableSourceImage(ConnectedPin, GenerationContext, MaxTextureSize);
+			RGB = true;
+		}
+		if (const UEdGraphPin* ConnectedPin = FollowInputPin(*TypedNodeFrom->APin()))
+		{
+			ANode = GenerateMutableSourceImage(ConnectedPin, GenerationContext, MaxTextureSize);
+		}
+
+		mu::NodeImageSwizzlePtr SwizzleNode = new mu::NodeImageSwizzle;
+		Result = SwizzleNode;
+
+		if (RGB && ANode)
+		{
+			SwizzleNode->SetFormat(mu::EImageFormat::IF_RGBA_UBYTE);
+			SwizzleNode->SetSource(0, RNode);
+			SwizzleNode->SetSourceChannel(0, 0);
+			SwizzleNode->SetSource(1, GNode);
+			SwizzleNode->SetSourceChannel(1, 0);
+			SwizzleNode->SetSource(2, BNode);
+			SwizzleNode->SetSourceChannel(2, 0);
+			SwizzleNode->SetSource(3, ANode);
+			SwizzleNode->SetSourceChannel(3, 0);
+		}
+		else if (RGB)
+		{
+			SwizzleNode->SetFormat(mu::EImageFormat::IF_RGB_UBYTE);
+			SwizzleNode->SetSource(0, RNode);
+			SwizzleNode->SetSourceChannel(0, 0);
+			SwizzleNode->SetSource(1, GNode);
+			SwizzleNode->SetSourceChannel(1, 0);
+			SwizzleNode->SetSource(2, BNode);
+			SwizzleNode->SetSourceChannel(2, 0);
+		}
+		else if (RNode)
+		{
+			SwizzleNode->SetFormat(mu::EImageFormat::IF_L_UBYTE);
+			SwizzleNode->SetSource(0, RNode);
+			SwizzleNode->SetSourceChannel(0, 0);
+		}
+		else if (ANode)
+		{
+			SwizzleNode->SetFormat(mu::EImageFormat::IF_L_UBYTE);
+			SwizzleNode->SetSource(0, ANode);
+			SwizzleNode->SetSourceChannel(0, 0);
+		}
+	}
+
+	else if (const UCustomizableObjectNodeTextureToChannels* TypedNodeTo = Cast<UCustomizableObjectNodeTextureToChannels>(Node))
+	{
+		mu::NodeImagePtr BaseNode;
+		if (const UEdGraphPin* ConnectedPin = FollowInputPin(*TypedNodeTo->InputPin()))
+		{
+			BaseNode = GenerateMutableSourceImage(ConnectedPin, GenerationContext, MaxTextureSize);
+		}
+
+		mu::NodeImageSwizzlePtr SwizzleNode = new mu::NodeImageSwizzle;
+		Result = SwizzleNode;
+		SwizzleNode->SetFormat(mu::EImageFormat::IF_L_UBYTE);
+		SwizzleNode->SetSource(0, BaseNode);
+
+		if (Pin == TypedNodeTo->RPin())
+		{
+			SwizzleNode->SetSourceChannel(0, 0);
+		}
+		else if (Pin == TypedNodeTo->GPin())
+		{
+			SwizzleNode->SetSourceChannel(0, 1);
+		}
+		else if (Pin == TypedNodeTo->BPin())
+		{
+			SwizzleNode->SetSourceChannel(0, 2);
+		}
+		else if (Pin == TypedNodeTo->APin())
+		{
+			SwizzleNode->SetSourceChannel(0, 3);
+		}
+		else
+		{
+			check(false);
+		}
+	}
+
+	else if (const UCustomizableObjectNodeTextureProject* TypedNodeProject = Cast<UCustomizableObjectNodeTextureProject>(Node))
+	{
+		mu::NodeImageProjectPtr ImageNode = new mu::NodeImageProject();
+		Result = ImageNode;
+
+		if (!FollowInputPin(*TypedNodeProject->MeshPin()))
+		{
+			GenerationContext.Compiler->CompilerLog(LOCTEXT("MissingMeshInProjector", "Texture projector does not have a Mesh. It will be ignored. "), Node, EMessageSeverity::Warning);
+			Result = nullptr;
+		}
+
+		ImageNode->SetLayout(TypedNodeProject->Layout);
+		ImageNode->SetImageSize( FUintVector2(TypedNodeProject->TextureSizeX, TypedNodeProject->TextureSizeY) );
+
+		if (const UEdGraphPin* ConnectedPin = FollowInputPin(*TypedNodeProject->AngleFadeStartPin()))
+		{
+			mu::NodeScalarPtr ScalarNode = GenerateMutableSourceFloat(ConnectedPin, GenerationContext);
+			ImageNode->SetAngleFadeStart(ScalarNode);
+		}
+		else
+		{
+			mu::NodeScalarConstantPtr ScalarNode = new mu::NodeScalarConstant;
+			float value = FCString::Atof(*TypedNodeProject->AngleFadeStartPin()->DefaultValue);
+			ScalarNode->SetValue(value);
+			ImageNode->SetAngleFadeStart(ScalarNode);
+		}
+
+		if (const UEdGraphPin* ConnectedPin = FollowInputPin(*TypedNodeProject->AngleFadeEndPin()))
+		{
+			mu::NodeScalarPtr ScalarNode = GenerateMutableSourceFloat(ConnectedPin, GenerationContext);
+			ImageNode->SetAngleFadeEnd(ScalarNode);
+		}
+		else
+		{
+			mu::NodeScalarConstantPtr ScalarNode = new mu::NodeScalarConstant;
+			float value = FCString::Atof(*TypedNodeProject->AngleFadeEndPin()->DefaultValue);
+			ScalarNode->SetValue(value);
+			ImageNode->SetAngleFadeEnd(ScalarNode);
+		}
+
+		if (const UEdGraphPin* ConnectedPin = FollowInputPin(*TypedNodeProject->MeshPin()))
+		{
+			FMutableGraphMeshGenerationData DummyMeshData;
+			mu::NodeMeshPtr MeshNode = GenerateMutableSourceMesh(ConnectedPin, GenerationContext, DummyMeshData);
+			ImageNode->SetMesh(MeshNode);
+		}
+
+		if (const UEdGraphPin* ConnectedPin = FollowInputPin(*TypedNodeProject->MeshMaskPin()))
+		{
+			mu::NodeImagePtr MeshMaskNode = GenerateMutableSourceImage(ConnectedPin, GenerationContext, MaxTextureSize);
+			ImageNode->SetTargetMask(MeshMaskNode);
+		}
+
+		if (const UEdGraphPin* ConnectedPin = FollowInputPin(*TypedNodeProject->ProjectorPin()))
+		{
+			mu::NodeProjectorPtr ProjectorNode = GenerateMutableSourceProjector(ConnectedPin, GenerationContext);
+			ImageNode->SetProjector(ProjectorNode);
+		}
+
+		int TexIndex = -1;// TypedNodeProject->OutputPins.Find((UEdGraphPin*)Pin);
+
+		for (int32 i = 0; i < TypedNodeProject->GetNumOutputs(); ++i)
+		{
+			if (TypedNodeProject->OutputPins(i) == Pin)
+			{
+				TexIndex = i;
+			}
+		}
+
+		check(TexIndex >= 0 && TexIndex < TypedNodeProject->GetNumTextures());
+
+		if (const UEdGraphPin* ConnectedPin = FollowInputPin(*TypedNodeProject->TexturePins(TexIndex)))
+		{
+			mu::NodeImagePtr SourceNode = GenerateMutableSourceImage(ConnectedPin, GenerationContext, MaxTextureSize);
+
+			// Images that are projected won't be scaled by any means, so we need to apply the lodbias settings here
+			if (GenerationContext.CurrentTextureLODBias > 0)
+			{
+				mu::NodeImageResizePtr ResizeImage = new mu::NodeImageResize();
+				ResizeImage->SetBase(SourceNode.get());
+				ResizeImage->SetRelative(true);
+				float factor = FMath::Pow(0.5f, GenerationContext.CurrentTextureLODBias);
+				ResizeImage->SetSize(factor, factor);
+				ResizeImage->SetMessageContext(Node);
+				SourceNode = ResizeImage;
+			}
+
+			ImageNode->SetImage(SourceNode);
+		}
+	}
+
+	else if (const UCustomizableObjectNodeTextureBinarise* TypedNodeTexBin = Cast<UCustomizableObjectNodeTextureBinarise>(Node))
+	{
+		mu::NodeImageBinarisePtr BinariseNode = new mu::NodeImageBinarise();
+		Result = BinariseNode;
+
+		if (const UEdGraphPin* ConnectedPin = FollowInputPin(*TypedNodeTexBin->GetBaseImagePin()))
+		{
+			mu::NodeImagePtr BaseImageNode = GenerateMutableSourceImage(ConnectedPin, GenerationContext, MaxTextureSize);
+			BinariseNode->SetBase(BaseImageNode);
+		}
+
+		if (const UEdGraphPin* ConnectedPin = FollowInputPin(*TypedNodeTexBin->GetThresholdPin()))
+		{
+			mu::NodeScalarPtr ThresholdNode = GenerateMutableSourceFloat(ConnectedPin, GenerationContext);
+			BinariseNode->SetThreshold(ThresholdNode);
+		}
+	}
+
+	else if (const UCustomizableObjectNodeTextureInvert* TypedNodeTexInv = Cast<UCustomizableObjectNodeTextureInvert>(Node))
+	{
+		mu::NodeImageInvertPtr InvertNode = new mu::NodeImageInvert();
+		Result = InvertNode;
+
+		if (const UEdGraphPin* ConnectedPin = FollowInputPin(*TypedNodeTexInv->GetBaseImagePin()))
+		{
+			mu::NodeImagePtr BaseImageNode = GenerateMutableSourceImage(ConnectedPin, GenerationContext, MaxTextureSize);
+			InvertNode->SetBase(BaseImageNode);
+		}
+	}
+
+	else if (const UCustomizableObjectNodeTextureColourMap* TypedNodeColourMap = Cast<UCustomizableObjectNodeTextureColourMap>(Node))
+	{
+		mu::NodeImageColourMapPtr ColourMapNode = new mu::NodeImageColourMap();
+
+		Result = ColourMapNode;
+
+		if (const UEdGraphPin* ConnectedPin = FollowInputPin(*TypedNodeColourMap->GetMapPin()))
+		{
+			mu::NodeImagePtr GradientImageNode = GenerateMutableSourceImage(ConnectedPin, GenerationContext, MaxTextureSize);
+			ColourMapNode->SetMap( GradientImageNode ); 
+		}
+
+		if (const UEdGraphPin* ConnectedPin = FollowInputPin(*TypedNodeColourMap->GetMaskPin()))
+		{
+			mu::NodeImagePtr GradientImageNode = GenerateMutableSourceImage(ConnectedPin, GenerationContext, MaxTextureSize);
+			ColourMapNode->SetMask( GradientImageNode ); 
+		}
+
+		if (const UEdGraphPin* ConnectedPin = FollowInputPin(*TypedNodeColourMap->GetBasePin()))
+		{
+			mu::NodeImagePtr SourceImageNode = GenerateMutableSourceImage(ConnectedPin, GenerationContext, MaxTextureSize);
+			ColourMapNode->SetBase( SourceImageNode );
+		}
+	}
+
+	else if (const UCustomizableObjectNodeTextureTransform* TypedNodeTransform = Cast<UCustomizableObjectNodeTextureTransform>(Node))
+	{
+		mu::NodeImageTransformPtr TransformNode = new mu::NodeImageTransform();
+		Result = TransformNode;
+		
+		if ( UEdGraphPin* BaseImagePin = FollowInputPin(*TypedNodeTransform->GetBaseImagePin()) )
+		{
+			mu::NodeImagePtr ImageNode = GenerateMutableSourceImage( BaseImagePin, GenerationContext, MaxTextureSize);
+			TransformNode->SetBase( ImageNode ); 
+		}
+
+		if ( UEdGraphPin* OffsetXPin = FollowInputPin(*TypedNodeTransform->GetOffsetXPin()) )
+		{
+			mu::NodeScalarPtr OffsetXNode = GenerateMutableSourceFloat(OffsetXPin, GenerationContext);
+			TransformNode->SetOffsetX( OffsetXNode ); 
+		}
+
+		if ( UEdGraphPin* OffsetYPin = FollowInputPin(*TypedNodeTransform->GetOffsetYPin()) )
+		{
+			mu::NodeScalarPtr OffsetYNode = GenerateMutableSourceFloat(OffsetYPin, GenerationContext);
+			TransformNode->SetOffsetY( OffsetYNode ); 
+		}
+
+		if ( UEdGraphPin* ScaleXPin = FollowInputPin(*TypedNodeTransform->GetScaleXPin()) )
+		{
+			mu::NodeScalarPtr ScaleXNode = GenerateMutableSourceFloat(ScaleXPin, GenerationContext);
+			TransformNode->SetScaleX( ScaleXNode ); 
+		}
+
+		if ( UEdGraphPin* ScaleYPin = FollowInputPin(*TypedNodeTransform->GetScaleYPin()) )
+		{
+			mu::NodeScalarPtr ScaleYNode = GenerateMutableSourceFloat(ScaleYPin, GenerationContext);
+			TransformNode->SetScaleY( ScaleYNode ); 
+		}
+
+		if ( UEdGraphPin* RotationPin = FollowInputPin(*TypedNodeTransform->GetRotationPin()) )
+		{
+			mu::NodeScalarPtr RotationNode = GenerateMutableSourceFloat(RotationPin, GenerationContext);
+			TransformNode->SetRotation( RotationNode ); 
+		}
+	}
+	// If the node is a plain colour node, generate an image out of it
+	else if (Pin->PinType.PinCategory == Helper_GetPinCategory(Schema->PC_Color))
+	{
+		mu::NodeColourPtr ColorNode = GenerateMutableSourceColor(Pin, GenerationContext);
+
+		mu::NodeImagePlainColourPtr ImageNode = new mu::NodeImagePlainColour;
+		ImageNode->SetSize(16, 16);
+		ImageNode->SetColour(ColorNode);
+		Result = ImageNode;
+	}
+
+	else if (const UCustomizableObjectNodeTable* TypedNodeTable = Cast<UCustomizableObjectNodeTable>(Node))
+	{
+		if (Pin->PinType.PinCategory == Schema->PC_Color)
+		{
+			mu::NodeColourPtr ColorNode = GenerateMutableSourceColor(Pin, GenerationContext);
+			
+			mu::NodeImagePlainColourPtr ImageNode = new mu::NodeImagePlainColour;
+
+			ImageNode->SetSize(16, 16);
+			ImageNode->SetColour(ColorNode);
+
+			Result = ImageNode;
+		}
+
+		else
+		{
+			mu::NodeImageTablePtr ImageTableNode = new mu::NodeImageTable();
+			Result = ImageTableNode;
+
+			mu::TablePtr Table = nullptr;
+
+			if (TypedNodeTable->Table)
+			{
+				FString ColumnName = Pin->PinFriendlyName.ToString();
+
+				if (Pin->PinType.PinCategory == Schema->PC_MaterialAsset)
+				{
+					ColumnName = GenerationContext.CurrentMaterialTableParameterId;
+				}
+
+				Table = GenerateMutableSourceTable(TypedNodeTable->Table->GetName(), Pin, GenerationContext);
+				
+				ImageTableNode->SetTable(Table);
+				ImageTableNode->SetColumn(TCHAR_TO_ANSI(*ColumnName));
+				ImageTableNode->SetParameterName(TCHAR_TO_ANSI(*TypedNodeTable->ParameterName));
+
+				GenerationContext.AddParameterNameUnique(Node, TypedNodeTable->ParameterName);
+				
+				if (Table->FindColumn(TCHAR_TO_ANSI(*ColumnName)) == -1)
+				{
+					FString Msg = FString::Printf(TEXT("Couldn't find pin column with name %s"), *ColumnName);
+					GenerationContext.Compiler->CompilerLog(FText::FromString(Msg), Node);
+				}
+			}
+			else
+			{
+				Table = new mu::Table();
+				ImageTableNode->SetTable(Table);
+
+				GenerationContext.Compiler->CompilerLog(LOCTEXT("ImageTableError", "Couldn't find the data table of the node."), Node);
+			}
+		}
+	}
+
+	else
+	{
+		GenerationContext.Compiler->CompilerLog(LOCTEXT("UnimplementedNode", "Node type not implemented yet."), Node);
+	}
+
+	GenerationContext.Generated.Add(Key, FGeneratedData(Node, Result));
+	GenerationContext.GeneratedNodes.Add(Node);
+
+	if (Result)
+	{
+		Result->SetMessageContext(Node);
+	}
+
+	return Result;
+}
diff --git a/Engine/Plugins/Experimental/Mutable/Source/CustomizableObjectEditor/Private/GenerateMutableSource/GenerateMutableSourceImage.h b/Engine/Plugins/Experimental/Mutable/Source/CustomizableObjectEditor/Private/GenerateMutableSource/GenerateMutableSourceImage.h
new file mode 100644
--- /dev/null
+++ b/Engine/Plugins/Experimental/Mutable/Source/CustomizableObjectEditor/Private/GenerateMutableSource/GenerateMutableSourceImage.h
@@ -0,0 +1,15 @@
+// Copyright Epic Games, Inc. All Rights Reserved.
+
+#pragma once
+
+#include "GenerateMutableSource/GenerateMutableSource.h"
+
+
+mu::ImagePtr ConvertTextureUnrealToMutable(UTexture2D* Texture, const UCustomizableObjectNode* Node, FCustomizableObjectCompiler* Compiler, bool bIsNormalComposite);
+
+
+mu::NodeImagePtr ResizeToMaxTextureSize(float MaxTextureSize, const UTexture2D* BaseTexture, mu::NodeImageConstantPtr ImageNode);
+
+
+/** Convert a CustomizableObject Source Graph into a mutable source graph. */
+mu::NodeImagePtr GenerateMutableSourceImage(const UEdGraphPin* Pin, FMutableGraphGenerationContext& GenerationContext, float MaxTextureSize);
\ No newline at end of file
diff --git a/Engine/Plugins/Experimental/Mutable/Source/CustomizableObjectEditor/Private/GenerateMutableSource/GenerateMutableSourceLayout.cpp b/Engine/Plugins/Experimental/Mutable/Source/CustomizableObjectEditor/Private/GenerateMutableSource/GenerateMutableSourceLayout.cpp
new file mode 100644
--- /dev/null
+++ b/Engine/Plugins/Experimental/Mutable/Source/CustomizableObjectEditor/Private/GenerateMutableSource/GenerateMutableSourceLayout.cpp
@@ -0,0 +1,122 @@
+// Copyright Epic Games, Inc. All Rights Reserved.
+
+#include "GenerateMutableSourceLayout.h"
+
+#include "GenerateMutableSource/GenerateMutableSourceFloat.h"
+
+#include "Templates/UnrealTemplate.h"
+#include "Engine/AssetManager.h"
+#include "Engine/TextureLODSettings.h"
+#include "ClothingAsset.h"
+#include "EdGraphToken.h"
+#include "AssetRegistry/AssetRegistryModule.h"
+#include "UnrealEditorPortabilityHelpers.h"
+#include "PlatformInfo.h"
+
+#include "CustomizableObjectCompiler.h"
+#include "UnrealToMutableTextureConversionUtils.h"
+#include "UnrealConversionUtils.h"
+#include "CustomizableObjectEditorModule.h"
+#include "CustomizableObject.h"
+#include "GraphTraversal.h"
+
+#include "MutableTools/Public/Node.h"
+#include "MutableTools/Public/NodeLayout.h"
+#include "MutableTools/Public/Compiler.h"
+#include "MutableTools/Public/Streams.h"
+
+#define LOCTEXT_NAMESPACE "CustomizableObjectEditor"
+
+
+mu::NodeLayoutPtr GenerateMutableSourceLayout(const UEdGraphPin * Pin, FMutableGraphGenerationContext & GenerationContext)
+{
+	check(Pin)
+	RETURN_ON_CYCLE(*Pin, GenerationContext)
+
+	CheckNumOutputs(*Pin, GenerationContext);
+
+	UCustomizableObjectNode* Node = CastChecked<UCustomizableObjectNode>(Pin->GetOwningNode());
+
+	const FGeneratedKey Key(reinterpret_cast<void*>(&GenerateMutableSourceLayout), *Pin, *Node, GenerationContext, true);
+	if (const FGeneratedData* Generated = GenerationContext.Generated.Find(Key))
+	{
+		return static_cast<mu::NodeLayout*>(Generated->Node.get());
+	}
+
+	mu::NodeLayoutPtr Result;
+	
+	if (const UCustomizableObjectNodeLayoutBlocks* TypedNodeBlocks = Cast<UCustomizableObjectNodeLayoutBlocks>(Node))
+	{
+		if (UCustomizableObjectNodeSkeletalMesh* SkeletalMeshNode = Cast<UCustomizableObjectNodeSkeletalMesh>(FollowOutputPin(*TypedNodeBlocks->OutputPin())->GetOwningNode()))
+		{
+			int32 LayoutIndex;
+			FString MaterialName;
+
+			if (!SkeletalMeshNode->CheckIsValidLayout(Pin, LayoutIndex, MaterialName))
+			{
+				FString msg = "Layouts ";
+				for (int32 i = 0; i < LayoutIndex; ++i)
+				{
+					msg += "UV" + FString::FromInt(i);
+					if (i < LayoutIndex - 1)
+					{
+						msg += ", ";
+					}
+				}
+				msg += " of " + MaterialName + " must be also connected to a Layout Blocks Node. ";
+				GenerationContext.Compiler->CompilerLog(FText::FromString(msg), Node, EMessageSeverity::Error);
+				return nullptr;
+			}
+		}
+
+		mu::NodeLayoutBlocksPtr LayoutNode = new mu::NodeLayoutBlocks;
+		Result = LayoutNode;
+
+		LayoutNode->SetGridSize(TypedNodeBlocks->Layout->GetGridSize().X, TypedNodeBlocks->Layout->GetGridSize().Y);
+		LayoutNode->SetMaxGridSize(TypedNodeBlocks->Layout->GetMaxGridSize().X, TypedNodeBlocks->Layout->GetMaxGridSize().Y);
+		LayoutNode->SetBlockCount(TypedNodeBlocks->Layout->Blocks.Num());
+
+		mu::EPackStrategy strategy = mu::EPackStrategy::RESIZABLE_LAYOUT;
+
+		switch (TypedNodeBlocks->Layout->GetPackingStrategy())
+		{
+		case ECustomizableObjectTextureLayoutPackingStrategy::Resizable:
+			strategy = mu::EPackStrategy::RESIZABLE_LAYOUT;
+			break;
+		case ECustomizableObjectTextureLayoutPackingStrategy::Fixed:
+			strategy = mu::EPackStrategy::FIXED_LAYOUT;
+			break;
+		default:
+			break;
+		}
+
+		LayoutNode->SetLayoutPackingStrategy(strategy);
+
+		for (int BlockIndex = 0; BlockIndex < TypedNodeBlocks->Layout->Blocks.Num(); ++BlockIndex)
+		{
+			LayoutNode->SetBlock(BlockIndex,
+				TypedNodeBlocks->Layout->Blocks[BlockIndex].Min.X,
+				TypedNodeBlocks->Layout->Blocks[BlockIndex].Min.Y,
+				TypedNodeBlocks->Layout->Blocks[BlockIndex].Max.X - TypedNodeBlocks->Layout->Blocks[BlockIndex].Min.X,
+				TypedNodeBlocks->Layout->Blocks[BlockIndex].Max.Y - TypedNodeBlocks->Layout->Blocks[BlockIndex].Min.Y);
+
+			LayoutNode->SetBlockPriority(BlockIndex, TypedNodeBlocks->Layout->Blocks[BlockIndex].Priority);
+		}
+	}
+	
+	else
+	{
+		GenerationContext.Compiler->CompilerLog(LOCTEXT("UnimplementedNode", "Node type not implemented yet."), Node);
+	}
+
+	GenerationContext.Generated.Add(Key, FGeneratedData(Node, Result));
+	GenerationContext.GeneratedNodes.Add(Node);
+
+	
+	if (Result)
+	{
+		Result->SetMessageContext(Node);
+	}
+
+	return Result;
+}
diff --git a/Engine/Plugins/Experimental/Mutable/Source/CustomizableObjectEditor/Private/GenerateMutableSource/GenerateMutableSourceLayout.h b/Engine/Plugins/Experimental/Mutable/Source/CustomizableObjectEditor/Private/GenerateMutableSource/GenerateMutableSourceLayout.h
new file mode 100644
--- /dev/null
+++ b/Engine/Plugins/Experimental/Mutable/Source/CustomizableObjectEditor/Private/GenerateMutableSource/GenerateMutableSourceLayout.h
@@ -0,0 +1,9 @@
+// Copyright Epic Games, Inc. All Rights Reserved.
+
+#pragma once
+
+#include "GenerateMutableSource/GenerateMutableSource.h"
+#include "NodeLayout.h"
+
+
+mu::NodeLayoutPtr GenerateMutableSourceLayout(const UEdGraphPin* Pin, FMutableGraphGenerationContext& GenerationContext);
