[mutable] Move Mutable plugin from the restricted folder to the experimental plugins folder.

diff --git a/Engine/Plugins/Experimental/Mutable/Source/CustomizableObjectEditor/Private/GenerateMutableSource/GenerateMutableSourceMesh.cpp b/Engine/Plugins/Experimental/Mutable/Source/CustomizableObjectEditor/Private/GenerateMutableSource/GenerateMutableSourceMesh.cpp
new file mode 100644
--- /dev/null
+++ b/Engine/Plugins/Experimental/Mutable/Source/CustomizableObjectEditor/Private/GenerateMutableSource/GenerateMutableSourceMesh.cpp
@@ -0,0 +1,2893 @@
+// Copyright Epic Games, Inc. All Rights Reserved.
+
+#include "GenerateMutableSourceMesh.h"
+
+#include "GenerateMutableSource/GenerateMutableSourceFloat.h"
+#include "GenerateMutableSource/GenerateMutableSourceLayout.h"
+#include "GenerateMutableSource/GenerateMutableSourceTable.h"
+
+#include "Templates/UnrealTemplate.h"
+#include "Engine/SkeletalMesh.h"
+#include "Engine/StaticMesh.h"
+#include "Engine/AssetManager.h"
+#include "Engine/TextureLODSettings.h"
+#include "PhysicsEngine/PhysicsAsset.h"
+#include "ClothingAsset.h"
+#include "MeshUtilities.h"
+#include "EdGraphToken.h"
+#include "AssetRegistry/AssetRegistryModule.h"
+#include "UnrealEditorPortabilityHelpers.h"
+#include "UnrealPortabilityHelpers.h"
+#include "PlatformInfo.h"
+#include "Templates/Decay.h"
+#include "GenericPlatform/GenericPlatformMath.h"
+#include "StaticMeshResources.h"
+#include "Animation/Skeleton.h"
+#include "Animation/MorphTarget.h"
+#include "Engine/SkeletalMeshSocket.h"
+
+#include "CustomizableObjectCompiler.h"
+#include "UnrealToMutableTextureConversionUtils.h"
+#include "UnrealConversionUtils.h"
+#include "CustomizableObjectEditorModule.h"
+#include "CustomizableObjectSystem.h"
+#include "CustomizableObject.h"
+#include "GraphTraversal.h"
+#include "MutableUtils.h"
+
+#include "MutableTools/Public/Node.h"
+#include "MutableTools/Public/NodeScalarConstant.h"
+#include "MutableTools/Public/NodeScalarEnumParameter.h"
+#include "MutableTools/Public/NodeLayout.h"
+#include "MutableTools/Public/NodeMeshConstant.h"
+#include "MutableTools/Public/NodeMeshMorph.h"
+#include "MutableTools/Public/NodeMeshMakeMorph.h"
+#include "MutableTools/Public/NodeMeshSwitch.h"
+#include "MutableTools/Public/NodeMeshApplyPose.h"
+#include "MutableTools/Public/NodeMeshVariation.h"
+#include "MutableTools/Public/NodeMeshGeometryOperation.h"
+#include "MutableTools/Public/NodeMeshReshape.h"
+#include "MutableTools/Public/NodeMeshTable.h"
+#include "MutableTools/Public/NodeSurfaceNew.h"
+#include "MutableTools/Public/Compiler.h"
+#include "MutableTools/Public/Streams.h"
+
+#define LOCTEXT_NAMESPACE "CustomizableObjectEditor"
+
+
+mu::MeshPtr ConvertSkeletalMeshToMutable(USkeletalMesh* InSkeletalMesh, int LOD, int MaterialIndex, FMutableGraphGenerationContext& GenerationContext, const UCustomizableObjectNode* CurrentNode)
+{
+	// Get the mesh generation flags to use
+	uint32 CurrentFlags = GenerationContext.MeshGenerationFlags.Last();
+	bool bIgnoreSkeleton = CurrentFlags & uint32(EMutableMeshConversionFlags::IgnoreSkinning);
+	bool bIgnorePhysics = CurrentFlags & uint32(EMutableMeshConversionFlags::IgnorePhysics);
+
+	mu::MeshPtr MutableMesh = new mu::Mesh();
+		
+	if(!InSkeletalMesh)
+	{
+		// Trying to convert a nullptr to mutable?
+		return nullptr;
+	}
+	
+	USkeleton* InSkeleton = InSkeletalMesh->GetSkeleton();
+	check(InSkeleton);
+
+	TObjectPtr<USkeletalMesh> SkeletalMesh = InSkeletalMesh;
+
+	FMutableComponentInfo& MutComponentInfo = GenerationContext.GetCurrentComponentInfo();
+	USkeletalMesh* ComponentRefSkeletalMesh = MutComponentInfo.RefSkeletalMesh;
+	USkeleton* ComponentRefSkeleton = MutComponentInfo.RefSkeleton;
+	check(ComponentRefSkeletalMesh);
+	check(ComponentRefSkeleton);
+	
+
+	// Apply removed bones
+	bool bHasBonesToRemove = false;
+	TMap<int32, int32> RemovedBonesActiveParentIndices;
+	if (!bIgnoreSkeleton)
+	{
+		for (int32 LODIndex = 0; LODIndex < Helper_GetLODNum(InSkeletalMesh); ++LODIndex)
+		{
+			if (Helper_GetLODInfoBonesToRemove(InSkeletalMesh, LODIndex).Num() > 0)
+			{
+				bHasBonesToRemove = true;
+				break;
+			}
+		}
+
+		if (bHasBonesToRemove)
+		{
+			IMeshUtilities& MeshUtilities = FModuleManager::Get().LoadModuleChecked<IMeshUtilities>("MeshUtilities");
+			// Keep the name since it is used at least when gathering info about morph targets. 
+			SkeletalMesh = (USkeletalMesh*)StaticDuplicateObject(SkeletalMesh, GetTransientPackage(), SkeletalMesh->GetFName(), EObjectFlags::RF_Transient);
+			check(SkeletalMesh);
+			
+			TMap<int32, bool> RemovedBoneIndices;
+			for (int32 LODIndex = 0; LODIndex < Helper_GetLODNum(SkeletalMesh); ++LODIndex)
+			{
+				TArray<FName> BoneNamesToRemove;
+
+				if (Helper_GetLODNum(InSkeletalMesh) > LODIndex)
+				{
+					for (const FBoneReference& BoneRef : Helper_GetLODInfoBonesToRemove(InSkeletalMesh, LODIndex))
+					{
+						BoneNamesToRemove.Add(BoneRef.BoneName);
+					}
+
+					if (BoneNamesToRemove.Num() > 0)
+					{
+						MeshUtilities.RemoveBonesFromMesh(SkeletalMesh, LODIndex, &BoneNamesToRemove);
+					}
+
+					if (LODIndex == LOD)
+					{
+						// Find Removed Bones Indices
+						for (const FName& Bone : BoneNamesToRemove)
+						{
+							int32 RefBoneIndex = SkeletalMesh->GetRefSkeleton().FindBoneIndex(Bone);
+							if (RefBoneIndex != INDEX_NONE)
+							{
+								RemovedBoneIndices.Add(RefBoneIndex, true);
+
+							}
+						}
+
+					}
+				}
+				else
+				{
+					GenerationContext.Compiler->CompilerLog(LOCTEXT("ReferenceLODMismatch", "The reference mesh has less LODs than the generated mesh."),
+						CurrentNode,
+						EMessageSeverity::Warning);
+					break;
+				}
+			}
+
+			// Reconstruct RefSkeleton's bone hierarchy
+			TMap<int32, TArray<int32>> SkeletalTree;
+			FReferenceSkeleton& RefSkeleton = InSkeletalMesh->GetRefSkeleton();
+			for (int32 i = 0; i < RefSkeleton.GetRawBoneNum(); ++i)
+			{
+				int32 ParentBoneIndex = RefSkeleton.GetRawParentIndex(i);
+				if (ParentBoneIndex == INDEX_NONE)
+				{
+					continue;
+				}
+
+				if (TArray<int32>* Children = SkeletalTree.Find(ParentBoneIndex))
+				{
+					Children->Add(i);
+				}
+				else
+				{
+					TArray<int32> NewChildrenArray;
+					NewChildrenArray.Add(i);
+					SkeletalTree.Add(ParentBoneIndex, NewChildrenArray);
+				}
+			}
+
+			// Search for the first active parent of each removed bone
+			// If a vertex is influenced by a removed bone the influences will be transferred to the parent bone.
+			for (const TPair<int32, bool>& RemovedBone : RemovedBoneIndices)
+			{
+				int32 ParentIndex = RefSkeleton.GetRawParentIndex(RemovedBone.Key);
+
+				if (TArray<int32>* Children = SkeletalTree.Find(RemovedBone.Key))
+				{
+					for (int32 i = 0; i < Children->Num(); ++i)
+					{
+						int32& ChildIndex = (*Children)[i];
+						RemovedBonesActiveParentIndices.Add(ChildIndex, ParentIndex);
+
+						if (RemovedBoneIndices.Find(ChildIndex))
+						{
+							continue;
+						}
+
+						if (TArray<int32>* ChildChildren = SkeletalTree.Find(ChildIndex))
+						{
+							for (int j = 0; j < ChildChildren->Num(); ++j)
+							{
+								Children->Add((*ChildChildren)[j]);
+							}
+						}
+					}
+				}
+
+				if (ParentIndex >= 0)
+				{
+					RemovedBonesActiveParentIndices.Add(RemovedBone.Key, ParentIndex);
+				}
+			}
+		}
+	}
+
+	const FSkeletalMeshModel* ImportedModel = Helper_GetImportedModel(SkeletalMesh);
+
+	// Check in case the data has changed
+	if (!(ImportedModel
+		&& ImportedModel->LODModels.Num() > LOD
+		&& ImportedModel->LODModels[LOD].Sections.Num() > MaterialIndex))
+	{
+		FString Msg;
+		if (!ImportedModel)
+		{
+			Msg = FString::Printf(TEXT("The SkeletalMesh [%s] doesn't have an imported resource."), *InSkeletalMesh->GetName());
+		}
+		else if (LOD >= ImportedModel->LODModels.Num())
+		{
+			Msg = FString::Printf(
+				TEXT("The SkeletalMesh [%s] doesn't have the expected number of LODs [need %d, has %d]. Changed after reimporting?"),
+				InSkeletalMesh ? *InSkeletalMesh->GetName() : TEXT("none"),
+				LOD + 1,
+				ImportedModel->LODModels.Num());
+		}
+		else
+		{
+			Msg = FString::Printf(
+				TEXT("The SkeletalMesh [%s] doesn't have the expected structure. Maybe the number of LODs [need %d, has %d] or Materials [need %d, has %d] has changed after reimporting?"),
+				*InSkeletalMesh->GetName(),
+				LOD + 1,
+				ImportedModel->LODModels.Num(),
+				MaterialIndex + 1,
+				ImportedModel->LODModels[LOD].Sections.Num()
+			);
+		}
+		GenerationContext.Compiler->CompilerLog(FText::FromString(Msg), CurrentNode);
+		return nullptr;
+	}
+
+	// Use this to set bone IDs in Mutable to specify to which skeleton they belong
+	int32 SkeletonID = 0;
+
+	// This map will be used to find and replace vertex influences.
+	TMap<int32, int32> InfluencesToReplaceMap;
+
+	// Check for a matching skeleton to the reference mesh
+	if (!bIgnoreSkeleton)
+	{
+		if (InSkeleton != ComponentRefSkeleton)
+		{
+			bool* SkeletonCompatibility = MutComponentInfo.SkeletonCompatibility.Find(InSkeleton);
+			
+			if(!SkeletonCompatibility)
+			{
+				// Check if the skeleton is compatible with the reference skeleton
+				const TMap<FName, uint32>& RefMeshBoneNamesToPathHash = MutComponentInfo.BoneNamesToPathHash;
+
+				const TArray<FMeshBoneInfo>& Bones = InSkeleton->GetReferenceSkeleton().GetRawRefBoneInfo();
+				const int32 NumBones = Bones.Num();
+
+				TMap<FName, uint32> BoneNamesToPathHash;
+				BoneNamesToPathHash.Reserve(NumBones);
+
+				for (int32 BoneIndex = 0; BoneIndex < NumBones; ++BoneIndex)
+				{
+					const FMeshBoneInfo& Bone = Bones[BoneIndex];
+
+					// Retrieve parent bone name and respective hash, root-bone is assumed to have a parent hash of 0
+					const FName ParentName = Bone.ParentIndex != INDEX_NONE ? Bones[Bone.ParentIndex].Name : NAME_None;
+					const uint32 ParentHash = Bone.ParentIndex != INDEX_NONE ? GetTypeHash(ParentName) : 0;
+
+					// Look-up the path-hash from root to the parent bone
+					const uint32* ParentPath = BoneNamesToPathHash.Find(ParentName);
+					const uint32 ParentPathHash = ParentPath ? *ParentPath : 0;
+
+					// Append parent hash to path to give full path hash to current bone
+					const uint32 BonePathHash = HashCombine(ParentPathHash, ParentHash);
+
+					// If the hash differs from the reference one it means skeletons are incompatible
+					if (const uint32* RefSMBonePathHash = RefMeshBoneNamesToPathHash.Find(Bone.Name); RefSMBonePathHash && *RefSMBonePathHash != BonePathHash)
+					{
+						// Different skeletons can't be used if they are incompatible with the reference skeleton.
+						FString Msg = FString::Printf(
+							TEXT("The SkeletalMesh [%s] with Skeleton [%s] is incompatible with the reference mesh [%s] which has [%s]. "
+								"Bone [%s] has a differnt parent on the Skeleton from the reference mesh."),
+							*InSkeletalMesh->GetName(),	*InSkeleton->GetName(),
+							*ComponentRefSkeletalMesh->GetName(), *ComponentRefSkeleton->GetName(),
+							*Bone.ExportName);
+
+						GenerationContext.Compiler->CompilerLog(FText::FromString(Msg), CurrentNode, EMessageSeverity::Info);
+
+						MutComponentInfo.SkeletonCompatibility.Add(InSkeleton, false);
+						return nullptr;
+					}
+
+					// Add path hash to current bone
+					BoneNamesToPathHash.Add(Bone.Name, BonePathHash);
+				}
+				
+				MutComponentInfo.SkeletonCompatibility.Add(InSkeleton, true);
+			}
+			else if(*SkeletonCompatibility == false)
+			{
+				// Incompatible skeleton 
+				return nullptr;
+			}
+			
+			// Is it a skeleton that we have already used?
+			SkeletonID = GenerationContext.ReferencedSkeletons.AddUnique(InSkeleton);
+		}
+		else
+		{
+			// Since we could have multiple reference skeletons, we can't assume its index is 0 
+			SkeletonID = GenerationContext.ReferencedSkeletons.Find(InSkeleton);
+		}
+
+		{
+			// Ensure the bones used by the Skeletal Mesh exits in the Mesh's Skeleton
+			const TArray<FMeshBoneInfo>& RawRefBoneInfo = InSkeletalMesh->GetRefSkeleton().GetRawRefBoneInfo();
+			const FReferenceSkeleton& InSkeletonRefSkeleton = InSkeleton->GetReferenceSkeleton();
+
+			bool bIsSkeletonMissingBones = false;
+
+			for (const FMeshBoneInfo& BoneInfo : RawRefBoneInfo)
+			{
+				if (InSkeletonRefSkeleton.FindRawBoneIndex(BoneInfo.Name) == INDEX_NONE)
+				{
+					bIsSkeletonMissingBones = true;
+					UE_LOG(LogMutable, Warning, TEXT("SkeletalMesh [%s] uses bone [%s] not present in skeleton [%s]."),
+						*InSkeletalMesh->GetName(),
+						*BoneInfo.ExportName,
+						*InSkeleton->GetName());
+				}
+			}
+
+			// Discard SkeletalMesh if some bones are missing
+			if (bIsSkeletonMissingBones)
+			{
+				FString Msg = FString::Printf(
+					TEXT("The Skeleton [%s] is missing bones that SkeletalMesh [%s] needs. The mesh will be discarded! Information about missing bones can be found in the Output Log."),
+					*InSkeleton->GetName(), *InSkeletalMesh->GetName());
+				
+				GenerationContext.Compiler->CompilerLog(FText::FromString(Msg), CurrentNode, EMessageSeverity::Warning);
+
+				return nullptr;
+			}
+		}
+
+		{
+			// Find removed bone and parent indices in the current section BoneMap.
+			const TArray<uint16>& BoneMap = ImportedModel->LODModels[LOD].Sections[MaterialIndex].BoneMap;
+
+			for (const TPair<int32,int32> it : RemovedBonesActiveParentIndices)
+			{
+				const int32 BoneMapIndex = BoneMap.Find(it.Key);
+				if (BoneMapIndex != INDEX_NONE)
+				{
+					const int32 ParentBoneMapIndex = BoneMap.Find(it.Value);
+					InfluencesToReplaceMap.Add(BoneMapIndex, ParentBoneMapIndex != INDEX_NONE ? ParentBoneMapIndex : 0);
+				}
+			}
+
+			// Removed bones might not be used in the current section
+			bHasBonesToRemove = !InfluencesToReplaceMap.IsEmpty();
+		}
+	}
+
+	// Vertices
+	TArray<FSoftSkinVertex> Vertices;
+	ImportedModel->LODModels[LOD].GetVertices(Vertices);
+	int VertexStart = ImportedModel->LODModels[LOD].Sections[MaterialIndex].GetVertexBufferIndex();
+	int VertexCount = ImportedModel->LODModels[LOD].Sections[MaterialIndex].GetNumVertices();
+
+	MutableMesh->GetVertexBuffers().SetElementCount(VertexCount);
+
+	const int32 VertexBuffersCount = 1 +
+		(GenerationContext.Options.bRealTimeMorphTargetsEnabled ? 2 : 0) +
+		(GenerationContext.Options.bClothingEnabled ? 1 : 0);
+
+	MutableMesh->GetVertexBuffers().SetBufferCount(VertexBuffersCount);
+
+	if (bIgnoreSkeleton)
+	{
+		// Create the mesh with the same data, but skinning is considered padding.
+		using namespace mu;
+		const int ElementSize = sizeof(FSoftSkinVertex);
+		const int ChannelCount = 9;
+		const MESH_BUFFER_SEMANTIC Semantics[ChannelCount] = { MBS_POSITION, MBS_TANGENT, MBS_BINORMAL, MBS_NORMAL, MBS_TEXCOORDS, MBS_TEXCOORDS, MBS_TEXCOORDS, MBS_TEXCOORDS, MBS_COLOUR };
+		const int SemanticIndices[ChannelCount] = { 0, 0, 0, 0, 0, 1, 2, 3, 0 };
+		const MESH_BUFFER_FORMAT Formats[ChannelCount] = { MBF_FLOAT32, MBF_FLOAT32, MBF_FLOAT32, MBF_FLOAT32, MBF_FLOAT32, MBF_FLOAT32, MBF_FLOAT32, MBF_FLOAT32, MBF_NUINT8 };
+		int Components[ChannelCount] = { 3, 3, 3, 4, 2, 2, 2, 2, 4 };
+
+		constexpr size_t SoftSkinVertexUVsElemSize = sizeof(TDecay<decltype(DeclVal<FSoftSkinVertex>().UVs[0])>::Type);
+		const int Offsets[ChannelCount] =
+		{
+			STRUCT_OFFSET(FSoftSkinVertex, Position),
+			STRUCT_OFFSET(FSoftSkinVertex, TangentX),
+			STRUCT_OFFSET(FSoftSkinVertex, TangentY),
+			STRUCT_OFFSET(FSoftSkinVertex, TangentZ),
+			STRUCT_OFFSET(FSoftSkinVertex, UVs) + 0 * SoftSkinVertexUVsElemSize,
+			STRUCT_OFFSET(FSoftSkinVertex, UVs) + 1 * SoftSkinVertexUVsElemSize,
+			STRUCT_OFFSET(FSoftSkinVertex, UVs) + 2 * SoftSkinVertexUVsElemSize,
+			STRUCT_OFFSET(FSoftSkinVertex, UVs) + 3 * SoftSkinVertexUVsElemSize,
+			STRUCT_OFFSET(FSoftSkinVertex, Color),
+		};
+
+		MutableMesh->GetVertexBuffers().SetBuffer(0, ElementSize, ChannelCount, Semantics, SemanticIndices, Formats, Components, Offsets);
+		FMemory::Memcpy(MutableMesh->GetVertexBuffers().GetBufferData(0), Vertices.GetData() + VertexStart, VertexCount* ElementSize);
+	}
+	else
+	{
+		int32 MaxSectionInfluences = ImportedModel->LODModels[LOD].Sections[MaterialIndex].MaxBoneInfluences;
+		using namespace mu;
+		const int ElementSize = sizeof(FSoftSkinVertex);
+		const int ChannelCount = 11;
+		const MESH_BUFFER_SEMANTIC Semantics[ChannelCount] = { MBS_POSITION, MBS_TANGENT, MBS_BINORMAL, MBS_NORMAL, MBS_TEXCOORDS, MBS_TEXCOORDS, MBS_TEXCOORDS, MBS_TEXCOORDS, MBS_COLOUR, MBS_BONEINDICES, MBS_BONEWEIGHTS };
+		const int SemanticIndices[ChannelCount] = { 0, 0, 0, 0, 0, 1, 2, 3, 0, 0, 0 };
+		const MESH_BUFFER_FORMAT Formats[ChannelCount] = { MBF_FLOAT32, MBF_FLOAT32, MBF_FLOAT32, MBF_FLOAT32, MBF_FLOAT32, MBF_FLOAT32, MBF_FLOAT32, MBF_FLOAT32, MBF_NUINT8, MBF_UINT16, MBF_NUINT8 };
+		int Components[ChannelCount] = { 3, 3, 3, 4, 2, 2, 2, 2, 4, 4, 4 };
+		if (GenerationContext.Options.bExtraBoneInfluencesEnabled && MaxSectionInfluences > 4)
+		{
+			Components[9] = EXTRA_BONE_INFLUENCES;
+			Components[10] = EXTRA_BONE_INFLUENCES;
+		}
+
+        constexpr size_t SoftSkinVertexUVsElemSize = sizeof(TDecay<decltype(DeclVal<FSoftSkinVertex>().UVs[0])>::Type); 
+		const int Offsets[ChannelCount] =
+		{
+			STRUCT_OFFSET(FSoftSkinVertex, Position),
+			STRUCT_OFFSET(FSoftSkinVertex, TangentX),
+			STRUCT_OFFSET(FSoftSkinVertex, TangentY),
+			STRUCT_OFFSET(FSoftSkinVertex, TangentZ),
+			STRUCT_OFFSET(FSoftSkinVertex, UVs) + 0 * SoftSkinVertexUVsElemSize,
+			STRUCT_OFFSET(FSoftSkinVertex, UVs) + 1 * SoftSkinVertexUVsElemSize,
+			STRUCT_OFFSET(FSoftSkinVertex, UVs) + 2 * SoftSkinVertexUVsElemSize,
+			STRUCT_OFFSET(FSoftSkinVertex, UVs) + 3 * SoftSkinVertexUVsElemSize,
+			STRUCT_OFFSET(FSoftSkinVertex, Color),
+			STRUCT_OFFSET(FSoftSkinVertex, InfluenceBones),
+			STRUCT_OFFSET(FSoftSkinVertex, InfluenceWeights),
+		};
+
+		if (GenerationContext.Options.bExtraBoneInfluencesEnabled && MaxSectionInfluences < EXTRA_BONE_INFLUENCES)
+		{
+			for (int32 VertexIndex = VertexStart; VertexIndex < VertexStart + VertexCount && VertexIndex < Vertices.Num(); ++VertexIndex)
+			{
+				FSoftSkinVertex& Vertex = Vertices[VertexIndex];
+
+				for (int32 i = MaxSectionInfluences; i < EXTRA_BONE_INFLUENCES; ++i)
+				{
+					check(Vertex.InfluenceWeights[i] == 0);
+				}
+			}
+		}
+		else if (!GenerationContext.Options.bExtraBoneInfluencesEnabled)
+		{
+			int32 MaxSectionBoneMapIndex = ImportedModel->LODModels[LOD].Sections[MaterialIndex].BoneMap.Num();
+
+			// Renormalize to 4 weights per vertex
+			for (int32 VertexIndex = VertexStart; VertexIndex < VertexStart + VertexCount && VertexIndex < Vertices.Num(); ++VertexIndex)
+			{
+				FSoftSkinVertex& Vertex = Vertices[VertexIndex];
+
+				const int32 MaxMutableWeights = 4;
+				int32 MaxOrderedWeighsIndices[MaxMutableWeights] = { -1, -1, -1, -1 };
+
+				// Transfer removed bones influences to parent bones
+				if (bHasBonesToRemove)
+				{
+					for (int32 i = 0; i < MaxSectionInfluences; ++i)
+					{
+						if (const int32* IndexPtr = InfluencesToReplaceMap.Find(Vertex.InfluenceBones[i]))
+						{
+							bool bParentFound = false;
+							int32 ParentIndex = *IndexPtr;
+							for (int32 j = 0; j < MaxSectionInfluences; ++j)
+							{
+								if (Vertex.InfluenceBones[j] == ParentIndex)
+								{
+									Vertex.InfluenceWeights[j] += Vertex.InfluenceWeights[i];
+									Vertex.InfluenceWeights[i] = 0.f;
+									bParentFound = true;
+									break;
+								}
+							}
+
+							if (!bParentFound)
+							{
+								Vertex.InfluenceBones[i] = ParentIndex;
+							}
+						}
+					}
+				}
+
+				// First get the indices of the 4 heaviest influences
+				for (int32 i = 0; i < MaxMutableWeights; ++i)
+				{
+					int32 CurrentMaxWeight = -1;
+
+					for (int32 j = 0; j < MaxSectionInfluences; ++j)
+					{
+						bool bIndexAlreadyUsed = false;
+
+						for (int32 k = 0; k < i; ++k)
+						{
+							if (MaxOrderedWeighsIndices[k] == j)
+							{
+								bIndexAlreadyUsed = true;
+								break;
+							}
+							else if (MaxOrderedWeighsIndices[k] < 0)
+							{
+								break;
+							}
+						}
+
+						if (!bIndexAlreadyUsed && Vertex.InfluenceWeights[j] > CurrentMaxWeight
+							&& Vertex.InfluenceBones[j] < MaxSectionBoneMapIndex)
+						{
+							MaxOrderedWeighsIndices[i] = j;
+							CurrentMaxWeight = Vertex.InfluenceWeights[j];
+						}
+					}
+				}
+
+				// Copy 4 heaviest influences to 4 first indices
+				for (int32 i = 0; i < MaxMutableWeights; ++i)
+				{
+					if (i < MaxSectionInfluences)
+					{
+						Vertex.InfluenceWeights[i] = Vertex.InfluenceWeights[MaxOrderedWeighsIndices[i]];
+						Vertex.InfluenceBones[i] = Vertex.InfluenceBones[MaxOrderedWeighsIndices[i]];
+					}
+					else
+					{
+						Vertex.InfluenceWeights[i] = 0;
+						Vertex.InfluenceBones[i] = 0;
+					}
+				}
+
+				// Actually renormalize the first 4 influences
+				int32 TotalWeight = 0;
+
+				for (int32 j = 0; j < MaxMutableWeights; ++j)
+				{
+					TotalWeight += Vertex.InfluenceWeights[j];
+				}
+
+				if (TotalWeight > 0)
+				{
+					int32 AssignedWeight = 0;
+
+					for (int32 j = 1; j < MAX_TOTAL_INFLUENCES; ++j)
+					{
+						if (j < MaxMutableWeights)
+						{
+							float Aux = Vertex.InfluenceWeights[j];
+							int32 Res = FMath::RoundToInt(Aux / TotalWeight * 255);
+							AssignedWeight += Res;
+							Vertex.InfluenceWeights[j] = Res;
+						}
+						else
+						{
+							Vertex.InfluenceWeights[j] = 0;
+						}
+					}
+
+					Vertex.InfluenceWeights[0] = 255 - AssignedWeight;
+				}
+				else
+				{
+					Vertex.InfluenceWeights[0] = 255;
+
+					for (int32 j = 1; j < MaxMutableWeights; ++j)
+					{
+						Vertex.InfluenceWeights[j] = 0;
+					}
+				}
+			}
+		}
+		else if (bHasBonesToRemove)
+		{
+			// Transfer removed bones influences to parent bones
+			for (int32 VertexIndex = VertexStart; VertexIndex < VertexStart + VertexCount && VertexIndex < Vertices.Num(); ++VertexIndex)
+			{
+				FSoftSkinVertex& Vertex = Vertices[VertexIndex];
+
+				for (int32 i = 0; i < MAX_TOTAL_INFLUENCES; ++i)
+				{
+					if (const int32* ParentBoneIndexPtr = InfluencesToReplaceMap.Find(Vertex.InfluenceBones[i]))
+					{
+						bool bParentFound = false;
+						int32 ParentIndex = *ParentBoneIndexPtr;
+						for (int32 j = 0; j < MAX_TOTAL_INFLUENCES; ++j)
+						{
+							if (Vertex.InfluenceBones[j] == ParentIndex)
+							{
+								Vertex.InfluenceWeights[j] += Vertex.InfluenceWeights[i];
+								Vertex.InfluenceWeights[i] = 0.f;
+								bParentFound = true;
+								break;
+							}
+						}
+
+						if (!bParentFound)
+						{
+							Vertex.InfluenceBones[i] = ParentIndex;
+						}
+					}
+				}
+			}
+		}
+
+		MutableMesh->GetVertexBuffers().SetBuffer(0, ElementSize, ChannelCount, Semantics, SemanticIndices, Formats, Components, Offsets);
+		FMemory::Memcpy(MutableMesh->GetVertexBuffers().GetBufferData(0), Vertices.GetData() + VertexStart, VertexCount * ElementSize);
+	}
+
+	int32 nextBufferIndex = 1;
+	if (GenerationContext.Options.bRealTimeMorphTargetsEnabled)
+	{
+		nextBufferIndex += 2;
+
+		// Find realtime MorphTargets to be used.
+		TArray<UMorphTarget*> UsedMorphTargets;
+		UsedMorphTargets.Empty(SkeletalMesh->GetMorphTargets().Num());
+
+		const UCustomizableObjectNodeSkeletalMesh* NodeTyped = Cast<UCustomizableObjectNodeSkeletalMesh>(CurrentNode);
+		check(NodeTyped);
+
+        // Add SkeletalMesh morphs to the usage override data structure.
+        // Notice this will only be populated here, when compiling.
+		TArray<FRealTimeMorphSelectionOverride>& RealTimeMorphTargetOverrides = GenerationContext.RealTimeMorphTargetsOverrides;
+        for (const TObjectPtr<UMorphTarget>& MorphTarget : SkeletalMesh->GetMorphTargets())
+		{
+            // Find if the morph target global override is already present
+            // and add it if needed.
+            FRealTimeMorphSelectionOverride* MorphFound = RealTimeMorphTargetOverrides.FindByPredicate(
+                    [MorphTarget](const FRealTimeMorphSelectionOverride& O) { return O.MorphName == MorphTarget->GetFName(); });
+
+            MorphFound = MorphFound 
+                    ? MorphFound 
+                    : &RealTimeMorphTargetOverrides.Emplace_GetRef(MorphTarget->GetFName());
+
+            const int32 AddedMeshNameIdx = MorphFound->SkeletalMeshesNames.AddUnique(SkeletalMesh->GetFName());
+
+			if (AddedMeshNameIdx >= MorphFound->Override.Num())
+			{
+				MorphFound->Override.Add(ECustomizableObjectSelectionOverride::NoOverride);
+			}
+        }
+
+        // Add SkeletalMesh node used defined realtime morph targets to a temporal array where
+        // the actual to be used real-time morphs names will be placed.        
+        TArray<FName> UsedMorphTargetsNames = [&]()
+        {
+            TArray<FName> MorphTargetsNames;
+            MorphTargetsNames.Reserve(SkeletalMesh->GetMorphTargets().Num());
+            
+            if (NodeTyped->bUseAllRealTimeMorphs)
+            {
+                TArray<TObjectPtr<UMorphTarget>>& SkeletalMeshMorphs = SkeletalMesh->GetMorphTargets();
+        
+                for (TObjectPtr<UMorphTarget>& MorphTarget : SkeletalMeshMorphs)
+                {
+                    check(MorphTarget);
+                    MorphTargetsNames.Add(MorphTarget->GetFName());
+                }
+            } 
+            else
+            {
+                for (const FString& MorphName : NodeTyped->UsedRealTimeMorphTargetNames)
+                {     
+                    MorphTargetsNames.Emplace(*MorphName);
+                }
+            }
+
+            return MorphTargetsNames;
+        }(); // lambda is invoked.
+
+        //Apply global morph targets overrides to the SkeletalMesh user defined RT morph targets. 
+        for (FRealTimeMorphSelectionOverride& MorphTargetOverride : RealTimeMorphTargetOverrides)
+        {
+            const ECustomizableObjectSelectionOverride OverrideValue = [&]() -> ECustomizableObjectSelectionOverride
+                {
+                    const ECustomizableObjectSelectionOverride GlobalOverrideValue = MorphTargetOverride.SelectionOverride;
+                    
+                    if (GlobalOverrideValue != ECustomizableObjectSelectionOverride::NoOverride)
+                    {
+                        return GlobalOverrideValue;
+                    } 
+                
+                    const int32 FoundIdx = 
+                            MorphTargetOverride.SkeletalMeshesNames.Find(SkeletalMesh->GetFName());
+
+                    if (FoundIdx != INDEX_NONE)
+                    {
+                        return MorphTargetOverride.Override[FoundIdx];
+                    }
+
+                    return ECustomizableObjectSelectionOverride::NoOverride;
+                }(); // lambda is invoked.
+
+            if (OverrideValue == ECustomizableObjectSelectionOverride::Enable)
+            {
+                UsedMorphTargetsNames.AddUnique(MorphTargetOverride.MorphName);
+            }
+            else if (OverrideValue == ECustomizableObjectSelectionOverride::Disable)
+            {
+                UsedMorphTargetsNames.Remove(MorphTargetOverride.MorphName);
+            }
+        }
+
+		for (UMorphTarget* MorphTarget : SkeletalMesh->GetMorphTargets())
+		{
+			if (!MorphTarget)
+			{
+				continue;
+			}
+
+			const bool bHasToBeAdded = UsedMorphTargetsNames.Contains(MorphTarget->GetFName());
+			if (bHasToBeAdded)
+			{
+				UsedMorphTargets.Add(MorphTarget);
+			}
+		}
+
+		// MorphTarget vertex info index.
+		{
+			using namespace mu;
+			const int ElementSize = sizeof(int32);
+			const int ChannelCount = 1;
+			const MESH_BUFFER_SEMANTIC Semantics[ChannelCount] = { MBS_OTHER };
+			const int SemanticIndices[ChannelCount] = { 0 };
+			const MESH_BUFFER_FORMAT Formats[ChannelCount] = { MBF_INT32 };
+			int Components[ChannelCount] = { 1 };
+			const int Offsets[ChannelCount] = { 0 };
+
+			MutableMesh->GetVertexBuffers().SetBuffer(1, ElementSize, ChannelCount, Semantics, SemanticIndices, Formats, Components, Offsets);
+		}
+
+		// MorphTarget vertex info count.
+		{
+			using namespace mu;
+			const int ElementSize = sizeof(int32);
+			const int ChannelCount = 1;
+			const MESH_BUFFER_SEMANTIC Semantics[ChannelCount] = { MBS_OTHER };
+			const int SemanticIndices[ChannelCount] = { 1 };
+			const MESH_BUFFER_FORMAT Formats[ChannelCount] = { MBF_INT32 };
+			int Components[ChannelCount] = { 1 };
+			const int Offsets[ChannelCount] = { 0 };
+
+			MutableMesh->GetVertexBuffers().SetBuffer(2, ElementSize, ChannelCount, Semantics, SemanticIndices, Formats, Components, Offsets);
+		}
+
+		// Setup MorphTarget reconstruction data.
+
+		TArrayView<int32> VertexMorphsCountBufferView(reinterpret_cast<int32*>(MutableMesh->GetVertexBuffers().GetBufferData(2)), VertexCount);
+		for (int32& Elem : VertexMorphsCountBufferView)
+		{
+			Elem = 0;
+		}
+
+		TArrayView<int32>  VertexMorphsInfoIndexBufferView(reinterpret_cast<int32*>(MutableMesh->GetVertexBuffers().GetBufferData(1)), VertexCount);
+
+		if (UsedMorphTargets.Num())
+		{
+			for (int32 VertexIdx = VertexStart; VertexIdx < VertexStart + VertexCount && VertexIdx < Vertices.Num(); ++VertexIdx)
+			{
+				TArray<FMorphTargetVertexData> MorphsUsed;
+				for (const UMorphTarget* MorphTarget : UsedMorphTargets)
+				{
+					if (!MorphTarget) continue;
+
+					const TArray<FMorphTargetLODModel>& MorphLODModels = MorphTarget->GetMorphLODModels();
+
+					if (!(MorphTarget && LOD < MorphLODModels.Num()))
+					{
+						continue;
+					}
+
+					const FMorphTargetDelta* VertexFound = MorphLODModels[LOD].Vertices.FindByPredicate(
+						[&VertexIdx](const FMorphTargetDelta& Elem) -> bool
+					{ 
+							return Elem.SourceIdx == VertexIdx; 
+					});
+
+					if (!VertexFound)
+					{
+						continue;
+					}
+
+					const FName MorphTargetName = MorphTarget->GetFName();
+
+					TArray<FMorphTargetInfo>& ContributingMorphTargetsInfo = GenerationContext.ContributingMorphTargetsInfo;
+
+					int32 DestMorphTargetIdx = ContributingMorphTargetsInfo.IndexOfByPredicate(
+						[&MorphTargetName](auto& MorphTargetInfo) { return MorphTargetName == MorphTargetInfo.Name; });
+
+					DestMorphTargetIdx = DestMorphTargetIdx != INDEX_NONE
+						? DestMorphTargetIdx
+						: ContributingMorphTargetsInfo.Emplace(FMorphTargetInfo{ MorphTargetName, LOD + 1 });
+
+					FMorphTargetInfo& MorphTargetInfo = ContributingMorphTargetsInfo[DestMorphTargetIdx];
+					MorphTargetInfo.LodNum = FMath::Max(MorphTargetInfo.LodNum, LOD + 1);
+
+					MorphsUsed.Emplace(FMorphTargetVertexData{ VertexFound->PositionDelta, VertexFound->TangentZDelta, DestMorphTargetIdx });
+				}
+
+				if (MorphsUsed.Num())
+				{
+					VertexMorphsInfoIndexBufferView[VertexIdx - VertexStart] = GenerationContext.MorphTargetReconstructionData.Num();
+					VertexMorphsCountBufferView[VertexIdx - VertexStart] = MorphsUsed.Num();
+
+					GenerationContext.MorphTargetReconstructionData.Append(MorphsUsed);
+				}
+			}
+		}
+	}
+
+	// Clothing vertex info.
+	if (GenerationContext.Options.bClothingEnabled)
+	{
+		{
+			using namespace mu;
+			const int ElementSize = sizeof(int32);
+			const int ChannelCount = 1;
+			const MESH_BUFFER_SEMANTIC Semantics[ChannelCount] = { MBS_OTHER };
+			const int SemanticIndices[ChannelCount] = { GenerationContext.Options.bRealTimeMorphTargetsEnabled ? 2 : 0 };
+			const MESH_BUFFER_FORMAT Formats[ChannelCount] = { MBF_INT32 };
+			int Components[ChannelCount] = { 1 };
+			const int Offsets[ChannelCount] = { 0 };
+
+			MutableMesh->GetVertexBuffers().SetBuffer(nextBufferIndex, ElementSize, ChannelCount, Semantics, SemanticIndices, Formats, Components, Offsets);
+		}
+
+		TArrayView<int32> ClothSectionBufferView(reinterpret_cast<int32*>(MutableMesh->GetVertexBuffers().GetBufferData(nextBufferIndex)), VertexCount);
+		for (int32& Elem : ClothSectionBufferView)
+		{
+			Elem = -1;
+		}
+
+		// Create new asset or find an already created one if the section has clothing assets.
+		// clothing assets are shared among all LODs in a section
+		const int32 ClothingAssetIndex = [&]() -> int32
+		{
+			UClothingAssetBase* ClothingAssetBase = SkeletalMesh->GetSectionClothingAsset(LOD, MaterialIndex);
+
+			if (!ClothingAssetBase)
+			{
+				return INDEX_NONE;
+			}
+
+			int32 FoundIndex = GenerationContext.ContributingClothingAssetsData.IndexOfByPredicate(
+				[AssetGuid = ClothingAssetBase->GetAssetGuid()](const FCustomizableObjectClothingAssetData& Asset){ return Asset.OriginalAssetGuid == AssetGuid; });
+			
+			if (FoundIndex != INDEX_NONE)
+			{
+				return FoundIndex;
+			}
+
+			UClothingAssetCommon* Asset = Cast<UClothingAssetCommon>(ClothingAssetBase);
+			if (!Asset)
+			{
+				return INDEX_NONE;
+			}
+
+			const int32 NewAssetIndex = GenerationContext.ContributingClothingAssetsData.AddDefaulted();
+			FCustomizableObjectClothingAssetData& AssetData = GenerationContext.ContributingClothingAssetsData[NewAssetIndex];
+			
+			AssetData.LodData = Asset->LodData;
+			AssetData.LodMap = Asset->LodMap;
+			AssetData.ReferenceBoneIndex = Asset->ReferenceBoneIndex;
+			AssetData.UsedBoneIndices = Asset->UsedBoneIndices;
+			AssetData.UsedBoneNames = Asset->UsedBoneNames;
+			AssetData.OriginalAssetGuid = Asset->GetAssetGuid();
+			AssetData.Name = Asset->GetFName();
+
+			// Store raw clothing config serialized raw data, and info to recreate it afterwards.
+			for ( const TPair<FName, TObjectPtr<UClothConfigBase>>& ClothConfig : Asset->ClothConfigs )
+			{
+				FCustomizableObjectClothConfigData& ConfigData = AssetData.ConfigsData.AddDefaulted_GetRef();
+				ConfigData.ClassPath = ClothConfig.Value->GetClass()->GetPathName();
+				ConfigData.ConfigName = ClothConfig.Key;
+				
+				FMemoryWriter MemoryWriter(ConfigData.ConfigBytes);
+                ClothConfig.Value->Serialize(MemoryWriter);
+			}
+
+			return NewAssetIndex;
+		}(); // lambda is invoked
+
+		if (ClothingAssetIndex != INDEX_NONE)
+		{
+			// Reserve first element as a way to indicate invalid data. Currently not used.
+			if (GenerationContext.ClothMeshToMeshVertData.Num() == 0)
+			{
+				FCustomizableObjectMeshToMeshVertData& FirstElem = GenerationContext.ClothMeshToMeshVertData.AddZeroed_GetRef();
+				FirstElem.SourceAssetIndex = INDEX_NONE;
+			}
+
+			const FSkelMeshSection& SectionData = ImportedModel->LODModels[LOD].Sections[MaterialIndex];
+			const TArray<FMeshToMeshVertData>& ClothMappingData = SectionData.ClothMappingDataLODs[0];
+
+
+			// Similar test as the one used on FSkeletalMeshObjectGPUSkin::FVertexFactoryData::InitAPEXClothVertexFactories
+			// Here should work as expexted, but in the reference code I'm not sure it always works. It is worth investigate
+			// in that direction if at some point multiple influences don't work as expected.
+			const bool bUseMutlipleInfluences = ClothMappingData.Num() > SectionData.NumVertices;
+
+			// Constant defined in ClothMeshUtils.cpp with the following comment:
+			// // This must match NUM_INFLUENCES_PER_VERTEX in GpuSkinCacheComputeShader.usf and GpuSkinVertexFactory.ush
+			// // TODO: Make this easier to change in without messing things up
+			// TODO: find a better place to keep this constant.
+			constexpr int32 NumInfluencesPerVertex = 5;
+
+			int32 MeshToMeshDataIndex = GenerationContext.ClothMeshToMeshVertData.Num();
+
+			constexpr int32 MaxSupportedInfluences = 1;
+			for (int32& Elem : ClothSectionBufferView)
+			{
+				Elem = MeshToMeshDataIndex;
+				MeshToMeshDataIndex += MaxSupportedInfluences;
+			}
+
+			const int32 ClothDataIndexBase = GenerationContext.ClothMeshToMeshVertData.Num();
+
+			const int32 ClothDataStride = bUseMutlipleInfluences ? NumInfluencesPerVertex : 1;
+			const int32 NumClothMappingDataVerts = ClothMappingData.Num() / ClothDataStride;
+
+			GenerationContext.ClothMeshToMeshVertData.Reserve(NumClothMappingDataVerts);
+			for (int32 Idx = 0; Idx < NumClothMappingDataVerts * ClothDataStride; Idx += ClothDataStride)
+			{
+				// If bUseMutlipleInfluences we will only take the element with higher weight ignoring the other ones.
+				TArrayView<const FMeshToMeshVertData> Influences(ClothMappingData.GetData() + Idx, ClothDataStride);
+				const FMeshToMeshVertData* MaxInfluence = MaxElement(Influences.begin(), Influences.end(),
+					[](const FMeshToMeshVertData& A, const FMeshToMeshVertData& B) { return A.Weight < B.Weight; });
+
+				GenerationContext.ClothMeshToMeshVertData.Emplace(*MaxInfluence);
+			}
+
+			TArrayView<FCustomizableObjectMeshToMeshVertData> AppendedClothingDataView
+			(GenerationContext.ClothMeshToMeshVertData.GetData() + ClothDataIndexBase, NumClothMappingDataVerts);
+
+			const FCustomizableObjectClothingAssetData& ClothingAssetData = GenerationContext.ContributingClothingAssetsData[ClothingAssetIndex];
+			const int16 ClothingAssetLODIndex = static_cast<int16>(ClothingAssetData.LodMap[LOD]);
+
+			for (FCustomizableObjectMeshToMeshVertData& ClothingDataElem : AppendedClothingDataView)
+			{
+				ClothingDataElem.SourceAssetIndex = static_cast<int16>(ClothingAssetIndex);
+				ClothingDataElem.SourceAssetLodIndex = ClothingAssetLODIndex;
+
+				// Currently if the cloth mapping uses multiple influences, these are ignored and only 
+				// the one with the highest weight is used. We set the weight to 1.0, but
+				// this value will be ignored anyway.
+				ClothingDataElem.Weight = 1.0f;
+			}
+		}
+	}
+
+	// Indices
+	{
+		int IndexStart = ImportedModel->LODModels[LOD].Sections[MaterialIndex].BaseIndex;
+		int IndexCount = ImportedModel->LODModels[LOD].Sections[MaterialIndex].NumTriangles * 3;
+		MutableMesh->GetIndexBuffers().SetBufferCount(1);
+		MutableMesh->GetIndexBuffers().SetElementCount(IndexCount);
+		MutableMesh->GetFaceBuffers().SetElementCount(IndexCount / 3);
+
+		using namespace mu;
+		// For some reason, the indices in 4.25 (and 4.24) are in different order in the Imported and Rendering data structures. The strange thing 
+		// is actually that the vertices in the imported model seem to match the rendering model indices. Maybe there is some mapping that
+		// we are missing, but for now this will do:
+		const int ElementSize = SkeletalMesh->GetResourceForRendering()->LODRenderData[LOD].MultiSizeIndexContainer.GetDataTypeSize();
+		void* IndexDataPointer = SkeletalMesh->GetResourceForRendering()->LODRenderData[LOD].MultiSizeIndexContainer.GetIndexBuffer()->GetPointerTo(IndexStart);
+		const int FinalElementSize = sizeof(uint32_t);
+		const int ChannelCount = 1;
+		const MESH_BUFFER_SEMANTIC Semantics[ChannelCount] = { MBS_VERTEXINDEX };
+		const int SemanticIndices[ChannelCount] = { 0 };
+		// We force 32 bit indices, since merging meshes may create vertex buffers bigger than the initial mesh
+		// and for now the mutable runtime doesn't handle it.
+		// \TODO: go back to 16-bit indices when possible.
+		MESH_BUFFER_FORMAT Formats[ChannelCount] = { MBF_UINT32 };
+		const int Components[ChannelCount] = { 1 };
+		const int Offsets[ChannelCount] = { 0 };
+
+		MutableMesh->GetIndexBuffers().SetBuffer(0, FinalElementSize, ChannelCount, Semantics, SemanticIndices, Formats, Components, Offsets);
+
+		// 32-bit to 32-bit
+		if (ElementSize == 4)
+		{
+			uint32* pDest = reinterpret_cast<uint32*>(MutableMesh->GetIndexBuffers().GetBufferData(0));
+			const uint32* pSource = reinterpret_cast<const uint32*>(IndexDataPointer);
+
+			for (int i = 0; i < IndexCount; ++i)
+			{
+				*pDest = *pSource - VertexStart;
+				check(*pDest < uint32(VertexCount));
+				++pDest;
+				++pSource;
+			}
+		}
+		// 16-bit to 16-bit
+		else if (ElementSize == 2 && Formats[0] == MBF_UINT16)
+		{
+			uint16* pDest = reinterpret_cast<uint16*>(MutableMesh->GetIndexBuffers().GetBufferData(0));
+			const uint16* pSource = reinterpret_cast<const uint16*>(IndexDataPointer);
+
+			for (int i = 0; i < IndexCount; ++i)
+			{
+				*pDest = *pSource - VertexStart;
+				check(*pDest < uint32(VertexCount));
+				++pDest;
+				++pSource;
+			}
+		}
+		// 16-bit to 32-bit
+		else if (ElementSize == 2 && Formats[0] == MBF_UINT32)
+		{
+			uint32* pDest = reinterpret_cast<uint32*>(MutableMesh->GetIndexBuffers().GetBufferData(0));
+			const uint16* pSource = reinterpret_cast<const uint16*>(IndexDataPointer);
+
+			for (int i = 0; i < IndexCount; ++i)
+			{
+				*pDest = *pSource - VertexStart;
+				check(*pDest < uint32(VertexCount));
+				++pDest;
+				++pSource;
+			}
+		}
+		else
+		{
+			// Unsupported case!
+			check(false);
+		}
+	}
+
+	if (!bIgnoreSkeleton)
+	{
+		// Skeleton
+		int32 NumBonesInMesh = ImportedModel->LODModels[LOD].Sections[MaterialIndex].BoneMap.Num();
+
+
+		// Add the bones in the order they are used in the mesh
+		TMap<int32, int32> InverseBoneMap; // Inverse of ImportedModel->LODModels[LOD].Chunks[MaterialIndex].BoneMap for the chunk's bone indices
+		TQueue<int32> SkippedBoneMapIndices; // The order and index of the bones in the BoneMap must remain the same even if some bones are removed
+		for (int32 BoneMapIndex = 0; BoneMapIndex < NumBonesInMesh; ++BoneMapIndex)
+		{
+			if (InfluencesToReplaceMap.Find(BoneMapIndex))
+			{
+				SkippedBoneMapIndices.Enqueue(BoneMapIndex); // Enqueue the skipped indices, will be used later by the remaining RequiredBones.
+				continue;
+			}
+
+			int32 RefSkeletonIndex = ImportedModel->LODModels[LOD].Sections[MaterialIndex].BoneMap[BoneMapIndex];
+			InverseBoneMap.Add(RefSkeletonIndex, BoneMapIndex);
+
+		}
+
+		// Add the data to mutable
+		int NumRequiredBones = ImportedModel->LODModels[LOD].RequiredBones.Num();
+		mu::SkeletonPtr MutableSkeleton = new mu::Skeleton;
+		MutableMesh->SetSkeleton(MutableSkeleton);
+		MutableMesh->SetBonePoseCount(NumRequiredBones);
+		MutableSkeleton->SetBoneCount(NumRequiredBones);
+		for (int32 RequiredBoneIndex = 0; RequiredBoneIndex < NumRequiredBones; ++RequiredBoneIndex)
+		{
+			int32 RefSkelIndex = ImportedModel->LODModels[LOD].RequiredBones[RequiredBoneIndex];
+
+			bool bSkinned = true;
+			// If the bone was required but not in the BoneMap (parent of a skinning bone) add it now.
+			if (!InverseBoneMap.Contains(RefSkelIndex))
+			{
+				bSkinned = false;
+				int32 NewBoneIndex = 0;
+				if (SkippedBoneMapIndices.Dequeue(NewBoneIndex)) // use skipped index
+					{
+					InverseBoneMap.Add(RefSkelIndex, NewBoneIndex);
+					}
+				else
+				{
+					InverseBoneMap.Add(RefSkelIndex, InverseBoneMap.Num());
+				}
+			}
+
+			int32 BoneMapIndex = InverseBoneMap[RefSkelIndex];
+			const FMeshBoneInfo& BoneInfo = SkeletalMesh->GetRefSkeleton().GetRefBoneInfo()[RefSkelIndex];
+			FString BoneName = BoneInfo.Name.ToString();
+
+			if (BoneMapIndex >= NumRequiredBones)
+			{
+				// This can only happen with invalid data.
+				UE_LOG(LogMutable, Warning, TEXT("ConvertSkeletalMeshToMutable: Source mesh [%s] section uses a bone [%s] that is not in its RequiredBones list."),
+					*SkeletalMesh->GetName(),
+					*BoneName);
+				MutableSkeleton->SetBoneCount(BoneMapIndex + 1);
+			}
+
+			MutableSkeleton->SetBoneName(BoneMapIndex, TCHAR_TO_ANSI(*BoneName));
+			int32 ParentRefSkeletonIndex = SkeletalMesh->GetRefSkeleton().GetParentIndex(RefSkelIndex);
+			int ParentBoneMapIndex = ParentRefSkeletonIndex >= 0 ? InverseBoneMap[ParentRefSkeletonIndex] : -1;
+			MutableSkeleton->SetBoneParent(BoneMapIndex, ParentBoneMapIndex);
+			MutableSkeleton->SetBoneId(BoneMapIndex,SkeletonID);
+
+			FMatrix44f BaseInvMatrix = SkeletalMesh->GetRefBasesInvMatrix()[RefSkelIndex];
+			FTransform3f BaseInvTransform;
+			BaseInvTransform.SetFromMatrix(BaseInvMatrix);
+			MutableMesh->SetBonePose(BoneMapIndex, TCHAR_TO_ANSI(*BoneName), BaseInvTransform.Inverse(), bSkinned);
+		}
+	}
+	
+	if (!bIgnorePhysics && SkeletalMesh->GetPhysicsAsset() && MutableMesh->GetSkeleton())
+	{
+		// Find BodySetups with relevant bones.
+		TArray<TObjectPtr<USkeletalBodySetup>>& SkeletalBodySetups = SkeletalMesh->GetPhysicsAsset()->SkeletalBodySetups;
+		
+		TArray<TObjectPtr<USkeletalBodySetup>> RelevantBodySetups;
+		RelevantBodySetups.Reserve(SkeletalBodySetups.Num());
+
+		for ( TObjectPtr<USkeletalBodySetup>& BodySetup : SkeletalBodySetups )
+		{
+			FString BodyBoneName = BodySetup->BoneName.ToString();
+
+			const int32 SkeletonBoneCount = MutableMesh->GetSkeleton()->GetBoneCount();
+			for ( int32 I = 0; I < SkeletonBoneCount; ++I )
+			{
+				FString SkeletonBoneName = MutableMesh->GetSkeleton()->GetBoneName( I );
+				
+				if (SkeletonBoneName.Equals(BodyBoneName))
+				{
+					RelevantBodySetups.Add(BodySetup);
+				}
+			}
+		}
+
+		mu::Ptr<mu::PhysicsBody> PhysicsBody = new mu::PhysicsBody;
+		
+		const int32 NumBodySetups = RelevantBodySetups.Num();	
+		PhysicsBody->SetBodyCount(NumBodySetups);
+
+		auto GetKBodyElemFlags = [](const FKShapeElem& KElem) -> uint32
+		{
+			uint8 ElemCollisionEnabled = static_cast<uint8>( KElem.GetCollisionEnabled() );
+			
+			uint32 Flags = static_cast<uint32>( ElemCollisionEnabled );
+			Flags = Flags | (static_cast<uint32>(KElem.GetContributeToMass()) << 8);
+
+			return Flags; 
+		};
+
+		for ( int32 B = 0; B < NumBodySetups; ++B )
+		{
+			TObjectPtr<USkeletalBodySetup>& BodySetup = RelevantBodySetups[B];
+			
+			FString BodyBoneName = BodySetup->BoneName.ToString();
+			PhysicsBody->SetBodyBoneName( B, TCHAR_TO_ANSI(*BodyBoneName) );
+			
+			const int32 NumSpheres = BodySetup->AggGeom.SphereElems.Num();
+			PhysicsBody->SetSphereCount( B, NumSpheres );
+
+			for ( int32 I = 0; I < NumSpheres; ++I )
+			{
+				const FKSphereElem& SphereElem = BodySetup->AggGeom.SphereElems[I];
+				PhysicsBody->SetSphere( B, I, FVector3f(SphereElem.Center), SphereElem.Radius );
+
+				const FString ElemName = SphereElem.GetName().ToString();
+				PhysicsBody->SetSphereName(B, I, TCHAR_TO_ANSI(*ElemName));	
+				PhysicsBody->SetSphereFlags(B, I, GetKBodyElemFlags(SphereElem));
+			}
+
+			const int32 NumBoxes = BodySetup->AggGeom.BoxElems.Num();
+			PhysicsBody->SetBoxCount( B, NumBoxes );
+
+			for ( int32 I = 0; I < NumBoxes; ++I )
+			{
+				const FKBoxElem& BoxElem = BodySetup->AggGeom.BoxElems[I];
+				PhysicsBody->SetBox( B, I, 
+						FVector3f(BoxElem.Center), 
+						FQuat4f(BoxElem.Rotation.Quaternion()), 
+						FVector3f(BoxElem.X, BoxElem.Y, BoxElem.Z));
+
+				const FString KElemName = BoxElem.GetName().ToString();
+				PhysicsBody->SetBoxName(B, I, TCHAR_TO_ANSI(*KElemName));
+				PhysicsBody->SetBoxFlags(B, I, GetKBodyElemFlags(BoxElem));
+			}
+
+			const int32 NumConvex = BodySetup->AggGeom.ConvexElems.Num();
+			PhysicsBody->SetConvexCount( B, NumConvex );
+			for ( int32 I = 0; I < NumConvex; ++I )
+			{
+				const FKConvexElem& ConvexElem = BodySetup->AggGeom.ConvexElems[I];
+
+				// Convert to FVector3f
+				TArray<FVector3f> VertexData;
+				VertexData.SetNumUninitialized( ConvexElem.VertexData.Num() );
+				for ( int32 Elem = VertexData.Num() - 1; Elem >= 0; --Elem )
+				{
+					VertexData[Elem] = FVector3f(ConvexElem.VertexData[Elem]);
+				}
+				
+				PhysicsBody->SetConvex( B, I, 
+						VertexData.GetData(), ConvexElem.VertexData.Num(), 
+						ConvexElem.IndexData.GetData(), ConvexElem.IndexData.Num(), 
+						FTransform3f(ConvexElem.GetTransform()) );
+				
+				const FString KElemName = ConvexElem.GetName().ToString();
+				PhysicsBody->SetConvexName(B, I, TCHAR_TO_ANSI(*KElemName));
+				PhysicsBody->SetConvexFlags(B, I, GetKBodyElemFlags(ConvexElem));
+			}
+
+			const int32 NumSphyls = BodySetup->AggGeom.SphylElems.Num();
+			PhysicsBody->SetSphylCount( B, NumSphyls );
+
+			for ( int32 I = 0; I < NumSphyls; ++I )
+			{
+				const FKSphylElem& SphylElem = BodySetup->AggGeom.SphylElems[I];
+				PhysicsBody->SetSphyl( B, I, 
+						FVector3f(SphylElem.Center), 
+						FQuat4f(SphylElem.Rotation.Quaternion()), 
+						SphylElem.Radius, SphylElem.Length );
+
+				const FString KElemName = SphylElem.GetName().ToString();
+				PhysicsBody->SetSphylName(B, I, TCHAR_TO_ANSI(*KElemName));
+				PhysicsBody->SetSphylFlags(B, I, GetKBodyElemFlags(SphylElem));
+			}
+
+			const int32 NumTaperedCapsules = BodySetup->AggGeom.TaperedCapsuleElems.Num();
+			PhysicsBody->SetTaperedCapsuleCount( B, NumTaperedCapsules );
+
+			for ( int32 I = 0; I < NumTaperedCapsules; ++I )
+			{
+				const FKTaperedCapsuleElem& TaperedCapsuleElem = BodySetup->AggGeom.TaperedCapsuleElems[I];
+				PhysicsBody->SetTaperedCapsule( B, I, 
+						FVector3f(TaperedCapsuleElem.Center), 
+						FQuat4f(TaperedCapsuleElem.Rotation.Quaternion()), 
+						TaperedCapsuleElem.Radius0, TaperedCapsuleElem.Radius1, TaperedCapsuleElem.Length );
+				
+				const FString KElemName = TaperedCapsuleElem.GetName().ToString();
+				PhysicsBody->SetTaperedCapsuleName(B, I, TCHAR_TO_ANSI(*KElemName));
+				PhysicsBody->SetTaperedCapsuleFlags( B, I, GetKBodyElemFlags(TaperedCapsuleElem));
+			}
+		}
+
+		MutableMesh->SetPhysicsBody(PhysicsBody);
+	}
+	
+	return MutableMesh;
+}
+
+
+mu::MeshPtr ConvertStaticMeshToMutable(UStaticMesh* StaticMesh, int LOD, int MaterialIndex, FMutableGraphGenerationContext& GenerationContext, const UCustomizableObjectNode* CurrentNode)
+{
+	if (!StaticMesh->GetRenderData() ||
+		!StaticMesh->GetRenderData()->LODResources.IsValidIndex(LOD) ||
+		!StaticMesh->GetRenderData()->LODResources[LOD].Sections.IsValidIndex(MaterialIndex))
+	{
+		FString Msg = FString::Printf(TEXT("Degenerated static mesh found for LOD %d Material %d. It will be ignored. "), LOD, MaterialIndex);
+		GenerationContext.Compiler->CompilerLog(FText::FromString(Msg), CurrentNode, EMessageSeverity::Warning);
+		return nullptr;
+	}
+
+	mu::MeshPtr MutableMesh = new mu::Mesh();
+
+	// Vertices
+	int VertexStart = StaticMesh->GetRenderData()->LODResources[LOD].Sections[MaterialIndex].MinVertexIndex;
+	int VertexCount = StaticMesh->GetRenderData()->LODResources[LOD].Sections[MaterialIndex].MaxVertexIndex - VertexStart + 1;
+
+	MutableMesh->GetVertexBuffers().SetElementCount(VertexCount);
+	{
+		using namespace mu;
+
+		MutableMesh->GetVertexBuffers().SetBufferCount(5);
+
+		// Position buffer
+		{
+			const FPositionVertexBuffer& VertexBuffer = Helper_GetStaticMeshPositionVertexBuffer(StaticMesh, LOD);
+
+			const int ElementSize = 12;
+			const int ChannelCount = 1;
+			const MESH_BUFFER_SEMANTIC Semantics[ChannelCount] = { MBS_POSITION };
+			const int SemanticIndices[ChannelCount] = { 0 };
+			const MESH_BUFFER_FORMAT Formats[ChannelCount] = { MBF_FLOAT32 };
+			const int Components[ChannelCount] = { 3 };
+			const int Offsets[ChannelCount] = { 0 };
+
+			MutableMesh->GetVertexBuffers().SetBuffer(MUTABLE_VERTEXBUFFER_POSITION, ElementSize, ChannelCount, Semantics, SemanticIndices, Formats, Components, Offsets);
+			FMemory::Memcpy(
+				MutableMesh->GetVertexBuffers().GetBufferData(MUTABLE_VERTEXBUFFER_POSITION),
+				&VertexBuffer.VertexPosition(VertexStart),
+				VertexCount * ElementSize);
+		}
+
+		// Tangent buffer
+		{
+			FStaticMeshVertexBuffer& VertexBuffer = Helper_GetStaticMeshVertexBuffer(StaticMesh, LOD);
+
+			MESH_BUFFER_SEMANTIC Semantics[2];
+			int SemanticIndices[2];
+			MESH_BUFFER_FORMAT Formats[2];
+			int Components[2];
+			int Offsets[2];
+
+			int currentChannel = 0;
+			int currentOffset = 0;
+
+			Semantics[currentChannel] = MBS_TANGENT;
+			SemanticIndices[currentChannel] = 0;
+			Formats[currentChannel] = MBF_PACKEDDIRS8;
+			Components[currentChannel] = 4;
+			Offsets[currentChannel] = currentOffset;
+			currentOffset += 4;
+			++currentChannel;
+
+			Semantics[currentChannel] = MBS_NORMAL;
+			SemanticIndices[currentChannel] = 0;
+			Formats[currentChannel] = MBF_PACKEDDIRS8;
+
+			Components[currentChannel] = 4;
+			Offsets[currentChannel] = currentOffset;
+			currentOffset += 4;
+			//++currentChannel;
+
+			MutableMesh->GetVertexBuffers().SetBuffer(MUTABLE_VERTEXBUFFER_TANGENT, currentOffset, 2, Semantics, SemanticIndices, Formats, Components, Offsets);
+
+			const uint8_t* pTangentData = static_cast<const uint8_t*>(VertexBuffer.GetTangentData());
+			FMemory::Memcpy(
+				MutableMesh->GetVertexBuffers().GetBufferData(MUTABLE_VERTEXBUFFER_TANGENT),
+				pTangentData + VertexStart * currentOffset,
+				VertexCount * currentOffset);
+		}
+
+		// Texture coordinates
+		{
+			FStaticMeshVertexBuffer& VertexBuffer = StaticMesh->GetRenderData()->LODResources[LOD].VertexBuffers.StaticMeshVertexBuffer;
+
+			int texChannels = VertexBuffer.GetNumTexCoords();
+			int ChannelCount = texChannels;
+
+			MESH_BUFFER_SEMANTIC* Semantics = new MESH_BUFFER_SEMANTIC[ChannelCount];
+			int* SemanticIndices = new int[ChannelCount];
+			MESH_BUFFER_FORMAT* Formats = new MESH_BUFFER_FORMAT[ChannelCount];
+			int* Components = new int[ChannelCount];
+			int* Offsets = new int[ChannelCount];
+
+			int currentChannel = 0;
+			int currentOffset = 0;
+
+			int texChannelSize;
+			MESH_BUFFER_FORMAT texChannelFormat;
+			if (VertexBuffer.GetUseFullPrecisionUVs())
+			{
+				texChannelSize = 2 * 4;
+				texChannelFormat = MBF_FLOAT32;
+			}
+			else
+			{
+				texChannelSize = 2 * 2;
+				texChannelFormat = MBF_FLOAT16;
+			}
+
+			for (int c = 0; c < texChannels; ++c)
+			{
+				Semantics[currentChannel] = MBS_TEXCOORDS;
+				SemanticIndices[currentChannel] = c;
+				Formats[currentChannel] = texChannelFormat;
+				Components[currentChannel] = 2;
+				Offsets[currentChannel] = currentOffset;
+				currentOffset += texChannelSize;
+				++currentChannel;
+			}
+
+			MutableMesh->GetVertexBuffers().SetBuffer(MUTABLE_VERTEXBUFFER_TEXCOORDS, currentOffset, ChannelCount, Semantics, SemanticIndices, Formats, Components, Offsets);
+
+			const uint8_t* pTextureCoordData = static_cast<const uint8_t*>(VertexBuffer.GetTexCoordData());
+			FMemory::Memcpy(
+				MutableMesh->GetVertexBuffers().GetBufferData(MUTABLE_VERTEXBUFFER_TEXCOORDS),
+				pTextureCoordData + VertexStart * currentOffset,
+				VertexCount * currentOffset);
+
+			delete[] Semantics;
+			delete[] SemanticIndices;
+			delete[] Formats;
+			delete[] Components;
+			delete[] Offsets;
+		}
+	}
+
+	// Indices
+	{
+		int IndexStart = StaticMesh->GetRenderData()->LODResources[LOD].Sections[MaterialIndex].FirstIndex;
+		int IndexCount = StaticMesh->GetRenderData()->LODResources[LOD].Sections[MaterialIndex].NumTriangles * 3;
+		MutableMesh->GetIndexBuffers().SetBufferCount(1);
+		MutableMesh->GetIndexBuffers().SetElementCount(IndexCount);
+		MutableMesh->GetFaceBuffers().SetElementCount(IndexCount / 3);
+
+		using namespace mu;
+		const int ElementSize = 2;
+		const int ChannelCount = 1;
+		const MESH_BUFFER_SEMANTIC Semantics[ChannelCount] = { MBS_VERTEXINDEX };
+		const int SemanticIndices[ChannelCount] = { 0 };
+		MESH_BUFFER_FORMAT Formats[ChannelCount] = { MBF_UINT16 };
+		const int Components[ChannelCount] = { 1 };
+		const int Offsets[ChannelCount] = { 0 };
+
+		MutableMesh->GetIndexBuffers().SetBuffer(0, ElementSize, ChannelCount, Semantics, SemanticIndices, Formats, Components, Offsets);
+
+		//if (ElementSize==4)
+		//{
+		//	uint32* pDest = reinterpret_cast<uint32*>( MutableMesh->GetIndexBuffers().GetBufferData(0) );
+		//	const uint32* pSource = reinterpret_cast<const uint32*>( TypedNode->StaticMesh->RenderData->LODResources[LOD].MultiSizeIndexContainer.GetIndexBuffer()->GetPointerTo(IndexStart) );
+
+		//	for ( int i=0; i<IndexCount; ++i )
+		//	{
+		//		*pDest = *pSource - VertexStart;
+		//		++pDest;
+		//		++pSource;
+		//	}
+		//}
+		//else
+		{
+			FIndexArrayView Source = StaticMesh->GetRenderData()->LODResources[LOD].IndexBuffer.GetArrayView();
+			//const uint16* pSource = reinterpret_cast<const uint16*>( StaticMesh->RenderData->LODResources[LOD].IndexBuffer.Indices.GetResourceData() );
+			//pSource += IndexStart;
+			uint16* pDest = reinterpret_cast<uint16*>(MutableMesh->GetIndexBuffers().GetBufferData(0));
+
+			for (int i = 0; i < IndexCount; ++i)
+			{
+				*pDest = Source[IndexStart + i] - VertexStart;
+				++pDest;
+			}
+		}
+	}
+
+
+	return MutableMesh;
+}
+
+
+// Convert a mesh constant to a mutable format.
+mu::MeshPtr GenerateMutableMesh(UObject * Mesh, int32 LOD, int32 MaterialIndex, FMutableGraphGenerationContext & GenerationContext, const UCustomizableObjectNode* CurrentNode)
+{
+	// Get the mesh generation flags to use
+	uint32 CurrentFlags = GenerationContext.MeshGenerationFlags.Last();
+	FMutableGraphGenerationContext::FGeneratedMeshData::FKey Key = { Mesh, LOD, MaterialIndex, CurrentFlags };
+	mu::MeshPtr MutableMesh = GenerationContext.FindGeneratedMesh(Key);
+	
+	if (!MutableMesh)
+	{
+		if (USkeletalMesh* SkeletalMesh = Cast<USkeletalMesh>(Mesh))
+		{
+			MutableMesh = ConvertSkeletalMeshToMutable(SkeletalMesh, LOD, MaterialIndex, GenerationContext, CurrentNode);
+		}
+		else if (UStaticMesh* StaticMesh = Cast<UStaticMesh>(Mesh))
+		{
+			MutableMesh = ConvertStaticMeshToMutable(StaticMesh, LOD, MaterialIndex, GenerationContext, CurrentNode);
+		}
+		else
+		{
+			GenerationContext.Compiler->CompilerLog(LOCTEXT("UnimplementedMesh", "Mesh type not implemented yet."), CurrentNode);
+		}
+
+		if (MutableMesh)
+		{
+			GenerationContext.GeneratedMeshes.Push({ Key, MutableMesh });
+		}
+	}
+
+	if (MutableMesh)
+	{
+		FMeshData MeshData; 
+		MeshData.Mesh = Mesh;
+		MeshData.LOD = LOD;
+		MeshData.MaterialIndex = MaterialIndex;
+		MeshData.Node = CurrentNode;
+		GenerationContext.PinData.GetCurrent().MeshesData.Add(MeshData); // Set::Emplace only supports single element constructors
+	}
+	
+	return MutableMesh;
+}
+
+
+mu::MeshPtr BuildMorphedMutableMesh(const UEdGraphPin* BaseSourcePin, const FString& MorphTargetName, FMutableGraphGenerationContext & GenerationContext, int32 RowIndex)
+{
+	check(BaseSourcePin);
+	SCOPED_PIN_DATA(GenerationContext, BaseSourcePin)
+
+	mu::MeshPtr MorphedSourceMesh;
+
+	if (!BaseSourcePin)
+	{
+		GenerationContext.Compiler->CompilerLog(LOCTEXT("NULLBaseSourcePin", "Morph base not set."), nullptr);
+		return nullptr;
+	}
+
+	int32 LODIndex = -1; // Initialization required to remove uninitialized warning.
+	int32 SectionIndex = -1;
+	
+	USkeletalMesh* SkeletalMesh = nullptr;
+	UCustomizableObjectNode* Node = Cast<UCustomizableObjectNode>(BaseSourcePin->GetOwningNode());
+
+	if (const UCustomizableObjectNodeSkeletalMesh* TypedNodeSkeletalMesh = Cast<UCustomizableObjectNodeSkeletalMesh>(Node))
+	{
+		int32 LayoutIndex;
+		TypedNodeSkeletalMesh->GetPinSection(*BaseSourcePin, LODIndex, SectionIndex, LayoutIndex);
+		SkeletalMesh = TypedNodeSkeletalMesh->SkeletalMesh;
+	}
+
+	else if (const UCustomizableObjectNodeTable* TypedNodeTable = Cast<UCustomizableObjectNodeTable>(Node))
+	{
+		TypedNodeTable->GetPinLODAndMaterial(BaseSourcePin, LODIndex, SectionIndex);
+		SkeletalMesh = TypedNodeTable->GetSkeletalMeshAt(BaseSourcePin, RowIndex);
+	}
+
+	if (SkeletalMesh)
+	{
+		// See if we need to select another LOD because of automatic LOD generation
+		if (GenerationContext.CurrentAutoLODStrategy == ECustomizableObjectAutomaticLODStrategy::AutomaticFromMesh)
+		{
+			if (SkeletalMesh && Helper_GetImportedModel(SkeletalMesh))
+			{
+				// Add the CurrentLOD being generated to the first LOD used from this SkeletalMesh
+				int32 CurrentLOD = LODIndex + GenerationContext.CurrentLOD;
+
+				// If the mesh has additional LODs and they use the same material, use them.
+				FSkeletalMeshModel* ImportedModel = Helper_GetImportedModel(SkeletalMesh);
+				if (ImportedModel->LODModels.Num() > CurrentLOD
+					&&
+					ImportedModel->LODModels[CurrentLOD].Sections.Num() > SectionIndex)
+				{
+					LODIndex = CurrentLOD;
+				}
+				else
+				{
+					// Find the closest valid LOD 
+					for (int32 LOD = ImportedModel->LODModels.Num() - 1; LOD >= 0; --LOD)
+					{
+						if (ImportedModel->LODModels[LOD].Sections.Num() > SectionIndex)
+						{
+							LODIndex = LOD;
+							break;
+						}
+					}
+				}
+			}
+		}
+
+
+		// Get the base mesh
+		mu::MeshPtr BaseSourceMesh = GenerateMutableMesh(SkeletalMesh, LODIndex, SectionIndex, GenerationContext, Node);
+		if (BaseSourceMesh)
+		{
+			// Clone it (it will probably be shared)
+			MorphedSourceMesh = BaseSourceMesh->Clone();
+
+			// Bake the morph in the new mutable mesh
+			UMorphTarget* MorphTarget = SkeletalMesh ? SkeletalMesh->FindMorphTarget(*MorphTargetName) : nullptr;
+
+			if (MorphTarget && MorphTarget->GetMorphLODModels().IsValidIndex(LODIndex))
+			{
+				int posBuf, posChannel;
+				MorphedSourceMesh->GetVertexBuffers().FindChannel(mu::MBS_POSITION, 0, &posBuf, &posChannel);
+				int posElemSize = MorphedSourceMesh->GetVertexBuffers().GetElementSize(posBuf);
+				int posOffset = MorphedSourceMesh->GetVertexBuffers().GetChannelOffset(posBuf, posChannel);
+				uint8* posBuffer = MorphedSourceMesh->GetVertexBuffers().GetBufferData(posBuf) + posOffset;
+
+				uint32 MaterialVertexStart = (uint32)Helper_GetImportedModel(SkeletalMesh)->LODModels[LODIndex].Sections[SectionIndex].GetVertexBufferIndex();
+				uint32 MeshVertexCount = (uint32)MorphedSourceMesh->GetVertexBuffers().GetElementCount();
+
+				const TArray<FMorphTargetLODModel>& MorphLODModels = MorphTarget->GetMorphLODModels();
+
+				for (int v = 0; v < MorphLODModels[LODIndex].Vertices.Num(); ++v)
+				{
+					const FMorphTargetDelta& data = MorphLODModels[LODIndex].Vertices[v];
+					if ((data.SourceIdx >= MaterialVertexStart) &&
+						((data.SourceIdx - MaterialVertexStart) < MeshVertexCount))
+					{
+						float* pPos = (float*)(posBuffer + posElemSize * (data.SourceIdx - MaterialVertexStart));
+						pPos[0] += data.PositionDelta[0];
+						pPos[1] += data.PositionDelta[1];
+						pPos[2] += data.PositionDelta[2];
+					}
+				}
+			}
+		}
+	}
+
+	return MorphedSourceMesh;
+}
+
+
+void GenerateMorphFactor(const UCustomizableObjectNode* Node, const UEdGraphPin& FactorPin, FMutableGraphGenerationContext& GenerationContext, mu::NodeMeshMorphPtr MeshNode)
+{
+	if (const UEdGraphPin* ConnectedPin = FollowInputPin(FactorPin))
+	{
+		UEdGraphNode* floatNode = ConnectedPin->GetOwningNode();
+		bool validStaticFactor = true;
+		
+		if (const UCustomizableObjectNodeFloatParameter* floatParameterNode = Cast<UCustomizableObjectNodeFloatParameter>(floatNode))
+		{
+			if (floatParameterNode->DefaultValue < -1.0f || floatParameterNode->DefaultValue > 1.0f)
+			{
+				validStaticFactor = false;
+				FString msg = FString::Printf(TEXT("Mesh morph nodes only accept factors between -1.0 and 1.0 inclusive but the default value of the float parameter node is (%f). Factor will be ignored."), floatParameterNode->DefaultValue);
+				GenerationContext.Compiler->CompilerLog(FText::FromString(msg), Node);
+			}
+			if (floatParameterNode->ParamUIMetadata.MinimumValue < -1.0f)
+			{
+				validStaticFactor = false;
+				FString msg = FString::Printf(TEXT("Mesh morph nodes only accept factors between -1.0 and 1.0 inclusive but the minimum UI value for the input float parameter node is (%f). Factor will be ignored."), floatParameterNode->ParamUIMetadata.MinimumValue);
+				GenerationContext.Compiler->CompilerLog(FText::FromString(msg), Node);
+			}
+			if (floatParameterNode->ParamUIMetadata.MaximumValue > 1.0f)
+			{
+				validStaticFactor = false;
+				FString msg = FString::Printf(TEXT("Mesh morph nodes only accept factors between -1.0 and 1.0 inclusive but the maximum UI value for the input float parameter node is (%f). Factor will be ignored."), floatParameterNode->ParamUIMetadata.MaximumValue);
+				GenerationContext.Compiler->CompilerLog(FText::FromString(msg), Node);
+			}
+		}
+		
+		else if (const UCustomizableObjectNodeFloatConstant* floatConstantNode = Cast<UCustomizableObjectNodeFloatConstant>(floatNode))
+		{
+			if (floatConstantNode->Value < -1.0f || floatConstantNode->Value > 1.0f)
+			{
+				validStaticFactor = false;
+				FString msg = FString::Printf(TEXT("Mesh morph nodes only accept factors between -1.0 and 1.0 inclusive but the value of the float constant node is (%f). Factor will be ignored."), floatConstantNode->Value);
+				GenerationContext.Compiler->CompilerLog(FText::FromString(msg), Node);
+			}
+		}
+
+		if (validStaticFactor)
+		{
+			mu::NodeScalarPtr FactorNode = GenerateMutableSourceFloat(ConnectedPin, GenerationContext);
+			MeshNode->SetFactor(FactorNode);
+		}
+	}
+}
+
+TArray<USkeletalMesh*> GetSkeletalMeshesForReshapeSelection(
+		const UEdGraphNode* SkeletalMeshOrTableNode, const UEdGraphPin* SourceMeshPin)
+{
+	TArray<USkeletalMesh*> SkeletalMeshes;
+
+	if (!(SkeletalMeshOrTableNode && SourceMeshPin))
+	{
+		return SkeletalMeshes;
+	}
+
+	if (const UCustomizableObjectNodeSkeletalMesh* SkeletalMeshNode = Cast<UCustomizableObjectNodeSkeletalMesh>(SkeletalMeshOrTableNode))
+	{
+		if (SkeletalMeshNode->SkeletalMesh)
+		{
+			SkeletalMeshes.Add(SkeletalMeshNode->SkeletalMesh);
+		}
+	}
+	else if (const UCustomizableObjectNodeTable* TableNode = Cast<UCustomizableObjectNodeTable>(SkeletalMeshOrTableNode))
+	{
+		if (TableNode->Table)
+		{
+			uint32 NumRows = TableNode->Table->GetRowNames().Num();
+
+			for (uint32 RowIndex = 0; RowIndex < NumRows; ++RowIndex)
+			{
+				USkeletalMesh* SkeletalMesh = TableNode->GetSkeletalMeshAt(SourceMeshPin, RowIndex);
+
+				if (SkeletalMesh)
+				{
+					SkeletalMeshes.Add(SkeletalMesh);
+				}
+			}
+		}	
+	}
+	else
+	{
+		checkf(false, TEXT("Node not expected."));
+	}
+
+	return SkeletalMeshes;
+}
+
+bool GetAndValidateReshapeBonesToDeform(
+	TArray<FString>& OutBonesToDeform,
+	const TArray<FMeshReshapeBoneReference>& InBonesToDeform,
+	const TArray<USkeletalMesh*>& SkeletalMeshes,
+	const UCustomizableObjectNode* Node,
+	FMutableGraphGenerationContext& GenerationContext)
+{
+	bool bSetRefreshWarning = false;
+
+	int32 NumBonesToDeform = InBonesToDeform.Num();
+	for (int32 BoneIndex = 0; BoneIndex < NumBonesToDeform; ++BoneIndex)
+	{
+		bool bMissingBone = true;
+		
+		const FName BoneName = InBonesToDeform[BoneIndex].BoneName;
+
+		for (const USkeletalMesh* SkeletalMesh : SkeletalMeshes)
+		{
+			if (SkeletalMesh->GetRefSkeleton().FindBoneIndex(BoneName) != INDEX_NONE)
+			{
+				OutBonesToDeform.AddUnique(BoneName.ToString());
+
+				bMissingBone = false;
+				break;
+			}
+		}
+
+		if (bMissingBone)
+		{
+			FString msg = FString::Printf(TEXT("Could not find the selected Bone to Deform [%s] in Skeleton"), *(BoneName.ToString()));
+
+			GenerationContext.Compiler->CompilerLog(FText::FromString(msg), Node);
+
+			bSetRefreshWarning = true;
+		}
+	}
+
+	return bSetRefreshWarning;
+}
+
+bool GetAndValidateReshapePhysicsToDeform(
+	TArray<FString>& OutPhysiscsToDeform,
+	const TArray<FMeshReshapeBoneReference>& InPhysicsToDeform,
+	const TArray<USkeletalMesh*>& SkeletalMeshes,
+	const UCustomizableObjectNode* Node,
+	FMutableGraphGenerationContext& GenerationContext)
+{
+	bool bSetRefreshWarning = false;
+
+	int32 NumPhysicsToDeform = InPhysicsToDeform.Num();
+	for (int32 BodyIndex = 0; BodyIndex < NumPhysicsToDeform; ++BodyIndex)
+	{
+		bool bMissingBone = true;
+		bool bMissingBody = true;
+
+		const FName BodyBoneName = InPhysicsToDeform[BodyIndex].BoneName;
+		for (const USkeletalMesh* SkeletalMesh : SkeletalMeshes)
+		{
+			if (SkeletalMesh->GetRefSkeleton().FindBoneIndex(BodyBoneName) != INDEX_NONE)
+			{
+				bMissingBone = false;
+			}
+
+			UPhysicsAsset* PhysicsAsset = SkeletalMesh->GetPhysicsAsset();
+			if (PhysicsAsset && !bMissingBone)
+			{
+				const int32 FoundIndex = PhysicsAsset->SkeletalBodySetups.IndexOfByPredicate(
+					[&BodyBoneName]( const TObjectPtr<USkeletalBodySetup>& Setup ) {  return Setup->BoneName == BodyBoneName; });
+
+				if (FoundIndex != INDEX_NONE)
+				{
+					OutPhysiscsToDeform.AddUnique(BodyBoneName.ToString());
+					bMissingBody = false;
+					break;
+				}
+			}
+		}
+
+		if (bMissingBone)
+		{
+			FString BoneName = InPhysicsToDeform[BodyIndex].BoneName.ToString();
+			FString msg = FString::Printf(TEXT("Could not find the selected Physics Body bone to Deform [%s] in Skeleton"), *BoneName);
+
+			GenerationContext.Compiler->CompilerLog(FText::FromString(msg), Node);
+
+			bSetRefreshWarning = true;
+		}
+
+		if (bMissingBody && !bMissingBone)
+		{
+			FString BoneName = InPhysicsToDeform[BodyIndex].BoneName.ToString();
+			FString msg = FString::Printf(TEXT("Selected Bone to Deform [%s] does not have any physics body attached."), *BoneName);
+
+			GenerationContext.Compiler->CompilerLog(FText::FromString(msg), Node);
+
+			bSetRefreshWarning = true;
+		}
+			
+	}
+	return bSetRefreshWarning;
+}
+
+mu::NodeMeshPtr GenerateMorphMesh(const UEdGraphPin* Pin,
+	TArray<UCustomizableObjectNodeMeshMorph*> TypedNodeMorphs,
+	int32 MorphIndex,
+	mu::NodeMeshPtr SourceNode,
+	FMutableGraphGenerationContext & GenerationContext,
+	FMutableGraphMeshGenerationData & MeshData,
+	const FString& TableColumnName = "")
+{
+	SCOPED_PIN_DATA(GenerationContext, Pin)
+	
+	// SkeletalMesh node
+	const UEdGraphNode * MeshNode = Pin->GetOwningNode();
+	check(MeshNode);
+	
+	// Current morph node
+	const UCustomizableObjectNode* MorphNode = TypedNodeMorphs[MorphIndex];
+	check(MorphNode);
+	
+	mu::NodeMeshMorphPtr Result = new mu::NodeMeshMorph();
+	Result->SetMorphCount(2);
+	
+	// Factor
+	GenerateMorphFactor(MorphNode, *TypedNodeMorphs[MorphIndex]->FactorPin(), GenerationContext, Result);
+	
+	// Base
+	if (MorphIndex == TypedNodeMorphs.Num() - 1)
+	{
+		Result->SetBase(SourceNode);
+	}
+	else
+	{
+		mu::NodeMeshPtr NextMorph = GenerateMorphMesh(Pin, TypedNodeMorphs, MorphIndex + 1, SourceNode, GenerationContext, MeshData, TableColumnName);
+		Result->SetBase(NextMorph);
+	}
+	
+	// Target
+	mu::NodeMeshPtr BaseSourceMesh = SourceNode;
+
+	mu::MeshPtr MorphedSourceMesh;
+
+	bool bSuccess = false;
+	
+	if (const UCustomizableObjectNodeTable* TypedNodeTable = Cast<UCustomizableObjectNodeTable>(Pin->GetOwningNode()))
+	{
+		// Generate a new Column for each morph
+		int32 NumRows = TypedNodeTable->Table->GetRowNames().Num();
+
+		// Should exist
+		mu::TablePtr Table = GenerationContext.GeneratedTables[TypedNodeTable->Table->GetName()];
+
+		FString ColumnName = TableColumnName + TypedNodeMorphs[MorphIndex]->MorphTargetName;
+		int32 ColumnIndex = -1;
+
+		for (int32 RowIndex = 0; RowIndex < NumRows; ++RowIndex)
+		{
+			ColumnIndex = Table->FindColumn(TCHAR_TO_ANSI(*ColumnName));
+
+			if (ColumnIndex == -1)
+			{
+				ColumnIndex = Table->AddColumn(TCHAR_TO_ANSI(*ColumnName), mu::TABLE_COLUMN_TYPE::TCT_MESH);
+			}
+
+			mu::MeshPtr MorphedSourceTableMesh = BuildMorphedMutableMesh(Pin, TypedNodeMorphs[MorphIndex]->MorphTargetName, GenerationContext, RowIndex);
+			Table->SetCell(ColumnIndex, RowIndex, MorphedSourceTableMesh.get());
+		}
+
+		if (ColumnIndex > -1)
+		{
+			bSuccess = true;
+
+			mu::NodeMeshTablePtr MorphedSourceMeshNodeTable = new mu::NodeMeshTable;
+			MorphedSourceMeshNodeTable->SetTable(Table);
+			MorphedSourceMeshNodeTable->SetColumn(TCHAR_TO_ANSI(*ColumnName));
+			MorphedSourceMeshNodeTable->SetParameterName(TCHAR_TO_ANSI(*TypedNodeTable->ParameterName));
+			MorphedSourceMeshNodeTable->SetMessageContext(MorphNode);
+
+			mu::NodeMeshMakeMorphPtr IdentityMorph = new mu::NodeMeshMakeMorph;
+			IdentityMorph->SetBase(BaseSourceMesh.get());
+			IdentityMorph->SetTarget(BaseSourceMesh.get());
+			IdentityMorph->SetMessageContext(MorphNode);
+
+			Result->SetMorph(0, IdentityMorph);
+
+			mu::NodeMeshMakeMorphPtr Morph = new mu::NodeMeshMakeMorph;
+			Morph->SetBase(BaseSourceMesh.get());
+			Morph->SetTarget(MorphedSourceMeshNodeTable.get());
+			Morph->SetMessageContext(MorphNode);
+
+			Result->SetMorph(1, Morph);
+		}
+	}
+	else
+	{
+		MorphedSourceMesh = BuildMorphedMutableMesh(Pin, TypedNodeMorphs[MorphIndex]->MorphTargetName, GenerationContext);
+
+		if (MorphedSourceMesh)
+		{
+			bSuccess = true;
+
+			mu::NodeMeshConstantPtr MorphedSourceMeshNode = new mu::NodeMeshConstant;
+			MorphedSourceMeshNode->SetValue(MorphedSourceMesh);
+			MorphedSourceMeshNode->SetMessageContext(MorphNode);
+
+			mu::NodeMeshMakeMorphPtr IdentityMorph = new mu::NodeMeshMakeMorph;
+			IdentityMorph->SetBase(BaseSourceMesh.get());
+			IdentityMorph->SetTarget(BaseSourceMesh.get());
+			IdentityMorph->SetMessageContext(MorphNode);
+
+			Result->SetMorph(0, IdentityMorph);
+
+			mu::NodeMeshMakeMorphPtr Morph = new mu::NodeMeshMakeMorph;
+			Morph->SetBase(BaseSourceMesh.get());
+			Morph->SetTarget(MorphedSourceMeshNode.get());
+			Morph->SetMessageContext(MorphNode);
+
+			Result->SetMorph(1, Morph);
+
+			UCustomizableObjectNodeMeshMorph* TypedMorphNode = TypedNodeMorphs[MorphIndex];
+
+			Result->SetDeformAllBones(TypedMorphNode->bDeformAllBones);
+			Result->SetDeformAllPhysics(TypedMorphNode->bDeformAllPhysicsBodies);
+			Result->SetReshapeSkeleton(TypedMorphNode->bReshapeSkeleton);
+			Result->SetReshapePhysicsVolumes(TypedMorphNode->bReshapePhysicsVolumes);	
+			{
+				const UEdGraphPin* ConnectedPin = FollowInputPin(*TypedMorphNode->MeshPin());
+				const UEdGraphPin* SourceMeshPin = ConnectedPin ? FindMeshBaseSource(*ConnectedPin, false) : nullptr;
+				const UEdGraphNode* SkeletalMeshNode = SourceMeshPin ? SourceMeshPin->GetOwningNode() : nullptr;
+
+				TArray<USkeletalMesh*> SkeletalMeshesToDeform = GetSkeletalMeshesForReshapeSelection(SkeletalMeshNode, SourceMeshPin);
+				
+				bool bWarningFound = false;
+				if (TypedMorphNode->bReshapeSkeleton && !TypedMorphNode->bDeformAllBones)
+				{
+					TArray<FString> BonesToDeform;
+					bWarningFound = GetAndValidateReshapeBonesToDeform(
+						BonesToDeform, TypedMorphNode->BonesToDeform, SkeletalMeshesToDeform, TypedMorphNode, GenerationContext);
+					
+					for (const FString& BoneName : BonesToDeform)
+					{
+						Result->AddBoneToDeform(TCHAR_TO_ANSI(*BoneName));
+					}
+				}
+
+				if (TypedMorphNode->bReshapePhysicsVolumes && !TypedMorphNode->bDeformAllPhysicsBodies)
+				{
+					TArray<FString> PhysicsToDeform;
+					bWarningFound = bWarningFound || GetAndValidateReshapePhysicsToDeform(
+						PhysicsToDeform, TypedMorphNode->PhysicsBodiesToDeform, SkeletalMeshesToDeform, TypedMorphNode, GenerationContext);
+	
+					for (const FString& PhysicsBoneName : PhysicsToDeform)
+					{
+						Result->AddPhysicsBodyToDeform(TCHAR_TO_ANSI(*PhysicsBoneName));
+					}	
+				}
+				
+				if (bWarningFound)
+				{
+					TypedMorphNode->SetRefreshNodeWarning();
+				}
+			}
+		}
+	}
+
+	if(!bSuccess)
+	{
+		GenerationContext.Compiler->CompilerLog(LOCTEXT("MorphGenerationFailed", "Failed to generate morph target."), MorphNode);
+	}
+
+	return Result;
+}
+
+
+void GenerateMorphTarget(const UCustomizableObjectNode* Node, const UEdGraphPin* BaseSourcePin, FMutableGraphGenerationContext& GenerationContext, mu::NodeMeshMorphPtr MeshNode, FString MorphName)
+{
+	SCOPED_PIN_DATA(GenerationContext, BaseSourcePin)
+
+	FMutableGraphMeshGenerationData DummyMeshData;
+	mu::NodeMeshPtr BaseSourceMesh = GenerateMutableSourceMesh(BaseSourcePin, GenerationContext, DummyMeshData);
+
+	mu::MeshPtr MorphedSourceMesh = BuildMorphedMutableMesh(BaseSourcePin, MorphName, GenerationContext);
+	if (MorphedSourceMesh)
+	{
+		mu::NodeMeshConstantPtr MorphedSourceMeshNode = new mu::NodeMeshConstant;
+		MorphedSourceMeshNode->SetValue(MorphedSourceMesh);
+		MorphedSourceMeshNode->SetMessageContext(Node);
+
+		mu::NodeMeshMakeMorphPtr IdentityMorph = new mu::NodeMeshMakeMorph;
+		IdentityMorph->SetBase(BaseSourceMesh.get());
+		IdentityMorph->SetTarget(BaseSourceMesh.get());
+		IdentityMorph->SetMessageContext(Node);
+
+		MeshNode->SetMorph(0, IdentityMorph);
+
+		mu::NodeMeshMakeMorphPtr Morph = new mu::NodeMeshMakeMorph;
+		Morph->SetBase(BaseSourceMesh.get());
+		Morph->SetTarget(MorphedSourceMeshNode.get());
+		Morph->SetMessageContext(Node);
+
+		MeshNode->SetMorph(1, Morph);
+	}
+	else
+	{
+		GenerationContext.Compiler->CompilerLog(LOCTEXT("MorphGenerationFailed", "Failed to generate morph target."), Node);
+	}
+}
+
+/** Convert a CustomizableObject Source Graph into a mutable source graph  */
+mu::NodeMeshPtr GenerateMutableSourceMesh(const UEdGraphPin * Pin,
+	FMutableGraphGenerationContext & GenerationContext,
+	FMutableGraphMeshGenerationData & MeshData)
+{
+	check(Pin)
+	RETURN_ON_CYCLE(*Pin, GenerationContext)
+	SCOPED_PIN_DATA(GenerationContext, Pin)
+
+	CheckNumOutputs(*Pin, GenerationContext);
+
+	UCustomizableObjectNode* Node = CastChecked<UCustomizableObjectNode>(Pin->GetOwningNode());
+
+	const FGeneratedKey Key(reinterpret_cast<void*>(&GenerateMutableSourceMesh), *Pin , *Node, GenerationContext, true);
+	if (const FGeneratedData* Generated = GenerationContext.Generated.Find(Key))
+	{
+		MeshData = Generated->meshData;
+		return static_cast<mu::NodeMesh*>(Generated->Node.get());
+	}
+	
+	if (Node->IsNodeOutDatedAndNeedsRefresh())
+	{
+		Node->SetRefreshNodeWarning();
+	}
+
+	//SkeletalMesh Result
+	mu::NodeMeshPtr Result;
+
+	//SkeletalMesh + Morphs Result
+	mu::NodeMeshPtr MorphResult;
+	
+	if (const UCustomizableObjectNodeSkeletalMesh* TypedNodeSkel = Cast<UCustomizableObjectNodeSkeletalMesh>(Node))
+	{
+		mu::NodeMeshConstantPtr MeshNode = new mu::NodeMeshConstant();
+		Result = MeshNode;
+
+		if (TypedNodeSkel->SkeletalMesh)
+		{
+			int32 LOD;
+			int32 SectionIndex;
+			int32 LayoutIndex;
+			TypedNodeSkel->GetPinSection(*Pin, LOD, SectionIndex, LayoutIndex);
+
+			check(SectionIndex < TypedNodeSkel->LODs[LOD].Materials.Num());
+
+			// See if we need to select another LOD because of automatic LOD generation
+			if (GenerationContext.CurrentAutoLODStrategy == ECustomizableObjectAutomaticLODStrategy::AutomaticFromMesh)
+			{
+				if (TypedNodeSkel->SkeletalMesh && Helper_GetImportedModel(TypedNodeSkel->SkeletalMesh))
+				{
+					// Add the CurrentLOD being generated to the first LOD used from this SkeletalMesh
+					int32 CurrentLOD = LOD + GenerationContext.CurrentLOD; 
+					
+					// If the mesh has additional LODs and they use the same material, use them.
+					FSkeletalMeshModel* ImportedModel = Helper_GetImportedModel(TypedNodeSkel->SkeletalMesh);
+					if (ImportedModel->LODModels.Num() > CurrentLOD
+						&&
+						ImportedModel->LODModels[CurrentLOD].Sections.Num() > SectionIndex)
+					{
+						LOD = CurrentLOD;
+					}
+					else
+					{
+						// Find the closest valid LOD 
+						for (int32 LODIndex = ImportedModel->LODModels.Num() - 1; LODIndex >= 0; --LODIndex)
+						{
+							if (ImportedModel->LODModels[LODIndex].Sections.Num() > SectionIndex)
+							{
+								LOD = LODIndex;
+								break;
+							}
+						}
+					}
+				}
+			}
+
+			mu::MeshPtr MutableMesh = GenerateMutableMesh(TypedNodeSkel->SkeletalMesh, LOD, SectionIndex, GenerationContext, TypedNodeSkel);
+			if (MutableMesh)
+			{
+				MeshNode->SetValue(MutableMesh);
+
+				if (TypedNodeSkel->SkeletalMesh->GetPhysicsAsset() && !GenerationContext.DiscartedPhysicsAssetMap.Find(TypedNodeSkel->SkeletalMesh->GetPhysicsAsset()))
+				{
+					bool bBoneMismatch = false;
+
+					// should be driven by build settings in some way, probably, or a configuration
+					const bool bOptedOutOfIncompleteBodyWarnings = false;
+
+					for (TObjectPtr<USkeletalBodySetup> BodySetup : TypedNodeSkel->SkeletalMesh->GetPhysicsAsset()->SkeletalBodySetups)
+					{
+						if (BodySetup->BoneName != NAME_None && TypedNodeSkel->SkeletalMesh->GetRefSkeleton().FindBoneIndex(BodySetup->BoneName) == INDEX_NONE)
+						{
+							if (!bOptedOutOfIncompleteBodyWarnings)
+							{
+								GenerationContext.DiscartedPhysicsAssetMap.Add(TypedNodeSkel->SkeletalMesh->GetPhysicsAsset(), GenerationContext.DiscartedPhysicsAssetMap.Num());
+								FString Msg = FString::Printf(TEXT("The PhysicsAsset %s has Bodies with no corresponding bones in the SkeletalMesh %s. It will not be used."), *TypedNodeSkel->SkeletalMesh->GetPhysicsAsset()->GetName(), *TypedNodeSkel->SkeletalMesh->GetName());
+								UE_LOG(LogMutable, Warning, TEXT("%s"), *Msg);
+							}
+
+							bBoneMismatch = true;
+							break;
+						}
+					}
+
+					if (!bBoneMismatch)
+					{
+						TSoftObjectPtr<UPhysicsAsset> PhysicsAsset = TypedNodeSkel->SkeletalMesh->GetPhysicsAsset();
+						GenerationContext.PhysicsAssetMap.Add(PhysicsAsset.ToString(), PhysicsAsset);
+						FString PhysicsAssetTag = FString("__PhysicsAsset:") + PhysicsAsset.ToString();
+						MutableMesh->SetTagCount(1);
+						MutableMesh->SetTag(0, TCHAR_TO_ANSI(*PhysicsAssetTag));
+					}
+				}
+
+				if (GenerationContext.Options.bClothingEnabled)
+				{
+					UClothingAssetBase* ClothingAssetBase = TypedNodeSkel->SkeletalMesh->GetSectionClothingAsset(LOD, SectionIndex);	
+					UClothingAssetCommon* ClothingAssetCommon = Cast<UClothingAssetCommon>(ClothingAssetBase);
+
+					if (ClothingAssetCommon && ClothingAssetCommon->PhysicsAsset)
+					{	
+						int32 AssetIndex = GenerationContext.ContributingClothingAssetsData.IndexOfByPredicate( 
+						[Guid = ClothingAssetBase->GetAssetGuid()](const FCustomizableObjectClothingAssetData& A)
+						{
+							return A.OriginalAssetGuid == Guid;
+						});
+
+						check(AssetIndex != INDEX_NONE);
+
+						TSoftObjectPtr<UPhysicsAsset> PhysicsAsset = ClothingAssetCommon->PhysicsAsset.Get();
+
+						FString ClothPhysicsAssetTag = FString("__ClothPhysicsAsset:") + FString::Printf(TEXT("%d_AssetIdx_"), AssetIndex) + PhysicsAsset.ToString();
+						
+						GenerationContext.PhysicsAssetMap.Add(PhysicsAsset.ToString(), ClothingAssetCommon->PhysicsAsset.Get());
+
+						int32 CurrentTagCount = MutableMesh->GetTagCount();
+						MutableMesh->SetTagCount(CurrentTagCount + 1);
+						MutableMesh->SetTag(CurrentTagCount, TCHAR_TO_ANSI(*ClothPhysicsAssetTag));
+					}
+				}
+
+				if (!TypedNodeSkel->AnimInstance.IsNull())
+				{
+					int32 SlotIndex= TypedNodeSkel->AnimBlueprintSlot;
+					GenerationContext.AnimBPAssetsMap.Add(TypedNodeSkel->AnimInstance.ToString(), TypedNodeSkel->AnimInstance);
+
+					FString AnimBPAssetTag = GenerateAnimationInstanceTag(TypedNodeSkel->AnimInstance.ToString(), SlotIndex);
+					int32 CurrentTagCount = MutableMesh->GetTagCount();
+					MutableMesh->SetTagCount(CurrentTagCount + 1);
+					MutableMesh->SetTag(CurrentTagCount, TCHAR_TO_ANSI(*AnimBPAssetTag));
+				}
+
+				for (const FGameplayTag& GamePlayTag : TypedNodeSkel->AnimationGameplayTags)
+				{
+					const FString AnimBPTag = GenerateGameplayTag(GamePlayTag.ToString());
+					const int32 CurrentTagCount = MutableMesh->GetTagCount();
+					MutableMesh->SetTagCount(CurrentTagCount + 1);
+					MutableMesh->SetTag(CurrentTagCount, TCHAR_TO_ANSI(*AnimBPTag));
+				}
+
+				MeshData.bHasVertexColors = TypedNodeSkel->SkeletalMesh->GetHasVertexColors();
+				FSkeletalMeshModel* importModel = Helper_GetImportedModel(TypedNodeSkel->SkeletalMesh);
+				MeshData.NumTexCoordChannels = importModel->LODModels[LOD].NumTexCoords;
+				MeshData.MaxBoneIndexTypeSizeBytes = importModel->LODModels[LOD].RequiredBones.Num() > 256 ? 2 : 1;
+				MeshData.MaxNumBonesPerVertex = Helper_GetMaxBoneInfluences(importModel->LODModels[LOD]);
+				
+				// When mesh data is combined we will get an upper and lower bound of the number of triangles.
+				MeshData.MaxNumTriangles = importModel->LODModels[LOD].Sections[SectionIndex].NumTriangles;
+				MeshData.MinNumTriangles = importModel->LODModels[LOD].Sections[SectionIndex].NumTriangles;
+			}
+
+			// Layouts
+			// If we didn't find a layout, but we are generating LODs and this LOD is automatic, reuse the first valid LOD layout
+			bool LayoutFound = false;
+
+			int LayoutLOD = LOD;
+			while (!LayoutFound
+				&&
+				LayoutLOD >= 0
+				)
+			{
+				if (TypedNodeSkel->LODs.IsValidIndex(LayoutLOD)
+					&&
+					TypedNodeSkel->LODs[LayoutLOD].Materials.IsValidIndex(SectionIndex))
+				{
+					MeshNode->SetLayoutCount(TypedNodeSkel->LODs[LayoutLOD].Materials[SectionIndex].LayoutPinsRef.Num());
+					for (LayoutIndex = 0; LayoutIndex < TypedNodeSkel->LODs[LayoutLOD].Materials[SectionIndex].LayoutPinsRef.Num(); ++LayoutIndex)
+					{
+						if (TypedNodeSkel->LODs[LayoutLOD].Materials[SectionIndex].LayoutPinsRef[LayoutIndex].Get())
+						{
+							if (const UEdGraphPin* ConnectedPin = FollowInputPin(*TypedNodeSkel->LODs[LayoutLOD].Materials[SectionIndex].LayoutPinsRef[LayoutIndex].Get()))
+							{
+								mu::NodeLayoutPtr LayoutNode = GenerateMutableSourceLayout(ConnectedPin, GenerationContext);
+								MeshNode->SetLayout(LayoutIndex, LayoutNode);
+								LayoutFound = true;
+							}
+						}
+					}
+				}
+
+				if (GenerationContext.CurrentAutoLODStrategy == ECustomizableObjectAutomaticLODStrategy::AutomaticFromMesh)
+				{
+					--LayoutLOD;
+				}
+				else
+				{
+					break;
+				}
+			}
+
+			if (!LayoutFound)
+			{
+				MeshNode->SetLayoutCount(1);
+				mu::NodeLayoutBlocksPtr LayoutNode = new mu::NodeLayoutBlocks();
+				int GridSize = 4;
+				LayoutNode->SetGridSize(GridSize, GridSize);
+				LayoutNode->SetMaxGridSize(GridSize, GridSize);
+				LayoutNode->SetLayoutPackingStrategy(mu::EPackStrategy::RESIZABLE_LAYOUT);
+				LayoutNode->SetBlockCount(1);
+				LayoutNode->SetBlock(0, 0, 0, GridSize, GridSize);
+				LayoutNode->SetBlockPriority(0, 0);
+				MeshNode->SetLayout(0, LayoutNode);
+
+				// We need it here because we create multiple nodes.
+				LayoutNode->SetMessageContext(Node);
+			}
+
+			// Applying Mesh Morph Nodes
+			if (GenerationContext.MeshMorphStack.Num())
+			{
+				MorphResult = GenerateMorphMesh(Pin, GenerationContext.MeshMorphStack, 0, Result, GenerationContext, MeshData);
+			}
+		}
+		else
+		{
+			GenerationContext.Compiler->CompilerLog(LOCTEXT("MissingskeletlMesh", "No Skeletal Mesh set in the SkeletalMesh node."), Node);
+		}
+	}
+
+	else if (const UCustomizableObjectNodeStaticMesh* TypedNodeStatic = Cast<UCustomizableObjectNodeStaticMesh>(Node))
+	{
+		if (TypedNodeStatic->StaticMesh == nullptr)
+		{
+			FString Msg = FString::Printf(TEXT("The UCustomizableObjectNodeStaticMesh node %s has no static mesh assigned"), *Node->GetName());
+			GenerationContext.Compiler->CompilerLog(FText::FromString(Msg), Node, EMessageSeverity::Warning);
+			return {};
+		}
+
+		if (TypedNodeStatic->StaticMesh->GetNumLODs() == 0)
+		{
+			FString Msg = FString::Printf(TEXT("The UCustomizableObjectNodeStaticMesh node %s has a static mesh assigned with no RenderData"), *Node->GetName());
+			GenerationContext.Compiler->CompilerLog(FText::FromString(Msg), Node, EMessageSeverity::Warning);
+			return {};
+		}
+
+		mu::NodeMeshConstantPtr MeshNode = new mu::NodeMeshConstant();
+		Result = MeshNode;
+
+		if (TypedNodeStatic->StaticMesh)
+		{
+			// TODO
+			int LOD = 0;
+
+			// Find out what material do we need
+			int MaterialIndex = 0;
+			for (; MaterialIndex < TypedNodeStatic->LODs[LOD].Materials.Num(); ++MaterialIndex)
+			{
+				if (TypedNodeStatic->LODs[LOD].Materials[MaterialIndex].MeshPinRef.Get() == Pin)
+				{
+					break;
+				}
+			}
+			check(MaterialIndex < TypedNodeStatic->LODs[LOD].Materials.Num());
+
+			mu::MeshPtr MutableMesh = GenerateMutableMesh(TypedNodeStatic->StaticMesh, LOD, MaterialIndex, GenerationContext, TypedNodeStatic);
+			if (MutableMesh)
+			{
+				MeshNode->SetValue(MutableMesh);
+
+				// Layouts
+				MeshNode->SetLayoutCount(1);
+
+				if (const UEdGraphPin* ConnectedPin = FollowInputPin(*TypedNodeStatic->LODs[LOD].Materials[MaterialIndex].LayoutPinRef.Get()))
+				{
+					mu::NodeLayoutPtr LayoutNode = GenerateMutableSourceLayout(ConnectedPin, GenerationContext);
+					MeshNode->SetLayout(0, LayoutNode);
+				}
+				else
+				{
+					mu::NodeLayoutBlocksPtr LayoutNode = new mu::NodeLayoutBlocks();
+					int GridSize = 4;
+					LayoutNode->SetGridSize(GridSize, GridSize);
+					LayoutNode->SetMaxGridSize(GridSize, GridSize);
+					LayoutNode->SetLayoutPackingStrategy(mu::EPackStrategy::RESIZABLE_LAYOUT);
+					LayoutNode->SetBlockCount(1);
+					LayoutNode->SetBlock(0, 0, 0, GridSize, GridSize);
+					LayoutNode->SetBlockPriority(0, 0);
+					MeshNode->SetLayout(0, LayoutNode);
+
+					// We need it here because we create multiple nodes.
+					LayoutNode->SetMessageContext(Node);
+				}
+			}
+			else
+			{
+				Result = nullptr;
+			}
+		}
+	}
+
+	else if (UCustomizableObjectNodeMeshMorph* TypedNodeMorph = Cast<UCustomizableObjectNodeMeshMorph>(Node))
+	{
+		if (const UEdGraphPin* ConnectedPin = FollowInputPin(*TypedNodeMorph->MeshPin()))
+		{
+			// Mesh Morph Stack Management
+			GenerationContext.MeshMorphStack.Push(TypedNodeMorph);
+			Result = GenerateMutableSourceMesh(ConnectedPin, GenerationContext, MeshData);
+			GenerationContext.MeshMorphStack.Pop(true);
+		}
+		else
+		{
+			mu::NodeMeshMorphPtr MeshNode = new mu::NodeMeshMorph();
+			Result = MeshNode;
+		}
+	}
+
+	else if (const UCustomizableObjectNodeMeshMorphStackApplication* TypedNodeMeshMorphStackApp = Cast< UCustomizableObjectNodeMeshMorphStackApplication >(Node))
+	{
+		if (const UEdGraphPin* ConnectedPin = FollowInputPin(*TypedNodeMeshMorphStackApp->GetStackPin()))
+		{
+			UEdGraphNode* OwningNode = ConnectedPin->GetOwningNode();
+			if (UCustomizableObjectNodeMeshMorphStackDefinition* TypedNodeMeshMorphStackDef = Cast<UCustomizableObjectNodeMeshMorphStackDefinition>(OwningNode))
+			{
+				// Checking if is out of data
+				if (TypedNodeMeshMorphStackDef->IsNodeOutDatedAndNeedsRefresh())
+				{
+					TypedNodeMeshMorphStackDef->SetRefreshNodeWarning();
+				}
+
+				mu::NodeMeshMorphPtr MeshNode = new mu::NodeMeshMorph();
+				Result = MeshNode;
+
+				MeshNode->SetMorphCount(2);
+
+				int32 NextPin = -1;
+				bool bAddNewMorph = true;
+
+				while (bAddNewMorph)
+				{
+					//Getting next connected pin index
+					NextPin = TypedNodeMeshMorphStackDef->NextConnectedPin(NextPin, TypedNodeMeshMorphStackApp->MorphNames);
+
+					UEdGraphPin* FactorPin = nullptr;
+
+					// If Next Pin is -1 then there is no pin connected, which is a warning
+					if (NextPin != -1)
+					{
+						FactorPin = TypedNodeMeshMorphStackDef->GetMorphPin(NextPin);
+					}
+
+					// Checking if there is another pin connected to add a new morph
+					bAddNewMorph = (TypedNodeMeshMorphStackDef->NextConnectedPin(NextPin, TypedNodeMeshMorphStackApp->MorphNames) != -1);
+
+					// Factor
+					GenerateMorphFactor(Node, *FactorPin, GenerationContext, MeshNode);
+
+					// Base
+					if (const UEdGraphPin* ConnectedMehsPin = FollowInputPin(*TypedNodeMeshMorphStackApp->GetMeshPin()))
+					{
+						mu::NodeMeshPtr SourceNode = nullptr;
+
+						if (bAddNewMorph)
+						{
+							mu::NodeMeshMorphPtr NextMorph = new mu::NodeMeshMorph();
+
+							SourceNode = NextMorph;
+							MeshNode->SetBase(SourceNode);
+
+							if (const UEdGraphPin* BaseSourcePin = FindMeshBaseSource(*ConnectedMehsPin, false))
+							{
+								FString MorphName = "";
+								if (NextPin != -1)
+								{
+									MorphName = TypedNodeMeshMorphStackDef->MorphNames[NextPin];
+								}
+
+								//Target
+								GenerateMorphTarget(Node, BaseSourcePin, GenerationContext, MeshNode, MorphName);
+							}
+
+							MeshNode = NextMorph;
+							MeshNode->SetMorphCount(2);
+						}
+						else
+						{
+							SourceNode = GenerateMutableSourceMesh(ConnectedMehsPin, GenerationContext, MeshData);
+							MeshNode->SetBase(SourceNode);
+
+							if (const UEdGraphPin* BaseSourcePin = FindMeshBaseSource(*ConnectedMehsPin, false))
+							{
+								FString MorphName = "";
+								if (NextPin != -1)
+								{
+									MorphName = TypedNodeMeshMorphStackDef->MorphNames[NextPin];
+								}
+
+								//Target
+								GenerateMorphTarget(Node, BaseSourcePin, GenerationContext, MeshNode, MorphName);
+							}
+						}
+					}
+				}
+			}
+			else
+			{
+				GenerationContext.Compiler->CompilerLog(LOCTEXT("MorphStackGenerationFailed", "Stack definition Generation failed."), Node);
+				Result = nullptr;
+			}
+		}
+		else
+		{
+			GenerationContext.Compiler->CompilerLog(LOCTEXT("MorphStackConnectionFailed", "Stack definition connection not found."), Node);
+			Result = nullptr;
+		}
+	}
+
+	else if (const UCustomizableObjectNodeMeshSwitch* TypedNodeMeshSwitch = Cast<UCustomizableObjectNodeMeshSwitch>(Node))
+	{
+		// Using a lambda so control flow is easier to manage.
+		Result = [&]()
+		{
+			const UEdGraphPin* SwitchParameter = TypedNodeMeshSwitch->SwitchParameter();
+
+			// Check Switch Parameter arity preconditions.
+			if (const UEdGraphPin* EnumPin = FollowInputPin(*SwitchParameter))
+			{
+				mu::NodeScalarPtr SwitchParam = GenerateMutableSourceFloat(EnumPin, GenerationContext);
+
+				// Switch Param not generated
+				if (!SwitchParam)
+				{
+					// Warn about a failure.
+					if (EnumPin)
+					{
+						const FText Message = LOCTEXT("FailedToGenerateSwitchParam", "Could not generate switch enum parameter. Please refesh the switch node and connect an enum.");
+						GenerationContext.Compiler->CompilerLog(Message, Node);
+					}
+
+					return Result;
+				}
+
+				if (SwitchParam->GetType() != mu::NodeScalarEnumParameter::GetStaticType())
+				{
+					const FText Message = LOCTEXT("WrongSwitchParamType", "Switch parameter of incorrect type.");
+					GenerationContext.Compiler->CompilerLog(Message, Node);
+
+					return Result;
+				}
+
+				const int32 NumSwitchOptions = TypedNodeMeshSwitch->GetNumElements();
+
+				mu::NodeScalarEnumParameter* EnumParameter = static_cast<mu::NodeScalarEnumParameter*>(SwitchParam.get());
+				if (NumSwitchOptions != EnumParameter->GetValueCount())
+				{
+					const FText Message = LOCTEXT("MismatchedSwitch", "Switch enum and switch node have different number of options. Please refresh the switch node to make sure the outcomes are labeled properly.");
+					GenerationContext.Compiler->CompilerLog(Message, Node);
+				}
+
+				mu::NodeMeshSwitchPtr SwitchNode = new mu::NodeMeshSwitch;
+				SwitchNode->SetParameter(SwitchParam);
+				SwitchNode->SetOptionCount(NumSwitchOptions);
+
+				for (int SelectorIndex = 0; SelectorIndex < NumSwitchOptions; ++SelectorIndex)
+				{
+					if (const UEdGraphPin* ConnectedPin = FollowInputPin(*TypedNodeMeshSwitch->GetElementPin(SelectorIndex)))
+					{
+						FMutableGraphMeshGenerationData ChildMeshData;
+						Result = GenerateMutableSourceMesh(ConnectedPin, GenerationContext, ChildMeshData);
+						SwitchNode->SetOption(SelectorIndex, Result);
+						MeshData.Combine(ChildMeshData);
+					}
+				}
+
+				Result = SwitchNode;
+				return Result;
+			}
+			else
+			{
+				GenerationContext.Compiler->CompilerLog(LOCTEXT("NoEnumParamInSwitch", "Switch nodes must have an enum switch parameter. Please connect an enum and refesh the switch node."), Node);
+				return Result;
+			}
+		}(); // invoke lambda.
+	}
+
+	else if (const UCustomizableObjectNodeMeshVariation* TypedNodeMeshVar = Cast<const UCustomizableObjectNodeMeshVariation>(Node))
+	{
+		mu::NodeMeshVariationPtr MeshNode = new mu::NodeMeshVariation();
+		Result = MeshNode;
+
+		if (const UEdGraphPin* ConnectedPin = FollowInputPin(*TypedNodeMeshVar->DefaultPin()))
+		{
+			FMutableGraphMeshGenerationData ChildMeshData;
+			mu::NodeMeshPtr ChildNode = GenerateMutableSourceMesh(ConnectedPin, GenerationContext, ChildMeshData);
+			if (ChildNode)
+			{
+				MeshNode->SetDefaultMesh(ChildNode.get());
+				MeshData.Combine(ChildMeshData);
+			}
+			else
+			{
+				GenerationContext.Compiler->CompilerLog(LOCTEXT("MeshFailed", "Mesh generation failed."), Node);
+			}
+		}
+		else
+		{
+			GenerationContext.Compiler->CompilerLog(LOCTEXT("MeshVarMissingDef", "Mesh variation node requires a default value."), Node);
+		}
+
+		MeshNode->SetVariationCount(TypedNodeMeshVar->Variations.Num());
+		for (int VariationIndex = 0; VariationIndex < TypedNodeMeshVar->Variations.Num(); ++VariationIndex)
+		{
+			const UEdGraphPin* VariationPin = TypedNodeMeshVar->VariationPin(VariationIndex);
+			if (!VariationPin) continue;
+
+			MeshNode->SetVariationTag(VariationIndex, TCHAR_TO_ANSI(*TypedNodeMeshVar->Variations[VariationIndex].Tag));
+			if (const UEdGraphPin* ConnectedPin = FollowInputPin(*VariationPin))
+			{
+				FMutableGraphMeshGenerationData DummyMeshData;
+				mu::NodeMeshPtr ChildNode = GenerateMutableSourceMesh(ConnectedPin, GenerationContext, DummyMeshData);
+				MeshNode->SetVariationMesh(VariationIndex, ChildNode.get());
+				MeshData.Combine(DummyMeshData);
+			}
+		}
+	}
+
+	else if (const UCustomizableObjectNodeMeshGeometryOperation* TypedNodeGeometry = Cast<const UCustomizableObjectNodeMeshGeometryOperation>(Node))
+	{
+		mu::Ptr<mu::NodeMeshGeometryOperation> MeshNode = new mu::NodeMeshGeometryOperation();
+		Result = MeshNode;
+
+		if (const UEdGraphPin* ConnectedPin = FollowInputPin(*TypedNodeGeometry->MeshAPin()))
+		{
+			FMutableGraphMeshGenerationData ChildMeshData;
+			mu::Ptr<mu::NodeMesh> ChildNode = GenerateMutableSourceMesh(ConnectedPin, GenerationContext, ChildMeshData);
+			if (ChildNode)
+			{
+				MeshNode->SetMeshA(ChildNode.get());
+				MeshData.Combine(ChildMeshData);
+			}
+			else
+			{
+				GenerationContext.Compiler->CompilerLog(LOCTEXT("MeshFailed", "Mesh generation failed."), Node);
+			}
+		}
+		else
+		{
+			GenerationContext.Compiler->CompilerLog(LOCTEXT("MeshGeometryMissingDef", "Mesh variation node requires a default value."), Node);
+		}
+
+		if (const UEdGraphPin* ConnectedPin = FollowInputPin(*TypedNodeGeometry->MeshBPin()))
+		{
+			FMutableGraphMeshGenerationData ChildMeshData;
+			mu::Ptr<mu::NodeMesh> ChildNode = GenerateMutableSourceMesh(ConnectedPin, GenerationContext, ChildMeshData);
+			if (ChildNode)
+			{
+				MeshNode->SetMeshB(ChildNode.get());
+				MeshData.Combine(ChildMeshData);
+			}
+			else
+			{
+				GenerationContext.Compiler->CompilerLog(LOCTEXT("MeshFailed", "Mesh generation failed."), Node);
+			}
+		}
+
+		if (const UEdGraphPin* ConnectedPin = FollowInputPin(*TypedNodeGeometry->ScalarAPin()))
+		{
+			mu::NodeScalarPtr ChildNode = GenerateMutableSourceFloat(ConnectedPin, GenerationContext);
+			if (ChildNode)
+			{
+				MeshNode->SetScalarA(ChildNode.get());
+			}
+			else
+			{
+				GenerationContext.Compiler->CompilerLog(LOCTEXT("MeshFailed", "Scalar generation failed."), Node);
+			}
+		}
+
+		if (const UEdGraphPin* ConnectedPin = FollowInputPin(*TypedNodeGeometry->ScalarBPin()))
+		{
+			mu::NodeScalarPtr ChildNode = GenerateMutableSourceFloat(ConnectedPin, GenerationContext);
+			if (ChildNode)
+			{
+				MeshNode->SetScalarB(ChildNode.get());
+			}
+			else
+			{
+				GenerationContext.Compiler->CompilerLog(LOCTEXT("MeshFailed", "Scalar generation failed."), Node);
+			}
+		}
+	}
+
+	else if (const UCustomizableObjectNodeMeshReshape* TypedNodeReshape = Cast<const UCustomizableObjectNodeMeshReshape>(Node))
+	{
+		mu::Ptr<mu::NodeMeshReshape> MeshNode = new mu::NodeMeshReshape();
+		Result = MeshNode;
+
+		if (const UEdGraphPin* ConnectedPin = FollowInputPin(*TypedNodeReshape->BaseMeshPin()))
+		{
+			FMutableGraphMeshGenerationData ChildMeshData;
+			mu::Ptr<mu::NodeMesh> ChildNode = GenerateMutableSourceMesh(ConnectedPin, GenerationContext, ChildMeshData);
+			if (ChildNode)
+			{
+				MeshNode->SetBaseMesh(ChildNode.get());
+				MeshData.Combine(ChildMeshData);
+			}
+			else
+			{
+				GenerationContext.Compiler->CompilerLog(LOCTEXT("MeshFailed", "Mesh generation failed."), Node);
+			}
+		}
+		else
+		{
+			GenerationContext.Compiler->CompilerLog(LOCTEXT("MeshGeometryMissingDef", "Mesh variation node requires a default value."), Node);
+		}
+	
+		{
+			MeshNode->SetReshapeSkeleton(TypedNodeReshape->bReshapeSkeleton);
+			MeshNode->SetReshapePhysicsVolumes(TypedNodeReshape->bReshapePhysicsVolumes);
+			MeshNode->SetEnableRigidParts(TypedNodeReshape->bEnableRigidParts);
+			MeshNode->SetDeformAllBones(TypedNodeReshape->bDeformAllBones);
+			MeshNode->SetDeformAllPhysics(TypedNodeReshape->bDeformAllPhysicsBodies);
+				
+			const UEdGraphPin* ConnectedPin = FollowInputPin(*TypedNodeReshape->BaseMeshPin());
+			const UEdGraphPin* SourceMeshPin = ConnectedPin ? FindMeshBaseSource(*ConnectedPin, false) : nullptr;
+			const UEdGraphNode* SkeletalMeshNode = SourceMeshPin ? SourceMeshPin->GetOwningNode() : nullptr;
+
+			TArray<USkeletalMesh*> SkeletalMeshesToDeform = GetSkeletalMeshesForReshapeSelection(SkeletalMeshNode, SourceMeshPin);
+
+			bool bWarningFound = false;
+			if (TypedNodeReshape->bReshapeSkeleton && !TypedNodeReshape->bDeformAllBones)
+			{
+				TArray<FString> BonesToDeform;
+				bWarningFound = GetAndValidateReshapeBonesToDeform(
+					BonesToDeform, TypedNodeReshape->BonesToDeform, SkeletalMeshesToDeform, TypedNodeReshape, GenerationContext);
+				
+				for (const FString& BoneName : BonesToDeform)
+				{
+					MeshNode->AddBoneToDeform(TCHAR_TO_ANSI(*BoneName));
+				}
+			}
+
+			if (TypedNodeReshape->bReshapePhysicsVolumes && !TypedNodeReshape->bDeformAllPhysicsBodies)
+			{
+				TArray<FString> PhysicsToDeform;
+				bWarningFound = bWarningFound || GetAndValidateReshapePhysicsToDeform(
+					PhysicsToDeform, TypedNodeReshape->PhysicsBodiesToDeform, SkeletalMeshesToDeform, TypedNodeReshape, GenerationContext);
+
+				for (const FString& PhysicsBoneName : PhysicsToDeform)
+				{
+					MeshNode->AddPhysicsBodyToDeform(TCHAR_TO_ANSI(*PhysicsBoneName));
+				}	
+			}
+			
+			if (bWarningFound)
+			{
+				Node->SetRefreshNodeWarning();
+			}		
+		}
+		// We don't need all the data for the shape meshes
+		const uint32 ShapeFlags = uint32(EMutableMeshConversionFlags::IgnoreSkinning) & 
+								  uint32(EMutableMeshConversionFlags::IgnorePhysics);
+
+		GenerationContext.MeshGenerationFlags.Push( ShapeFlags );
+			
+		constexpr int32 PinNotSetValue = TNumericLimits<int32>::Max();
+		int32 BaseShapeTriangleCount = PinNotSetValue;
+		if (const UEdGraphPin* ConnectedPin = FollowInputPin(*TypedNodeReshape->BaseShapePin()))
+		{
+			FMutableGraphMeshGenerationData ChildMeshData;
+			mu::Ptr<mu::NodeMesh> ChildNode = GenerateMutableSourceMesh(ConnectedPin, GenerationContext, ChildMeshData);
+	
+			if (ChildNode)
+			{
+				BaseShapeTriangleCount = ChildMeshData.MaxNumTriangles == ChildMeshData.MinNumTriangles ? ChildMeshData.MaxNumTriangles : -1;
+				MeshNode->SetBaseShape(ChildNode.get());	
+			}
+			else
+			{
+				GenerationContext.Compiler->CompilerLog(LOCTEXT("MeshFailed", "Mesh generation failed."), Node);
+			}
+		}
+
+		int32 TargetShapeTriangleCount = PinNotSetValue;
+		if (const UEdGraphPin* ConnectedPin = FollowInputPin(*TypedNodeReshape->TargetShapePin()))
+		{
+			FMutableGraphMeshGenerationData ChildMeshData;
+			mu::Ptr<mu::NodeMesh> ChildNode = GenerateMutableSourceMesh(ConnectedPin, GenerationContext, ChildMeshData);
+			
+			if (ChildNode)
+			{
+				TargetShapeTriangleCount = ChildMeshData.MaxNumTriangles == ChildMeshData.MinNumTriangles ? ChildMeshData.MaxNumTriangles : -1;
+				MeshNode->SetTargetShape(ChildNode.get());
+			}
+			else
+			{
+				GenerationContext.Compiler->CompilerLog(LOCTEXT("MeshFailed", "Mesh generation failed."), Node);
+			}
+		}
+
+		// If any of the shape pins is not set, don't warn about it.
+		if ( BaseShapeTriangleCount != PinNotSetValue && TargetShapeTriangleCount != PinNotSetValue )
+		{ 
+			if ( BaseShapeTriangleCount != TargetShapeTriangleCount || BaseShapeTriangleCount == -1 || TargetShapeTriangleCount == -1)
+			{
+				GenerationContext.Compiler->CompilerLog(LOCTEXT("ReshapeMeshShapeIncompatible", 
+					"Base and Target Shapes might not be compatible. Don't have the same number of triangles."), Node, EMessageSeverity::Warning);
+			}
+		}
+
+		GenerationContext.MeshGenerationFlags.Pop();
+	}
+
+	else if (const UCustomizableObjectNodeAnimationPose* TypedNode = Cast<UCustomizableObjectNodeAnimationPose>(Node))
+	{
+		if (const UEdGraphPin* ConnectedPin = FollowInputPin(*TypedNode->GetInputMeshPin()))
+		{
+			mu::Ptr<mu::NodeMesh> InputMeshNode = GenerateMutableSourceMesh(ConnectedPin, GenerationContext, MeshData);
+
+			TArray<FString> ArrayBoneName;
+			TArray<FTransform> ArrayTransform;
+			UCustomizableObjectNodeAnimationPose::StaticRetrievePoseInformation(TypedNode->PoseAsset, GenerationContext.GetCurrentComponentInfo().RefSkeletalMesh, ArrayBoneName, ArrayTransform);
+			mu::NodeMeshApplyPosePtr NodeMeshApplyPose = CreateNodeMeshApplyPose(InputMeshNode, GenerationContext.Object, ArrayBoneName, ArrayTransform);
+
+			if (NodeMeshApplyPose)
+			{
+				Result = NodeMeshApplyPose;
+			}
+			else
+			{
+				FString msg = FString::Printf(TEXT("Couldn't get bone transform information from a Pose Asset."));
+				GenerationContext.Compiler->CompilerLog(FText::FromString(msg), Node);
+
+				Result = nullptr;
+			}
+		}
+	}
+
+	else if (const UCustomizableObjectNodeTable* TypedNodeTable = Cast<UCustomizableObjectNodeTable>(Node))
+	{
+		mu::NodeMeshTablePtr MeshTableNode = new mu::NodeMeshTable();
+		Result = MeshTableNode;
+
+		mu::TablePtr Table = nullptr;
+
+		if (TypedNodeTable->Table)
+		{
+			Table = GenerateMutableSourceTable(TypedNodeTable->Table->GetName(), Pin, GenerationContext);
+	
+			USkeletalMesh* SkeletalMesh = TypedNodeTable->GetColumnDefaultAssetByType<USkeletalMesh>(Pin);
+			
+			int32 CurrentLOD = 0;
+			int32 MaterialIndex = 0;
+			
+			TypedNodeTable->GetPinLODAndMaterial(Pin, CurrentLOD, MaterialIndex);
+
+			if (GenerationContext.CurrentAutoLODStrategy == ECustomizableObjectAutomaticLODStrategy::AutomaticFromMesh)
+			{
+				CurrentLOD += GenerationContext.CurrentLOD;
+
+				if (SkeletalMesh)
+				{
+					FSkeletalMeshModel* ImportedModel = Helper_GetImportedModel(SkeletalMesh);
+
+					// Checking if the current LOD is valid
+					if (ImportedModel->LODModels.Num() <= CurrentLOD
+						||
+						ImportedModel->LODModels[CurrentLOD].Sections.Num() <= MaterialIndex)
+					{
+						// Find the closest valid LOD 
+						for (int32 LODIndex = ImportedModel->LODModels.Num() - 1; LODIndex >= 0; --LODIndex)
+						{
+							if (ImportedModel->LODModels[LODIndex].Sections.Num() > MaterialIndex)
+							{
+								CurrentLOD = LODIndex;
+								break;
+							}
+						}
+					}
+				}
+			}
+
+			FString ColumnName = TypedNodeTable->GetMutableColumnName(Pin, CurrentLOD);
+
+			MeshTableNode->SetTable(Table);
+			MeshTableNode->SetColumn(TCHAR_TO_ANSI(*ColumnName));
+			MeshTableNode->SetParameterName(TCHAR_TO_ANSI(*TypedNodeTable->ParameterName));
+
+			GenerationContext.AddParameterNameUnique(Node, TypedNodeTable->ParameterName);
+
+			if (SkeletalMesh)
+			{
+				// TODO: this should be made for all the meshes of the Column to support meshes with different values
+				// Filling Mesh Data
+				FSkeletalMeshModel* importModel = Helper_GetImportedModel(SkeletalMesh);
+				MeshData.bHasVertexColors = SkeletalMesh->GetHasVertexColors();
+				MeshData.NumTexCoordChannels = importModel->LODModels[CurrentLOD].NumTexCoords;
+				MeshData.MaxBoneIndexTypeSizeBytes = importModel->LODModels[CurrentLOD].RequiredBones.Num() > 256 ? 2 : 1;
+				MeshData.MaxNumBonesPerVertex = Helper_GetMaxBoneInfluences(importModel->LODModels[CurrentLOD]);
+
+				// When mesh data is combined we will get an upper and lower bound of the number of triangles.
+				MeshData.MaxNumTriangles = importModel->LODModels[CurrentLOD].Sections[MaterialIndex].NumTriangles;
+				MeshData.MinNumTriangles = importModel->LODModels[CurrentLOD].Sections[MaterialIndex].NumTriangles;
+			}
+
+			TArray<UCustomizableObjectLayout*> Layouts = TypedNodeTable->GetLayouts(Pin);
+			MeshTableNode->SetLayoutCount(Layouts.Num());
+
+			if (Layouts.Num())
+			{
+				// Generating node Layouts
+				for (int32 i = 0; i < Layouts.Num(); ++i)
+				{
+					mu::NodeLayoutBlocksPtr LayoutNode = new mu::NodeLayoutBlocks;
+
+					LayoutNode->SetGridSize(Layouts[i]->GetGridSize().X, Layouts[i]->GetGridSize().Y);
+					LayoutNode->SetMaxGridSize(Layouts[i]->GetMaxGridSize().X, Layouts[i]->GetMaxGridSize().Y);
+					LayoutNode->SetBlockCount(Layouts[i]->Blocks.Num());
+
+					mu::EPackStrategy strategy = mu::EPackStrategy::RESIZABLE_LAYOUT;
+
+					switch (Layouts[i]->GetPackingStrategy())
+					{
+					case ECustomizableObjectTextureLayoutPackingStrategy::Resizable:
+						strategy = mu::EPackStrategy::RESIZABLE_LAYOUT;
+						break;
+					case ECustomizableObjectTextureLayoutPackingStrategy::Fixed:
+						strategy = mu::EPackStrategy::FIXED_LAYOUT;
+						break;
+					default:
+						break;
+					}
+
+					LayoutNode->SetLayoutPackingStrategy(strategy);
+
+					if (Layouts[i]->Blocks.Num())
+					{
+						for (int BlockIndex = 0; BlockIndex < Layouts[i]->Blocks.Num(); ++BlockIndex)
+						{
+							LayoutNode->SetBlock(BlockIndex,
+								Layouts[i]->Blocks[BlockIndex].Min.X,
+								Layouts[i]->Blocks[BlockIndex].Min.Y,
+								Layouts[i]->Blocks[BlockIndex].Max.X - Layouts[i]->Blocks[BlockIndex].Min.X,
+								Layouts[i]->Blocks[BlockIndex].Max.Y - Layouts[i]->Blocks[BlockIndex].Min.Y);
+
+							LayoutNode->SetBlockPriority(BlockIndex, Layouts[i]->Blocks[BlockIndex].Priority);
+						}
+					}
+					else
+					{
+						LayoutNode->SetBlock(0, 0, 0, Layouts[i]->GetGridSize().X, Layouts[i]->GetGridSize().Y);
+					}
+
+					MeshTableNode->SetLayout(i, LayoutNode);
+				}
+			}
+
+			// Applying Mesh Morph Nodes
+			if (GenerationContext.MeshMorphStack.Num())
+			{
+				MorphResult = GenerateMorphMesh(Pin, GenerationContext.MeshMorphStack, 0, Result, GenerationContext, MeshData, ColumnName);
+			}
+
+			if (Table->FindColumn(TCHAR_TO_ANSI(*ColumnName)) == -1)
+			{
+				FString Msg = FString::Printf(TEXT("Couldn't find pin column with name %s"), *ColumnName);
+				GenerationContext.Compiler->CompilerLog(FText::FromString(Msg), Node);
+			}
+		}
+		else
+		{
+			Table = new mu::Table();
+			MeshTableNode->SetTable(Table);
+
+			GenerationContext.Compiler->CompilerLog(LOCTEXT("ImageTableError", "Couldn't find the data table of the node."), Node);
+		}
+	}
+	
+	else
+	{
+		GenerationContext.Compiler->CompilerLog(LOCTEXT("UnimplementedMeshNode", "Mesh node type not implemented yet."), Node);
+	}
+	
+	GenerationContext.Generated.Add(Key, FGeneratedData(Node, Result, &MeshData));
+	GenerationContext.GeneratedNodes.Add(Node);
+
+	// We return the mesh modified by morphs if there is any
+	if (MorphResult)
+	{
+		Result = MorphResult;
+	}
+
+	if (Result)
+	{
+		Result->SetMessageContext(Node);
+	}
+
+	return Result;
+}
