[mutable] Move Mutable plugin from the restricted folder to the experimental plugins folder.

diff --git a/Engine/Plugins/Experimental/Mutable/Source/CustomizableObjectEditor/Private/GenerateMutableSource/GenerateMutableSourceMesh.h b/Engine/Plugins/Experimental/Mutable/Source/CustomizableObjectEditor/Private/GenerateMutableSource/GenerateMutableSourceMesh.h
new file mode 100644
--- /dev/null
+++ b/Engine/Plugins/Experimental/Mutable/Source/CustomizableObjectEditor/Private/GenerateMutableSource/GenerateMutableSourceMesh.h
@@ -0,0 +1,20 @@
+// Copyright Epic Games, Inc. All Rights Reserved.
+
+#pragma once
+
+#include "GenerateMutableSource/GenerateMutableSource.h"
+
+
+mu::MeshPtr ConvertSkeletalMeshToMutable(USkeletalMesh* InSkeletalMesh, int LOD, int MaterialIndex, FMutableGraphGenerationContext& GenerationContext, const UCustomizableObjectNode* CurrentNode);
+
+
+mu::MeshPtr ConvertStaticMeshToMutable(UStaticMesh* StaticMesh, int LOD, int MaterialIndex, FMutableGraphGenerationContext& GenerationContext, const UCustomizableObjectNode* CurrentNode);
+
+
+mu::MeshPtr GenerateMutableMesh(UObject * Mesh, int32 LOD, int32 MaterialIndex, FMutableGraphGenerationContext & GenerationContext, const UCustomizableObjectNode* CurrentNode);
+
+
+mu::MeshPtr BuildMorphedMutableMesh(const UEdGraphPin * BaseSourcePin, const FString& MorphTargetName, FMutableGraphGenerationContext & GenerationContext, int32 RowIndex = 0);
+
+/** Convert a CustomizableObject Source Graph into a mutable source graph. */
+mu::NodeMeshPtr GenerateMutableSourceMesh(const UEdGraphPin* Pin, FMutableGraphGenerationContext& GenerationContext, FMutableGraphMeshGenerationData& MeshData);
diff --git a/Engine/Plugins/Experimental/Mutable/Source/CustomizableObjectEditor/Private/GenerateMutableSource/GenerateMutableSourceModifier.cpp b/Engine/Plugins/Experimental/Mutable/Source/CustomizableObjectEditor/Private/GenerateMutableSource/GenerateMutableSourceModifier.cpp
new file mode 100644
--- /dev/null
+++ b/Engine/Plugins/Experimental/Mutable/Source/CustomizableObjectEditor/Private/GenerateMutableSource/GenerateMutableSourceModifier.cpp
@@ -0,0 +1,207 @@
+// Copyright Epic Games, Inc. All Rights Reserved.
+
+#include "GenerateMutableSourceModifier.h"
+
+#include "GenerateMutableSource/GenerateMutableSourceFloat.h"
+#include "GenerateMutableSource/GenerateMutableSourceMesh.h"
+
+#include "Templates/UnrealTemplate.h"
+#include "Engine/SkeletalMesh.h"
+#include "Engine/StaticMesh.h"
+#include "Engine/AssetManager.h"
+#include "Engine/TextureLODSettings.h"
+#include "ClothingAsset.h"
+#include "EdGraphToken.h"
+#include "AssetRegistry/AssetRegistryModule.h"
+#include "UnrealEditorPortabilityHelpers.h"
+#include "PlatformInfo.h"
+
+#include "CustomizableObjectCompiler.h"
+#include "UnrealToMutableTextureConversionUtils.h"
+#include "UnrealConversionUtils.h"
+#include "CustomizableObjectEditorModule.h"
+#include "CustomizableObject.h"
+#include "GraphTraversal.h"
+#include "MutableUtils.h"
+
+#include "MutableTools/Public/NodeMeshConstant.h"
+#include "MutableTools/Public/NodeMeshTransform.h"
+#include "MutableTools/Public/Compiler.h"
+#include "MutableTools/Public/Streams.h"
+#include "MutableTools/Public/NodeModifierMeshClipMorphPlane.h"
+#include "MutableTools/Public/NodeModifierMeshClipWithMesh.h"
+#include "MutableTools/Public/NodeModifierMeshClipDeform.h"
+
+#define LOCTEXT_NAMESPACE "CustomizableObjectEditor"
+
+
+mu::NodeModifierPtr GenerateMutableSourceModifier(const UEdGraphPin * Pin, FMutableGraphGenerationContext & GenerationContext)
+{
+	check(Pin)
+	RETURN_ON_CYCLE(*Pin, GenerationContext)
+
+	CheckNumOutputs(*Pin, GenerationContext);
+	
+	UCustomizableObjectNode* Node = CastChecked<UCustomizableObjectNode>(Pin->GetOwningNode());
+
+	const FGeneratedKey Key(reinterpret_cast<void*>(&GenerateMutableSourceModifier), *Pin, *Node, GenerationContext, true);
+	if (const FGeneratedData* Generated = GenerationContext.Generated.Find(Key))
+	{
+		return static_cast<mu::NodeModifier*>(Generated->Node.get());
+	}
+	
+	mu::NodeModifierPtr Result;
+
+	bool bDoNotAddToGeneratedCache = false; // TODO Remove on MTBL-829 
+
+	if (const UCustomizableObjectNodeMeshClipMorph* TypedNodeClip = Cast<UCustomizableObjectNodeMeshClipMorph>(Node))
+	{
+		mu::NodeModifierMeshClipMorphPlanePtr ClipNode = new mu::NodeModifierMeshClipMorphPlane();
+		Result = ClipNode;
+
+		const FVector Origin = TypedNodeClip->GetOriginWithOffset();
+		const FVector& Normal = TypedNodeClip->Normal;
+
+		ClipNode->SetPlane(Origin.X, Origin.Y, Origin.Z, Normal.X, Normal.Y, Normal.Z);
+		ClipNode->SetParams(TypedNodeClip->B, TypedNodeClip->Exponent);
+		ClipNode->SetMorphEllipse(TypedNodeClip->Radius, TypedNodeClip->Radius2, TypedNodeClip->RotationAngle);
+
+		ClipNode->SetVertexSelectionBone(TCHAR_TO_ANSI(*TypedNodeClip->BoneName.ToString()), TypedNodeClip->MaxEffectRadius);
+
+		for (const FString& Tag : TypedNodeClip->Tags)
+		{
+			 ClipNode->AddTag(TCHAR_TO_ANSI(*Tag));
+		}
+	}
+
+	else if (const UCustomizableObjectNodeMeshClipDeform* TypedNodeClipDeform = Cast<UCustomizableObjectNodeMeshClipDeform>(Node))
+	{
+		mu::Ptr<mu::NodeModifierMeshClipDeform> ClipNode = new mu::NodeModifierMeshClipDeform();
+		Result = ClipNode;
+	
+		if (const UEdGraphPin* ConnectedPin = FollowInputPin(*TypedNodeClipDeform->ClipShapePin()))
+		{
+			FMutableGraphMeshGenerationData DummyMeshData;
+			mu::NodeMeshPtr ClipMesh = GenerateMutableSourceMesh(ConnectedPin, GenerationContext, DummyMeshData);
+
+			ClipNode->SetClipMesh(ClipMesh.get());
+
+			mu::EShapeBindingMethod BindingMethod = mu::EShapeBindingMethod::ClipDeformClosestProject;
+			switch(TypedNodeClipDeform->BindingMethod)
+			{
+				case EShapeBindingMethod::ClosestProject:
+					BindingMethod = mu::EShapeBindingMethod::ClipDeformClosestProject;
+					break;
+				case EShapeBindingMethod::NormalProject:
+					BindingMethod = mu::EShapeBindingMethod::ClipDeformNormalProject;
+					break;
+				case EShapeBindingMethod::ClosestToSurface:
+					BindingMethod = mu::EShapeBindingMethod::ClipDeformClosestToSurface;
+					break;
+				default:
+					check(false);
+					break;
+			}
+
+			ClipNode->SetBindingMethod(BindingMethod);
+		}
+	
+		for (const FString& Tag : TypedNodeClipDeform->Tags)
+		{
+			ClipNode->AddTag(TCHAR_TO_ANSI(*Tag));
+		}		
+	}
+
+	else if (const UCustomizableObjectNodeMeshClipWithMesh* TypedNodeClipMesh = Cast<UCustomizableObjectNodeMeshClipWithMesh>(Node))
+	{
+		// MeshClipWithMesh can be connected to multiple objects, so the compiled NodeModifierMeshClipWithMesh
+		// needs to be different for each object. If it were added to the Generated cache, all the objects would get the same.
+		bDoNotAddToGeneratedCache = true;
+
+		mu::NodeModifierMeshClipWithMeshPtr ClipNode = new mu::NodeModifierMeshClipWithMesh();
+		Result = ClipNode;
+
+		if (const UEdGraphPin* ConnectedPin = FollowInputPin(*TypedNodeClipMesh->ClipMeshPin()))
+		{
+			FMutableGraphMeshGenerationData DummyMeshData;
+			mu::NodeMeshPtr ClipMesh = GenerateMutableSourceMesh(ConnectedPin, GenerationContext, DummyMeshData);
+
+			FPinDataValue* PinData = GenerationContext.PinData.Find(ConnectedPin);
+			for (const FMeshData& MeshData : PinData->MeshesData)
+			{
+				bool bClosed = true;
+				if (const USkeletalMesh* SkeletalMesh = Cast<USkeletalMesh>(MeshData.Mesh))
+				{
+					bClosed = IsMeshClosed(SkeletalMesh, MeshData.LOD, MeshData.MaterialIndex);
+				}
+				else if (const UStaticMesh* StaticMesh = Cast<UStaticMesh>(MeshData.Mesh))
+				{
+					bClosed = IsMeshClosed(StaticMesh, MeshData.LOD, MeshData.MaterialIndex);
+				}
+				else
+				{
+					GenerationContext.Compiler->CompilerLog(LOCTEXT("UnimplementedNode", "Node type not implemented yet."), MeshData.Node);
+				}
+
+				if (!bClosed)
+				{
+					GenerationContext.Compiler->CompilerLog(LOCTEXT("Clipping mesh", "Clipping mesh not closed (i.e., it does not enclose a volume)."), MeshData.Node, EMessageSeverity::Warning);
+				}
+			}
+
+			if (Cast<UCustomizableObjectNodeStaticMesh>(ConnectedPin->GetOwningNode()))
+			{
+				mu::NodeMeshTransformPtr TransformMesh = new mu::NodeMeshTransform();
+				TransformMesh->SetSource(ClipMesh.get());
+
+				FMatrix Matrix = TypedNodeClipMesh->Transform.ToMatrixWithScale();
+				const float M[16] = {
+				  	float(Matrix.M[0][0]),float(Matrix.M[1][0]),float(Matrix.M[2][0]),float(Matrix.M[3][0]),
+					float(Matrix.M[0][1]),float(Matrix.M[1][1]),float(Matrix.M[2][1]),float(Matrix.M[3][1]),
+					float(Matrix.M[0][2]),float(Matrix.M[1][2]),float(Matrix.M[2][2]),float(Matrix.M[3][2]),
+					float(Matrix.M[0][3]),float(Matrix.M[1][3]),float(Matrix.M[2][3]),float(Matrix.M[3][3])
+				};
+
+				TransformMesh->SetTransform(M);
+				ClipMesh = TransformMesh;
+			}
+
+			ClipNode->SetClipMesh(ClipMesh.get());
+		}
+
+		for (const FString& Tag : TypedNodeClipMesh->Tags)
+		{
+			 ClipNode->AddTag(TCHAR_TO_ANSI(*Tag));
+		}
+
+		if (TypedNodeClipMesh->CustomizableObjectToClipWith != nullptr)
+		{
+			TArray<mu::NodeModifierMeshClipWithMeshPtr>* ArrayDataPtr = GenerationContext.MapClipMeshNodeToMutableClipMeshNodeArray.Find(Cast<UCustomizableObjectNodeMeshClipWithMesh>(Pin->GetOwningNode()));
+
+			if (ArrayDataPtr == nullptr)
+			{
+				TArray<mu::NodeModifierMeshClipWithMeshPtr> ArrayData;
+				ArrayData.Add(ClipNode);
+				UCustomizableObjectNodeMeshClipWithMesh* CastedNode = Cast<UCustomizableObjectNodeMeshClipWithMesh>(Pin->GetOwningNode());
+				GenerationContext.MapClipMeshNodeToMutableClipMeshNodeArray.Add(CastedNode, ArrayData);
+			}
+			else
+			{
+				ArrayDataPtr->AddUnique(ClipNode);
+			}
+		}
+	}
+
+	else
+	{
+		GenerationContext.Compiler->CompilerLog(LOCTEXT("UnimplementedNode", "Node type not implemented yet."), Node);
+	}
+
+	if (!bDoNotAddToGeneratedCache)
+	{
+		GenerationContext.Generated.Add(Key, FGeneratedData(Node, Result));
+	}
+	GenerationContext.GeneratedNodes.Add(Node);
+	
+	return Result;
+}
\ No newline at end of file
diff --git a/Engine/Plugins/Experimental/Mutable/Source/CustomizableObjectEditor/Private/GenerateMutableSource/GenerateMutableSourceModifier.h b/Engine/Plugins/Experimental/Mutable/Source/CustomizableObjectEditor/Private/GenerateMutableSource/GenerateMutableSourceModifier.h
new file mode 100644
--- /dev/null
+++ b/Engine/Plugins/Experimental/Mutable/Source/CustomizableObjectEditor/Private/GenerateMutableSource/GenerateMutableSourceModifier.h
@@ -0,0 +1,8 @@
+// Copyright Epic Games, Inc. All Rights Reserved.
+
+#pragma once
+
+#include "GenerateMutableSource/GenerateMutableSource.h"
+
+
+mu::NodeModifierPtr GenerateMutableSourceModifier(const UEdGraphPin* Pin, FMutableGraphGenerationContext& GenerationContext);
diff --git a/Engine/Plugins/Experimental/Mutable/Source/CustomizableObjectEditor/Private/GenerateMutableSource/GenerateMutableSourceProjector.cpp b/Engine/Plugins/Experimental/Mutable/Source/CustomizableObjectEditor/Private/GenerateMutableSource/GenerateMutableSourceProjector.cpp
new file mode 100644
--- /dev/null
+++ b/Engine/Plugins/Experimental/Mutable/Source/CustomizableObjectEditor/Private/GenerateMutableSource/GenerateMutableSourceProjector.cpp
@@ -0,0 +1,172 @@
+// Copyright Epic Games, Inc. All Rights Reserved.
+
+#include "GenerateMutableSourceProjector.h"
+
+#include "GenerateMutableSource/GenerateMutableSourceFloat.h"
+
+#include "Templates/UnrealTemplate.h"
+#include "GenericPlatform/GenericPlatformMath.h"
+#include "Engine/AssetManager.h"
+#include "Engine/TextureLODSettings.h"
+#include "ClothingAsset.h"
+#include "MeshUtilities.h"
+#include "EdGraphToken.h"
+#include "AssetRegistry/AssetRegistryModule.h"
+#include "UnrealEditorPortabilityHelpers.h"
+#include "PlatformInfo.h"
+
+#include "CustomizableObjectCompiler.h"
+#include "UnrealToMutableTextureConversionUtils.h"
+#include "UnrealConversionUtils.h"
+#include "CustomizableObjectEditorModule.h"
+#include "CustomizableObject.h"
+
+#include "MutableTools/Public/Node.h"
+#include "MutableTools/Public/NodeScalarConstant.h"
+#include "MutableTools/Public/NodeProjector.h"
+#include "MutableTools/Public/Compiler.h"
+#include "MutableTools/Public/Streams.h"
+
+#define LOCTEXT_NAMESPACE "CustomizableObjectEditor"
+
+
+mu::NodeProjectorPtr GenerateMutableSourceProjector(const UEdGraphPin* Pin, FMutableGraphGenerationContext& GenerationContext)
+{
+	check(Pin)
+	RETURN_ON_CYCLE(*Pin, GenerationContext)
+	
+	CheckNumOutputs(*Pin, GenerationContext);
+	
+	UCustomizableObjectNode* Node = CastChecked<UCustomizableObjectNode>(Pin->GetOwningNode());
+
+	const FGeneratedKey Key(reinterpret_cast<void*>(&GenerateMutableSourceProjector), *Pin, *Node, GenerationContext, true);
+	if (const FGeneratedData* Generated = GenerationContext.Generated.Find(Key))
+	{
+		return static_cast<mu::NodeProjector*>(Generated->Node.get());
+	}
+
+	mu::NodeProjectorPtr Result;
+	
+	if (const UCustomizableObjectNodeProjectorConstant* TypedNodeConst = Cast<UCustomizableObjectNodeProjectorConstant>(Node))
+	{
+		mu::NodeProjectorConstantPtr ProjectorNode = new mu::NodeProjectorConstant();
+		Result = ProjectorNode;
+
+		switch ((int)TypedNodeConst->Value.ProjectionType)
+		{
+		case 0:
+			ProjectorNode->SetValue(
+				mu::PROJECTOR_TYPE::PLANAR,
+				TypedNodeConst->Value.Position[0], TypedNodeConst->Value.Position[1], TypedNodeConst->Value.Position[2],
+				TypedNodeConst->Value.Direction[0], TypedNodeConst->Value.Direction[1], TypedNodeConst->Value.Direction[2],
+				TypedNodeConst->Value.Up[0], TypedNodeConst->Value.Up[1], TypedNodeConst->Value.Up[2],
+				TypedNodeConst->Value.Scale[0], TypedNodeConst->Value.Scale[1], TypedNodeConst->Value.Scale[2],
+				TypedNodeConst->Value.Angle);
+			break;
+
+		case 1:
+		{
+			// Apply strange swizzle for scales
+			// TODO: try to avoid this
+			float Radius = FMath::Abs(TypedNodeConst->Value.Scale[0] / 2.0f);
+			float Height = TypedNodeConst->Value.Scale[2];
+			ProjectorNode->SetValue(
+				mu::PROJECTOR_TYPE::CYLINDRICAL,
+				TypedNodeConst->Value.Position[0], TypedNodeConst->Value.Position[1], TypedNodeConst->Value.Position[2],
+				-TypedNodeConst->Value.Direction[0], -TypedNodeConst->Value.Direction[1], -TypedNodeConst->Value.Direction[2],
+				-TypedNodeConst->Value.Up[0], -TypedNodeConst->Value.Up[1], -TypedNodeConst->Value.Up[2],
+				-Height, Radius, Radius,
+				TypedNodeConst->Value.Angle);
+			break;
+		}
+
+		case 2:
+			ProjectorNode->SetValue(
+				mu::PROJECTOR_TYPE::WRAPPING,
+				TypedNodeConst->Value.Position[0], TypedNodeConst->Value.Position[1], TypedNodeConst->Value.Position[2],
+				TypedNodeConst->Value.Direction[0], TypedNodeConst->Value.Direction[1], TypedNodeConst->Value.Direction[2],
+				TypedNodeConst->Value.Up[0], TypedNodeConst->Value.Up[1], TypedNodeConst->Value.Up[2],
+				TypedNodeConst->Value.Scale[0], TypedNodeConst->Value.Scale[1], TypedNodeConst->Value.Scale[2],
+				TypedNodeConst->Value.Angle);
+			break;
+
+		default:
+			// Not implemented.
+			check(false);
+		}
+
+	}
+
+	else if (const UCustomizableObjectNodeProjectorParameter* TypedNodeParam = Cast<UCustomizableObjectNodeProjectorParameter>(Node))
+	{
+		mu::NodeProjectorParameterPtr ProjectorNode = new mu::NodeProjectorParameter();
+		Result = ProjectorNode;
+
+		GenerationContext.AddParameterNameUnique(Node, TypedNodeParam->ParameterName);
+
+		ProjectorNode->SetName(TCHAR_TO_ANSI(*TypedNodeParam->ParameterName));
+		ProjectorNode->SetUid(TCHAR_TO_ANSI(*GenerationContext.GetNodeIdUnique(Node).ToString()));
+		switch ((int)TypedNodeParam->DefaultValue.ProjectionType)
+		{
+		case 0:
+			ProjectorNode->SetDefaultValue(
+				mu::PROJECTOR_TYPE::PLANAR,
+				TypedNodeParam->DefaultValue.Position[0], TypedNodeParam->DefaultValue.Position[1], TypedNodeParam->DefaultValue.Position[2],
+				TypedNodeParam->DefaultValue.Direction[0], TypedNodeParam->DefaultValue.Direction[1], TypedNodeParam->DefaultValue.Direction[2],
+				TypedNodeParam->DefaultValue.Up[0], TypedNodeParam->DefaultValue.Up[1], TypedNodeParam->DefaultValue.Up[2],
+				TypedNodeParam->DefaultValue.Scale[0], TypedNodeParam->DefaultValue.Scale[1], TypedNodeParam->DefaultValue.Scale[2],
+				TypedNodeParam->DefaultValue.Angle);
+			break;
+
+		case 1:
+		{
+			// Apply strange swizzle for scales
+			// TODO: try to avoid this
+			float Radius = FMath::Abs(TypedNodeParam->DefaultValue.Scale[0] / 2.0f);
+			float Height = TypedNodeParam->DefaultValue.Scale[2];
+			ProjectorNode->SetDefaultValue(
+				mu::PROJECTOR_TYPE::CYLINDRICAL,
+				TypedNodeParam->DefaultValue.Position[0], TypedNodeParam->DefaultValue.Position[1], TypedNodeParam->DefaultValue.Position[2],
+				-TypedNodeParam->DefaultValue.Direction[0], -TypedNodeParam->DefaultValue.Direction[1], -TypedNodeParam->DefaultValue.Direction[2],
+				-TypedNodeParam->DefaultValue.Up[0], -TypedNodeParam->DefaultValue.Up[1], -TypedNodeParam->DefaultValue.Up[2],
+				-Height, Radius, Radius,
+				TypedNodeParam->DefaultValue.Angle);
+			break;
+		}
+
+		case 2:
+			ProjectorNode->SetDefaultValue(
+				mu::PROJECTOR_TYPE::WRAPPING,
+				TypedNodeParam->DefaultValue.Position[0], TypedNodeParam->DefaultValue.Position[1], TypedNodeParam->DefaultValue.Position[2],
+				TypedNodeParam->DefaultValue.Direction[0], TypedNodeParam->DefaultValue.Direction[1], TypedNodeParam->DefaultValue.Direction[2],
+				TypedNodeParam->DefaultValue.Up[0], TypedNodeParam->DefaultValue.Up[1], TypedNodeParam->DefaultValue.Up[2],
+				TypedNodeParam->DefaultValue.Scale[0], TypedNodeParam->DefaultValue.Scale[1], TypedNodeParam->DefaultValue.Scale[2],
+				TypedNodeParam->DefaultValue.Angle);
+			break;
+
+		default:
+			// Not implemented.
+			check(false);
+		}
+
+		GenerationContext.ParameterUIDataMap.Add(TypedNodeParam->ParameterName, FParameterUIData(
+			TypedNodeParam->ParameterName,
+			TypedNodeParam->ParamUIMetadata,
+			EMutableParameterType::Projector));
+	}
+
+	else
+	{
+		GenerationContext.Compiler->CompilerLog(LOCTEXT("UnimplementedNode", "Node type not implemented yet."), Node);
+	}
+
+	GenerationContext.Generated.Add(Key, FGeneratedData(Node, Result));
+	GenerationContext.GeneratedNodes.Add(Node);
+
+	if (Result)
+	{
+		Result->SetMessageContext(Node);
+	}
+
+	return Result;
+}
diff --git a/Engine/Plugins/Experimental/Mutable/Source/CustomizableObjectEditor/Private/GenerateMutableSource/GenerateMutableSourceProjector.h b/Engine/Plugins/Experimental/Mutable/Source/CustomizableObjectEditor/Private/GenerateMutableSource/GenerateMutableSourceProjector.h
new file mode 100644
--- /dev/null
+++ b/Engine/Plugins/Experimental/Mutable/Source/CustomizableObjectEditor/Private/GenerateMutableSource/GenerateMutableSourceProjector.h
@@ -0,0 +1,9 @@
+// Copyright Epic Games, Inc. All Rights Reserved.
+
+#pragma once
+
+#include "GenerateMutableSource/GenerateMutableSource.h"
+
+
+/** Convert a CustomizableObject Source Graph into a mutable source graph. */
+mu::NodeProjectorPtr GenerateMutableSourceProjector(const UEdGraphPin* Pin, FMutableGraphGenerationContext& GenerationContext);
diff --git a/Engine/Plugins/Experimental/Mutable/Source/CustomizableObjectEditor/Private/GenerateMutableSource/GenerateMutableSourceSurface.cpp b/Engine/Plugins/Experimental/Mutable/Source/CustomizableObjectEditor/Private/GenerateMutableSource/GenerateMutableSourceSurface.cpp
new file mode 100644
--- /dev/null
+++ b/Engine/Plugins/Experimental/Mutable/Source/CustomizableObjectEditor/Private/GenerateMutableSource/GenerateMutableSourceSurface.cpp
@@ -0,0 +1,1460 @@
+// Copyright Epic Games, Inc. All Rights Reserved.
+
+#include "GenerateMutableSourceSurface.h"
+
+#include "GenerateMutableSource/GenerateMutableSourceFloat.h"
+#include "GenerateMutableSource/GenerateMutableSourceColor.h"
+#include "GenerateMutableSource/GenerateMutableSourceImage.h"
+#include "GenerateMutableSource/GenerateMutableSourceMesh.h"
+#include "GenerateMutableSource/GenerateMutableSourceTable.h"
+#include "GenerateMutableSource/GenerateMutableSourceGroupProjector.h"
+
+#include "Templates/UnrealTemplate.h"
+#include "GenericPlatform/GenericPlatformMath.h"
+#include "Interfaces/ITargetPlatform.h"
+#include "Engine/SkeletalMesh.h"
+#include "Engine/AssetManager.h"
+#include "Engine/TextureLODSettings.h"
+#include "ClothingAsset.h"
+#include "MeshUtilities.h"
+#include "EdGraphToken.h"
+#include "AssetRegistry/AssetRegistryModule.h"
+#include "Materials/MaterialInstance.h"
+#include "UnrealEditorPortabilityHelpers.h"
+#include "PlatformInfo.h"
+#include "Engine/SkeletalMeshSocket.h"
+
+#include "CustomizableObjectCompiler.h"
+#include "UnrealConversionUtils.h"
+#include "CustomizableObjectEditorModule.h"
+#include "CustomizableObjectSystem.h"
+#include "CustomizableObject.h"
+#include "GraphTraversal.h"
+#include "MutableUtils.h"
+#include "CustomizableObject/Public/MutableMeshBufferUtils.h"
+
+#include "MutableTools/Public/Node.h"
+#include "MutableTools/Public/NodeColourConstant.h"
+#include "MutableTools/Public/NodeImageBinarise.h"
+#include "MutableTools/Public/NodeImageConstant.h"
+#include "MutableTools/Public/NodeImageMipmap.h"
+#include "MutableTools/Public/NodeImageFormat.h"
+#include "MutableTools/Public/NodeImageResize.h"
+#include "MutableTools/Public/NodeImageNormalComposite.h"
+#include "MutableTools/Public/NodePatchImage.h"
+#include "MutableTools/Public/NodeMeshConstant.h"
+#include "MutableTools/Public/NodeMeshFormat.h"
+#include "MutableTools/Public/NodeMeshFragment.h"
+#include "MutableTools/Public/NodeSurfaceNew.h"
+#include "MutableTools/Public/NodeSurfaceEdit.h"
+#include "MutableTools/Public/NodeSurfaceVariation.h"
+#include "MutableTools/Public/NodePatchMesh.h"
+#include "MutableTools/Public/Compiler.h"
+#include "MutableTools/Public/Streams.h"
+
+#include "MutableRuntime/Public/Mesh.h"
+#include "MutableRuntime/Public/Skeleton.h"
+#include "MutableRuntime/Public/Image.h"
+
+#define LOCTEXT_NAMESPACE "CustomizableObjectEditor"
+
+
+/**
+ * Show a compilation warning when any of the meshes behind the pin InMeshPin has a UV that is not normalized.
+ * 
+ * @param InMeshPin Meshes connected to this pin (directly or indirectly through switch/variation nodes).
+ * @param OperationNode Node which is performing the block operation. 
+ */
+void LayoutOperationUVNormalizedWarning(FMutableGraphGenerationContext& GenerationContext, const UEdGraphPin* InMeshPin, const UCustomizableObjectNode* OperationNode, const int32 UVIndex)
+{
+	if (const UEdGraphPin* ConnectedPin = FollowInputPin(*InMeshPin))
+	{
+		FPinDataValue* PinData = GenerationContext.PinData.Find(ConnectedPin);
+		if (PinData)
+		{
+			for (const FMeshData& MeshData : PinData->MeshesData)
+			{
+				if (const UCustomizableObjectNodeMesh* Node = Cast<UCustomizableObjectNodeMesh>(MeshData.Node))
+				{
+					bool bNormalized = true;
+					if (const USkeletalMesh* SkeletalMesh = Cast<USkeletalMesh>(MeshData.Mesh))
+					{
+						bNormalized = IsUVNormalized(*SkeletalMesh, MeshData.LOD, MeshData.MaterialIndex, UVIndex);
+					}
+					else if (const UStaticMesh* StaticMesh = Cast<UStaticMesh>(MeshData.Mesh))
+					{
+						bNormalized = IsUVNormalized(*StaticMesh, MeshData.LOD, MeshData.MaterialIndex, UVIndex);
+					}
+
+					if (!bNormalized)
+					{
+						FText Text = FText::Format(LOCTEXT("UVNotNormalized", "UV from mesh {0} not normalized. Required to perform texture layout operations."), FText::FromString(*MeshData.Mesh->GetName()));
+						GenerationContext.Compiler->CompilerLog(Text, OperationNode);
+					}
+				}
+			}
+		}
+	}
+}
+
+void SetSurfaceFormat( mu::FMeshBufferSet& OutVertexBufferFormat, mu::FMeshBufferSet& OutIndexBufferFormat, const FMutableGraphMeshGenerationData& MeshData, 
+					   bool bWithExtraBoneInfluences, bool bWithRealTimeMorphs, bool bWithClothing )
+{
+	// Limit skinning weights if necessary
+	// \todo: make it more flexible to support 3 or 5 or 1 weight, since there is support for this in 4.25
+	const int32 MutableBonesPerVertex = bWithExtraBoneInfluences ? EXTRA_BONE_INFLUENCES : 4;
+	
+	if (MutableBonesPerVertex != MeshData.MaxNumBonesPerVertex)
+	{
+		UE_LOG(LogMutable, Log, TEXT("Mesh bone number adjusted from %d to %d."), MeshData.MaxNumBonesPerVertex, MutableBonesPerVertex);
+	}
+
+	int MutableBufferCount = MUTABLE_VERTEXBUFFER_TEXCOORDS + 1;
+	if (MeshData.bHasVertexColors)
+	{
+		++MutableBufferCount;
+	}
+
+	if (MeshData.MaxNumBonesPerVertex > 0 && MeshData.MaxBoneIndexTypeSizeBytes > 0)
+	{
+		++MutableBufferCount;
+	}
+
+	if (bWithRealTimeMorphs)
+	{
+		MutableBufferCount += 2;
+	}
+
+	if (bWithClothing)
+	{
+		++MutableBufferCount;
+	}
+
+	OutVertexBufferFormat.SetBufferCount(MutableBufferCount);
+
+	int32 CurrentVertexBuffer = 0;
+
+	// Vertex buffer
+	MutableMeshBufferUtils::SetupVertexPositionsBuffer(CurrentVertexBuffer, OutVertexBufferFormat);
+	++CurrentVertexBuffer;
+
+	// Tangent buffer
+	MutableMeshBufferUtils::SetupTangentBuffer(CurrentVertexBuffer, OutVertexBufferFormat);
+	++CurrentVertexBuffer;
+
+	// Texture coords buffer
+	MutableMeshBufferUtils::SetupTexCoordinatesBuffer(CurrentVertexBuffer, MeshData.NumTexCoordChannels, OutVertexBufferFormat);
+	++CurrentVertexBuffer;
+
+	// Skin buffer
+	if (MeshData.MaxNumBonesPerVertex > 0 && MeshData.MaxBoneIndexTypeSizeBytes > 0)
+	{
+		MutableMeshBufferUtils::SetupSkinBuffer(CurrentVertexBuffer, MeshData.MaxBoneIndexTypeSizeBytes, MutableBonesPerVertex, OutVertexBufferFormat);
+		++CurrentVertexBuffer;
+	}
+
+	// Colour buffer
+	if (MeshData.bHasVertexColors)
+	{
+		MutableMeshBufferUtils::SetupVertexColorBuffer(CurrentVertexBuffer, OutVertexBufferFormat);
+		++CurrentVertexBuffer;
+	}
+
+	// MorphTarget vertex tracking info buffers
+	if (bWithRealTimeMorphs)
+	{
+		{
+			using namespace mu;
+			const int ElementSize = sizeof(int32);
+			constexpr int ChannelCount = 1;
+			const MESH_BUFFER_SEMANTIC Semantics[ChannelCount] = { MBS_OTHER };
+			const int SemanticIndices[ChannelCount] = { 0 };
+			const MESH_BUFFER_FORMAT Formats[ChannelCount] = { MBF_INT32 };
+			const int Components[ChannelCount] = { 1 };
+			const int Offsets[ChannelCount] = { 0 };
+
+			OutVertexBufferFormat.SetBuffer(CurrentVertexBuffer, ElementSize, ChannelCount, Semantics, SemanticIndices, Formats, Components, Offsets);
+			++CurrentVertexBuffer;
+		}
+
+		{
+			using namespace mu;
+			const int ElementSize = sizeof(int32);
+			constexpr int ChannelCount = 1;
+			const MESH_BUFFER_SEMANTIC Semantics[ChannelCount] = { MBS_OTHER };
+			const int SemanticIndices[ChannelCount] = { 1 };
+			const MESH_BUFFER_FORMAT Formats[ChannelCount] = { MBF_INT32 };
+			const int Components[ChannelCount] = { 1 };
+			const int Offsets[ChannelCount] = { 0 };
+
+			OutVertexBufferFormat.SetBuffer(CurrentVertexBuffer, ElementSize, ChannelCount, Semantics, SemanticIndices, Formats, Components, Offsets);
+			++CurrentVertexBuffer;
+		}
+	}
+
+	//Clothing Data Buffer.
+	if (bWithClothing)
+	{
+		using namespace mu;
+		const int ElementSize = sizeof(int32);
+		constexpr int ChannelCount = 1;
+		const MESH_BUFFER_SEMANTIC Semantics[ChannelCount] = { MBS_OTHER };
+		const int SemanticIndices[ChannelCount] = { bWithRealTimeMorphs ? 2 : 0 };
+		const MESH_BUFFER_FORMAT Formats[ChannelCount] = { MBF_INT32 };
+		const int Components[ChannelCount] = { 1 };
+		const int Offsets[ChannelCount] = { 0 };
+
+		OutVertexBufferFormat.SetBuffer(CurrentVertexBuffer, ElementSize, ChannelCount, Semantics, SemanticIndices, Formats, Components, Offsets);
+		++CurrentVertexBuffer;
+	}
+
+
+	// Index buffer
+	MutableMeshBufferUtils::SetupIndexBuffer(OutIndexBufferFormat);
+}
+
+
+mu::NodeSurfacePtr GenerateMutableSourceSurface(const UEdGraphPin * Pin, FMutableGraphGenerationContext & GenerationContext, FMutableGraphSurfaceGenerationData& SurfaceData)
+{
+	check(Pin)
+	RETURN_ON_CYCLE(*Pin, GenerationContext)
+
+	CheckNumOutputs(*Pin, GenerationContext);
+
+	UCustomizableObjectNode* Node = CastChecked<UCustomizableObjectNode>(Pin->GetOwningNode());
+
+	const FGeneratedKey Key(reinterpret_cast<void*>(&GenerateMutableSourceSurface), *Pin, *Node, GenerationContext, true);
+	if (const FGeneratedData* Generated = GenerationContext.Generated.Find(Key))
+	{
+		SurfaceData.NodeMaterial = Cast<UCustomizableObjectNodeMaterial>(Generated->Source);
+		return static_cast<mu::NodeSurface*>(Generated->Node.get());
+	}
+	
+	mu::NodeSurfacePtr Result;
+
+	const int32 LOD = Node->IsAffectedByLOD() ? GenerationContext.CurrentLOD : 0;
+	
+	if (UCustomizableObjectNode* CustomObjNode = Cast<UCustomizableObjectNode>(Node))
+	{
+		if (CustomObjNode->IsNodeOutDatedAndNeedsRefresh())
+		{
+			CustomObjNode->SetRefreshNodeWarning();
+		}
+	}
+
+	if (UCustomizableObjectNodeMaterial* TypedNodeMat = Cast<UCustomizableObjectNodeMaterial>(Node))
+	{
+		if (TypedNodeMat->MeshComponentIndex != GenerationContext.CurrentMeshComponent)
+		{
+			return Result;
+		}
+
+		SurfaceData.NodeMaterial = TypedNodeMat; // Save the NodeMaterial for the calling recursive calls
+
+		// NodeCopyMaterial. Special case when the TypedNodeMat is a NodeCopyMaterial. The TypedNodeMat pointer now points to the parent NodeMaterial except when reading the mesh pin, which comes from the NodeCopyMaterial.
+		UCustomizableObjectNodeMaterial* TypedNodeMaterial = TypedNodeMat;
+		UCustomizableObjectNodeMaterial* TypedNodeCopyMaterial = TypedNodeMat;
+
+		if (UCustomizableObjectNodeCopyMaterial* TypedDerivedNodeCopyMaterial = Cast<UCustomizableObjectNodeCopyMaterial>(TypedNodeMat))
+		{
+			UEdGraphPin* MaterialPin = TypedDerivedNodeCopyMaterial->GetMaterialPin();
+			if (const UEdGraphPin* ConnectedPin = FollowInputPin(*MaterialPin))
+			{
+				FMutableGraphSurfaceGenerationData ParentSurfaceData;
+				GenerateMutableSourceSurface(ConnectedPin, GenerationContext, ParentSurfaceData);
+
+				if (ParentSurfaceData.NodeMaterial)
+				{
+					UEdGraphNode* NodeMaterial = const_cast<UCustomizableObjectNodeMaterial*>(ParentSurfaceData.NodeMaterial);
+
+					if (NodeMaterial->IsA(UCustomizableObjectNodeMaterial::StaticClass()) && !NodeMaterial->IsA(UCustomizableObjectNodeCopyMaterial::StaticClass()))
+					{
+						TypedNodeMaterial = static_cast<UCustomizableObjectNodeMaterial*>(NodeMaterial);
+						TypedNodeMat = TypedNodeMaterial;
+					}
+					else
+					{
+						GenerationContext.Compiler->CompilerLog(LOCTEXT("CopyMaterialInput", "Copy Material Node can only have a Material Node as an input."), TypedDerivedNodeCopyMaterial);
+					}
+				}
+			}
+		}
+
+		const UEdGraphPin* ConnectedMaterialPin = FollowInputPin(*TypedNodeMat->GetMeshPin());
+		// Warn when texture connections are  improperly used by connecting them directly to material inputs when no layout is used
+		// TODO: delete the if clause and the warning when static meshes are operational again
+		if (ConnectedMaterialPin)
+		{
+			if (const UEdGraphPin* StaticMeshPin = FindMeshBaseSource(*ConnectedMaterialPin, true))
+			{
+				const UCustomizableObjectNode* StaticMeshNode = CastChecked<UCustomizableObjectNode>(StaticMeshPin->GetOwningNode());
+				GenerationContext.Compiler->CompilerLog(LOCTEXT("UnsupportedStaticMeshes", "Static meshes are currently not supported as material meshes"), StaticMeshNode);
+			}
+		}
+
+		mu::NodeSurfaceNewPtr SurfNode = new mu::NodeSurfaceNew();
+		Result = SurfNode;
+
+		int32 ReferencedMaterialsIndex = -1;
+		if (TypedNodeMat->Material)
+		{
+			const int32 lastMaterialAmount = GenerationContext.ReferencedMaterials.Num();
+			ReferencedMaterialsIndex = GenerationContext.ReferencedMaterials.AddUnique(TypedNodeMat->Material);
+			// Used ReferencedMaterialsIndex instead of TypedNodeMat->Material->GetName() to prevent material name collisions
+			SurfNode->SetCustomID(ReferencedMaterialsIndex);
+
+			// Take slot name from skeletal mesh if one can be found, else leave empty.
+			// Keep Referenced Materials and Materail Slot Names synchronized even if no material name can be found.
+			const bool IsNewSlotName = GenerationContext.ReferencedMaterialSlotNames.Num() == ReferencedMaterialsIndex;
+			check(IsNewSlotName == lastMaterialAmount < GenerationContext.ReferencedMaterials.Num());
+			bool SlotNameFound = false;
+			check(GenerationContext.ReferencedMaterialSlotNames.Num() >= ReferencedMaterialsIndex);
+			if (IsNewSlotName || GenerationContext.ReferencedMaterialSlotNames[ReferencedMaterialsIndex].IsNone())
+			{
+				if (ConnectedMaterialPin)
+				{
+					if (const UEdGraphPin* SkeletalMeshPin = FindMeshBaseSource(*ConnectedMaterialPin, false))
+					{
+						if (const UCustomizableObjectNodeSkeletalMesh* SkeletalMeshNode = Cast<UCustomizableObjectNodeSkeletalMesh>(SkeletalMeshPin->GetOwningNode()))
+						{
+							if (const FSkeletalMaterial* SkeletalMaterial = SkeletalMeshNode->GetSkeletalMaterialFor(SkeletalMeshPin))
+							{
+								if (IsNewSlotName)
+								{
+									GenerationContext.ReferencedMaterialSlotNames.Add(SkeletalMaterial->MaterialSlotName);
+									SlotNameFound = true;
+								}
+								else if (GenerationContext.ReferencedMaterialSlotNames[ReferencedMaterialsIndex].IsNone())
+								{
+									GenerationContext.ReferencedMaterialSlotNames[ReferencedMaterialsIndex] = SkeletalMaterial->MaterialSlotName;
+									SlotNameFound = true;
+								}
+							}
+						}
+					}
+				}
+
+				// No name was found, we need to keep index parity, so we add empty value. We may find the value later.
+				if (IsNewSlotName && !SlotNameFound)
+				{
+					GenerationContext.ReferencedMaterialSlotNames.Add(FName(NAME_None));
+				}
+			}
+		}
+
+		mu::NodeMeshPtr MeshNode;
+
+
+		TypedNodeMat = TypedNodeCopyMaterial; // NodeCopyMaterial. Start reading mesh pin. Set TypedNodeMat pointer to NodeCopyMaterial
+
+		if (const UEdGraphPin* ConnectedPin = FollowInputPin(*TypedNodeMat->GetMeshPin()))
+		{
+			FMutableGraphMeshGenerationData MeshData;
+			MeshNode = GenerateMutableSourceMesh(ConnectedPin, GenerationContext, MeshData);
+
+			if (MeshNode)
+			{
+				SurfNode->SetMeshCount(1);
+
+				mu::NodeMeshFormatPtr MeshFormatNode = new mu::NodeMeshFormat();
+				MeshFormatNode->SetSource(MeshNode.get());
+				SetSurfaceFormat( 
+						MeshFormatNode->GetVertexBuffers(), MeshFormatNode->GetIndexBuffers(), MeshData,
+						GenerationContext.Options.bExtraBoneInfluencesEnabled,
+						GenerationContext.Options.bRealTimeMorphTargetsEnabled, 
+						GenerationContext.Options.bClothingEnabled);
+				MeshFormatNode->SetMessageContext(Node);
+
+				SurfNode->SetMesh(0, MeshFormatNode);
+			}
+			else
+			{
+				GenerationContext.Compiler->CompilerLog(LOCTEXT("MeshFailed", "Mesh generation failed."), Node);
+			}
+		}
+
+		TypedNodeMat = TypedNodeMaterial; // NodeCopyMaterial. End reading mesh pin. Set TypedNodeMat pointer back to the parent NodeMaterial
+
+		TMap<FString, float> TextureNameToProjectionResFactor;
+		FString AlternateResStateName;
+
+		bool bMaterialPinConnected = FollowInputPin(*TypedNodeMat->GetMaterialAssetPin()) != nullptr;
+		TArray<FGuid> NodeTableParametersGenerated;
+
+		// Checking if we should not use the material of the table node even if it is linked to the material node
+		if (const UEdGraphPin* ConnectedPin = FollowInputPin(*TypedNodeMat->GetMaterialAssetPin()))
+		{
+			if (const UCustomizableObjectNodeTable* TypedNodeTable = Cast< UCustomizableObjectNodeTable >(ConnectedPin->GetOwningNode()))
+			{
+				GenerateMutableSourceTable(TypedNodeTable->Table->GetName(), ConnectedPin, GenerationContext);
+
+				// Checking if this material has some parameters modified by the table node linked to it
+				if (GenerationContext.GeneratedParametersInTables.Contains(TypedNodeTable))
+				{
+					NodeTableParametersGenerated = GenerationContext.GeneratedParametersInTables[TypedNodeTable];
+
+					UMaterialInstance* TableMaterial = TypedNodeTable->GetColumnDefaultAssetByType<UMaterialInstance>(ConnectedPin);
+
+					if (TableMaterial && TypedNodeMat->Material)
+					{
+						// Checking if the reference material of the Table Node has the same parent as the material of the Material Node 
+						if (TableMaterial->GetMaterial() != TypedNodeMat->Material->GetMaterial())
+						{
+							bMaterialPinConnected = false;
+
+							FString msg = FString::Printf(TEXT("Material from NodeTable is an instance from a different Parent"));
+							GenerationContext.Compiler->CompilerLog(FText::FromString(msg), TypedNodeMat);
+						}
+					}
+					else
+					{
+						bMaterialPinConnected = false;
+
+						if (!TypedNodeMat->Material)
+						{
+							FString msg = FString::Printf(TEXT("Need to select a Material to work with Node Table Materials."));
+							GenerationContext.Compiler->CompilerLog(FText::FromString(msg), TypedNodeMat);
+						}
+					}
+				}
+				else
+				{
+					bMaterialPinConnected = false;
+				}
+			}
+		}
+
+		const int32 NumImages = TypedNodeMat->GetNumParameters(EMaterialParameterType::Texture);
+		SurfNode->SetImageCount(NumImages);
+		for (int32 ImageIndex = 0; ImageIndex < NumImages; ++ImageIndex)
+		{
+			mu::NodeImagePtr GroupProjectionImg;
+			UTexture2D* GroupProjectionReferenceTexture = nullptr;
+			const UEdGraphPin* ImagePin = TypedNodeMat->GetParameterPin(EMaterialParameterType::Texture, ImageIndex);
+			const FString ImageName = TypedNodeMat->GetParameterName(EMaterialParameterType::Texture, ImageIndex).ToString();
+			const FGuid ImageId = TypedNodeMat->GetParameterId(EMaterialParameterType::Texture, ImageIndex);
+
+			FString MaterialImageId = FGroupProjectorImageInfo::GenerateId(TypedNodeMat, ImageIndex);
+			bool bShareProjectionTexturesBetweenLODs = false;
+			FGroupProjectorImageInfo* ProjectorInfo = GenerationContext.GroupProjectorLODCache.Find(MaterialImageId);
+
+			if (!ProjectorInfo) // No previous LOD of this material generated the image.
+			{
+				bool bIsGroupProjectorImage = false;
+
+				GroupProjectionImg = GenerateMutableGroupProjection(LOD, ImageIndex, MeshNode, GenerationContext,
+				TypedNodeMat, bShareProjectionTexturesBetweenLODs, bIsGroupProjectorImage,
+				GroupProjectionReferenceTexture, TextureNameToProjectionResFactor, AlternateResStateName,
+				nullptr);
+
+				if (GroupProjectionImg.get() || TypedNodeMaterial->IsImageMutableMode(ImageIndex) || (bMaterialPinConnected && NodeTableParametersGenerated.Contains(ImageId)))
+				{
+					// Get the reference texture
+					UTexture2D* ReferenceTexture = nullptr;
+
+					GenerationContext.CurrentMaterialTableParameter = ImageName;
+					GenerationContext.CurrentMaterialTableParameterId = ImageId.ToString();
+
+					if (GroupProjectionImg.get() && GroupProjectionReferenceTexture)
+					{
+						ReferenceTexture = GroupProjectionReferenceTexture;
+					}
+					else
+					{
+						ReferenceTexture = TypedNodeMat->GetImageReferenceTexture(ImageIndex);
+
+						if (!ReferenceTexture && !GroupProjectionImg.get())
+						{
+							if (ImagePin)
+							{
+								if (const UEdGraphPin* ConnectedPin = FollowInputPin(*ImagePin))
+								{
+									ReferenceTexture = FindReferenceImage(ConnectedPin, GenerationContext);
+								}
+							}
+							if (!ReferenceTexture && bMaterialPinConnected && NodeTableParametersGenerated.Contains(ImageId))
+							{
+								if (const UEdGraphPin* ConnectedPin = FollowInputPin(*TypedNodeMat->GetMaterialAssetPin()))
+								{
+									ReferenceTexture = FindReferenceImage(ConnectedPin, GenerationContext);
+								}
+							}
+							if (!ReferenceTexture)
+							{
+								ReferenceTexture = TypedNodeMat->GetImageValue(ImageIndex);
+							}
+						}
+					}
+
+					FGeneratedImageProperties Props;
+					if (ReferenceTexture)
+					{
+						// Store properties for the generated images
+						Props.TextureParameterName = ImageName;
+						Props.CompressionSettings = ReferenceTexture->CompressionSettings;
+						Props.Filter = ReferenceTexture->Filter;
+						Props.SRGB = ReferenceTexture->SRGB;
+ 						Props.LODBias = 0;
+						Props.MipGenSettings = ReferenceTexture->MipGenSettings;
+						Props.MaxTextureSize = GetMaxTextureSize(ReferenceTexture, GenerationContext);
+						Props.LODGroup = ReferenceTexture->LODGroup;
+						Props.AddressX = ReferenceTexture->AddressX;
+						Props.AddressY = ReferenceTexture->AddressY;
+						Props.bFlipGreenChannel = ReferenceTexture->bFlipGreenChannel;
+					}
+					else if (!GroupProjectionImg.get())
+					{
+						// warning!
+						FString msg = FString::Printf(TEXT("The Reference texture for material image [%s] is not set and it couldn't be found automatically."), *ImageName);
+						GenerationContext.Compiler->CompilerLog(FText::FromString(msg), Node);
+					}
+
+					GenerationContext.ImageProperties.Add(Props);
+					SurfaceData.ImageProperties = Props;
+					
+					// Calculate the LODBias for this texture
+					int32 LODBias = ComputeLODBias(GenerationContext, ReferenceTexture, Props.MaxTextureSize, TypedNodeMat, ImageIndex);
+
+					GenerationContext.CurrentTextureLODBias = LODBias;
+
+					// Generate the texture nodes
+					mu::NodeImagePtr ImageNode = [&]()
+					{
+						if (TypedNodeMaterial->IsImageMutableMode(ImageIndex))
+						{
+							if (ImagePin)
+							{
+								if (const UEdGraphPin* ConnectedPin = FollowInputPin(*ImagePin))
+								{
+									return GenerateMutableSourceImage(ConnectedPin, GenerationContext, Props.MaxTextureSize);
+								}
+							}
+							
+							if (bMaterialPinConnected && NodeTableParametersGenerated.Contains(ImageId))
+							{
+								if (const UEdGraphPin* ConnectedPin = FollowInputPin(*TypedNodeMat->GetMaterialAssetPin()))
+								{
+									return GenerateMutableSourceImage(ConnectedPin, GenerationContext, Props.MaxTextureSize);
+								}
+							}
+
+							// Else
+							{
+								UTexture2D* Texture2D = TypedNodeMat->GetImageValue(ImageIndex);
+
+								const mu::NodeImageConstantPtr ConstImageNode = new mu::NodeImageConstant();
+								GenerationContext.ArrayTextureUnrealToMutableTask.Add(FTextureUnrealToMutableTask(ConstImageNode, Texture2D, Node));
+
+								return ResizeToMaxTextureSize(Props.MaxTextureSize, Texture2D, ConstImageNode);
+							}
+						}
+						else
+						{
+							return mu::NodeImagePtr();
+						}
+					}();
+
+					if (GroupProjectionImg.get())
+					{
+						ImageNode = GroupProjectionImg;
+					}
+
+					if (GenerationContext.Options.TargetPlatform && ReferenceTexture)
+					{
+						int LayerIndex = 0;
+
+						TArray< TArray<FName> > PlatformFormats;
+						GenerationContext.Options.TargetPlatform->GetTextureFormats(ReferenceTexture, PlatformFormats);
+
+						bool bHaveFormatPlatformFormats = PlatformFormats.Num() > 0;
+
+						mu::NodeImagePtr LastImage = ImageNode;
+
+						// To apply LOD bias						
+						if (LODBias > 0)
+						{
+							mu::NodeImageResizePtr ResizeImage = new mu::NodeImageResize();
+							ResizeImage->SetBase(LastImage.get());
+							ResizeImage->SetRelative(true);
+							float factor = FMath::Pow(0.5f, LODBias);
+							ResizeImage->SetSize(factor, factor);
+							ResizeImage->SetMessageContext(Node);
+							LastImage = ResizeImage;
+						}
+						
+						mu::NodeImageMipmapPtr MipmapImage = new mu::NodeImageMipmap();
+						MipmapImage->SetSource(LastImage.get());
+						MipmapImage->SetMipmapGenerationSettings(mu::EMipmapFilterType::MFT_SimpleAverage, mu::EAddressMode::AM_NONE, 1.0f, false);
+
+						MipmapImage->SetMessageContext(Node);
+						LastImage = MipmapImage;
+
+						// Apply composite image. This needs to be computed after mipmaps generation. 	
+						if (ReferenceTexture && ReferenceTexture->CompositeTexture && ReferenceTexture->CompositeTextureMode != CTM_Disabled)
+						{
+							mu::NodeImageNormalCompositePtr CompositedImage = new mu::NodeImageNormalComposite();
+							CompositedImage->SetBase( LastImage.get() );
+							CompositedImage->SetPower( ReferenceTexture->CompositePower );
+
+							mu::ECompositeImageMode CompositeImageMode = [CompositeTextureMode = ReferenceTexture->CompositeTextureMode]() -> mu::ECompositeImageMode
+							{
+								switch(CompositeTextureMode)
+								{
+									case CTM_NormalRoughnessToRed  : return mu::ECompositeImageMode::CIM_NormalRoughnessToRed;
+									case CTM_NormalRoughnessToGreen: return mu::ECompositeImageMode::CIM_NormalRoughnessToGreen;
+									case CTM_NormalRoughnessToBlue : return mu::ECompositeImageMode::CIM_NormalRoughnessToBlue;
+									case CTM_NormalRoughnessToAlpha: return mu::ECompositeImageMode::CIM_NormalRoughnessToAlpha;
+									
+									default: return mu::ECompositeImageMode::CIM_Disabled;
+								}
+							}();
+				
+							CompositedImage->SetMode( CompositeImageMode );							
+
+							mu::NodeImageConstantPtr CompositeNormalImage = new mu::NodeImageConstant();
+
+							UTexture2D* ReferenceCompositeNormalTexture = Cast<UTexture2D>( ReferenceTexture->CompositeTexture );
+							if (ReferenceCompositeNormalTexture)
+							{
+								GenerationContext.ArrayTextureUnrealToMutableTask.Add(FTextureUnrealToMutableTask(CompositeNormalImage, ReferenceCompositeNormalTexture, Node, true));
+
+								mu::NodeImageMipmapPtr NormalCompositeMipmapImage = new mu::NodeImageMipmap();
+								NormalCompositeMipmapImage->SetSource( CompositeNormalImage );
+								NormalCompositeMipmapImage->SetMipmapGenerationSettings(mu::EMipmapFilterType::MFT_SimpleAverage, mu::EAddressMode::AM_NONE, 1.0f, true);
+								
+								CompositedImage->SetNormal( NormalCompositeMipmapImage );
+							}
+
+							LastImage = CompositedImage;
+						}
+
+						if (bHaveFormatPlatformFormats)
+						{
+							mu::NodeImageFormatPtr FormatImage = new mu::NodeImageFormat();
+							FormatImage->SetSource(LastImage.get());
+							FormatImage->SetFormat(mu::EImageFormat::IF_RGBA_UBYTE);
+							FormatImage->SetMessageContext(Node);
+							LastImage = FormatImage;
+
+							if (GenerationContext.Options.bTextureCompression)
+							{
+								check(PlatformFormats[0].Num() > LayerIndex);
+								FString FormatWithoutPrefix = PlatformFormats[0][LayerIndex].ToString();
+
+								FormatWithoutPrefix.RemoveFromStart("XBOXONE_");
+								FormatWithoutPrefix.RemoveFromStart("XBOXONEGDK_");
+								FormatWithoutPrefix.RemoveFromStart("XSX_");
+								FormatWithoutPrefix.RemoveFromStart("PS4_");
+								FormatWithoutPrefix.RemoveFromStart("PS5_");
+								FormatWithoutPrefix.RemoveFromStart("SWITCH_");
+								FormatWithoutPrefix.RemoveFromStart("OODLE_");
+								FormatWithoutPrefix.RemoveFromStart("TFO_");
+
+								// Special case for textures compressed with OODLE with the format "PlatformPrefix_OODLE_TextureFormat"
+								FString OodleString = TEXT("_OODLE_");
+								FString LeftSplit, RightSplit;
+								bool bSplitString = FormatWithoutPrefix.Split(OodleString, &LeftSplit, &RightSplit);
+
+								if (bSplitString)
+								{									
+									FormatWithoutPrefix = RightSplit;
+								}
+
+								if (FormatWithoutPrefix == TEXT("AutoDXT"))
+								{
+									FormatImage->SetFormat(mu::EImageFormat::IF_BC1, mu::EImageFormat::IF_BC3);
+								}
+								else if ((FormatWithoutPrefix == TEXT("AutoASTC")) || (FormatWithoutPrefix == TEXT("ASTC_RGBAuto")))
+								{
+									FormatImage->SetFormat(mu::EImageFormat::IF_ASTC_4x4_RGB_LDR, mu::EImageFormat::IF_ASTC_4x4_RGBA_LDR);
+								}
+								else
+								{
+									mu::EImageFormat mutableFormat = mu::EImageFormat::IF_RGBA_UBYTE;
+
+									if (FormatWithoutPrefix == TEXT("DXT1")) mutableFormat = mu::EImageFormat::IF_BC1;
+									else if (FormatWithoutPrefix == TEXT("DXT3")) mutableFormat = mu::EImageFormat::IF_BC2;
+									else if (FormatWithoutPrefix == TEXT("DXT5")) mutableFormat = mu::EImageFormat::IF_BC3;
+									else if (FormatWithoutPrefix == TEXT("BC1")) mutableFormat = mu::EImageFormat::IF_BC1;
+									else if (FormatWithoutPrefix == TEXT("BC2")) mutableFormat = mu::EImageFormat::IF_BC2;
+									else if (FormatWithoutPrefix == TEXT("BC3")) mutableFormat = mu::EImageFormat::IF_BC3;
+									else if (FormatWithoutPrefix == TEXT("BC4")) mutableFormat = mu::EImageFormat::IF_BC4;
+									else if (FormatWithoutPrefix == TEXT("BC5")) mutableFormat = mu::EImageFormat::IF_BC5;
+									else if (FormatWithoutPrefix == TEXT("ASTC_RGB")) mutableFormat = mu::EImageFormat::IF_ASTC_4x4_RGB_LDR;
+									else if (FormatWithoutPrefix == TEXT("ASTC_RGBA")) mutableFormat = mu::EImageFormat::IF_ASTC_4x4_RGBA_LDR;
+									else if (FormatWithoutPrefix == TEXT("ASTC_NormalRG")) mutableFormat = mu::EImageFormat::IF_ASTC_4x4_RG_LDR;
+									else if (FormatWithoutPrefix == TEXT("G8")) mutableFormat = mu::EImageFormat::IF_L_UBYTE;
+									else if (FormatWithoutPrefix == TEXT("BGRA8")) mutableFormat = mu::EImageFormat::IF_RGBA_UBYTE;
+									else
+									{
+										UE_LOG(LogMutable, Warning, TEXT("Unexpected image format [%s]."), *FormatWithoutPrefix);
+									}
+
+									FormatImage->SetFormat(mutableFormat);
+								}
+							}
+
+						}
+
+						ImageNode = LastImage;
+					}
+
+					mu::NodeImagePtr ImageNodePtr = ImageNode;
+					SurfNode->SetImage(ImageIndex, ImageNodePtr);
+
+					FString SurfNodeImageName = FString::Printf(TEXT("%d"), GenerationContext.ImageProperties.Num() - 1);
+
+					// Encoding material layer in mutable name
+					const int32 LayerIndex = TypedNodeMat->GetParameterLayerIndex(EMaterialParameterType::Texture, ImageIndex);
+					if (LayerIndex != -1)
+					{
+						SurfNodeImageName += "-MutableLayerParam:" + FString::FromInt(LayerIndex);
+					}
+
+					SurfNode->SetImageName(ImageIndex, TCHAR_TO_ANSI(*SurfNodeImageName));
+					int32 UVLayout = TypedNodeMat->GetImageUVLayout(ImageIndex);
+					SurfNode->SetImageLayoutIndex(ImageIndex, UVLayout);
+					SurfNode->SetImageAdditionalNames(ImageIndex, TCHAR_TO_ANSI(*TypedNodeMat->Material->GetName()), TCHAR_TO_ANSI(*ImageName));
+
+					if (bShareProjectionTexturesBetweenLODs && bIsGroupProjectorImage)
+					{
+						// Add to the GroupProjectorLODCache to potentially reuse this projection texture in higher LODs
+						ensure(LOD == 0);
+						float* AlternateProjectionResFactor = TextureNameToProjectionResFactor.Find(ImageName);
+						GenerationContext.GroupProjectorLODCache.Add(MaterialImageId,
+							FGroupProjectorImageInfo(ImageNodePtr, ImageName, ImageName, TypedNodeMat,
+							AlternateProjectionResFactor ? *AlternateProjectionResFactor : 0.f, AlternateResStateName, SurfNode, UVLayout));
+					}
+				}
+			}
+			else
+			{
+				ensure(LOD > 0);
+				check(ProjectorInfo->SurfNode->GetImage(ImageIndex) == ProjectorInfo->ImageNode);
+				SurfNode->SetImage(ImageIndex, ProjectorInfo->ImageNode);
+				SurfNode->SetImageName(ImageIndex, TCHAR_TO_ANSI(*ProjectorInfo->TextureName));
+				SurfNode->SetImageLayoutIndex(ImageIndex, ProjectorInfo->UVLayout);
+
+				TextureNameToProjectionResFactor.Add(ProjectorInfo->RealTextureName, ProjectorInfo->AlternateProjectionResolutionFactor);
+				AlternateResStateName = ProjectorInfo->AlternateResStateName;
+			}
+		}
+
+		const int32 NumVectors = TypedNodeMat->GetNumParameters(EMaterialParameterType::Vector);
+		SurfNode->SetVectorCount(NumVectors);
+		for (int32 VectorIndex = 0; VectorIndex < NumVectors; ++VectorIndex)
+		{
+			const UEdGraphPin* VectorPin = TypedNodeMat->GetParameterPin(EMaterialParameterType::Vector, VectorIndex);
+			bool bVectorPinConnected = VectorPin && FollowInputPin(*VectorPin);
+
+			FString VectorName = TypedNodeMat->GetParameterName(EMaterialParameterType::Vector, VectorIndex).ToString();
+			FGuid VectorId = TypedNodeMat->GetParameterId(EMaterialParameterType::Vector, VectorIndex);
+
+			if (bVectorPinConnected || (bMaterialPinConnected && NodeTableParametersGenerated.Contains(VectorId)))
+			{
+				const UEdGraphPin* SourcePin = nullptr;
+
+				if (bVectorPinConnected)
+				{
+					SourcePin = VectorPin;
+				}
+				else if (bMaterialPinConnected)
+				{
+					GenerationContext.CurrentMaterialTableParameter = VectorName;
+					GenerationContext.CurrentMaterialTableParameterId = VectorId.ToString();
+
+					SourcePin = TypedNodeMat->GetMaterialAssetPin();
+				}
+
+				if (const UEdGraphPin* ConnectedPin = FollowInputPin(*SourcePin))
+				{
+					mu::NodeColourPtr ColorNode = GenerateMutableSourceColor(ConnectedPin, GenerationContext);
+
+					// Encoding material layer in mutable name
+					if (const int32 LayerIndex = TypedNodeMat->GetParameterLayerIndex(EMaterialParameterType::Vector, VectorIndex); LayerIndex != INDEX_NONE)
+					{
+						VectorName += "-MutableLayerParam:" + FString::FromInt(LayerIndex);
+					}
+
+					SurfNode->SetVector(VectorIndex, ColorNode);
+					SurfNode->SetVectorName(VectorIndex, TCHAR_TO_ANSI(*VectorName));
+				}
+			}
+		}
+
+		const int32 NumScalar = TypedNodeMat->GetNumParameters(EMaterialParameterType::Scalar);
+		SurfNode->SetScalarCount(NumScalar);
+		for (int32 ScalarIndex = 0; ScalarIndex < NumScalar; ++ScalarIndex)
+		{
+			const UEdGraphPin* ScalarPin = TypedNodeMat->GetParameterPin(EMaterialParameterType::Scalar, ScalarIndex);
+			bool bScalarPinConnected = ScalarPin && FollowInputPin(*ScalarPin);
+
+			FString ScalarName = TypedNodeMat->GetParameterName(EMaterialParameterType::Scalar, ScalarIndex).ToString();
+			FGuid ScalarId = TypedNodeMat->GetParameterId(EMaterialParameterType::Scalar, ScalarIndex);
+
+			if (bScalarPinConnected || (bMaterialPinConnected && NodeTableParametersGenerated.Contains(ScalarId)))
+			{
+				const UEdGraphPin* SourcePin = nullptr;
+
+				if (bScalarPinConnected)
+				{
+					SourcePin = ScalarPin;
+				}
+				else if (bMaterialPinConnected)
+				{
+					GenerationContext.CurrentMaterialTableParameter = ScalarName;
+					GenerationContext.CurrentMaterialTableParameterId = ScalarId.ToString();
+
+					SourcePin = TypedNodeMat->GetMaterialAssetPin();
+				}
+
+				if (const UEdGraphPin* ConnectedPin = FollowInputPin(*SourcePin))
+				{
+					mu::NodeScalarPtr ScalarNode = GenerateMutableSourceFloat(ConnectedPin, GenerationContext);
+
+					// Encoding material layer in mutable name
+					if (const int32 LayerIndex = TypedNodeMat->GetParameterLayerIndex(EMaterialParameterType::Scalar, ScalarIndex); LayerIndex != INDEX_NONE)
+					{
+						ScalarName += "-MutableLayerParam:" + FString::FromInt(LayerIndex);
+					}
+
+					SurfNode->SetScalar(ScalarIndex, ScalarNode);
+					SurfNode->SetScalarName(ScalarIndex, TCHAR_TO_ANSI(*ScalarName));
+				}
+			}
+		}
+
+		for (const FString& Tag : TypedNodeMat->Tags)
+		{
+			SurfNode->AddTag(TCHAR_TO_ANSI(*Tag));
+		}
+
+		TArray<mu::NodeSurfaceNewPtr>* ArraySurfaceNodePtr = GenerationContext.MapMaterialNodeToMutableSurfaceNodeArray.Find(TypedNodeMat);
+		if (ArraySurfaceNodePtr == nullptr)
+		{
+			TArray<mu::NodeSurfaceNewPtr> ArraySurfaceNode;
+			ArraySurfaceNode.Add(SurfNode);
+			ArraySurfaceNodePtr = &GenerationContext.MapMaterialNodeToMutableSurfaceNodeArray.Add(TypedNodeMat, ArraySurfaceNode);
+		}
+		else
+		{
+			ArraySurfaceNodePtr->AddUnique(SurfNode);
+		}
+
+		// If an alternate resolution for a particular state is present, clone the surface node, add the image resizing and inject the surface variation node
+		if (TextureNameToProjectionResFactor.Num() > 0 && !AlternateResStateName.IsEmpty())
+		{
+			mu::NodeSurfaceNewPtr SurfNode2 = new mu::NodeSurfaceNew;
+
+			SurfNode2->SetCustomID(ReferencedMaterialsIndex);
+
+			SurfNode2->SetMeshCount(SurfNode->GetMeshCount());
+			SurfNode2->SetMesh(0, SurfNode->GetMesh(0));
+
+			SurfNode2->SetVectorCount(SurfNode->GetVectorCount());
+
+			for (int32 VectorParamIndex = 0; VectorParamIndex < SurfNode->GetVectorCount(); ++VectorParamIndex)
+			{
+				SurfNode2->SetVector(VectorParamIndex, SurfNode->GetVector(VectorParamIndex));
+				SurfNode2->SetVectorName(VectorParamIndex, SurfNode->GetVectorName(VectorParamIndex));
+			}
+
+			SurfNode2->SetScalarCount(SurfNode->GetScalarCount());
+
+			for (int32 ScalarParamIndex = 0; ScalarParamIndex < SurfNode->GetScalarCount(); ++ScalarParamIndex)
+			{
+				SurfNode2->SetScalar(ScalarParamIndex, SurfNode->GetScalar(ScalarParamIndex));
+				SurfNode2->SetScalarName(ScalarParamIndex, SurfNode->GetScalarName(ScalarParamIndex));
+			}
+
+			for (const FString& Tag : TypedNodeMat->Tags)
+			{
+				SurfNode2->AddTag(TCHAR_TO_ANSI(*Tag));
+			}
+
+			SurfNode2->SetImageCount(SurfNode->GetImageCount());
+
+			for (int32 ImageIndex = 0; ImageIndex < SurfNode->GetImageCount(); ++ImageIndex)
+			{
+				const FString ImageName = TypedNodeMat->GetParameterName(EMaterialParameterType::Texture, ImageIndex).ToString();
+				if (float* ResolutionFactor = TextureNameToProjectionResFactor.Find(ImageName))
+				{
+					FString MaterialImageId = FGroupProjectorImageInfo::GenerateId(TypedNodeMat, ImageIndex);
+					FGroupProjectorImageInfo* ProjectorInfo = GenerationContext.GroupProjectorLODCache.Find(MaterialImageId);
+
+					if (!ProjectorInfo || !ProjectorInfo->bIsAlternateResolutionResized)
+					{
+						mu::NodeImageResizePtr NodeImageResize = new mu::NodeImageResize;
+						NodeImageResize->SetRelative(true);
+						NodeImageResize->SetSize(*ResolutionFactor, *ResolutionFactor);
+						NodeImageResize->SetBase(SurfNode->GetImage(ImageIndex));
+
+						SurfNode2->SetImage(ImageIndex, NodeImageResize);
+
+						if (ProjectorInfo)
+						{
+							ensure(LOD == 0);
+							ProjectorInfo->ImageResizeNode = NodeImageResize;
+							ProjectorInfo->bIsAlternateResolutionResized = true;
+						}
+					}
+					else
+					{
+						ensure(LOD > 0);
+						check(ProjectorInfo->bIsAlternateResolutionResized);
+						SurfNode2->SetImage(ImageIndex, ProjectorInfo->ImageResizeNode);
+					}
+				}
+				else
+				{
+					SurfNode2->SetImage(ImageIndex, SurfNode->GetImage(ImageIndex));
+				}
+
+				SurfNode2->SetImageName(ImageIndex, SurfNode->GetImageName(ImageIndex));
+				SurfNode2->SetImageLayoutIndex(ImageIndex, SurfNode->GetImageLayoutIndex(ImageIndex));
+			}
+
+			ArraySurfaceNodePtr->AddUnique(SurfNode2);
+
+			mu::NodeSurfaceVariationPtr SurfaceVariation = new mu::NodeSurfaceVariation;
+			SurfaceVariation->SetVariationType(mu::NodeSurfaceVariation::VariationType::State);
+			SurfaceVariation->SetVariationCount(1);
+			SurfaceVariation->SetVariationTag(0, TCHAR_TO_ANSI(*AlternateResStateName));
+
+			SurfaceVariation->AddDefaultSurface(&*SurfNode);
+			SurfaceVariation->AddVariationSurface(0, &*SurfNode2);
+
+			Result = SurfaceVariation;
+		}
+	}
+
+	else if (UCustomizableObjectNodeExtendMaterial* TypedNodeExt = Cast<UCustomizableObjectNodeExtendMaterial>(Node))
+	{
+		mu::NodeSurfaceEditPtr SurfNode = new mu::NodeSurfaceEdit();
+		Result = SurfNode;
+
+		[&] // Using a lambda so control flow is easier to manage.
+		{
+			UCustomizableObjectNodeMaterial* OriginalParentMaterialNode = TypedNodeExt->GetParentMaterialNode();
+			UCustomizableObjectNodeMaterial* ParentMaterialNode = OriginalParentMaterialNode;
+
+			// Copy material
+			UCustomizableObjectNodeMaterial* CopyMaterialParentMaterialNode = nullptr;
+			const UCustomizableObjectNodeCopyMaterial* ParentMaterialCopyNodeCast = Cast<UCustomizableObjectNodeCopyMaterial>(OriginalParentMaterialNode);
+			if (ParentMaterialCopyNodeCast)
+			{
+				CopyMaterialParentMaterialNode = ParentMaterialCopyNodeCast->GetMaterialNode();
+				if (!CopyMaterialParentMaterialNode)
+				{
+					GenerationContext.Compiler->CompilerLog(LOCTEXT("ParentMissing", "Base Material not set (or not found)."), ParentMaterialCopyNodeCast);
+					return;
+				}
+			}
+		
+			if (!ParentMaterialNode)
+			{
+				GenerationContext.Compiler->CompilerLog(LOCTEXT("ParentMissing", "Parent node not set (or not found)."), Node);
+				return;
+			}
+
+			if (TypedNodeExt->IsNodeOutDatedAndNeedsRefresh())
+			{
+				GenerationContext.Compiler->CompilerLog(LOCTEXT("ParentMissing", "Parent node not set (or not found)."), Node);
+				return;
+			}
+
+			// Parent, probably generated, will be retrieved from the cache
+			FMutableGraphSurfaceGenerationData DummySurfaceData;
+			mu::NodeSurfacePtr ParentNode = GenerateMutableSourceSurface(ParentMaterialNode->OutputPin(), GenerationContext, DummySurfaceData);
+			SurfNode->SetParent(ParentNode.get());
+
+			mu::NodeMeshPtr AddMeshNode;
+			FMutableGraphMeshGenerationData MeshData;
+			if (const UEdGraphPin* ConnectedPin = FollowInputPin(*TypedNodeExt->AddMeshPin()))
+			{
+				AddMeshNode = GenerateMutableSourceMesh(ConnectedPin, GenerationContext, MeshData);
+			}
+
+			if (AddMeshNode)
+			{
+				mu::NodeMeshFormatPtr MeshFormat = new mu::NodeMeshFormat();
+				SetSurfaceFormat( 
+						MeshFormat->GetVertexBuffers(), MeshFormat->GetIndexBuffers(), MeshData,
+						GenerationContext.Options.bExtraBoneInfluencesEnabled,
+						GenerationContext.Options.bRealTimeMorphTargetsEnabled, 
+						GenerationContext.Options.bClothingEnabled );
+
+				MeshFormat->SetSource(AddMeshNode.get());
+
+				mu::NodePatchMeshPtr MeshPatch = new mu::NodePatchMesh();
+				MeshPatch->SetAdd(MeshFormat.get());
+
+				SurfNode->SetMesh(MeshPatch.get());
+				MeshPatch->SetMessageContext(Node);
+			}
+
+			if (ParentMaterialCopyNodeCast) // CopyMaterial swap.
+			{
+				ParentMaterialNode = CopyMaterialParentMaterialNode; 
+			}
+		
+			const int32 NumImages = ParentMaterialNode->GetNumParameters(EMaterialParameterType::Texture);
+			SurfNode->SetImageCount(NumImages);
+			for (int32 ImageIndex = 0; ImageIndex < NumImages; ++ImageIndex)
+			{
+				mu::NodeImagePtr ImageNode;
+				
+				if (!ImageNode) // If
+				{
+					const FString MaterialImageId = FGroupProjectorImageInfo::GenerateId(TypedNodeExt, ImageIndex);
+					const FGroupProjectorImageInfo* ProjectorInfo = GenerationContext.GroupProjectorLODCache.Find(MaterialImageId);
+
+					if (ProjectorInfo)
+					{
+						ensure(LOD > 0);
+						check(ProjectorInfo->SurfNode->GetImage(ImageIndex) == ProjectorInfo->ImageNode);
+						ImageNode = ProjectorInfo->ImageNode;
+
+						//TextureNameToProjectionResFactor.Add(ProjectorInfo->RealTextureName, ProjectorInfo->AlternateProjectionResolutionFactor);
+						//AlternateResStateName = ProjectorInfo->AlternateResStateName;
+					}
+				}
+
+				if (!ImageNode) // Else if
+				{
+					bool bShareProjectionTexturesBetweenLODs = false;
+					bool bIsGroupProjectorImage = false;
+					UTexture2D * GroupProjectionReferenceTexture = nullptr;
+					TMap<FString, float> TextureNameToProjectionResFactor;
+					FString AlternateResStateName;
+					
+					ImageNode = GenerateMutableGroupProjection(LOD, ImageIndex, AddMeshNode, GenerationContext,
+						TypedNodeExt, bShareProjectionTexturesBetweenLODs, bIsGroupProjectorImage,
+						GroupProjectionReferenceTexture, TextureNameToProjectionResFactor, AlternateResStateName,
+						ParentMaterialNode);
+				}
+				
+				if (!ImageNode) // Else if
+				{
+					const FGuid ImageId = ParentMaterialNode->GetParameterId(EMaterialParameterType::Texture, ImageIndex);
+					
+					if (TypedNodeExt->UsesImage(ImageId))
+					{
+						check(ParentMaterialNode->IsImageMutableMode(ImageIndex)); // Ensured at graph time. If it fails, something is wrong.
+						
+						float MaxTextureSize = 0.f;
+
+						// The static_cast is correct because we now it's a material node due to the ParentMaterialNodeCast
+						if (const mu::NodeSurfaceNewPtr ParentNode2 = static_cast<mu::NodeSurfaceNew*>(ParentNode.get()))
+						{
+							FString Aux(ParentNode2->GetImageName(ImageIndex));
+
+							if (!Aux.IsEmpty())
+							{
+								check(Aux.IsNumeric());
+								const int32 ImagePropertiesIndex = FCString::Atoi(*Aux);
+								check(GenerationContext.ImageProperties.IsValidIndex(ImagePropertiesIndex));
+								MaxTextureSize = GenerationContext.ImageProperties[ImagePropertiesIndex].MaxTextureSize;
+							}
+						}
+
+						if (const UEdGraphPin* ConnectedPin = FollowInputPin(*TypedNodeExt->GetUsedImagePin(ImageId)))
+						{
+							ImageNode = GenerateMutableSourceImage(ConnectedPin, GenerationContext, DummySurfaceData.ImageProperties.MaxTextureSize);
+						}
+					}
+				}
+
+				if (ImageNode)
+				{
+					SurfNode->SetImage(ImageIndex, ImageNode);
+				}
+
+				const int32 UVIndex = ParentMaterialNode->GetImageUVLayout(ImageIndex);
+				LayoutOperationUVNormalizedWarning(GenerationContext, ParentMaterialNode->GetMeshPin(), Node, UVIndex);
+				LayoutOperationUVNormalizedWarning(GenerationContext, TypedNodeExt->AddMeshPin(), Node, UVIndex);
+			}
+		
+			for (const FString& Tag : TypedNodeExt->Tags)
+			{
+				SurfNode->AddTag(TCHAR_TO_ANSI(*Tag));
+			}
+		}();
+	}
+
+	else if (const UCustomizableObjectNodeRemoveMesh* TypedNodeRem = Cast<UCustomizableObjectNodeRemoveMesh>(Node))
+	{
+		mu::NodeSurfaceEditPtr SurfNode = new mu::NodeSurfaceEdit();
+		Result = SurfNode;
+
+		UCustomizableObjectNodeMaterialBase* ParentMaterialNode = TypedNodeRem->GetParentMaterialNode();
+		if (!ParentMaterialNode)
+		{
+			GenerationContext.Compiler->CompilerLog(LOCTEXT("ParentMissing", "Parent node not set (or not found)."), Node);
+		}
+		else
+		{
+			// Parent, probably generated, will be retrieved from the cache
+			FMutableGraphSurfaceGenerationData DummySurfaceData;
+			mu::NodeSurfacePtr ParentNode = GenerateMutableSourceSurface(ParentMaterialNode->OutputPin(), GenerationContext, DummySurfaceData);
+			SurfNode->SetParent(ParentNode.get());
+
+			mu::NodeMeshPtr RemoveMeshNode;
+
+			if (const UEdGraphPin* ConnectedPin = FollowInputPin(*TypedNodeRem->RemoveMeshPin()))
+			{
+				FMutableGraphMeshGenerationData DummyMeshData;
+				RemoveMeshNode = GenerateMutableSourceMesh(ConnectedPin, GenerationContext, DummyMeshData);
+			}
+
+			if (RemoveMeshNode)
+			{
+				mu::NodePatchMeshPtr MeshPatch = new mu::NodePatchMesh();
+				MeshPatch->SetRemove(RemoveMeshNode.get());
+				SurfNode->SetMesh(MeshPatch.get());
+				MeshPatch->SetMessageContext(Node);
+			}
+		}
+	}
+
+	else if (const UCustomizableObjectNodeRemoveMeshBlocks* TypedNodeRemBlocks = Cast<UCustomizableObjectNodeRemoveMeshBlocks>(Node))
+	{
+		mu::NodeSurfaceEditPtr SurfNode = new mu::NodeSurfaceEdit();
+		Result = SurfNode;
+
+		UCustomizableObjectNodeMaterial* ParentMaterialNode = TypedNodeRemBlocks->GetParentMaterialNode();
+
+		if (!ParentMaterialNode)
+		{
+			GenerationContext.Compiler->CompilerLog(LOCTEXT("ParentMissing", "Parent node not set (or not found)."), Node);
+		}
+		else
+		{
+			// Parent, probably generated, will be retrieved from the cache
+			FMutableGraphSurfaceGenerationData DummySurfaceData;
+			mu::NodeSurfacePtr ParentNode = GenerateMutableSourceSurface(ParentMaterialNode->OutputPin(), GenerationContext, DummySurfaceData);
+			SurfNode->SetParent(ParentNode.get());
+
+			const UEdGraphPin* BaseSourcePin = FindMeshBaseSource(*ParentMaterialNode->OutputPin(), false);
+
+			if (!BaseSourcePin)
+			{
+				GenerationContext.Compiler->CompilerLog(LOCTEXT("ParentMissing", "Parent node not set (or not found)."), Node);
+			}
+			else
+			{
+				FMutableGraphMeshGenerationData DummyMeshData;
+				mu::NodeMeshPtr BaseSourceMesh = GenerateMutableSourceMesh(BaseSourcePin, GenerationContext, DummyMeshData);
+
+				mu::NodeMeshFragmentPtr MeshFrag = new mu::NodeMeshFragment();
+				MeshFrag->SetMesh(BaseSourceMesh);
+				MeshFrag->SetLayoutOrGroup(TypedNodeRemBlocks->ParentLayoutIndex);
+				MeshFrag->SetFragmentType(mu::NodeMeshFragment::FT_LAYOUT_BLOCKS);
+
+				// Add the block indices
+				TArray<int> FoundBlockIndices;
+				const TArray<UCustomizableObjectLayout*> Layouts = ParentMaterialNode->GetLayouts();
+
+				if (!Layouts.IsValidIndex(TypedNodeRemBlocks->ParentLayoutIndex))
+				{
+					GenerationContext.Compiler->CompilerLog(LOCTEXT("Invalid layout", "Node refers to an invalid texture layout. Add a layout to the selected parent material, then select some block(s)."), Node);
+					UE_LOG(LogMutable, Warning, TEXT("In object [%s] UCustomizableObjectNodeRemoveMeshBlocks refers to an invalid texture layout index %d. Parent node has %d layouts."),
+						*GenerationContext.Object->GetName(), TypedNodeRemBlocks->ParentLayoutIndex, Layouts.Num());
+				}
+				else
+				{
+					const UCustomizableObjectLayout* Layout = Layouts[TypedNodeRemBlocks->ParentLayoutIndex];
+
+					if (Cast<UCustomizableObjectNodeMaterial>(ParentMaterialNode))
+					{
+						for (const FGuid& Id : TypedNodeRemBlocks->BlockIds)
+						{
+							int BlockIndex = Layout->FindBlock(Id);
+							if (Layout->Blocks.IsValidIndex(BlockIndex))
+							{
+								FoundBlockIndices.Add(BlockIndex);
+							}
+							else
+							{
+								GenerationContext.Compiler->CompilerLog(LOCTEXT("InvalidLayoutBlock", "Node refers to an invalid layout block."), Node);
+								UE_LOG(LogMutable, Warning, TEXT("In object [%s] UCustomizableObjectNodeRemoveMeshBlocks refers to an invalid layout block id %s. Parent node has %d blocks."),
+									*GenerationContext.Object->GetName(), *Id.ToString(), Layout->Blocks.Num());
+							}
+						}
+					}
+				}
+
+				MeshFrag->SetBlockCount(FoundBlockIndices.Num());
+				for (int i = 0; i < FoundBlockIndices.Num(); ++i)
+				{
+					MeshFrag->SetBlock(i, FoundBlockIndices[i]);
+				}
+
+				mu::NodePatchMeshPtr MeshPatch = new mu::NodePatchMesh();
+				MeshPatch->SetRemove(MeshFrag.get());
+				SurfNode->SetMesh(MeshPatch.get());
+				MeshPatch->SetMessageContext(Node);
+			}
+
+			LayoutOperationUVNormalizedWarning(GenerationContext, ParentMaterialNode->GetMeshPin(), Node, TypedNodeRemBlocks->ParentLayoutIndex);
+		}
+	}
+
+	else if (UCustomizableObjectNodeEditMaterial* TypedNodeEdit = Cast<UCustomizableObjectNodeEditMaterial>(Node))
+	{
+		mu::NodeSurfaceEditPtr SurfNode = new mu::NodeSurfaceEdit();
+		Result = SurfNode;
+
+		UCustomizableObjectNodeMaterial* ParentMaterialNode = TypedNodeEdit->GetParentMaterialNode();
+		if (!ParentMaterialNode)
+		{
+			GenerationContext.Compiler->CompilerLog(LOCTEXT("ParentMissing", "Parent node not set (or not found)."), Node);
+		}
+		else
+		{
+			// Parent, probably generated, will be retrieved from the cache
+			FMutableGraphSurfaceGenerationData DummySurfaceData;
+			mu::NodeSurfacePtr ParentNode = GenerateMutableSourceSurface(ParentMaterialNode->OutputPin(), GenerationContext, DummySurfaceData);
+			SurfNode->SetParent(ParentNode.get());
+
+
+			const int32 NumImages = ParentMaterialNode->GetNumParameters(EMaterialParameterType::Texture);
+			SurfNode->SetImageCount(NumImages);
+			for (int32 ImageIndex = 0; ImageIndex < NumImages; ++ImageIndex)
+			{
+				const FGuid ImageId = ParentMaterialNode->GetParameterId(EMaterialParameterType::Texture, ImageIndex);
+
+				if (TypedNodeEdit->UsesImage(ImageId))
+				{
+					check(ParentMaterialNode->IsImageMutableMode(ImageIndex)); // Ensured at graph time. If it fails, something is wrong.
+					
+					const UEdGraphPin* ConnectedImagePin = FollowInputPin(*TypedNodeEdit->GetUsedImagePin(ImageId));
+					
+					mu::NodePatchImagePtr ImagePatchNode = new mu::NodePatchImage;
+					ImagePatchNode->SetMessageContext(Node);
+
+					// \todo: expose these two options?
+					ImagePatchNode->SetBlendType(mu::EBlendType::BT_BLEND);
+					ImagePatchNode->SetApplyToAlphaChannel(true);
+
+					// Calculate the LODBias for this texture
+					UTexture2D* ReferenceTexture = ParentMaterialNode->GetImageReferenceTexture(ImageIndex);
+					int32 LODBias = ComputeLODBias(GenerationContext, ReferenceTexture, ReferenceTexture ? ReferenceTexture->MaxTextureSize : 0, ParentMaterialNode, ImageIndex);
+
+					GenerationContext.CurrentTextureLODBias = LODBias;
+
+					mu::NodeImagePtr ImageNode = GenerateMutableSourceImage(ConnectedImagePin, GenerationContext, 0.f);
+					ImagePatchNode->SetImage(ImageNode);
+
+					if (const UEdGraphPin* ConnectedMaskPin = FollowInputPin(*TypedNodeEdit->GetUsedImageMaskPin(ImageId)))
+					{
+						mu::NodeImagePtr MaskNode = GenerateMutableSourceImage(ConnectedMaskPin, GenerationContext, 0.f);
+						ImagePatchNode->SetMask(MaskNode);
+					}
+
+					// Add the block indices
+					TArray<int> FoundBlockIndices;
+					const TArray<UCustomizableObjectLayout*> Layouts = ParentMaterialNode->GetLayouts();
+
+					if (!Layouts.IsValidIndex(TypedNodeEdit->ParentLayoutIndex))
+					{
+						GenerationContext.Compiler->CompilerLog(LOCTEXT("Invalid layout", "Node refers to an invalid texture layout."), Node);
+						UE_LOG(LogMutable, Warning, TEXT("In object [%s] UCustomizableObjectNodeEditMaterial refers to an invalid texture layout index %d. Parent node has %d layouts."),
+							*GenerationContext.Object->GetName(), TypedNodeEdit->ParentLayoutIndex, Layouts.Num());
+					}
+					else
+					{
+						if (const UCustomizableObjectNodeMaterial* ParentMaterial = Cast<UCustomizableObjectNodeMaterial>(ParentMaterialNode))
+						{
+							if (TypedNodeEdit->ParentLayoutIndex == ParentMaterial->GetImageUVLayout(ImageIndex))
+							{
+								const UCustomizableObjectLayout* Layout = Layouts[TypedNodeEdit->ParentLayoutIndex];
+
+								for (const FGuid& Id : TypedNodeEdit->BlockIds)
+								{
+									int BlockIndex = Layout->FindBlock(Id);
+									if (Layout->Blocks.IsValidIndex(BlockIndex))
+									{
+										FoundBlockIndices.Add(BlockIndex);
+									}
+									else
+									{
+										GenerationContext.Compiler->CompilerLog(LOCTEXT("InvalidLayoutBlock", "Node refers to an invalid layout block."), Node);
+										UE_LOG(LogMutable, Warning, TEXT("In object [%s] UCustomizableObjectNodeEditMaterial refers to an invalid layout block id %s. Parent node has %d blocks."),
+											*GenerationContext.Object->GetName(), *Id.ToString(), Layout->Blocks.Num());
+									}
+								}
+							}
+							else
+							{
+								GenerationContext.Compiler->CompilerLog(LOCTEXT("InvalidLayoutTexture", "Texture layout does not match the layout of Edit Material node or parent material changed."), Cast<UCustomizableObjectNode>(ConnectedImagePin->GetOwningNode()));
+							}
+						}
+					}
+
+					ImagePatchNode->SetBlockCount(FoundBlockIndices.Num());
+					for (int BlockIndex = 0; BlockIndex < FoundBlockIndices.Num(); ++BlockIndex)
+					{
+						ImagePatchNode->SetBlock(BlockIndex, FoundBlockIndices[BlockIndex]);
+					}
+
+					SurfNode->SetPatch(ImageIndex, ImagePatchNode);
+				}
+				
+				LayoutOperationUVNormalizedWarning(GenerationContext, ParentMaterialNode->GetMeshPin(), Node, ParentMaterialNode->GetImageUVLayout(ImageIndex));
+			}
+		}
+	}
+
+	else if (const UCustomizableObjectNodeMorphMaterial* TypedNodeMorph = Cast<UCustomizableObjectNodeMorphMaterial>(Node))
+	{
+		mu::NodeSurfaceEditPtr SurfNode = new mu::NodeSurfaceEdit();
+		Result = SurfNode;
+
+		UCustomizableObjectNodeMaterialBase* ParentMaterialNode = TypedNodeMorph->GetParentMaterialNode();
+		if (!ParentMaterialNode)
+		{
+			GenerationContext.Compiler->CompilerLog(LOCTEXT("ParentMissing", "Parent node not set (or not found)."), Node);
+		}
+		else
+		{
+			// Parent, probably generated, will be retrieved from the cache
+			FMutableGraphSurfaceGenerationData DummySurfaceData;
+			mu::NodeSurfacePtr ParentNode = GenerateMutableSourceSurface(ParentMaterialNode->OutputPin(), GenerationContext, DummySurfaceData);
+			SurfNode->SetParent(ParentNode.get());
+
+			const UEdGraphPin* BaseSourcePin = FindMeshBaseSource(*ParentMaterialNode->OutputPin(), false);
+
+			if (!BaseSourcePin)
+			{
+				GenerationContext.Compiler->CompilerLog(LOCTEXT("ParentMissing", "Parent node not set (or not found)."), Node);
+			}
+			else
+			{
+				//Morph Mesh
+				mu::MeshPtr MorphedSourceMesh = BuildMorphedMutableMesh(BaseSourcePin, TypedNodeMorph->MorphTargetName, GenerationContext);
+
+				mu::NodeMeshConstantPtr MorphedSourceMeshNode = new mu::NodeMeshConstant;
+				MorphedSourceMeshNode->SetMessageContext(Node);					
+				MorphedSourceMeshNode->SetValue(MorphedSourceMesh);
+				
+				SurfNode->SetMorph(MorphedSourceMeshNode);
+				
+				if (const UEdGraphPin* ConnectedPin = FollowInputPin(*TypedNodeMorph->FactorPin()))
+				{
+					UEdGraphNode* floatNode = ConnectedPin->GetOwningNode();
+					bool validStaticFactor = true;
+					if (const UCustomizableObjectNodeFloatParameter* floatParameterNode = Cast<UCustomizableObjectNodeFloatParameter>(floatNode))
+					{
+						if (floatParameterNode->DefaultValue < -1.0f || floatParameterNode->DefaultValue > 1.0f)
+						{
+							validStaticFactor = false;
+							FString msg = FString::Printf(TEXT("Mesh morph nodes only accept factors between -1.0 and 1.0 inclusive but the default value of the float parameter node is (%f). Factor will be ignored."), floatParameterNode->DefaultValue);
+							GenerationContext.Compiler->CompilerLog(FText::FromString(msg), Node);
+						}
+						if (floatParameterNode->ParamUIMetadata.MinimumValue < -1.0f)
+						{
+							validStaticFactor = false;
+							FString msg = FString::Printf(TEXT("Mesh morph nodes only accept factors between -1.0 and 1.0 inclusive but the minimum UI value for the input float parameter node is (%f). Factor will be ignored."), floatParameterNode->ParamUIMetadata.MinimumValue);
+							GenerationContext.Compiler->CompilerLog(FText::FromString(msg), Node);
+						}
+						if (floatParameterNode->ParamUIMetadata.MaximumValue > 1.0f)
+						{
+							validStaticFactor = false;
+							FString msg = FString::Printf(TEXT("Mesh morph nodes only accept factors between -1.0 and 1.0 inclusive but the maximum UI value for the input float parameter node is (%f). Factor will be ignored."), floatParameterNode->ParamUIMetadata.MaximumValue);
+							GenerationContext.Compiler->CompilerLog(FText::FromString(msg), Node);
+						}
+					}
+					else if (const UCustomizableObjectNodeFloatConstant* floatConstantNode = Cast<UCustomizableObjectNodeFloatConstant>(floatNode))
+					{
+						if (floatConstantNode->Value < -1.0f || floatConstantNode->Value > 1.0f)
+						{
+							validStaticFactor = false;
+							FString msg = FString::Printf(TEXT("Mesh morph nodes only accept factors between -1.0 and 1.0 inclusive but the value of the float constant node is (%f). Factor will be ignored."), floatConstantNode->Value);
+							GenerationContext.Compiler->CompilerLog(FText::FromString(msg), Node);
+						}
+					}
+
+					if (validStaticFactor)
+					{
+						mu::NodeScalarPtr FactorNode = GenerateMutableSourceFloat(ConnectedPin, GenerationContext);
+						SurfNode->SetFactor(FactorNode);
+					}
+				}
+			}
+		}
+	}
+
+	else if (const UCustomizableObjectNodeMaterialVariation* TypedNodeVar = Cast<UCustomizableObjectNodeMaterialVariation>(Node))
+	{
+		mu::NodeSurfaceVariationPtr SurfNode = new mu::NodeSurfaceVariation();
+		Result = SurfNode;
+
+		mu::NodeSurfaceVariation::VariationType muType = mu::NodeSurfaceVariation::VariationType::Tag;
+		switch (TypedNodeVar->Type)
+		{
+		case ECustomizableObjectNodeMaterialVariationType::Tag: muType = mu::NodeSurfaceVariation::VariationType::Tag; break;
+		case ECustomizableObjectNodeMaterialVariationType::State: muType = mu::NodeSurfaceVariation::VariationType::State; break;
+		default:
+			check(false);
+			break;
+		}
+		SurfNode->SetVariationType(muType);
+
+		for (const UEdGraphPin* ConnectedPin : FollowInputPinArray(*TypedNodeVar->DefaultPin()))
+		{
+			// Is it a modifier?
+			FMutableGraphSurfaceGenerationData DummySurfaceData;
+			mu::NodeSurfacePtr ChildNode = GenerateMutableSourceSurface(ConnectedPin, GenerationContext, DummySurfaceData);
+			if (ChildNode)
+			{
+				SurfNode->AddDefaultSurface(ChildNode.get());
+			}
+			else
+			{
+				GenerationContext.Compiler->CompilerLog(LOCTEXT("SurfaceFailed", "Surface generation failed."), Node);
+			}
+		}
+
+		SurfNode->SetVariationCount(TypedNodeVar->Variations.Num());
+		for (int VariationIndex = 0; VariationIndex < TypedNodeVar->Variations.Num(); ++VariationIndex)
+		{
+			mu::NodeSurfacePtr VariationSurfaceNode;
+
+			if (UEdGraphPin* VariationPin = TypedNodeVar->VariationPin(VariationIndex))
+			{
+				SurfNode->SetVariationTag(VariationIndex, TCHAR_TO_ANSI(*TypedNodeVar->Variations[VariationIndex].Tag));
+				for (const UEdGraphPin* ConnectedPin : FollowInputPinArray(*VariationPin))
+				{
+					// Is it a modifier?
+					FMutableGraphSurfaceGenerationData DummySurfaceData;
+					mu::NodeSurfacePtr ChildNode = GenerateMutableSourceSurface(ConnectedPin, GenerationContext, DummySurfaceData);
+					if (ChildNode)
+					{
+						SurfNode->AddVariationSurface(VariationIndex, ChildNode.get());
+					}
+					else
+					{
+						GenerationContext.Compiler->CompilerLog(LOCTEXT("SurfaceModifierFailed", "Surface generation failed."), Node);
+					}
+				}
+			}
+		}
+	}
+
+	else
+	{
+		GenerationContext.Compiler->CompilerLog(LOCTEXT("UnimplementedNode", "Node type not implemented yet."), Node);
+	}
+
+
+	if (Result)
+	{
+		Result->SetMessageContext(Node);
+	}
+
+	GenerationContext.Generated.Add(Key, FGeneratedData(Node, Result));
+	GenerationContext.GeneratedNodes.Add(Node);
+
+	return Result;
+}
diff --git a/Engine/Plugins/Experimental/Mutable/Source/CustomizableObjectEditor/Private/GenerateMutableSource/GenerateMutableSourceSurface.h b/Engine/Plugins/Experimental/Mutable/Source/CustomizableObjectEditor/Private/GenerateMutableSource/GenerateMutableSourceSurface.h
new file mode 100644
--- /dev/null
+++ b/Engine/Plugins/Experimental/Mutable/Source/CustomizableObjectEditor/Private/GenerateMutableSource/GenerateMutableSourceSurface.h
@@ -0,0 +1,8 @@
+// Copyright Epic Games, Inc. All Rights Reserved.
+
+#pragma once
+
+#include "GenerateMutableSource/GenerateMutableSource.h"
+
+
+mu::NodeSurfacePtr GenerateMutableSourceSurface(const UEdGraphPin * Pin, FMutableGraphGenerationContext & GenerationContext, FMutableGraphSurfaceGenerationData& SurfaceData);
