[mutable] Move Mutable plugin from the restricted folder to the experimental plugins folder.

diff --git a/Engine/Plugins/Experimental/Mutable/Source/CustomizableObjectEditor/Private/GenerateMutableSource/GenerateMutableSourceTable.cpp b/Engine/Plugins/Experimental/Mutable/Source/CustomizableObjectEditor/Private/GenerateMutableSource/GenerateMutableSourceTable.cpp
new file mode 100644
--- /dev/null
+++ b/Engine/Plugins/Experimental/Mutable/Source/CustomizableObjectEditor/Private/GenerateMutableSource/GenerateMutableSourceTable.cpp
@@ -0,0 +1,685 @@
+// Copyright Epic Games, Inc. All Rights Reserved.
+
+#include "GenerateMutableSourceTable.h"
+
+#include "GenerateMutableSource/GenerateMutableSourceFloat.h"
+#include "GenerateMutableSource/GenerateMutableSourceMesh.h"
+
+#include "Templates/UnrealTemplate.h"
+#include "Engine/SkeletalMesh.h"
+#include "Engine/AssetManager.h"
+#include "Engine/TextureLODSettings.h"
+#include "ClothingAsset.h"
+#include "MeshUtilities.h"
+#include "EdGraphToken.h"
+#include "AssetRegistry/AssetRegistryModule.h"
+#include "Materials/MaterialInstance.h"
+#include "UnrealEditorPortabilityHelpers.h"
+#include "PlatformInfo.h"
+
+#include "CustomizableObjectCompiler.h"
+#include "CustomizableObjectEditorModule.h"
+#include "CustomizableObject.h"
+
+#include "MutableTools/Public/Compiler.h"
+#include "MutableTools/Public/Streams.h"
+
+#include "MutableRuntime/Public/Mesh.h"
+#include "MutableRuntime/Public/Skeleton.h"
+#include "MutableRuntime/Public/Image.h"
+
+#define LOCTEXT_NAMESPACE "CustomizableObjectEditor"
+
+
+void FillTableColumn(const UEdGraphPin* Pin, mu::TablePtr MutableTable, FString ColumnName, FString RowName, int32 RowIdx, uint8* CellData, FProperty* Property, FMutableGraphGenerationContext& GenerationContext)
+{
+	UCustomizableObjectNodeTable* CustomNodeTable = Cast<UCustomizableObjectNodeTable>(Pin->GetOwningNode());
+	int32 CurrentColumn;
+
+	// Getting property type
+	if (const FSoftObjectProperty* SoftObjectProperty = CastField<FSoftObjectProperty>(Property))
+	{
+		UObject* Object = SoftObjectProperty->GetPropertyValue(CellData).LoadSynchronous();
+
+		if (!Object)
+		{
+			return;
+		}
+
+		if (USkeletalMesh* SkeletalMesh = Cast<USkeletalMesh>(Object))
+		{
+			// Getting Animation Blueprint and Animation Slot
+			FString AnimBP, AnimSlot, GameplayTag, AnimBPAssetTag;
+			TArray<FGameplayTag> GameplayTags;
+
+			CustomNodeTable->GetAnimationColumns(ColumnName, AnimBP, AnimSlot, GameplayTag);
+
+			if (!AnimBP.IsEmpty())
+			{
+				if (!AnimSlot.IsEmpty())
+				{
+					if (CustomNodeTable->Table)
+					{
+						uint8* AnimRowData = CustomNodeTable->Table->FindRowUnchecked(FName(*RowName));
+
+						if (AnimRowData)
+						{
+							int32 SlotIndex = -1;
+
+							// Getting animation slot row value from data table
+							if (FProperty* AnimSlotProperty = CustomNodeTable->Table->FindTableProperty(FName(*AnimSlot)))
+							{
+								uint8* AnimSlotData = AnimSlotProperty->ContainerPtrToValuePtr<uint8>(AnimRowData, 0);
+
+								if (AnimSlotData)
+								{
+									if (const FIntProperty* IntProperty = CastField<FIntProperty>(AnimSlotProperty))
+									{
+										SlotIndex = IntProperty->GetPropertyValue(AnimSlotData);
+									}
+								}
+							}
+
+							if (SlotIndex > -1)
+							{
+								// Getting animation instance soft class from data table
+								if (FProperty* AnimBPProperty = CustomNodeTable->Table->FindTableProperty(FName(*AnimBP)))
+								{
+									uint8* AnimBPData = AnimBPProperty->ContainerPtrToValuePtr<uint8>(AnimRowData, 0);
+
+									if (AnimBPData)
+									{
+										if (const FSoftClassProperty* SoftClassProperty = CastField<FSoftClassProperty>(AnimBPProperty))
+										{
+											TSoftClassPtr<UAnimInstance> AnimInstance(SoftClassProperty->GetPropertyValue(AnimBPData).ToSoftObjectPath());
+
+											if (AnimInstance)
+											{
+												GenerationContext.AnimBPAssetsMap.Add(AnimInstance.ToString(), AnimInstance);
+
+												AnimBPAssetTag = GenerateAnimationInstanceTag(AnimInstance.ToString(), SlotIndex);
+											}
+										}
+									}
+								}
+							}
+							else
+							{
+								FString msg = FString::Printf(TEXT("Could not found the Slot column of the animation blueprint column [%s] for the mesh column [%s] row [%s]."), *AnimBP, *ColumnName, *RowName);
+								GenerationContext.Compiler->CompilerLog(FText::FromString(msg), CustomNodeTable);
+							}
+						}
+					}
+				}
+				else
+				{
+					FString msg = FString::Printf(TEXT("Could not found the Slot column of the animation blueprint column [%s] for the mesh column [%s]."), *AnimBP, *ColumnName);
+					GenerationContext.Compiler->CompilerLog(FText::FromString(msg), CustomNodeTable);
+				}
+			}
+
+			// Getting Gameplay tags
+			if (!GameplayTag.IsEmpty())
+			{
+				if (CustomNodeTable->Table)
+				{
+					uint8* GameplayRowData = CustomNodeTable->Table->FindRowUnchecked(FName(*RowName));
+
+					if (GameplayRowData)
+					{
+						// Getting animation tag row value from data table
+						if (FProperty* GameplayTagProperty = CustomNodeTable->Table->FindTableProperty(FName(*GameplayTag)))
+						{
+							uint8* GameplayTagData = GameplayTagProperty->ContainerPtrToValuePtr<uint8>(GameplayRowData, 0);
+
+							if (const FStructProperty* StructProperty = CastField<FStructProperty>(GameplayTagProperty))
+							{
+								if (StructProperty->Struct == TBaseStructure<FGameplayTagContainer>::Get())
+								{
+									if (GameplayTagData)
+									{
+										const FGameplayTagContainer TagContainer = *(FGameplayTagContainer*)GameplayTagData;
+										TagContainer.GetGameplayTagArray(GameplayTags);
+									}
+								}
+							}
+						}
+					}
+				}
+			}
+
+			// Getting reference Mesh
+			USkeletalMesh* ReferenceSkeletalMesh = CustomNodeTable->GetColumnDefaultAssetByType<USkeletalMesh>(ColumnName);
+
+			if (!ReferenceSkeletalMesh)
+			{
+				FString msg = FString::Printf(TEXT("Reference Skeletal Mesh not found for column [%s]."), *ColumnName);
+				GenerationContext.Compiler->CompilerLog(FText::FromString(msg), CustomNodeTable);
+
+				return;
+			}
+
+			int ReferenceNumLODs = 0;
+			ReferenceNumLODs = Helper_GetLODInfoArray(ReferenceSkeletalMesh).Num();
+
+			int NumLODs = Helper_GetLODInfoArray(SkeletalMesh).Num();
+
+			if (NumLODs > ReferenceNumLODs)
+			{
+				FString msg = FString::Printf(TEXT("Mesh from column [%s] row [%s] has more LODs than the reference mesh. Some will be ignored"), *ColumnName, *RowName);
+				GenerationContext.Compiler->CompilerLog(FText::FromString(msg), CustomNodeTable);
+			}
+
+			int32 LODs = NumLODs <= ReferenceNumLODs ? NumLODs : ReferenceNumLODs;
+
+			// Generating skeletal mesh columns
+			for (int32 LODIndex = 0; LODIndex < LODs; ++LODIndex)
+			{
+				int32 NumMaterials = Helper_GetImportedModel(SkeletalMesh)->LODModels[LODIndex].Sections.Num();
+				int32 ReferenceNumMaterials = Helper_GetImportedModel(ReferenceSkeletalMesh)->LODModels[LODIndex].Sections.Num();
+
+				if (NumMaterials > ReferenceNumMaterials)
+				{
+					FString msg = FString::Printf(TEXT("Mesh from column [%s] row [%s] has more Sections than the reference mesh. Some will be ignored"), *ColumnName, *RowName);
+					GenerationContext.Compiler->CompilerLog(FText::FromString(msg), CustomNodeTable);
+				}
+
+				int32 Materials = NumMaterials <= ReferenceNumMaterials ? NumMaterials : ReferenceNumMaterials;
+
+				for (int32 MatIndex = 0; MatIndex < Materials; ++MatIndex)
+				{
+					FString MutableColumnName = ColumnName + FString::Printf(TEXT(" LOD_%d "), LODIndex) + FString::Printf(TEXT("Mat_%d"), MatIndex);
+
+					CurrentColumn = MutableTable.get()->FindColumn(TCHAR_TO_ANSI(*MutableColumnName));
+
+					if (CurrentColumn == -1)
+					{
+						CurrentColumn = MutableTable->AddColumn(TCHAR_TO_ANSI(*MutableColumnName), mu::TABLE_COLUMN_TYPE::TCT_MESH);
+					}
+
+					mu::MeshPtr MutableMesh = GenerateMutableMesh(SkeletalMesh, LODIndex, MatIndex, GenerationContext, CustomNodeTable);
+
+					if (MutableMesh)
+					{
+						if (!AnimBPAssetTag.IsEmpty())
+						{
+							int32 CurrentTagCount = MutableMesh->GetTagCount();
+							MutableMesh->SetTagCount(CurrentTagCount + 1);
+							MutableMesh->SetTag(CurrentTagCount, TCHAR_TO_ANSI(*AnimBPAssetTag));
+						}
+
+						for (const FGameplayTag& Tag : GameplayTags)
+						{
+							const FString AnimBPTag = GenerateGameplayTag(Tag.ToString());
+							const int32 CurrentTagCount = MutableMesh->GetTagCount();
+							MutableMesh->SetTagCount(CurrentTagCount + 1);
+							MutableMesh->SetTag(CurrentTagCount, TCHAR_TO_ANSI(*AnimBPTag));
+						}
+
+						MutableTable->SetCell(CurrentColumn, RowIdx, MutableMesh.get());
+					}
+					else
+					{
+						FString msg = FString::Printf(TEXT("Error converting skeletal mesh LOD %d, Material %d from column [%s] row [%s] to mutable."),
+							LODIndex, MatIndex, *ColumnName, *RowName);
+						GenerationContext.Compiler->CompilerLog(FText::FromString(msg), CustomNodeTable);
+					}
+				}
+			}
+		}
+
+		else if (UStaticMesh* StaticMesh = Cast<UStaticMesh>(Object))
+		{
+			// Getting reference Mesh
+			UStaticMesh* ReferenceStaticMesh = CustomNodeTable->GetColumnDefaultAssetByType<UStaticMesh>(ColumnName);
+
+			if (!ReferenceStaticMesh)
+			{
+				FString msg = FString::Printf(TEXT("Reference Static Mesh not found for column [%s]."), *ColumnName);
+				GenerationContext.Compiler->CompilerLog(FText::FromString(msg), CustomNodeTable);
+
+				return;
+			}
+
+			int NumLODs = StaticMesh->GetRenderData()->LODResources.Num();
+			int ReferenceNumLODs = ReferenceStaticMesh->GetRenderData()->LODResources.Num();
+
+			if (NumLODs > ReferenceNumLODs)
+			{
+				FString msg = FString::Printf(TEXT("Mesh from column [%s] row [%s] has more LODs than the reference mesh. Some will be ignored"), *ColumnName, *RowName);
+				GenerationContext.Compiler->CompilerLog(FText::FromString(msg), CustomNodeTable);
+			}
+
+			int32 LODs = NumLODs <= ReferenceNumLODs ? NumLODs : ReferenceNumLODs;
+
+			for (int32 LODIndex = 0; LODIndex < LODs; ++LODIndex)
+			{
+				int32 NumMaterials = StaticMesh->GetRenderData()->LODResources[LODIndex].Sections.Num();
+				int32 ReferenceNumMaterials = ReferenceStaticMesh->GetRenderData()->LODResources[LODIndex].Sections.Num();
+
+				if (NumMaterials > ReferenceNumMaterials)
+				{
+					FString msg = FString::Printf(TEXT("Mesh from column [%s] row [%s] has more Sections than the reference mesh. Some will be ignored"), *ColumnName, *RowName);
+					GenerationContext.Compiler->CompilerLog(FText::FromString(msg), CustomNodeTable);
+				}
+
+				int32 Materials = NumMaterials <= ReferenceNumMaterials ? NumMaterials : ReferenceNumMaterials;
+
+				for (int32 MatIndex = 0; MatIndex < Materials; ++MatIndex)
+				{
+					FString MutableColumnName = ColumnName + FString::Printf(TEXT(" LOD_%d "), LODIndex) + FString::Printf(TEXT("Mat_%d"), MatIndex);
+
+					CurrentColumn = MutableTable.get()->FindColumn(TCHAR_TO_ANSI(*MutableColumnName));
+
+					if (CurrentColumn == -1)
+					{
+						CurrentColumn = MutableTable->AddColumn(TCHAR_TO_ANSI(*MutableColumnName), mu::TABLE_COLUMN_TYPE::TCT_MESH);
+					}
+
+					mu::MeshPtr MutableMesh = GenerateMutableMesh(StaticMesh, LODIndex, MatIndex, GenerationContext, CustomNodeTable);
+
+					if (MutableMesh)
+					{
+						MutableTable->SetCell(CurrentColumn, RowIdx, MutableMesh.get());
+					}
+					else
+					{
+						FString msg = FString::Printf(TEXT("Error converting skeletal mesh LOD %d, Material %d from column [%s] row [%s] to mutable."),
+							LODIndex, MatIndex, *ColumnName, *RowName);
+
+						GenerationContext.Compiler->CompilerLog(FText::FromString(msg), CustomNodeTable);
+					}
+				}
+			}
+		}
+
+		else if (UTexture2D* Texture = Cast<UTexture2D>(Object))
+		{
+			// Getting column index from column name
+			CurrentColumn = MutableTable->FindColumn(TCHAR_TO_ANSI(*ColumnName));
+
+			if (CurrentColumn == -1)
+			{
+				CurrentColumn = MutableTable->AddColumn(TCHAR_TO_ANSI(*ColumnName), mu::TABLE_COLUMN_TYPE::TCT_IMAGE);
+			}
+
+			GenerationContext.ArrayTextureUnrealToMutableTask.Add(FTextureUnrealToMutableTask(MutableTable, Texture, CustomNodeTable, CurrentColumn, RowIdx));
+		}
+
+		else if (UMaterialInstance* Material = Cast<UMaterialInstance>(Object))
+		{
+			UMaterialInstance* ReferenceMaterial = CustomNodeTable->GetColumnDefaultAssetByType<UMaterialInstance>(ColumnName);
+			
+			if (!GenerationContext.GeneratedParametersInTables.Contains(CustomNodeTable))
+			{
+				GenerationContext.GeneratedParametersInTables.Add(CustomNodeTable);
+			}
+
+			TArray<FGuid>& Parameters = GenerationContext.GeneratedParametersInTables[CustomNodeTable];
+
+			if (!ReferenceMaterial)
+			{
+				FString msg = FString::Printf(TEXT("Reference Material not found for column [%s]."), *ColumnName);
+				GenerationContext.Compiler->CompilerLog(FText::FromString(msg), CustomNodeTable);
+
+				return;
+			}
+
+			if (ReferenceMaterial->GetMaterial() == Material->GetMaterial())
+			{
+				// Getting parameter Guids
+				TArray<FMaterialParameterInfo > ParameterInfos;
+				TArray<FGuid> ParameterGuids;
+
+				ReferenceMaterial->GetAllParameterInfoOfType(EMaterialParameterType::Texture, ParameterInfos, ParameterGuids);
+
+				// Number of modified texture parameters in the Material Instance
+				int32 ModifiedTextureParameters = 0;
+				
+				for (FTextureParameterValue ReferenceTexture : ReferenceMaterial->TextureParameterValues)
+				{
+					if (Cast<UTexture2D>(ReferenceTexture.ParameterValue))
+					{
+						int32 ColumnIndex;
+
+						FString TextureParameterName = ReferenceTexture.ParameterInfo.Name.ToString();
+						FString ParameterGuid;
+						
+						for (int32 i = 0; i < ParameterInfos.Num(); ++i)
+						{
+							if (ParameterInfos[i].Name == ReferenceTexture.ParameterInfo.Name)
+							{
+								ParameterGuid = ParameterGuids[i].ToString();
+								Parameters.Add(ParameterGuids[i]);
+
+								break;
+							}
+						}
+
+						// Getting column index from parameter name
+						ColumnIndex = MutableTable->FindColumn(TCHAR_TO_ANSI(*ParameterGuid));
+
+						if (ColumnIndex == -1)
+						{
+							// If there is no column with the parameters name, we generate a new one
+							ColumnIndex = MutableTable->AddColumn(TCHAR_TO_ANSI(*ParameterGuid), mu::TABLE_COLUMN_TYPE::TCT_IMAGE);
+						}
+
+						UTexture* OutTexture = nullptr;
+						UTexture2D* TextureToConvert = nullptr;
+						
+						// Getting the parameter value from the parent material
+						if (Material->GetMaterial()->GetTextureParameterValue(ReferenceTexture.ParameterInfo.Name, OutTexture))
+						{
+							TextureToConvert = Cast<UTexture2D>(OutTexture);
+						}
+
+						// Getting the parameter value from the instance if it has been modified
+						for (FTextureParameterValue InstanceTexture : Material->TextureParameterValues)
+						{
+							if (InstanceTexture.ParameterInfo.Name == ReferenceTexture.ParameterInfo.Name)
+							{
+								TextureToConvert = Cast<UTexture2D>(InstanceTexture.ParameterValue);
+								ModifiedTextureParameters++;
+								
+								break;
+							}
+						}
+
+						if (TextureToConvert)
+						{
+							GenerationContext.ArrayTextureUnrealToMutableTask.Add(FTextureUnrealToMutableTask(MutableTable, TextureToConvert, CustomNodeTable, ColumnIndex, RowIdx));
+						}
+						else
+						{
+							FString msg = FString::Printf(TEXT("Didn't find Material Texture Parameter [%s] from parent Material of column [%s] row [%s]."), *TextureParameterName, *ColumnName, *RowName);
+							GenerationContext.Compiler->CompilerLog(FText::FromString(msg), CustomNodeTable);
+						}
+					}
+				}
+
+				// Checking if all modifiable textures of the material instance have been modified
+				if (Material->TextureParameterValues.Num() > ModifiedTextureParameters)
+				{
+					int32 ParametersDiff = Material->TextureParameterValues.Num() - ModifiedTextureParameters;
+					FString msg =
+						FString::Printf(TEXT("Material Instance [%s] from column [%s] row [%s] has %d modifiable Textures that will not be modified, they are non-modifiable parameters in the Default Material Instance"),
+						*Material->GetName(), *ColumnName, *RowName, ParametersDiff);
+					GenerationContext.Compiler->CompilerLog(FText::FromString(msg), CustomNodeTable);
+				}
+
+				ParameterInfos.Empty();
+				ParameterGuids.Empty();
+
+				ReferenceMaterial->GetAllParameterInfoOfType(EMaterialParameterType::Vector, ParameterInfos, ParameterGuids);
+
+				// Number of modified Vector parameters in the Material Instance
+				int32 ModifiedVectorParameters = 0;
+
+				for (FVectorParameterValue ReferenceVector : ReferenceMaterial->VectorParameterValues)
+				{
+					int32 ColumnIndex;
+
+					FString ParameterGuid;
+
+					for (int32 i = 0; i < ParameterInfos.Num(); ++i)
+					{
+						if (ParameterInfos[i].Name == ReferenceVector.ParameterInfo.Name)
+						{
+							ParameterGuid = ParameterGuids[i].ToString();
+							Parameters.Add(ParameterGuids[i]);
+
+							break;
+						}
+					}
+					
+					// Getting column index from parameter name
+					ColumnIndex = MutableTable->FindColumn(TCHAR_TO_ANSI(*ParameterGuid));
+
+					if (ColumnIndex == -1)
+					{
+						// If there is no column with the parameters name, we generate a new one
+						ColumnIndex = MutableTable->AddColumn(TCHAR_TO_ANSI(*ParameterGuid), mu::TABLE_COLUMN_TYPE::TCT_COLOUR);
+					}
+
+					// Getting the parameter value from the parent material
+					FLinearColor VectorValue;
+
+					Material->GetMaterial()->GetVectorParameterValue(ReferenceVector.ParameterInfo.Name, VectorValue);
+					
+					// Getting the parameter value from the instance if it has been modified
+					for (FVectorParameterValue InstanceVector : Material->VectorParameterValues)
+					{
+						if (InstanceVector.ParameterInfo.Name == ReferenceVector.ParameterInfo.Name)
+						{
+							VectorValue = InstanceVector.ParameterValue;
+							ModifiedVectorParameters++;
+							
+							break;
+						}
+					}
+
+					// Setting cell value
+					MutableTable->SetCell(ColumnIndex, RowIdx, VectorValue.R, VectorValue.G, VectorValue.B);
+				}
+
+				// Checking if all modifiable vectors of the material instance have been modified
+				if (Material->VectorParameterValues.Num() > ModifiedVectorParameters)
+				{
+					int32 ParametersDiff = Material->VectorParameterValues.Num() - ModifiedVectorParameters;
+
+					FString msg =
+						FString::Printf(TEXT("Material Instance [%s] from column [%s] row [%s] has %d modifiable Vectors that will not be modified, they are non-modifiable parameters in the Default Material Instance"),
+							*Material->GetName(), *ColumnName, *RowName, ParametersDiff);
+					GenerationContext.Compiler->CompilerLog(FText::FromString(msg), CustomNodeTable);
+				}
+
+				// Number of modified Float parameters in the Material Instance
+				int32 ModifiedFloatParameters = 0;
+
+				ParameterInfos.Empty();
+				ParameterGuids.Empty();
+
+				ReferenceMaterial->GetAllParameterInfoOfType(EMaterialParameterType::Scalar, ParameterInfos, ParameterGuids);
+
+				for (FScalarParameterValue ReferenceScalar : ReferenceMaterial->ScalarParameterValues)
+				{
+					int32 ColumnIndex;
+
+					FString ParameterGuid;
+
+					for (int32 i = 0; i < ParameterInfos.Num(); ++i)
+					{
+						if (ParameterInfos[i].Name == ReferenceScalar.ParameterInfo.Name)
+						{
+							ParameterGuid = ParameterGuids[i].ToString();
+							Parameters.Add(ParameterGuids[i]);
+							
+							break;
+						}
+					}
+
+					// Getting column index from parameter name
+					ColumnIndex = MutableTable->FindColumn(TCHAR_TO_ANSI(*ParameterGuid));
+
+					if (ColumnIndex == -1)
+					{
+						// If there is no column with the parameters name, we generate a new one
+						ColumnIndex = MutableTable->AddColumn(TCHAR_TO_ANSI(*ParameterGuid), mu::TABLE_COLUMN_TYPE::TCT_SCALAR);
+					}
+
+					// Getting the parameter value from the parent material
+					float ScalarValue;
+					Material->GetMaterial()->GetScalarParameterValue(ReferenceScalar.ParameterInfo.Name, ScalarValue);
+					
+					// Getting the parameter value from the instance if it has been modified
+					for (FScalarParameterValue InstanceScalar : Material->ScalarParameterValues)
+					{
+						if (InstanceScalar.ParameterInfo.Name == ReferenceScalar.ParameterInfo.Name)
+						{
+							ScalarValue = InstanceScalar.ParameterValue;
+							ModifiedFloatParameters++;
+							
+							break;
+						}
+					}
+
+					// Setting cell value
+					MutableTable->SetCell(ColumnIndex, RowIdx, ScalarValue);
+				}
+
+				// Checking if all modifiable floats of the material instance have been modified
+				if (Material->ScalarParameterValues.Num() > ModifiedFloatParameters)
+				{
+					int32 ParametersDiff = Material->ScalarParameterValues.Num() - ModifiedFloatParameters;
+
+					FString msg =
+						FString::Printf(TEXT("Material Instance [%s] from column [%s] row [%s] has %d modifiable Scalars that will not be modified, they are non-modifiable parameters in the Default Material Instance"),
+							*Material->GetName(), *ColumnName, *RowName, ParametersDiff);
+					GenerationContext.Compiler->CompilerLog(FText::FromString(msg), CustomNodeTable);
+				}
+			}
+			else
+			{
+				FString msg = FString::Printf(TEXT("Material from column [%s] row [%s] is a diferent instance than the Reference Material of the table."), *ColumnName, *RowName);
+				GenerationContext.Compiler->CompilerLog(FText::FromString(msg), CustomNodeTable);
+			}
+		}
+	}
+
+	else if (const FStructProperty* StructProperty = CastField<FStructProperty>(Property))
+	{
+		if (StructProperty->Struct == TBaseStructure<FLinearColor>::Get())
+		{
+			CurrentColumn = MutableTable->FindColumn(TCHAR_TO_ANSI(*ColumnName));
+
+			if (CurrentColumn == -1)
+			{
+				CurrentColumn = MutableTable->AddColumn(TCHAR_TO_ANSI(*ColumnName), mu::TABLE_COLUMN_TYPE::TCT_COLOUR);
+			}
+
+			// Setting cell value
+			FLinearColor Value = *(FLinearColor*)CellData;
+			MutableTable->SetCell(CurrentColumn, RowIdx, Value.R, Value.G, Value.B);
+		}
+	}
+
+	else if (const FNumericProperty* NumProperty = CastField<FFloatProperty>(Property))
+	{
+		CurrentColumn = MutableTable->FindColumn(TCHAR_TO_ANSI(*ColumnName));
+
+		if (CurrentColumn == -1)
+		{
+			CurrentColumn = MutableTable->AddColumn(TCHAR_TO_ANSI(*ColumnName), mu::TABLE_COLUMN_TYPE::TCT_SCALAR);
+		}
+
+		// Setting cell value
+		float Value = NumProperty->GetFloatingPointPropertyValue(CellData);
+		MutableTable->SetCell(CurrentColumn, RowIdx, Value);
+	}
+}
+
+
+mu::TablePtr GenerateMutableSourceTable(const FString& TableName, const UEdGraphPin* Pin, FMutableGraphGenerationContext& GenerationContext)
+{
+	SCOPED_PIN_DATA(GenerationContext, Pin)
+	
+	if (mu::TablePtr* Result = GenerationContext.GeneratedTables.Find(TableName))
+	{
+		return *Result;
+	}
+
+	mu::TablePtr MutableTable = new mu::Table();
+	
+	UCustomizableObjectNode* Node = CastChecked<UCustomizableObjectNode>(Pin->GetOwningNode());
+	if (Node->IsNodeOutDatedAndNeedsRefresh())
+	{
+		Node->SetRefreshNodeWarning();
+	}
+
+	if (const UCustomizableObjectNodeTable* TypedTable = Cast<UCustomizableObjectNodeTable>(Node))
+	{
+		UDataTable* Table = TypedTable->Table;
+
+		if (!Table)
+		{
+			FString msg = "Couldn't find the Data Table asset in the Node.";
+			GenerationContext.Compiler->CompilerLog(FText::FromString(msg), Node);
+
+			return MutableTable;
+		}
+
+		const UScriptStruct* TableStruct = Table->GetRowStruct();
+
+		if (TableStruct)
+		{
+			// Getting names of the column headers
+			TArray<FString> ColumnHeaders = Table->GetColumnTitles();
+
+			// Getting names of the rows to access the information
+			TArray<FName> RowNames = Table->GetRowNames();
+
+			// index to iterate the column headers
+			int32 ColumnHeaderIndex = 1;
+
+			// Adding and filling Name Column
+			MutableTable->AddColumn("Name", mu::TABLE_COLUMN_TYPE::TCT_STRING);
+
+			// Add metadata
+			FParameterUIData ParameterUIData(
+				TypedTable->ParameterName,
+				TypedTable->ParamUIMetadata,
+				EMutableParameterType::Int);
+			
+			ParameterUIData.IntegerParameterGroupType = TypedTable->bAddNoneOption ? ECustomizableObjectGroupType::COGT_ONE_OR_NONE : ECustomizableObjectGroupType::COGT_ONE;
+
+			// Adding a None Option
+			MutableTable->SetNoneOption(TypedTable->bAddNoneOption);
+
+			for (int32 i = 0; i < RowNames.Num(); ++i)
+			{
+				MutableTable->AddRow(i);
+				FString RowName= RowNames[i].ToString();
+				MutableTable->SetCell(0, i, TCHAR_TO_ANSI(*RowName));
+				ParameterUIData.ArrayIntegerParameterOption.Add(FIntegerParameterUIData(
+					RowName,
+					FMutableParamUIMetadata()));
+			}
+
+			for (TFieldIterator<FProperty> It(TableStruct); It; ++It)
+			{
+				FProperty* Property = *It;
+
+				for (int32 RowIndex = 0; RowIndex < RowNames.Num(); ++RowIndex)
+				{
+					// Getting Row Data
+					uint8* RowData = Table->FindRowUnchecked(RowNames[RowIndex]);
+
+					if (RowData)
+					{
+						// Getting Cell Data
+						uint8* CellData = Property->ContainerPtrToValuePtr<uint8>(RowData, 0);
+
+						if (CellData)
+						{
+							FillTableColumn(Pin, MutableTable, ColumnHeaders[ColumnHeaderIndex], RowNames[RowIndex].ToString(), RowIndex, CellData, Property, GenerationContext);
+						}
+					}
+				}
+
+				ColumnHeaderIndex++;
+			}
+
+			GenerationContext.ParameterUIDataMap.Add(TypedTable->ParameterName, ParameterUIData);
+		}
+	}
+	
+	else
+	{
+		GenerationContext.Compiler->CompilerLog(LOCTEXT("UnimplementedNode", "Node type not implemented yet."), Node);
+	}
+
+	GenerationContext.GeneratedTables.Add(TableName, MutableTable);
+
+	return MutableTable;
+}
diff --git a/Engine/Plugins/Experimental/Mutable/Source/CustomizableObjectEditor/Private/GenerateMutableSource/GenerateMutableSourceTable.h b/Engine/Plugins/Experimental/Mutable/Source/CustomizableObjectEditor/Private/GenerateMutableSource/GenerateMutableSourceTable.h
new file mode 100644
--- /dev/null
+++ b/Engine/Plugins/Experimental/Mutable/Source/CustomizableObjectEditor/Private/GenerateMutableSource/GenerateMutableSourceTable.h
@@ -0,0 +1,8 @@
+// Copyright Epic Games, Inc. All Rights Reserved.
+
+#pragma once
+
+#include "GenerateMutableSource/GenerateMutableSource.h"
+
+
+mu::TablePtr GenerateMutableSourceTable(const FString& TableName, const UEdGraphPin* Pin, FMutableGraphGenerationContext& GenerationContext);
diff --git a/Engine/Plugins/Experimental/Mutable/Source/CustomizableObjectEditor/Private/GraphTraversal.cpp b/Engine/Plugins/Experimental/Mutable/Source/CustomizableObjectEditor/Private/GraphTraversal.cpp
new file mode 100644
--- /dev/null
+++ b/Engine/Plugins/Experimental/Mutable/Source/CustomizableObjectEditor/Private/GraphTraversal.cpp
@@ -0,0 +1,490 @@
+// Copyright Epic Games, Inc. All Rights Reserved.
+
+#include "GraphTraversal.h"
+#include "CustomizableObject.h"
+#include "CustomizableObjectPin.h"
+#include "Nodes/CustomizableObjectNodeEnumParameter.h"
+#include "Nodes/CustomizableObjectNodeExtendMaterial.h"
+#include "Nodes/CustomizableObjectNodeExternalPin.h"
+#include "Nodes/CustomizableObjectNodeMaterialVariation.h"
+#include "Nodes/CustomizableObjectNodeMeshGeometryOperation.h"
+#include "Nodes/CustomizableObjectNodeMeshMorph.h"
+#include "Nodes/CustomizableObjectNodeMeshMorphStackApplication.h"
+#include "Nodes/CustomizableObjectNodeMeshMorphStackDefinition.h"
+#include "Nodes/CustomizableObjectNodeMeshReshape.h"
+#include "Nodes/CustomizableObjectNodeMeshSwitch.h"
+#include "Nodes/CustomizableObjectNodeMeshVariation.h"
+#include "Nodes/CustomizableObjectNodeObject.h"
+#include "Nodes/CustomizableObjectNodeSkeletalMesh.h"
+#include "Nodes/CustomizableObjectNodeStaticMesh.h"
+#include "Nodes/CustomizableObjectNodeTable.h"
+
+
+/** Function to follow an External Pin node.
+ *
+ * @param Node Possible node to follow.
+ * @param PinsToVisit Discovered pins (output parameter).
+ * @return True if the node was a External Pin node (node followed).
+ */
+bool CastExternalPinNode(const UEdGraphNode& Node, TArray<const UEdGraphPin*>& PinsToVisit)
+{
+	const UCustomizableObjectNodeExternalPin* ExternalPinNode = Cast<UCustomizableObjectNodeExternalPin>(&Node);
+	if (ExternalPinNode)
+	{
+		if (const UCustomizableObjectNodeExposePin* ExposePinNode = ExternalPinNode->GetNodeExposePin())
+		{
+			if (const UEdGraphPin* ExposePin = ExposePinNode->InputPin())
+			{
+				PinsToVisit.Add(ExposePin);
+			}
+		}
+	}
+
+	return ExternalPinNode != nullptr;
+}
+
+
+/** Function to follow an Expose Pin node.
+ *
+ * @param Node Possible node to follow.
+ * @param PinsToVisit Discovered pins (output parameter).
+ * @return True if the node was a Expose Pin node (node followed).
+ */
+bool CastExposePinNode(const UEdGraphNode& Node, TArray<const UEdGraphPin*>& PinsToVisit)
+{
+	const UCustomizableObjectNodeExposePin* ExposePinNode = Cast<UCustomizableObjectNodeExposePin>(&Node);
+	if (ExposePinNode)
+	{
+		const UCustomizableObjectNodeExternalPin* NodeExternalPin = nullptr;
+		for (TObjectIterator<UCustomizableObjectNodeExternalPin> It; It; ++It)
+		{
+			if ((*It)->GetNodeExposePin() == ExposePinNode)
+			{
+				NodeExternalPin = *It;
+				break;
+			}
+		}
+		
+		if (NodeExternalPin)
+		{
+			if (const UEdGraphPin* ExternalPin = NodeExternalPin->GetExternalPin())
+			{
+				PinsToVisit.Add(ExternalPin);
+			}
+		}
+	}
+
+	return ExposePinNode != nullptr;
+}
+
+
+/** Follow the given pin returning its connected pin.
+ *
+ * - Skips all orphan pins.
+ * - Follows External Pin nodes.
+ *
+ * @param Pin Pin to follow.
+ * @param CycleDetected If provided, it will set to true if a cycle has been found.
+ * @param CastNode Function to follow External Pin nodes (following an output pin) or Expose Pin nodes (following an input pin).
+ * @param Direction Direction of the pin to explore.
+ */
+TArray<UEdGraphPin*> FollowPinArray(const UEdGraphPin& Pin, bool* CycleDetected, const EEdGraphPinDirection Direction, const TFunction<bool(const UEdGraphNode&, TArray<const UEdGraphPin*>&)> CastNode)
+{
+	check(Pin.Direction == Direction); // Can only follow input pins. To follow output pins see FollowOutputPin, but be aware of its limitations!
+
+	if (CycleDetected)
+	{
+		*CycleDetected = false;
+	}
+	
+	TArray<UEdGraphPin*> Result;
+	
+	TSet<const UEdGraphPin*> Visited;
+
+	TArray<const UEdGraphPin*> PinsToVisit;
+	PinsToVisit.Add(&Pin);
+	while (PinsToVisit.Num())
+	{
+		const UEdGraphPin& CurrentPin = *PinsToVisit.Pop();
+
+		if (IsPinOrphan(CurrentPin))
+		{
+			continue;
+		}
+
+		Visited.FindOrAdd(&CurrentPin, CycleDetected);
+		if (CycleDetected && *CycleDetected)
+		{
+			continue;
+		}
+
+		for (UEdGraphPin* LinkedPin : CurrentPin.LinkedTo)
+		{
+			if (!IsPinOrphan(*LinkedPin))
+			{
+				if (!CastNode(*LinkedPin->GetOwningNode(), PinsToVisit))
+				{
+					Result.Add(LinkedPin);
+				}
+			}
+		}
+		
+	}
+
+	return Result;
+}
+
+
+TArray<UEdGraphPin*> FollowInputPinArray(const UEdGraphPin& Pin, bool* CycleDetected)
+{
+	return FollowPinArray(Pin, CycleDetected, EGPD_Input, &CastExternalPinNode);
+}
+
+
+UEdGraphPin* FollowInputPin(const UEdGraphPin& Pin, bool* CycleDetected)
+{
+	TArray<UEdGraphPin*> Result = FollowInputPinArray(Pin, CycleDetected);
+	check(Result.Num() <= 1); // Use FollowPinImmersive if the pin can have more than one input.
+
+	if (!Result.IsEmpty())
+	{
+		return Result[0];
+	}
+	else
+	{
+		return nullptr;
+	}
+}
+
+
+TArray<UEdGraphPin*> FollowOutputPinArray(const UEdGraphPin& Pin, bool* CycleDetected)
+{
+	return FollowPinArray(Pin, CycleDetected, EGPD_Output, &CastExposePinNode);
+}
+
+
+UEdGraphPin* FollowOutputPin(const UEdGraphPin& Pin, bool* CycleDetected)
+{
+	TArray<UEdGraphPin*> Result = FollowOutputPinArray(Pin, CycleDetected);
+	check(Result.Num() <= 1); // Use FollowPinImmersive if the pin can have more than one input.
+
+	if (!Result.IsEmpty())
+	{
+		return Result[0];
+	}
+	else
+	{
+		return nullptr;
+	}
+}
+
+
+UCustomizableObjectNodeObject* GetRootNode(UCustomizableObject* Object, bool& bOutMultipleBaseObjectsFound)
+{
+	// Look for the base object node
+	UCustomizableObjectNodeObject* Root = nullptr;
+	TArray<UCustomizableObjectNodeObject*> ObjectNodes;
+	Object->Source->GetNodesOfClass<UCustomizableObjectNodeObject>(ObjectNodes);
+
+	bOutMultipleBaseObjectsFound = false;
+
+	for (TArray<UCustomizableObjectNodeObject*>::TIterator It(ObjectNodes); It; ++It)
+	{
+		if ((*It)->bIsBase)
+		{
+			if (Root)
+			{
+				bOutMultipleBaseObjectsFound = true;
+				break;
+			}
+			else
+			{
+				Root = *It;
+			}
+		}
+	}
+
+	return Root;
+}
+
+
+bool GetParentsUntilRoot(UCustomizableObject* Object, TArray<UCustomizableObjectNodeObject*>& ArrayNodeObject, TArray<UCustomizableObject*>& ArrayCustomizableObject)
+{
+	bool MultipleBaseObjectsFound;
+	UCustomizableObjectNodeObject* Root = GetRootNode(Object, MultipleBaseObjectsFound);
+
+	bool bSuccess = true;
+
+	if (!MultipleBaseObjectsFound && (Root != nullptr))
+	{
+		if (!ArrayCustomizableObject.Contains(Object))
+		{
+			ArrayNodeObject.Add(Root);
+			ArrayCustomizableObject.Add(Object);
+		}
+		else
+		{
+			// This object has already been visted which means that there is a Cycle between Customizable Objects
+			return false;
+		}
+
+		if (Root->ParentObject != nullptr)
+		{
+			bSuccess = GetParentsUntilRoot(Root->ParentObject, ArrayNodeObject, ArrayCustomizableObject);
+
+			UCustomizableObjectNodeObject* ParentRoot = GetRootNode(Root->ParentObject, MultipleBaseObjectsFound);
+
+			if (!MultipleBaseObjectsFound && (ParentRoot != nullptr))
+			{
+				Root->SetMeshComponentNumFromParent(ParentRoot->NumMeshComponents);
+			}
+		}
+	}
+
+	return bSuccess;
+}
+
+
+bool HasCandidateAsParent(UCustomizableObjectNodeObject* Node, UCustomizableObject* ParentCandidate)
+{
+	if (Node->ParentObject == ParentCandidate)
+	{
+		return true;
+	}
+
+	if (Node->ParentObject != nullptr)
+	{
+		bool MultipleBaseObjectsFound;
+		UCustomizableObjectNodeObject* ParentNodeObject = GetRootNode(Node->ParentObject, MultipleBaseObjectsFound);
+
+		if ((ParentNodeObject->ParentObject == nullptr) || MultipleBaseObjectsFound)
+		{
+			return false;
+		}
+		else
+		{
+			return HasCandidateAsParent(ParentNodeObject, ParentCandidate);
+		}
+	}
+
+	return false;
+}
+
+
+UCustomizableObject* GetFullGraphRootObject(UCustomizableObjectNodeObject* Node, TArray<UCustomizableObject*>& VisitedObjects)
+{
+	if (Node->ParentObject != nullptr)
+	{
+		VisitedObjects.Add(Node->ParentObject);
+
+		bool MultipleBaseObjectsFound;
+		UCustomizableObjectNodeObject* Root = GetRootNode(Node->ParentObject, MultipleBaseObjectsFound);
+
+		if (Root->ParentObject == nullptr)
+		{
+			if (MultipleBaseObjectsFound)
+			{
+				return nullptr;
+			}
+			else
+			{
+				return Node->ParentObject;
+			}
+		}
+		else
+		{
+			if (VisitedObjects.Contains(Root->ParentObject))
+			{
+				//There is a cycle
+				return nullptr;
+			}
+			else
+			{
+				return GetFullGraphRootObject(Root, VisitedObjects);
+			}
+		}
+	}
+
+	return nullptr;
+}
+
+
+UCustomizableObjectNodeObject* GetFullGraphRootNodeObject(UCustomizableObjectNodeObject* Node, TArray<UCustomizableObject*>& VisitedObjects)
+{
+	if (Node->ParentObject != nullptr)
+	{
+		VisitedObjects.Add(Node->ParentObject);
+
+		bool MultipleBaseObjectsFound;
+		UCustomizableObjectNodeObject* Root = GetRootNode(Node->ParentObject, MultipleBaseObjectsFound);
+
+		if (Root->ParentObject == nullptr)
+		{
+			if (MultipleBaseObjectsFound)
+			{
+				return nullptr;
+			}
+			else
+			{
+				return Root;
+			}
+		}
+		else
+		{
+			if (VisitedObjects.Contains(Root->ParentObject))
+			{
+				//There is a cycle
+				return nullptr;
+			}
+			else
+			{
+				return GetFullGraphRootNodeObject(Root, VisitedObjects);
+			}
+		}
+	}
+
+	return nullptr;
+}
+
+
+const UEdGraphPin* FindMeshBaseSource(const UEdGraphPin& Pin, const bool bOnlyLookForStaticMesh)
+{
+	check(Pin.Direction == EGPD_Output);
+	check(Pin.PinType.PinCategory == UEdGraphSchema_CustomizableObject::PC_Mesh || Pin.PinType.PinCategory == UEdGraphSchema_CustomizableObject::PC_Material);
+	
+	const UEdGraphNode* Node = Pin.GetOwningNode();
+	check(Node);
+
+	if (Cast<UCustomizableObjectNodeSkeletalMesh>(Node))
+	{
+		if (!bOnlyLookForStaticMesh)
+		{
+			return &Pin;
+		}
+	}
+
+	else if (Cast<UCustomizableObjectNodeStaticMesh>(Node))
+	{
+		return &Pin;
+	}
+
+	else if (const UCustomizableObjectNodeMeshGeometryOperation* TypedNodeGeom = Cast<UCustomizableObjectNodeMeshGeometryOperation>(Node))
+	{
+		if (const UEdGraphPin* ConnectedPin = FollowInputPin(*TypedNodeGeom->MeshAPin()))
+		{
+			return FindMeshBaseSource(*ConnectedPin, bOnlyLookForStaticMesh);
+		}
+	}
+
+	else if (const UCustomizableObjectNodeMeshReshape* TypedNodeReshape = Cast<UCustomizableObjectNodeMeshReshape>(Node))
+	{
+		if (const UEdGraphPin* ConnectedPin = FollowInputPin(*TypedNodeReshape->BaseMeshPin()))
+		{
+			return FindMeshBaseSource(*ConnectedPin, bOnlyLookForStaticMesh);
+		}
+	}
+
+	else if (const UCustomizableObjectNodeMeshMorph* TypedNodeMorph = Cast<UCustomizableObjectNodeMeshMorph>(Node))
+	{
+		if (const UEdGraphPin* ConnectedPin = FollowInputPin(*TypedNodeMorph->MeshPin()))
+		{
+			return FindMeshBaseSource(*ConnectedPin, bOnlyLookForStaticMesh);
+		}
+	}
+
+	else if (const UCustomizableObjectNodeMeshSwitch* TypedNodeSwitch = Cast<UCustomizableObjectNodeMeshSwitch>(Node))
+	{
+		if (const UEdGraphPin* EnumParameterPin = FollowInputPin(*TypedNodeSwitch->SwitchParameter()))
+		{
+			if (const UCustomizableObjectNodeEnumParameter* EnumNode = Cast<UCustomizableObjectNodeEnumParameter>(EnumParameterPin->GetOwningNode()))
+			{
+				if (const UEdGraphPin* DefaultPin = TypedNodeSwitch->GetElementPin(EnumNode->DefaultIndex))
+				{
+					if (const UEdGraphPin* ConnectedPin = FollowInputPin(*DefaultPin))
+					{
+						return FindMeshBaseSource(*ConnectedPin, bOnlyLookForStaticMesh);
+					}
+				}
+			}
+		}
+	}
+
+	else if (const UCustomizableObjectNodeMeshVariation* TypedNodeMeshVar = Cast<UCustomizableObjectNodeMeshVariation>(Node))
+	{
+		if (const UEdGraphPin* ConnectedPin = FollowInputPin(*TypedNodeMeshVar->DefaultPin()))
+		{
+			return FindMeshBaseSource(*ConnectedPin, bOnlyLookForStaticMesh);
+		}
+
+		for (int32 i = 0; i < TypedNodeMeshVar->GetNumVariations(); ++i)
+		{
+			if (const UEdGraphPin* ConnectedPin = FollowInputPin(*TypedNodeMeshVar->VariationPin(i)))
+			{
+				return FindMeshBaseSource(*ConnectedPin, bOnlyLookForStaticMesh);
+			}
+		}
+	}
+
+	else if (const UCustomizableObjectNodeMaterial* TypedNodeMat = Cast<UCustomizableObjectNodeMaterial>(Node))
+	{
+		if (const UEdGraphPin* ConnectedPin = FollowInputPin(*TypedNodeMat->GetMeshPin()))
+		{
+			return FindMeshBaseSource(*ConnectedPin, bOnlyLookForStaticMesh);
+		}
+	}
+
+	else if (const UCustomizableObjectNodeMaterialVariation* TypedNodeMatVar = Cast<UCustomizableObjectNodeMaterialVariation>(Node))
+	{
+		if (const UEdGraphPin* ConnectedPin = FollowInputPin(*TypedNodeMatVar->DefaultPin()))
+		{
+			return FindMeshBaseSource(*ConnectedPin, bOnlyLookForStaticMesh);
+		}
+	}
+
+	else if (const UCustomizableObjectNodeExtendMaterial* TypedNodeExtend = Cast<UCustomizableObjectNodeExtendMaterial>(Node))
+	{
+		if (const UEdGraphPin* ConnectedPin = FollowInputPin(*TypedNodeExtend->AddMeshPin()))
+		{
+			return FindMeshBaseSource(*ConnectedPin, bOnlyLookForStaticMesh);
+		}
+	}
+	
+	else if (const UCustomizableObjectNodeMeshMorphStackDefinition* TypedNodeMorphStackDef = Cast<UCustomizableObjectNodeMeshMorphStackDefinition>(Node))
+	{
+		if (const UEdGraphPin* ConnectedPin = FollowInputPin(*TypedNodeMorphStackDef->GetMeshPin()))
+		{
+			return FindMeshBaseSource(*ConnectedPin, bOnlyLookForStaticMesh);
+		}
+	}
+
+	else if (const UCustomizableObjectNodeMeshMorphStackApplication* TypedNodeMorphStackApp = Cast<UCustomizableObjectNodeMeshMorphStackApplication>(Node))
+	{
+		if (const UEdGraphPin* ConnectedPin = FollowInputPin(*TypedNodeMorphStackApp->GetMeshPin()))
+		{
+			return FindMeshBaseSource(*ConnectedPin, bOnlyLookForStaticMesh);
+		}
+	}
+
+	else if (const UCustomizableObjectNodeMeshReshape* NodeMeshReshape = Cast<UCustomizableObjectNodeMeshReshape>(Node))
+	{
+		if (const UEdGraphPin* ConnectedPin = FollowInputPin(*NodeMeshReshape->BaseMeshPin()))
+		{
+			return FindMeshBaseSource(*ConnectedPin, bOnlyLookForStaticMesh);
+		}
+	}
+	
+	else if (Cast<UCustomizableObjectNodeTable>(Node))
+	{
+		if (!bOnlyLookForStaticMesh)
+		{
+			return &Pin;
+		}
+	}
+
+	else
+	{
+		unimplemented(); // Case missing.
+	}
+	
+	return nullptr;
+}
diff --git a/Engine/Plugins/Experimental/Mutable/Source/CustomizableObjectEditor/Private/MutableMeshPreviewUtils.cpp b/Engine/Plugins/Experimental/Mutable/Source/CustomizableObjectEditor/Private/MutableMeshPreviewUtils.cpp
new file mode 100644
--- /dev/null
+++ b/Engine/Plugins/Experimental/Mutable/Source/CustomizableObjectEditor/Private/MutableMeshPreviewUtils.cpp
@@ -0,0 +1,809 @@
+// Copyright Epic Games, Inc. All Rights Reserved.
+
+#include "MutableMeshPreviewUtils.h"
+#include "MutableMeshBufferUtils.h"
+#include "UnrealConversionUtils.h"
+
+#include "Animation/Skeleton.h"
+#include "CustomizableObject/Private/CustomizableObjectSystemPrivate.h"
+#include "Materials/Material.h"
+#include "MutableRuntime/Private/MeshPrivate.h"
+#include "MutableRuntime/Private/OpMeshFormat.h"
+#include "MutableRuntime/Public/Mesh.h"
+#include "MutableRuntime/Public/MutableTrace.h"
+#include "MutableRuntime/Public/Parameters.h"
+#include "MutableRuntime/Public/Ptr.h"
+
+
+namespace MutableMeshPreviewUtils
+{
+	namespace
+	{
+		/** It generates a copy of the provided constant mesh but making sure the buffer structure is the one
+		 * unreal and our systems want. This is because some of the meshes we provide to mutable do not all have the
+		 * same structure.
+		 * @param InMutableMesh Mutable mesh whose buffers should be formatted following the desired buffer structure.
+		 * @return An identical mutable mesh to the provided one in terms of data but with the buffer structure that UE
+		 * requires
+		 */
+		mu::MeshPtrConst GenerateCompatibleMutableMesh(mu::MeshPtrConst InMutableMesh)
+		{
+			// 1) Generate a mutable mesh to be the one hosting the buffer format we desire (it will be required for 
+			mu::MeshPtr FormattedMutableMesh = new::mu::Mesh();
+			mu::FMeshBufferSet& FormattedVertexBuffers = FormattedMutableMesh->GetVertexBuffers();
+			mu::FMeshBufferSet& FormattedIndexBuffers = FormattedMutableMesh->GetIndexBuffers();
+			
+			// 2) Generate the buffers for that mesh with the format unreal expects
+			{
+				// 2.1)		Locate what buffers are present on the mutable constant mesh
+				const mu::FMeshBufferSet& OriginVertexBuffers = InMutableMesh->GetVertexBuffers();
+
+				// Data extracted from the original buffers to be used later when setting up the formatted mesh
+				bool bOriginDoesHaveVertexColorData = false;
+				int32 NumbOfTextCoordChannels = 0;
+				int32 MaxNumBonesPerVertex = 0;
+
+				// Contingency to avoid mismatches between the ideal semantic indices on the provided mutable mesh against the
+				// ones that should be .
+				int TextureSemanticIndices[MutableMeshBufferUtils::MaxTexCordChannelCount] = {-1,-1,-1,-1};
+				
+				// Vertex buffers
+				const int32 MutableMeshBuffersCount = OriginVertexBuffers.GetBufferCount();
+				for (int32 BufferIndex = 0; BufferIndex < MutableMeshBuffersCount; BufferIndex++)
+				{
+					const int32 BufferChannelCount = OriginVertexBuffers.GetBufferChannelCount(BufferIndex);
+					for (int32 ChannelIndex = 0; ChannelIndex < BufferChannelCount; ChannelIndex ++)
+					{
+						mu::MESH_BUFFER_SEMANTIC BufferSemantic = mu::MESH_BUFFER_SEMANTIC::MBS_NONE;
+						mu::MESH_BUFFER_FORMAT BufferFormat = mu::MESH_BUFFER_FORMAT::MBF_NONE;
+						int32 BufferComponentCount = 0;
+						int32 SemanticIndex = 0;
+
+						// Get the data from mutable that we require from the selected buffer set buffer
+						OriginVertexBuffers.GetChannel
+						(
+							BufferIndex,
+							ChannelIndex,
+							&(BufferSemantic),
+							&(SemanticIndex),
+							&(BufferFormat),
+							&(BufferComponentCount),
+							nullptr
+						);
+
+						switch (BufferSemantic)
+						{
+							case mu::MESH_BUFFER_SEMANTIC::MBS_COLOUR:
+							{
+								// Does the original mesh have vert color data?
+								bOriginDoesHaveVertexColorData = true;
+								break;
+							}
+
+							case mu::MESH_BUFFER_SEMANTIC::MBS_TEXCOORDS:
+							{
+								// Store the texture semantic index to be later used during the buffer formatting process on  mu::MeshFormat
+								TextureSemanticIndices[NumbOfTextCoordChannels] = SemanticIndex;
+									
+								// Store the amount of located texture data channels
+								NumbOfTextCoordChannels++;
+								break;
+							}
+
+							case mu::MESH_BUFFER_SEMANTIC::MBS_BONEINDICES:
+							{
+								// Store the amount of bones a vertex can be skinned to.
+								MaxNumBonesPerVertex = BufferComponentCount;
+								break;
+							}
+
+							default:
+							{
+								break;
+							}
+						}
+					}
+				}
+				
+				// 2.2)		Generate the mesh buffer Sets so they follow the expected structure.
+
+				// Setup the buffers to later setup the channels each of them have
+				int32 MutableBufferCount = MUTABLE_VERTEXBUFFER_TEXCOORDS + 1;
+				{
+					// Bone Indices
+					{
+						// They are mandatory when generating a new skeletal mesh
+						MutableBufferCount++;
+					}
+					
+					// (Optional) Vertex colors
+					if (bOriginDoesHaveVertexColorData)
+					{
+						MutableBufferCount++;
+					}
+				}
+				FormattedVertexBuffers.SetBufferCount(MutableBufferCount);
+
+				// Setup the channels of each buffer
+
+				int32 CurrentVertexBuffer = 0;
+				{
+					// Vertex buffer (positions)
+					MutableMeshBufferUtils::SetupVertexPositionsBuffer(CurrentVertexBuffer,FormattedVertexBuffers);
+					CurrentVertexBuffer++;
+
+					// Tangent buffer
+					MutableMeshBufferUtils::SetupTangentBuffer(CurrentVertexBuffer,FormattedVertexBuffers);
+					CurrentVertexBuffer++;
+				
+					// Texture coords buffer
+					MutableMeshBufferUtils::SetupTexCoordinatesBuffer(CurrentVertexBuffer,NumbOfTextCoordChannels,FormattedVertexBuffers, TextureSemanticIndices);
+					CurrentVertexBuffer++;
+				
+					// Skin buffer
+					{
+						// Use the index buffer to know the amount of bones
+						const int32 NumberOfBones = InMutableMesh->GetVertexCount() * MaxNumBonesPerVertex;;
+					
+						// Failing this check would mean that the provided mutable mesh is in very bad shape (missing verts or
+						// the BONE_INDICES buffer channel
+						if (!NumberOfBones)
+						{
+							UE_LOG(LogTemp,Warning,TEXT("Unable to extract the bone count data from the provided mutable mesh. Make sure it has at least buffer data  "))
+						}
+
+						const int32 MaxBoneIndexTypeSizeBytes = NumberOfBones > 256 ? 2 : 1;
+						MutableMeshBufferUtils::SetupSkinBuffer(CurrentVertexBuffer,MaxBoneIndexTypeSizeBytes,MaxNumBonesPerVertex,FormattedVertexBuffers);
+						CurrentVertexBuffer++;
+					}
+
+					// Colour buffer
+					if (bOriginDoesHaveVertexColorData)
+					{
+						MutableMeshBufferUtils::SetupVertexColorBuffer(CurrentVertexBuffer,FormattedVertexBuffers);
+						CurrentVertexBuffer++;
+					}
+
+					// Index buffer
+					MutableMeshBufferUtils::SetupIndexBuffer(FormattedIndexBuffers);
+				}
+			}
+			
+			// At this point a new mutable mesh with no data but with the appropriate structure for UE to be able to convert it to an
+			// Skeletal mesh is ready to be used as an argument for the process of reformatting the Mutable Mesh we want to preview
+
+			// 3) Call mu::MeshFormat to get the mesh with the wanted buffer structure to have the data of the original mesh
+			mu::MeshPtrConst ResultMutableMesh = mu::MeshFormat(
+				InMutableMesh.get(),
+				FormattedMutableMesh.get(),
+				false,
+				true,
+				true,
+				false,
+				false,
+				false );
+			
+			return ResultMutableMesh;
+		}
+
+		/** Method designed to scan the vertex buffers of the provided mutable mesh and tell the caller if the buffers are
+		 * set up following what UE requires. 
+		 * @param InMutableMesh The Mutable Mesh that is wanted to be checked.
+		 * @param bLogFindingsToConsole If true the console will display not only if the buffers were found but also if they are
+		 * badly located inside the buffer set.
+		 * @return True if the buffers are where they should be, false if not correctly placed or missing.
+		 */
+		bool AreVertexBuffersCorrectlyLocated(const mu::MeshPtrConst InMutableMesh, bool bLogFindingsToConsole = false )
+		{
+			bool bCanBeConverted = true;
+
+			// At the moment requires the presence of all "mandatory" buffers. Only meshes designed to be displayed by unreal
+			// will therefore be processable at the moment and not always
+			const mu::FMeshBufferSet& VertexBuffers = InMutableMesh->GetVertexBuffers();
+
+			int32 Buffer = -1;
+			int32 Channel = -1;
+
+			// POSITION BUFFER INDEX CHECK
+			VertexBuffers.FindChannel(mu::MESH_BUFFER_SEMANTIC::MBS_POSITION, 0, &Buffer, &Channel);
+			if (Buffer == -1)
+			{
+				if (bLogFindingsToConsole)
+				{
+					UE_LOG(LogMutable, Warning, TEXT("\tVertex Positions Buffer not found."));
+				}
+
+				bCanBeConverted = false;
+			}
+			else if (Buffer != MUTABLE_VERTEXBUFFER_POSITION)
+			{
+				if (bLogFindingsToConsole)
+				{
+					UE_LOG(LogMutable, Warning,
+					       TEXT(
+						       "\tBad Vertex Positions buffer position ( at %d instead of %d )."
+					       ), Buffer, MUTABLE_VERTEXBUFFER_POSITION)
+				}
+
+				bCanBeConverted = false;
+			}
+
+			// TANGENT BUFFER INDEX CHECK
+			VertexBuffers.FindChannel(mu::MESH_BUFFER_SEMANTIC::MBS_TANGENT, 0, &Buffer, &Channel);
+			if (Buffer == -1)
+			{
+				if (bLogFindingsToConsole)
+				{
+					UE_LOG(LogMutable, Warning, TEXT("\tTangents Buffer not found."));
+				}
+
+				bCanBeConverted = false;
+			}
+			else if (Buffer != MUTABLE_VERTEXBUFFER_TANGENT)
+			{
+				if (bLogFindingsToConsole)
+				{
+					UE_LOG(LogMutable, Warning,
+					       TEXT(
+						       "\tBad Tangent buffer position (  at %d instead of %d ). "
+					       ), Buffer, MUTABLE_VERTEXBUFFER_TANGENT);
+				}
+
+				bCanBeConverted = false;
+			}
+
+			// TEX COORDS BUFFER INDEX CHECK
+			VertexBuffers.FindChannel(mu::MESH_BUFFER_SEMANTIC::MBS_TEXCOORDS, 0, &Buffer, &Channel);
+			if (Buffer == -1)
+			{
+				if (bLogFindingsToConsole)
+				{
+					UE_LOG(LogMutable, Warning, TEXT("\tTexCords Buffer not found."));
+				}
+
+				bCanBeConverted = false;
+			}
+			else if (Buffer != MUTABLE_VERTEXBUFFER_TEXCOORDS)
+			{
+				if (bLogFindingsToConsole)
+				{
+					UE_LOG(LogMutable, Warning,
+					       TEXT(
+						       "\tBad TexCoords buffer position ( at %d instead of %d ). "
+					       ), Buffer, MUTABLE_VERTEXBUFFER_TEXCOORDS);
+				}
+
+				bCanBeConverted = false;
+			}
+			
+			return bCanBeConverted;
+		}
+
+		/** Loads up the OutSkeletalMesh with all the bone data found on the InMutableMesh if found to be a skeleton there. If not it will fail
+		 * @param InMutableMesh - The mutable mesh to read bone data from
+		 * @param OutFinalBoneNames - An array of names containing the names of the bones to be used on the skeletal mesh
+		 * @param OutBoneMap - Array with all the bones used on the final mesh
+		 */ 
+		void PrepareSkeletonData(mu::MeshPtrConst InMutableMesh, TArray<FName>& OutFinalBoneNames, TArray<uint16>& OutBoneMap)
+		{
+			if (!InMutableMesh)
+			{
+				return;
+			}
+			
+			// Helpers
+			TMap<FName, int16> BoneNameToIndex;
+			TArray<int16> ParentBoneIndices;
+			TArray<int16> ParentChain;
+			TArray<bool> UsedBones;
+			
+
+			// Use first valid LOD bone count as a potential total number of bones, used for pre-allocating data arrays
+			if ( InMutableMesh && InMutableMesh->GetSkeleton())
+			{
+				const int32 TotalPossibleBones = InMutableMesh->GetSkeleton()->GetBoneCount();
+				
+				// Helpers
+				BoneNameToIndex.Empty(TotalPossibleBones);
+				ParentBoneIndices.Empty(TotalPossibleBones);
+
+				ParentChain.SetNumZeroed(TotalPossibleBones);
+
+				// Out Data
+				OutFinalBoneNames.Reserve(TotalPossibleBones);
+			}
+
+			// Locate what bones are being used by asking each vertex with what bone index it is related with.
+			// We end up having an array of Bool variables where each index corresponds to a bone on the mutable skeleton
+			// and the value determines if the bone is found to be used by any of the vertices located on the
+			// mutable mesh surfaces
+			{
+				if (!InMutableMesh || !InMutableMesh->GetSkeleton())
+				{
+					return;
+				}
+
+				const mu::Skeleton* Skeleton = InMutableMesh->GetSkeleton().get();
+				const int32 BoneCount = Skeleton->GetBoneCount();
+
+				const mu::FMeshBufferSet& MutableMeshVertexBuffers = InMutableMesh->GetVertexBuffers();
+
+				const int32 NumVerticesLODModel = InMutableMesh->GetVertexCount();
+				const int32 SurfaceCount = InMutableMesh->GetSurfaceCount();
+
+				mu::MESH_BUFFER_FORMAT BoneIndexFormat = mu::MBF_NONE;
+				int32 BoneIndexComponents = 0;
+				int32 BoneIndexOffset = 0;
+				int32 BoneIndexBuffer = -1;
+				int32 BoneIndexChannel = -1;
+				MutableMeshVertexBuffers.FindChannel(mu::MBS_BONEINDICES, 0, &BoneIndexBuffer, &BoneIndexChannel);
+				if (BoneIndexBuffer >= 0 || BoneIndexChannel >= 0)
+				{
+					MutableMeshVertexBuffers.GetChannel(BoneIndexBuffer, BoneIndexChannel,
+						nullptr, nullptr, &BoneIndexFormat, &BoneIndexComponents, &BoneIndexOffset);
+				}
+
+				const int32 ElementSize = MutableMeshVertexBuffers.GetElementSize(BoneIndexBuffer);
+
+				const uint8_t* BufferStart = MutableMeshVertexBuffers.GetBufferData(BoneIndexBuffer) + BoneIndexOffset;
+
+				const int NumBoneInfluences = BoneIndexComponents;
+				{
+					UsedBones.Empty(BoneCount);
+					UsedBones.AddDefaulted(BoneCount);
+
+					for (int32 Surface = 0; Surface < SurfaceCount; Surface++)
+					{
+						int32 FirstIndex;
+						int32 IndexCount;
+						int32 FirstVertex;
+						int32 VertexCount;
+						InMutableMesh->GetSurface(Surface, &FirstVertex, &VertexCount, &FirstIndex, &IndexCount);
+
+						if (VertexCount == 0 || IndexCount == 0)
+						{
+							continue;
+						}
+
+						const uint8_t* VertexBoneIndexPtr = BufferStart + FirstVertex * ElementSize;
+
+						if (BoneIndexFormat == mu::MBF_UINT8)
+						{
+							for (int32 v = 0; v < VertexCount; ++v)
+							{
+								for (int32 i = 0; i < BoneIndexComponents; ++i)
+								{
+									const size_t SectionBoneIndex = VertexBoneIndexPtr[i];
+									UsedBones[SectionBoneIndex] = true;
+								}
+								VertexBoneIndexPtr += ElementSize;
+							}
+						}
+						else if (BoneIndexFormat == mu::MBF_UINT16)
+						{
+							for (int32 v = 0; v < VertexCount; ++v)
+							{
+								for (int32 i = 0; i < BoneIndexComponents; ++i)
+								{
+									const size_t SectionBoneIndex = ((uint16_t*)VertexBoneIndexPtr)[i];
+									UsedBones[SectionBoneIndex] = true;
+								}
+								VertexBoneIndexPtr += ElementSize;
+							}
+						}
+						else if (BoneIndexFormat == mu::MBF_UINT32)
+						{
+							for (int32 v = 0; v < VertexCount; ++v)
+							{
+								for (int32 i = 0; i < BoneIndexComponents; ++i)
+								{
+									const size_t SectionBoneIndex = ((uint32_t*)VertexBoneIndexPtr)[i];
+									UsedBones[SectionBoneIndex] = true;
+								}
+								VertexBoneIndexPtr += ElementSize;
+							}
+						}
+						else
+						{
+							// Unsupported bone index format in generated mutable mesh
+							check(false);
+						}
+
+					}
+				}
+
+				// Fill the BoneMap with the data found on the mutable skeleton and with the aid of the
+				// bool array with all the bones we are using in this mesh from all the ones found on the mu::Skeleton
+				{
+					OutBoneMap.Reserve(BoneCount);
+
+					for (int32 BoneIndex = 0; BoneIndex < BoneCount; ++BoneIndex)
+					{
+						if (!UsedBones[BoneIndex])
+						{
+							// Add root as a placeholder
+							OutBoneMap.Add(0);
+							continue;
+						}
+
+						FName BoneName = Skeleton->GetBoneName(BoneIndex);
+
+						int16* FinalBoneIndexPtr = BoneNameToIndex.Find(BoneName);
+						if (!FinalBoneIndexPtr) // New bone!
+						{
+							// Ensure parent chain is valid
+							int16 FinalParentIndex = INDEX_NONE;
+
+							const int16 MutParentIndex = Skeleton->GetBoneParent(BoneIndex);
+							if (MutParentIndex != INDEX_NONE) // Bone has a parent, ensure the parent exists
+							{
+								// Ensure parent chain until root exists
+								uint16 ParentChainCount = 0;
+
+								// Find parents to add
+								int16 NextMutParentIndex = MutParentIndex;
+								while (NextMutParentIndex != INDEX_NONE)
+								{
+									if (int16* Found = BoneNameToIndex.Find(Skeleton->GetBoneName(NextMutParentIndex)))
+									{
+										FinalParentIndex = *Found;
+										break;
+									}
+
+									ParentChain[ParentChainCount] = NextMutParentIndex;
+									++ParentChainCount;
+
+									NextMutParentIndex = Skeleton->GetBoneParent(NextMutParentIndex);
+								}
+
+								// Add parent bones to the list and to the active bones array
+								for (int16 ParentChainIndex = ParentChainCount - 1; ParentChainIndex >= 0; --ParentChainIndex)
+								{
+									FName ChainParentName = Skeleton->GetBoneName(ParentChain[ParentChainIndex]);
+
+									// Set the parent for the bone we're about to add (previous ParentIndex)
+									ParentBoneIndices.Add(FinalParentIndex);
+
+									// Add the parent
+									FinalParentIndex = OutFinalBoneNames.Num();
+									BoneNameToIndex.Add(ChainParentName, FinalParentIndex);
+									OutFinalBoneNames.Add(ChainParentName);
+								}
+							}
+
+							const int16 NewBoneIndex = OutFinalBoneNames.Num();
+							OutBoneMap.Add(NewBoneIndex);
+
+							// Add new bone to the ReferenceSkeleton
+							BoneNameToIndex.Add(BoneName, NewBoneIndex);
+							OutFinalBoneNames.Add(BoneName);
+							ParentBoneIndices.Add(FinalParentIndex);
+						}
+						else
+						{
+							const int16 FinalBoneIndex = *FinalBoneIndexPtr;
+							OutBoneMap.Add(FinalBoneIndex);
+						}
+					}
+				}
+
+				OutBoneMap.Shrink();
+			}
+
+			OutFinalBoneNames.Shrink();
+		}
+
+		/** Provides the OutSkeletalMesh with a new reference skeleton based on the data found on the other method inputs
+		 * @param InRefSkeletalMesh - The skeletal mesh to be used as reference for the copying of data not held by the mutable
+		 * @param OutSkeletalMesh - The skeletal mesh whose reference skeleton we want to generate.
+		 * @return True if the operation could be performed successfully, false if not.
+		 */
+		bool BuildSkeletalMeshSkeletonData(const USkeletalMesh* InRefSkeletalMesh, USkeletalMesh* OutSkeletalMesh)
+		{
+			if (!OutSkeletalMesh || !InRefSkeletalMesh)
+			{
+				return false;
+			}
+
+			// get a copy of the skeleton used by the reference skeletal mesh
+			const USkeleton* OriginalSkeleton = InRefSkeletalMesh->GetSkeleton();
+			if (!OriginalSkeleton)
+			{
+				return false;
+			}
+			
+			TObjectPtr<USkeleton> Skeleton = DuplicateObject(OriginalSkeleton,OriginalSkeleton->GetExternalPackage()) ;
+			OutSkeletalMesh->SetSkeleton(Skeleton);
+			
+			const FReferenceSkeleton& SourceReferenceSkeleton = Skeleton->GetReferenceSkeleton();
+			const int32 SourceBoneCount = SourceReferenceSkeleton.GetRawBoneNum();
+			
+			// Used to add them in strictly increasing order
+			TArray<bool> UsedBones;
+			UsedBones.Init(true, SourceBoneCount);	
+
+			// New reference skeleton
+			FReferenceSkeleton RefSkeleton;
+			
+			// Build RefSkeleton
+			UnrealConversionUtils::BuildRefSkeleton(
+				nullptr,
+				SourceReferenceSkeleton,
+				UsedBones,
+				RefSkeleton,
+				Skeleton);
+
+			OutSkeletalMesh->SetRefSkeleton(RefSkeleton);
+
+			OutSkeletalMesh->GetRefBasesInvMatrix().Empty(RefSkeleton.GetRawBoneNum());
+			OutSkeletalMesh->CalculateInvRefMatrices();
+
+			return true;
+		}
+
+		/** Prepares the OutSkeletal mesh with the required materials taking in consideration the amount of surfaces found on
+		 * the InMutableMesh. As it is working with a mutable mesh, and a mutable mesh only represents one lod, the target lod is 0
+		 * @param InMutableMesh - The reference mutable mesh to be reading data from
+		 * @param OutSkeletalMesh - The skeletal mesh to be getting it's materials set up.
+		 */
+		void BuildSkeletalMeshElementData(const mu::MeshPtrConst InMutableMesh,USkeletalMesh* OutSkeletalMesh)
+		{
+			// Set up Unreal's default material
+			UMaterialInterface* UnrealMaterial = UMaterial::GetDefaultMaterial(MD_Surface);
+			OutSkeletalMesh->GetMaterials().SetNum(1);
+			OutSkeletalMesh->GetMaterials()[0] = UnrealMaterial;
+
+			constexpr int32 MeshLODIndex = 0;
+			UnrealConversionUtils::BuildSkeletalMeshElementDataAtLOD(MeshLODIndex,InMutableMesh,OutSkeletalMesh);
+		}
+
+
+		/** Prepares the skeletal mesh to be able to be rendered. It should provide you with a mesh ready to be rendered
+		 * @param InMutableMesh - The mutable mesh to be used as reference
+		 * @param InRefSkeletalMesh - A reference skeletal mesh to aid on the preparation of the OutSkeletalMesh
+		 * @param InBoneMap - An array with all the bones used by the mesh.
+		 * @param OutSkeletalMesh - The skeletal mesh to set up.
+		 * @return True if the operation could be performed successfully, false if not.
+		 */
+		bool BuildSkeletalMeshRenderData(
+			const mu::MeshPtrConst InMutableMesh, const USkeletalMesh* InRefSkeletalMesh,
+			const TArray<uint16>& InBoneMap, USkeletalMesh* OutSkeletalMesh)
+		{
+			OutSkeletalMesh->SetHasBeenSimplified(false);
+			OutSkeletalMesh->SetHasVertexColors(false);
+
+			// Constant values since we are working with a simple mutable mesh (represents only one LOD)
+			constexpr int32 MeshLODIndex = 0;
+			constexpr int32 LODIndex = 0;
+			
+			// Detect if the produced mesh is correct
+			bool bSuccess = true;
+
+			FSkeletalMeshLODRenderData& LODModel = Helper_GetLODData(OutSkeletalMesh)[MeshLODIndex];
+			const mu::FMeshBufferSet& MutableMeshVertexBuffers = InMutableMesh->GetVertexBuffers();
+			
+			// Load buffer data found on the mutable model onto the out skeletal mesh
+			// It includes vertex and index buffers
+			{
+				const int32 NumVerticesLODModel = InMutableMesh->GetVertexCount();
+				//const int32 SurfaceCount = InMutableMesh->GetSurfaceCount();
+
+				mu::MESH_BUFFER_FORMAT boneIndexFormat = mu::MBF_NONE;
+				int32 BoneIndexComponents = 0;
+				int32 BoneIndexOffset = 0;
+				int32 BoneIndexBuffer = -1;
+				int32 BoneIndexChannel = -1;
+				MutableMeshVertexBuffers.FindChannel(mu::MBS_BONEINDICES, 0, &BoneIndexBuffer, &BoneIndexChannel);
+				if (BoneIndexBuffer >= 0 || BoneIndexChannel >= 0)
+				{
+					MutableMeshVertexBuffers.GetChannel(BoneIndexBuffer, BoneIndexChannel,
+					nullptr, nullptr, &boneIndexFormat, &BoneIndexComponents, &BoneIndexOffset);
+				}
+
+				const int32 NumBoneInfluences = BoneIndexComponents;
+
+				UnrealConversionUtils::SetupRenderSections(
+					InMutableMesh,
+					OutSkeletalMesh,
+					MeshLODIndex,
+					NumBoneInfluences,
+					InBoneMap);
+			
+				UnrealConversionUtils::CopyMutableVertexBuffers(
+					OutSkeletalMesh,
+					NumVerticesLODModel,
+					NumBoneInfluences,
+					BoneIndexBuffer,
+					MutableMeshVertexBuffers,
+					MeshLODIndex,
+					boneIndexFormat);
+			}
+			
+			
+			{
+				MUTABLE_CPUPROFILER_SCOPE(UpdateSkeletalMesh_SurfaceLoop_UpdateActiveAndRequiredBones);
+
+				// Generate the active bones array based on the bones found on BoneMap
+				TArray<uint16> ActiveBones;
+				{
+					ActiveBones.Reserve(InBoneMap.Num());
+					for (const uint16& BoneMapIndex : InBoneMap)
+					{
+						ActiveBones.AddUnique(BoneMapIndex);
+					}
+					ActiveBones.Sort();
+				}
+				
+				LODModel.ActiveBoneIndices.Append(ActiveBones);
+				LODModel.RequiredBones.Append(ActiveBones);
+			}
+			
+			if (bSuccess)
+			{
+				bSuccess = UnrealConversionUtils::CopyMutableIndexBuffers(InMutableMesh,LODModel);
+			}
+
+			
+			// Update LOD and streaming data
+			if (InRefSkeletalMesh && LODIndex >= 0 && LODIndex < Helper_GetLODNum(InRefSkeletalMesh))
+			{
+				const FSkeletalMeshLODRenderData& ReferenceLODModel = Helper_GetLODData(InRefSkeletalMesh)[LODIndex];
+
+				LODModel.bIsLODOptional = ReferenceLODModel.bIsLODOptional;
+				LODModel.bStreamedDataInlined = ReferenceLODModel.bStreamedDataInlined;
+			}
+
+
+			if (!bSuccess)
+			{
+				// End with failure
+				return false;
+			}
+
+			if (InRefSkeletalMesh)
+			{
+				const FBoxSphereBounds Bounds = ((USkeletalMesh*)InRefSkeletalMesh)->GetBounds();
+				OutSkeletalMesh->SetImportedBounds(Bounds);
+			}
+			
+			return true;
+		}
+	}
+
+	/*
+	 * Publicly accessible functions of MutableMeshPreviewUtils
+	 */
+
+	USkeletalMesh* GenerateSkeletalMeshFromMutableMesh(mu::MeshPtrConst InMutableMesh,
+		const USkeletalMesh* InReferenceSkeletalMesh)
+	{
+		if (!InMutableMesh ||
+			!InReferenceSkeletalMesh)
+		{
+			return nullptr;
+		}
+
+		// Buffer structures checks and corrective operations
+		{
+			// If the semantic channels for the required buffer data are on not the buffer they are expected to be then
+			// and only then format the mesh to follow what our systems and Unreal expect in terms of channel on buffer structure
+			if (!AreVertexBuffersCorrectlyLocated(InMutableMesh))
+			{
+				UE_LOG(LogMutable, Warning,
+				       TEXT(
+					       "Restructuring Mutable Mesh Buffers:  "
+				       ));
+				
+				// Since we know the mesh does not have the required buffer structure do format it to match our requirements
+				InMutableMesh = GenerateCompatibleMutableMesh(InMutableMesh);
+
+				// If the conversion was required then make sure all went well
+				if (!AreVertexBuffersCorrectlyLocated(InMutableMesh,true))
+				{
+					// The formatting operation failed, do not proceed
+					UE_LOG(LogMutable,Error,TEXT("Restructuring of mutable mesh buffers failed: Aborting Skeletal Mesh generation. "));
+					return nullptr;
+				}
+				
+				UE_LOG(LogMutable,Log,TEXT("Buffers Restructuring was succesfull."));
+			}
+
+			// We should now have a mutable mesh following a buffer structure where the required buffers are there to be
+			// later found and used to generate a proper Skeletal Mesh
+		}
+		
+		
+		// Bone processed data containers
+		TArray<uint16> BoneMap;
+		TArray<FName> BoneNames;
+
+		// Provided a mutable skeleton do load the BoneNames and BoneMap from there
+		if (InMutableMesh->GetSkeleton())
+		{
+			// Prepare the skeleton data and get the bone names to be used on the generated mesh alongside with the indices of
+			// such bones on the final mesh.
+			PrepareSkeletonData(InMutableMesh,BoneNames,BoneMap);
+		}
+		else
+		{
+			// No mutable skeleton has been loaded.
+			// Use the data already set on the reference skeletal mesh to do the work of PrepareSkeletonData
+			// but without having access to the mutable skeleton
+
+			// Add all reference bones set on the reference skeleton onto our new mesh.
+			const FReferenceSkeleton& ReferenceMeshReferenceSkeleton = InReferenceSkeletalMesh->GetRefSkeleton();
+			for	(int32 RefBone = 0; RefBone < ReferenceMeshReferenceSkeleton.GetRawBoneNum(); RefBone++)
+			{
+				BoneMap.Add(RefBone);
+			}
+		}
+		
+		// Generate a new skeletal mesh to be filled up with the data found on the mutable mesh and the reference skeletal mesh
+		USkeletalMesh* GeneratedSkeletalMesh =
+			NewObject<USkeletalMesh>(GetTransientPackage(), NAME_None, RF_Transient);
+
+		// build the reference skeleton for the Generated Skeletal mesh
+		bool bSuccess = BuildSkeletalMeshSkeletonData(InReferenceSkeletalMesh,GeneratedSkeletalMesh);
+		if (!bSuccess)
+		{
+			return nullptr;
+		}
+		
+		// Load the data present on the reference skeletal mesh that could not be grabbed from the mu::Mesh
+		GeneratedSkeletalMesh->SetImportedBounds(InReferenceSkeletalMesh->GetBounds()); // TODO Generate Bounds
+		GeneratedSkeletalMesh->SetPhysicsAsset(InReferenceSkeletalMesh->GetPhysicsAsset());
+		GeneratedSkeletalMesh->SetEnablePerPolyCollision(InReferenceSkeletalMesh->GetEnablePerPolyCollision());
+
+		GeneratedSkeletalMesh->AllocateResourceForRendering();
+
+		// Prepare the mesh to be filled with rendering data (buffers)
+		{
+			Helper_GetLODData(GeneratedSkeletalMesh).Add(( Helper_GetLODData(GeneratedSkeletalMesh).Num()) ? &Helper_GetLODData(GeneratedSkeletalMesh)[0] : new Helper_LODDataType());
+			Helper_GetLODInfoArray(GeneratedSkeletalMesh).Add(FSkeletalMeshLODInfo());
+
+	#if WITH_EDITORONLY_DATA
+			Helper_GetLODInfoArray(GeneratedSkeletalMesh).Last().BuildSettings.bUseFullPrecisionUVs = true;
+	#endif
+
+			constexpr int32 LODIndex = 0;
+			if (Helper_GetLODInfoArray(InReferenceSkeletalMesh).IsValidIndex(LODIndex))
+			{
+				Helper_GetLODInfoArray(GeneratedSkeletalMesh).Last().ScreenSize = Helper_GetLODInfoArray(InReferenceSkeletalMesh)[LODIndex].ScreenSize;
+				Helper_GetLODInfoArray(GeneratedSkeletalMesh).Last().LODHysteresis = Helper_GetLODInfoArray(InReferenceSkeletalMesh)[LODIndex].LODHysteresis;
+				Helper_GetLODInfoArray(GeneratedSkeletalMesh).Last().bSupportUniformlyDistributedSampling = Helper_GetLODInfoArray(InReferenceSkeletalMesh)[LODIndex].bSupportUniformlyDistributedSampling;
+				Helper_GetLODInfoArray(GeneratedSkeletalMesh).Last().bAllowCPUAccess = Helper_GetLODInfoArray(InReferenceSkeletalMesh)[LODIndex].bAllowCPUAccess;
+			}
+			else
+			{
+				Helper_GetLODInfoArray(GeneratedSkeletalMesh).Last().ScreenSize = 0.3f / (LODIndex + 1);
+				Helper_GetLODInfoArray(GeneratedSkeletalMesh).Last().LODHysteresis = 0.02f;
+				UE_LOG(LogMutable, Warning, TEXT("Setting default values for LOD ScreenSize and LODHysteresis because the values cannot be found in the reference mesh"));
+			}
+		}
+
+		// Set the material data and initialize the sections that will be used later
+		BuildSkeletalMeshElementData(InMutableMesh,GeneratedSkeletalMesh);
+		
+		// Load all data from the mutable mesh onto the skeletal mesh object
+		bSuccess = BuildSkeletalMeshRenderData(InMutableMesh,InReferenceSkeletalMesh,BoneMap,GeneratedSkeletalMesh);
+		if (!bSuccess)
+		{
+			return nullptr;
+		}
+
+		// Prepares the mesh to be rendered
+		GeneratedSkeletalMesh->InitResources();
+		
+		return GeneratedSkeletalMesh;
+	}
+
+}
+
+
+
+
+
+
+
+
+
diff --git a/Engine/Plugins/Experimental/Mutable/Source/CustomizableObjectEditor/Private/MutableUtils.cpp b/Engine/Plugins/Experimental/Mutable/Source/CustomizableObjectEditor/Private/MutableUtils.cpp
new file mode 100644
--- /dev/null
+++ b/Engine/Plugins/Experimental/Mutable/Source/CustomizableObjectEditor/Private/MutableUtils.cpp
@@ -0,0 +1,279 @@
+// Copyright Epic Games, Inc. All Rights Reserved.
+
+#include "MutableUtils.h"
+
+#include "Rendering/SkeletalMeshModel.h"
+
+
+TArray<FVector2f> GetUV(const USkeletalMesh& SkeletalMesh, const int32 LODIndex, const int32 SectionIndex, const int32 UVIndex)
+{
+	TArray<FVector2f> Result;
+
+	const FSkeletalMeshModel* ImportedModel = SkeletalMesh.GetImportedModel();
+	const FSkeletalMeshLODModel& LOD = ImportedModel->LODModels[LODIndex];
+	const FSkelMeshSection& Section = LOD.Sections[SectionIndex];
+	
+	TArray<FSoftSkinVertex> Vertices;
+	LOD.GetVertices(Vertices);
+
+	TArray<uint32> Indices;
+	SkeletalMesh.GetResourceForRendering()->LODRenderData[LODIndex].MultiSizeIndexContainer.GetIndexBuffer(Indices);
+	
+	int32 IndexIndex = Section.BaseIndex;
+	for (uint32 FaceIndex = 0; FaceIndex < Section.NumTriangles; ++FaceIndex, IndexIndex += 3)
+	{
+		Result.Add(Vertices[Indices[IndexIndex + 0]].UVs[UVIndex]);
+		Result.Add(Vertices[Indices[IndexIndex + 1]].UVs[UVIndex]);
+
+		Result.Add(Vertices[Indices[IndexIndex + 1]].UVs[UVIndex]);
+		Result.Add(Vertices[Indices[IndexIndex + 2]].UVs[UVIndex]);
+		
+		Result.Add(Vertices[Indices[IndexIndex + 2]].UVs[UVIndex]);
+		Result.Add(Vertices[Indices[IndexIndex + 0]].UVs[UVIndex]);
+	}
+
+	return Result;
+}
+
+TArray<FVector2f> GetUV(const UStaticMesh& StaticMesh, const int32 LODIndex, const int32 SectionIndex, const int32 UVIndex)
+{
+	TArray<FVector2f> Result;
+
+	const FStaticMeshRenderData* RenderData = StaticMesh.GetRenderData();
+	const FStaticMeshLODResources& LOD = RenderData->LODResources[LODIndex];
+	const FStaticMeshSection& Section = LOD.Sections[SectionIndex];
+	
+	const FStaticMeshVertexBuffer& VertexBuffer = LOD.VertexBuffers.StaticMeshVertexBuffer;
+	
+	FIndexArrayView Indices = LOD.IndexBuffer.GetArrayView();
+
+	int32 IndexIndex = Section.FirstIndex;
+	for (uint32 FaceIndex = 0; FaceIndex < Section.NumTriangles; ++FaceIndex, IndexIndex += 3)
+	{
+		Result.Add(VertexBuffer.GetVertexUV(Indices[IndexIndex + 0], UVIndex));
+		Result.Add(VertexBuffer.GetVertexUV(Indices[IndexIndex + 1], UVIndex));
+
+		Result.Add(VertexBuffer.GetVertexUV(Indices[IndexIndex + 1], UVIndex));
+		Result.Add(VertexBuffer.GetVertexUV(Indices[IndexIndex + 2], UVIndex));
+		
+		Result.Add(VertexBuffer.GetVertexUV(Indices[IndexIndex + 2], UVIndex));
+		Result.Add(VertexBuffer.GetVertexUV(Indices[IndexIndex + 0], UVIndex));
+	}
+	
+	return Result;
+}
+
+bool HasNormalizedBounds(const FVector2f& Point)
+{
+	return Point.X >= 0.0f && Point.X <= 1.0 && Point.Y >= 0.0f && Point.Y <= 1.0;
+}
+
+
+bool IsUVNormalized(const USkeletalMesh& SkeletalMesh, const int32 LODIndex, const int32 SectionIndex, const int32 UVIndex)
+{
+	const FSkeletalMeshModel* ImportedModel = SkeletalMesh.GetImportedModel();
+	const FSkeletalMeshLODModel& LOD = ImportedModel->LODModels[LODIndex];
+	const FSkelMeshSection& Section = LOD.Sections[SectionIndex];
+	
+	TArray<FSoftSkinVertex> Vertices;
+	LOD.GetVertices(Vertices);
+
+	TArray<uint32> Indices;
+	SkeletalMesh.GetResourceForRendering()->LODRenderData[LODIndex].MultiSizeIndexContainer.GetIndexBuffer(Indices);
+	
+	for (int32 VertexIndex = Section.BaseVertexIndex; VertexIndex < Section.NumVertices; ++VertexIndex)
+	{
+		if (!HasNormalizedBounds(Vertices[Indices[VertexIndex]].UVs[UVIndex]))
+		{
+			return false;
+		}
+	}
+
+	return true;
+}
+
+
+bool IsUVNormalized(const UStaticMesh& StaticMesh, const int32 LODIndex, const int32 SectionIndex, const int32 UVIndex)
+{
+	const FStaticMeshRenderData* RenderData = StaticMesh.GetRenderData();
+	const FStaticMeshLODResources& LOD = RenderData->LODResources[LODIndex];
+	const FStaticMeshSection& Section = LOD.Sections[SectionIndex];
+	
+	const FStaticMeshVertexBuffer& VertexBuffer = LOD.VertexBuffers.StaticMeshVertexBuffer;
+	
+	FIndexArrayView Indices = LOD.IndexBuffer.GetArrayView();
+
+	for (uint32 VertexIndex = Section.MinVertexIndex; VertexIndex < Section.MaxVertexIndex; ++VertexIndex)
+	{
+		if (!HasNormalizedBounds(VertexBuffer.GetVertexUV(Indices[VertexIndex + 0], UVIndex)))
+		{
+			return false;
+		}
+	}
+	
+	return true;
+}
+
+
+bool IsMeshClosed(const USkeletalMesh* Mesh, int LOD, int MaterialIndex)
+{
+	check(Mesh->GetImportedModel() &&
+		Mesh->GetImportedModel()->LODModels.IsValidIndex(LOD) &&
+		Mesh->GetImportedModel()->LODModels[LOD].Sections.IsValidIndex(MaterialIndex));
+
+	const float VertexCollapseEpsilonSquared = 0.0001f * 0.0001f;
+
+	const FSkeletalMeshLODModel& LODResource = Mesh->GetImportedModel()->LODModels[LOD];
+	const FSkelMeshSection& Section = LODResource.Sections[MaterialIndex];
+
+	const FSoftSkinVertex* VertexBuffer = Section.SoftVertices.GetData();
+	const uint32 VertexStart = Section.BaseVertexIndex;
+	const uint32 VertexCount = Section.NumVertices;
+
+	// Collapse vertices
+	TArray<uint32> CollapsedVertexMap;
+	CollapsedVertexMap.SetNum(VertexCount);
+
+
+	for (uint32 V = VertexStart; V < VertexCount; ++V)
+	{
+		CollapsedVertexMap[V] = V;
+
+		for (uint32 CandidateIndex = VertexStart; CandidateIndex < V; ++CandidateIndex)
+		{
+			uint32 Collapsed = CollapsedVertexMap[CandidateIndex];
+
+			FVector3f Vec = VertexBuffer[V].Position - VertexBuffer[Collapsed].Position;
+
+			if (FVector3f::DotProduct(Vec, Vec) <= VertexCollapseEpsilonSquared)
+			{
+				CollapsedVertexMap[V] = Collapsed;
+				break;
+			}
+		}
+	}
+
+	// Count faces per edge
+	TMap<TPair<uint32, uint32>, uint32> FaceCountPerEdge;
+
+	const uint32* IndexBuffer = LODResource.IndexBuffer.GetData();
+	const uint32 IndexStart = Section.BaseIndex;
+	const uint32 IndexCount = Section.NumTriangles * 3;
+
+	for (uint32 I = IndexStart; I < IndexCount; I += 3)
+	{
+		uint32 Face[3];
+		Face[0] = CollapsedVertexMap[IndexBuffer[I + 0]];
+		Face[1] = CollapsedVertexMap[IndexBuffer[I + 1]];
+		Face[2] = CollapsedVertexMap[IndexBuffer[I + 2]];
+
+		for (int E = 0; E < 3; ++E)
+		{
+			uint32 Vertex0 = Face[E];
+			uint32 Vertex1 = Face[(E + 1) % 3];
+
+			if (Vertex0 == Vertex1)
+			{
+				return false; // Degenerated mesh
+			}
+
+			TPair<uint32, uint32> Edge;
+			Edge.Key = FGenericPlatformMath::Min(Vertex0, Vertex1);
+			Edge.Value = FGenericPlatformMath::Max(Vertex0, Vertex1);
+
+			FaceCountPerEdge.FindOrAdd(Edge)++;
+		}
+	}
+
+	// See if every edge has 2 faces
+	for (const TPair<TPair<uint32, uint32>, uint32>& Entry : FaceCountPerEdge)
+	{
+		if (Entry.Value != 2)
+		{
+			return false;
+		}
+	}
+
+	return true;
+}
+
+
+bool IsMeshClosed(const UStaticMesh* Mesh, int LOD, int MaterialIndex)
+{
+	check(Mesh->GetRenderData() &&
+		Mesh->GetRenderData()->LODResources.IsValidIndex(LOD) &&
+		Mesh->GetRenderData()->LODResources[LOD].Sections.IsValidIndex(MaterialIndex));
+
+	const float VertexCollapseEpsilonSquared = 0.0001f * 0.0001f;
+
+	const FStaticMeshLODResources& LODResource = Mesh->GetRenderData()->LODResources[LOD];
+	const FStaticMeshSection& Section = LODResource.Sections[MaterialIndex];
+
+	const uint32 VertexStart = Section.MinVertexIndex;
+	const uint32 VertexCount = Section.MaxVertexIndex - VertexStart + 1;
+
+	// Collapse vertices
+	TArray<uint32> CollapsedVertexMap;
+	CollapsedVertexMap.SetNum(VertexCount);
+
+	for (uint32 V = VertexStart; V < VertexCount; ++V)
+	{
+		CollapsedVertexMap[V] = V;
+
+		for (uint32 CandidateIndex = VertexStart; CandidateIndex < V; ++CandidateIndex)
+		{
+			uint32 Collapsed = CollapsedVertexMap[CandidateIndex];
+
+			FVector3f Vec = LODResource.VertexBuffers.PositionVertexBuffer.VertexPosition(V) - LODResource.VertexBuffers.PositionVertexBuffer.VertexPosition(Collapsed);
+
+			if (FVector3f::DotProduct(Vec, Vec) <= VertexCollapseEpsilonSquared)
+			{
+				CollapsedVertexMap[V] = Collapsed;
+				break;
+			}
+		}
+	}
+
+	// Count faces per edge
+	TMap<TPair<uint32, uint32>, uint32> FaceCountPerEdge;
+
+	const FIndexArrayView& IndexBuffer = LODResource.IndexBuffer.GetArrayView();
+	const uint32 IndexStart = Section.FirstIndex;
+	const uint32 IndexCount = Section.NumTriangles * 3;
+
+	for (uint32 I = IndexStart; I < IndexCount; I += 3)
+	{
+		uint32 Face[3];
+		Face[0] = CollapsedVertexMap[IndexBuffer[I + 0]];
+		Face[1] = CollapsedVertexMap[IndexBuffer[I + 1]];
+		Face[2] = CollapsedVertexMap[IndexBuffer[I + 2]];
+
+		for (int E = 0; E < 3; ++E)
+		{
+			uint32 Vertex0 = Face[E];
+			uint32 Vertex1 = Face[(E + 1) % 3];
+
+			if (Vertex0 == Vertex1)
+			{
+				return false; // Degenerated mesh
+			}
+
+			TPair<uint32, uint32> Edge;
+			Edge.Key = FGenericPlatformMath::Min(Vertex0, Vertex1);
+			Edge.Value = FGenericPlatformMath::Max(Vertex0, Vertex1);
+
+			FaceCountPerEdge.FindOrAdd(Edge)++;
+		}
+	}
+
+	// See if every edge has 2 faces
+	for (const TPair<TPair<uint32, uint32>, uint32>& Entry : FaceCountPerEdge)
+	{
+		if (Entry.Value != 2)
+		{
+			return false;
+		}
+	}
+
+	return true;
+}
\ No newline at end of file
diff --git a/Engine/Plugins/Experimental/Mutable/Source/CustomizableObjectEditor/Private/Nodes/CustomizableObjectNode.cpp b/Engine/Plugins/Experimental/Mutable/Source/CustomizableObjectEditor/Private/Nodes/CustomizableObjectNode.cpp
new file mode 100644
--- /dev/null
+++ b/Engine/Plugins/Experimental/Mutable/Source/CustomizableObjectEditor/Private/Nodes/CustomizableObjectNode.cpp
@@ -0,0 +1,639 @@
+// Copyright Epic Games, Inc. All Rights Reserved.
+
+#include "Nodes/CustomizableObjectNode.h"
+
+#include "Nodes/CustomizableObjectNodeExposePin.h"
+#include "CustomizableObjectEditorModule.h"
+#include "Toolkits/ToolkitManager.h"
+
+#include "UnrealEditorPortabilityHelpers.h"
+
+#include "CustomizableObjectEditor.h"
+#include "CustomizableObjectCustomVersion.h"
+#include "CustomizableObjectGraph.h"
+#include "EdGraphSchema_CustomizableObject.h"
+
+#include "CustomizableObjectGraph.h"
+#include "CustomizableObjectPin.h"
+#include "EdGraphSchema_CustomizableObject.h"
+#include "GraphTraversal.h"
+#include "Runtime/Launch/Resources/Version.h" // Include engine version macros
+
+#define LOCTEXT_NAMESPACE "CustomizableObjectEditor"
+
+
+/** Copy pin data from old pin to new pin. Keeps the id of the new pin. */
+void CopyPin(UEdGraphPin& NewPin, UEdGraphPin& OldPin)
+{
+	FGuid PinId = NewPin.PinId;
+	
+	NewPin.CopyPersistentDataFromOldPin(OldPin);
+	NewPin.PinId = PinId;
+	NewPin.bHidden = OldPin.bHidden;
+}
+
+
+UCustomizableObjectGraph* UCustomizableObjectNode::GetCustomizableObjectGraph() const
+{
+	return Cast<UCustomizableObjectGraph>(GetOuter());
+}
+
+bool UCustomizableObjectNode::IsSingleOutputNode() const
+{
+	return false;
+}
+
+
+UEdGraphPin* UCustomizableObjectNode::CustomCreatePin(EEdGraphPinDirection Direction, const FName& Type, const FName& Name, bool bIsArray)
+{
+	UEdGraphPin* Pin = CreatePin(Direction, Type, Name);
+	if (bIsArray)
+	{
+		Pin->PinType.ContainerType = EPinContainerType::Array;
+	}
+
+	return Pin;
+}
+
+
+UEdGraphPin* UCustomizableObjectNode::CustomCreatePin(EEdGraphPinDirection Direction, const FName& Type, const FName& Name, UCustomizableObjectNodePinData* PinData)
+{
+	UEdGraphPin* Pin = CreatePin(Direction, Type, Name);
+	PinsDataId.Add(Pin->PinId, PinData);
+	
+	return Pin;
+}
+
+
+bool UCustomizableObjectNode::ShouldBreakExistingConnections(const UEdGraphPin* InputPin, const UEdGraphPin* OutputPin) const
+{
+	return IsSingleOutputNode();
+}
+
+void UCustomizableObjectNode::PostEditChangeProperty(FPropertyChangedEvent& PropertyChangedEvent)
+{
+	Super::PostEditChangeProperty(PropertyChangedEvent);
+
+	FPostEditChangePropertyDelegateParameters Parameters;
+	Parameters.Node = this;
+	Parameters.FPropertyChangedEvent = &PropertyChangedEvent;
+
+	PostEditChangePropertyDelegate.Broadcast(Parameters);
+	PostEditChangePropertyRegularDelegate.Broadcast(this, PropertyChangedEvent);
+}
+
+
+TSharedPtr<ICustomizableObjectEditor> UCustomizableObjectNode::GetGraphEditor() const
+{
+	const UCustomizableObject* CustomizableObject = Cast<UCustomizableObject>(GetCustomizableObjectGraph()->GetOuter());
+
+	TSharedPtr<ICustomizableObjectEditor> CustomizableObjectEditor;
+	if (CustomizableObject)
+	{
+		TSharedPtr< IToolkit > FoundAssetEditor = FToolkitManager::Get().FindEditorForAsset(CustomizableObject);
+		if (FoundAssetEditor.IsValid())
+		{
+			return StaticCastSharedPtr<ICustomizableObjectEditor>(FoundAssetEditor);
+		}
+	}
+
+	return TSharedPtr<ICustomizableObjectEditor>();
+}
+
+
+bool UCustomizableObjectNode::CustomRemovePin(UEdGraphPin& Pin)
+{
+	PinsDataId.Remove(Pin.PinId);
+	
+	return RemovePin(&Pin);	
+}
+
+
+void UCustomizableObjectNode::GetInputPins(TArray<class UEdGraphPin*>& OutInputPins) const
+{
+	OutInputPins.Empty();
+
+	for (int32 PinIndex = 0; PinIndex < Pins.Num(); PinIndex++)
+	{
+		if (Pins[PinIndex]->Direction == EGPD_Input)
+		{
+			OutInputPins.Add(Pins[PinIndex]);
+		}
+	}
+}
+
+
+void UCustomizableObjectNode::GetOutputPins(TArray<class UEdGraphPin*>& OutOutputPins) const
+{
+	OutOutputPins.Empty();
+
+	for (int32 PinIndex = 0; PinIndex < Pins.Num(); PinIndex++)
+	{
+		if (Pins[PinIndex]->Direction == EGPD_Output)
+		{
+			OutOutputPins.Add(Pins[PinIndex]);
+		}
+	}
+}
+
+
+UEdGraphPin* UCustomizableObjectNode::GetOutputPin(int32 OutputIndex) const
+{
+	for (int32 PinIndex = 0, FoundOutputs = 0; PinIndex < Pins.Num(); PinIndex++)
+	{
+		if (Pins[PinIndex]->Direction == EGPD_Output)
+		{
+			if (OutputIndex == FoundOutputs)
+			{
+				return Pins[PinIndex];
+			}
+			else
+			{
+				FoundOutputs++;
+			}
+		}
+	}
+
+	return NULL;
+}
+
+
+void UCustomizableObjectNode::SetRefreshNodeWarning()
+{
+	if (!bHasCompilerMessage && ErrorType < EMessageSeverity::Warning)
+	{
+		GetGraph()->NotifyGraphChanged();
+	
+		bHasCompilerMessage = true;
+		ErrorType = EMessageSeverity::Warning;
+		ErrorMsg = GetRefreshMessage();
+	}
+}
+
+
+void UCustomizableObjectNode::RemoveWarnings()
+{
+	bHasCompilerMessage = false;
+	ErrorType = 0;
+	ErrorMsg.Empty();
+}
+
+
+void UCustomizableObjectNode::AllocateDefaultPins()
+{
+	AllocateDefaultPins(nullptr);
+}
+
+
+void UCustomizableObjectNode::ReconstructNode()
+{
+	ReconstructNode(CreateRemapPinsDefault());
+}
+
+
+void UCustomizableObjectNode::ReconstructNode(UCustomizableObjectNodeRemapPins* RemapPinsAction)
+{
+	Modify();
+
+	// Break any single sided links. All connections must always be present in both nodes
+	for (int32 PinIndex = 0; PinIndex < Pins.Num(); ++PinIndex)
+	{
+		UEdGraphPin* Pin = Pins[PinIndex];
+		TArray<class UEdGraphPin*>& LinkedToRef = Pin->LinkedTo;
+		for (int32 LinkIdx=0; LinkIdx < LinkedToRef.Num(); LinkIdx++)
+		{
+			UEdGraphPin* OtherPin = LinkedToRef[LinkIdx];
+			// If we are linked to a pin that its owner doesn't know about, break that link
+			if (OtherPin && !OtherPin->GetOwningNode()->Pins.Contains(OtherPin))
+			{
+				Pin->LinkedTo.Remove(OtherPin);
+			}
+		}
+	}
+
+	RemoveWarnings();
+
+	// Move the existing orphan and non orphan pins to a saved array
+	TArray<UEdGraphPin*> OldPins(Pins); // We can not empty Pins at this point since it will break all FEdGraphPinReference during the reconstruction.
+
+	// Recreate the new pins
+	AllocateDefaultPins(RemapPinsAction);
+	
+	// Try to remap orphan and non orphan pins.
+	TArray<UEdGraphPin*> NewPins;
+	NewPins.Reset(Pins.Num() - OldPins.Num());
+	for (UEdGraphPin* Pin : Pins)
+	{
+		if (!OldPins.Contains(Pin))
+		{
+			NewPins.Add(Pin);
+		}
+	}
+
+	TMap<UEdGraphPin*, UEdGraphPin*> PinsToRemap;
+	TArray<UEdGraphPin*> PinsToOrphan;
+	RemapPinsAction->RemapPins(OldPins, NewPins, PinsToRemap, PinsToOrphan);
+
+	// Check only.
+	for (const TTuple<UEdGraphPin*, UEdGraphPin*>& Pair : PinsToRemap)
+	{
+		check(NewPins.Contains(Pair.Value)); // Can only remap and old pin to a new pin.
+		check(OldPins.Contains(Pair.Key));
+	}
+
+	RemapPins(PinsToRemap);
+	RemapPinsData(PinsToRemap);
+	
+	// Check only.
+	for (UEdGraphPin* Pin : PinsToOrphan)
+	{
+		check(OldPins.Contains(Pin)); // Can only orphan old pins.
+	}
+
+	bool bOrphanedPin = false;
+	for (UEdGraphPin* OldPin : OldPins)
+	{
+		OldPin->Modify();
+
+		if (PinsToOrphan.Contains(OldPin))
+		{
+			bOrphanedPin = bOrphanedPin || !OldPin->bOrphanedPin;
+			OrphanPin(*OldPin);
+		
+			// Move pin to the end.
+			Pins.RemoveSingle(OldPin);
+			Pins.Add(OldPin);
+		}
+		else
+		{
+			// Remove the old pin
+			OldPin->BreakAllPinLinks();
+			
+			CustomRemovePin(*OldPin);
+		}
+	}
+
+	if (UEdGraph* Graph = GetCustomizableObjectGraph())
+	{
+		if (bOrphanedPin)
+		{
+			FCustomizableObjectEditorLogger::CreateLog(LOCTEXT("OrphanPinsWarningReconstruct", "Failed to remap old pins"))
+			.Severity(EMessageSeverity::Warning)
+			.Node(*this)
+			.Log();
+		}
+
+		Graph->NotifyGraphChanged();
+	}
+	
+	PostReconstructNodeDelegate.Broadcast();
+}
+
+
+void UCustomizableObjectNode::DestroyNode()
+{
+	Super::DestroyNode();
+	DestroyNodeDelegate.Broadcast();
+}
+
+
+void UCustomizableObjectNode::AutowireNewNode(UEdGraphPin* FromPin)
+{
+	if (!FromPin)
+	{
+		return;
+	}
+
+	const UEdGraphSchema_CustomizableObject* Schema = GetDefault<UEdGraphSchema_CustomizableObject>();
+
+	for ( int PinIndex=0; PinIndex<Pins.Num(); ++PinIndex )
+	{
+		if ( FromPin->PinType.PinCategory == Pins[PinIndex]->PinType.PinCategory 
+			&&
+			FromPin->Direction != Pins[PinIndex]->Direction )
+		{
+			UEdGraphNode* OwningNode = FromPin->GetOwningNode(); // TryCreateConnection can reconstruct the node invalidating the FromPin. Save the OwningNode here.
+			if (Schema->TryCreateConnection(FromPin, Pins[PinIndex]))
+			{
+				OwningNode->NodeConnectionListChanged();
+				NodeConnectionListChanged();
+			}
+			break;
+		}
+	}
+}
+
+
+void UCustomizableObjectNode::NodeConnectionListChanged()
+{
+	Super::NodeConnectionListChanged();
+	NodeConnectionListChangedDelegate.Broadcast();
+}
+
+
+
+void UCustomizableObjectNode::PostInitProperties()
+{
+	Super::PostInitProperties();
+
+	RemoveWarnings();
+}
+
+
+void UCustomizableObjectNode::BackwardsCompatibleFixup()
+{
+	const int32 CustomizableObjectCustomVersion = GetLinkerCustomVersion(FCustomizableObjectCustomVersion::GUID);
+	
+	// Fix UE large world coordinates automatic pin conversion. 
+	// Now all pins with PinCategory == FName("float") get automatically changed to the new double type
+	if (CustomizableObjectCustomVersion < FCustomizableObjectCustomVersion::FixBlueprintPinsUseRealNumbers)
+	{
+		for (UEdGraphPin* Pin : Pins)
+		{
+			FEdGraphPinType* const PinType = &Pin->PinType;
+
+			if (PinType->PinCategory == TEXT("real") && PinType->PinSubCategory == TEXT("double"))
+			{
+				PinType->PinCategory = TEXT("float");
+				PinType->PinSubCategory = FName();
+			}
+		}
+	}
+	
+	if (CustomizableObjectCustomVersion < FCustomizableObjectCustomVersion::AutomaticNodeMaterialPerformance)
+	{
+		for (TTuple<FEdGraphPinReference, UCustomizableObjectNodePinData*> Pair : PinsData_DEPRECATED)
+		{
+			PinsDataId.Add(Pair.Key.Get()->PinId, Pair.Value);
+		}
+
+		PinsData_DEPRECATED.Empty();
+	}
+
+	if (CustomizableObjectCustomVersion < FCustomizableObjectCustomVersion::FixBlueprintPinsUseRealNumbersAgain)
+	{
+		for (UEdGraphPin* Pin : Pins)
+		{
+			FEdGraphPinType* const PinType = &Pin->PinType;
+
+			if (PinType->PinCategory == TEXT("real") && PinType->PinSubCategory == TEXT("double"))
+			{
+				PinType->PinCategory = TEXT("float");
+				PinType->PinSubCategory = FName();
+			}
+		}
+	}
+}
+
+
+bool UCustomizableObjectNode::CanConnect( const UEdGraphPin* InOwnedInputPin, const UEdGraphPin* InOutputPin, bool& bOutIsOtherNodeBlocklisted, bool& bOutArePinsCompatible) const
+{
+	bOutIsOtherNodeBlocklisted = false;
+
+	bOutArePinsCompatible = InOwnedInputPin->PinType.PinCategory == InOutputPin->PinType.PinCategory;
+	return bOutArePinsCompatible;
+}
+
+void UCustomizableObjectNode::MarkForReconstruct()
+{
+	TSharedPtr<FCustomizableObjectEditor> Editor = StaticCastSharedPtr<FCustomizableObjectEditor>(GetGraphEditor());
+
+	// This could be called if a skeletal mesh asset in the graph is deleted without the editor being open. 
+	if (Editor)
+	{
+		Editor->MarkForReconstruct(this);
+	}
+}
+
+
+UCustomizableObjectNodeRemapPins* UCustomizableObjectNode::CreateRemapPinsDefault() const
+{
+	return CreateRemapPinsByName();
+}
+
+
+UCustomizableObjectNodeRemapPinsByName* UCustomizableObjectNode::CreateRemapPinsByName() const
+{
+	return NewObject<UCustomizableObjectNodeRemapPinsByName>();
+}
+
+
+UCustomizableObjectNodeRemapPinsByPosition* UCustomizableObjectNode::CreateRemapPinsByPosition() const
+{
+	return NewObject<UCustomizableObjectNodeRemapPinsByPosition>();
+}
+
+
+void UCustomizableObjectNode::RemapPins(const TMap<UEdGraphPin*, UEdGraphPin*>& PinsToRemap)
+{
+	for (const TTuple<UEdGraphPin*, UEdGraphPin*>& Pair : PinsToRemap)
+	{
+		CopyPin(*Pair.Value, *Pair.Key);
+	}
+	
+	RemapPinsDelegate.Broadcast(PinsToRemap);
+}
+
+
+void UCustomizableObjectNode::RemapPinsData(const TMap<UEdGraphPin*, UEdGraphPin*>& PinsToRemap)
+{
+	for (const TTuple<UEdGraphPin*, UEdGraphPin*>& Pair : PinsToRemap)
+	{
+		// Move pin data.
+		if (TObjectPtr<UCustomizableObjectNodePinData>* PinDataOldPin = PinsDataId.Find(Pair.Key->PinId))
+		{
+			PinsDataId[Pair.Value->PinId] = *PinDataOldPin;		
+		}
+	}
+}
+
+
+void UCustomizableObjectNode::AddPinData(const UEdGraphPin& Pin, UCustomizableObjectNodePinData& PinData)
+{
+	PinsDataId.Add(Pin.PinId, &PinData);
+}
+
+
+TArray<UEdGraphPin*> UCustomizableObjectNode::GetAllOrphanPins() const
+{
+	TArray<UEdGraphPin*> OrphanPins;
+
+	for (UEdGraphPin* Pin : Pins)
+	{
+		if (Pin->bOrphanedPin)
+		{
+			OrphanPins.Add(Pin);
+		}
+	}
+
+	return OrphanPins;
+}
+
+
+TArray<UEdGraphPin*> UCustomizableObjectNode::GetAllNonOrphanPins() const
+{
+	TArray<UEdGraphPin*> NonOrphanPins;
+	NonOrphanPins.Reserve(Pins.Num());
+	
+	for (UEdGraphPin* Pin : Pins)
+	{
+		if (!Pin->bOrphanedPin)
+		{
+			NonOrphanPins.Add(Pin);
+		}
+	}
+
+	return NonOrphanPins;
+}
+
+
+UCustomizableObjectNodePinData::UCustomizableObjectNodePinData()
+{
+	SetFlags(RF_Transactional);
+}
+
+
+void UCustomizableObjectNode::Serialize(FArchive& Ar)
+{
+	Super::Serialize(Ar);
+	
+	Ar.UsingCustomVersion(FCustomizableObjectCustomVersion::GUID);
+}
+
+
+void UCustomizableObjectNode::PostLoad()
+{
+	// Do not do work here. Do work at PostBackwardsCompatibleFixup.
+	Super::PostLoad();
+}
+
+
+int32 UCustomizableObjectNode::GetLOD() const
+{
+	// Search recursively all parent nodes until a UCustomizableObjectNodeObject is found.
+	// Once found, obtain the matching LOD.
+	TQueue<const UCustomizableObjectNode*> PotentialCustomizableNodeObjects;
+	PotentialCustomizableNodeObjects.Enqueue(this);
+
+	const UCustomizableObjectNode* CurrentElement;
+	while (PotentialCustomizableNodeObjects.Dequeue(CurrentElement))
+	{
+		for (UEdGraphPin* Pin : CurrentElement->GetAllNonOrphanPins())
+		{
+			if (Pin->Direction == EEdGraphPinDirection::EGPD_Output)
+			{
+				for (UEdGraphPin* LinkedPin : Pin->LinkedTo)
+				{
+					if (UCustomizableObjectNodeObject* CurrentCustomizableObjectNode = Cast<UCustomizableObjectNodeObject>(LinkedPin->GetOwningNode()))
+					{
+						return CurrentCustomizableObjectNode->GetLOD(LinkedPin);
+					}
+					else
+					{
+						const UCustomizableObjectNode* Node = Cast<UCustomizableObjectNode>(LinkedPin->GetOwningNode());
+						check(Node); // All nodes inherit from UCustomizableObjectNode
+						PotentialCustomizableNodeObjects.Enqueue(Node);
+					}
+				}
+			}
+		}
+	}
+
+	return -1; // UCustomizableObjectNodeObject not found.
+}
+
+
+TArray<UCustomizableObjectNodeObject*> UCustomizableObjectNode::GetParentObjectNodes(const int LOD) const
+{
+	// Search recursively all parent nodes. Add all Object nodes found.
+	check(LOD >= 0);
+
+	TArray<UCustomizableObjectNodeObject*> Result;
+
+	TQueue<const UCustomizableObjectNode*> PotentialCustomizableNodeObjects;
+	PotentialCustomizableNodeObjects.Enqueue(this);
+
+	const UCustomizableObjectNode* CurrentElement;
+	while (PotentialCustomizableNodeObjects.Dequeue(CurrentElement))
+	{
+		for (UEdGraphPin* Pin : CurrentElement->GetAllNonOrphanPins())
+		{
+			if (Pin->Direction == EEdGraphPinDirection::EGPD_Output)
+			{
+				for (UEdGraphPin* LinkedPin : FollowOutputPinArray(*Pin))
+				{
+					UCustomizableObjectNode* Node = Cast<UCustomizableObjectNode>(LinkedPin->GetOwningNode());
+					check(Node); // All nodes inherit from UCustomizableObjectNode.
+					PotentialCustomizableNodeObjects.Enqueue(Node);
+
+					if (UCustomizableObjectNodeObject* CurrentCustomizableObjectNode = Cast<UCustomizableObjectNodeObject>(Node))
+					{
+						const int32 LODPin = CurrentCustomizableObjectNode->GetLOD(LinkedPin);
+						if (LODPin == -1 || LODPin == LOD) // Add to parents if it is not connected to a LOD pin or it is connected to the given LOD pin.
+						{
+							Result.Add(CurrentCustomizableObjectNode);
+
+							// Explore external CO.
+							if (CurrentCustomizableObjectNode->bIsBase)
+							{
+								if (UCustomizableObjectNodeObjectGroup* ObjectGroupNode = GetCustomizableObjectExternalNode<UCustomizableObjectNodeObjectGroup>(CurrentCustomizableObjectNode->ParentObject, CurrentCustomizableObjectNode->ParentObjectGroupId))
+								{
+									Result.Append(ObjectGroupNode->GetParentObjectNodes(LOD)); // Recursive call.
+								}
+							}
+						}
+					}
+				}
+			}
+		}
+	}
+
+	return Result;
+}
+
+
+void UCustomizableObjectNode::SetPinHidden(UEdGraphPin& Pin, bool bHidden)
+{
+	Pin.SafeSetHidden(bHidden && CanPinBeHidden(Pin));
+
+	GetGraph()->NotifyGraphChanged();
+}
+
+
+void UCustomizableObjectNode::SetPinHidden(const TArray<UEdGraphPin*>& PinsToHide, bool bHidden)
+{
+	for (UEdGraphPin* Pin : PinsToHide)
+	{
+		Pin->SafeSetHidden(bHidden && CanPinBeHidden(*Pin));
+	}
+
+	GetGraph()->NotifyGraphChanged();
+}
+
+
+bool UCustomizableObjectNode::CanPinBeHidden(const UEdGraphPin& Pin) const
+{
+	return !Pin.LinkedTo.Num() && !Pin.bOrphanedPin;
+}
+
+
+TSharedPtr<SWidget> UCustomizableObjectNode::CustomizePinDetails(UEdGraphPin& Pin)
+{
+	return nullptr;
+}
+
+
+UCustomizableObjectNodePinData* UCustomizableObjectNode::GetPinData(const UEdGraphPin& Pin) const
+{
+	if (TObjectPtr<UCustomizableObjectNodePinData> const* Result = PinsDataId.Find(Pin.PinId))
+	{
+		return *Result;	
+	}
+	else
+	{
+		return nullptr;
+	}
+}
+
+
+#undef LOCTEXT_NAMESPACE
