[mutable] Move Mutable plugin from the restricted folder to the experimental plugins folder.

diff --git a/Engine/Plugins/Experimental/Mutable/Source/CustomizableObjectEditor/Private/Nodes/CustomizableObjectNodeAnimationPose.cpp b/Engine/Plugins/Experimental/Mutable/Source/CustomizableObjectEditor/Private/Nodes/CustomizableObjectNodeAnimationPose.cpp
new file mode 100644
--- /dev/null
+++ b/Engine/Plugins/Experimental/Mutable/Source/CustomizableObjectEditor/Private/Nodes/CustomizableObjectNodeAnimationPose.cpp
@@ -0,0 +1,181 @@
+// Copyright Epic Games, Inc. All Rights Reserved.
+
+#include "Nodes/CustomizableObjectNodeAnimationPose.h"
+#include "CustomizableObjectEditorModule.h"
+#include "Animation/AnimInstanceProxy.h"
+#include "Animation/AnimInstance.h"
+#include "Animation/DebugSkelMeshComponent.h"
+#include "Animation/AnimInstance.h"
+#include "Animation/AnimSingleNodeInstance.h"
+#include "Animation/AnimationPoseData.h"
+
+#include "IPersonaPreviewScene.h"
+
+#define LOCTEXT_NAMESPACE "CustomizableObjectEditor"
+
+
+PRAGMA_DISABLE_DEPRECATION_WARNINGS
+
+
+UCustomizableObjectNodeAnimationPose::UCustomizableObjectNodeAnimationPose() : Super()
+	, PoseAsset(nullptr)
+{
+}
+
+
+void UCustomizableObjectNodeAnimationPose::PostEditChangeProperty(FPropertyChangedEvent& PropertyChangedEvent)
+{
+	Super::PostEditChangeProperty(PropertyChangedEvent);
+}
+
+
+void UCustomizableObjectNodeAnimationPose::AllocateDefaultPins(UCustomizableObjectNodeRemapPins* RemapPins)
+{
+	const UEdGraphSchema_CustomizableObject* Schema = GetDefault<UEdGraphSchema_CustomizableObject>();
+
+	FString PinName = TEXT("Input Mesh");
+	UEdGraphPin* PinImagePin = CustomCreatePin(EGPD_Input, Schema->PC_Mesh, FName(*PinName));
+	PinImagePin->bDefaultValueIsIgnored = true;
+
+	PinName = TEXT("Output Mesh");
+	PinImagePin = CustomCreatePin(EGPD_Output, Schema->PC_Mesh, FName(*PinName));
+	PinImagePin->bDefaultValueIsIgnored = true;
+}
+
+
+UEdGraphPin* UCustomizableObjectNodeAnimationPose::GetInputMeshPin() const
+{
+	FString PinName = FString(TEXT("Input Mesh"));
+
+	UEdGraphPin* Pin = FindPin(PinName);
+
+	return Pin;
+}
+
+
+FText UCustomizableObjectNodeAnimationPose::GetNodeTitle(ENodeTitleType::Type TitleType) const
+{
+	if (PoseAsset)
+	{
+		FFormatNamedArguments Args;
+		Args.Add(TEXT("SkeletamMeshName"), FText::FromString(PoseAsset->GetName()));
+
+		return FText::Format(LOCTEXT("AnimationPose_Title", "{SkeletamMeshName}\nAnimation Pose"), Args);
+	}
+	else
+	{
+		return LOCTEXT("PoseMesh", "Pose Mesh");
+	}
+}
+
+
+FLinearColor UCustomizableObjectNodeAnimationPose::GetNodeTitleColor() const
+{
+	const UEdGraphSchema_CustomizableObject* Schema = GetDefault<UEdGraphSchema_CustomizableObject>();
+	return Schema->GetPinTypeColor(Schema->PC_Mesh);
+}
+
+
+void UCustomizableObjectNodeAnimationPose::StaticRetrievePoseInformation(UPoseAsset* PoseAsset, USkeletalMesh* RefSkeletalMesh, TArray<FString>& OutArrayBoneName, TArray<FTransform>& OutArrayTransform)
+{
+
+	// Need this for FCompactPose
+	FMemMark Mark(FMemStack::Get());
+	
+	UDebugSkelMeshComponent* SkeletalMeshComponent = NewObject<UDebugSkelMeshComponent>();
+	SkeletalMeshComponent->SetSkeletalMesh(RefSkeletalMesh);
+	SkeletalMeshComponent->SetAnimationMode(EAnimationMode::AnimationBlueprint);
+	SkeletalMeshComponent->AllocateTransformData();
+	SkeletalMeshComponent->SetAnimation(PoseAsset);
+	SkeletalMeshComponent->RefreshBoneTransforms();
+	SkeletalMeshComponent->InitAnim(false);
+
+	UAnimInstance* AnimInstance = SkeletalMeshComponent->GetAnimInstance();
+	FBoneContainer& BoneContainerCopy = AnimInstance->GetRequiredBones();
+
+	FCompactPose OutPose; // Needs a FMemMark declared before in the stack context so that the memory allocated by the FCompactPose is freed correctly
+	OutPose.SetBoneContainer(&BoneContainerCopy);
+
+	FBlendedCurve OutCurve;
+	FStackCustomAttributes OutAttributes;
+	FAnimationPoseData OutAnimData(OutPose, OutCurve, OutAttributes);
+	PoseAsset->GetBaseAnimationPose(OutAnimData);
+	OutPose = OutAnimData.GetPose();
+	OutCurve = OutAnimData.GetCurve();
+
+	FBlendedHeapCurve AnimCurves = SkeletalMeshComponent->AnimCurves;
+	const int32 NumElement = AnimCurves.CurveWeights.Num();
+
+	for (int i = 0; i < NumElement; ++i)
+	{
+		OutCurve.CurveWeights.Add(AnimCurves.CurveWeights[i]);
+	}
+
+	OutCurve.bInitialized = AnimCurves.bInitialized;
+
+	// Assuming one single pose, with a weigth set to 1.0
+	FAnimExtractContext ExtractionContext;
+	ExtractionContext.bExtractRootMotion = false;
+	ExtractionContext.CurrentTime = 0.0f;
+
+#if (ENGINE_MAJOR_VERSION == 4 && ENGINE_MINOR_VERSION >= 24) || ENGINE_MAJOR_VERSION >= 5
+	OutCurve.UIDToArrayIndexLUT = AnimCurves.UIDToArrayIndexLUT;
+
+	const TArray<FSmartName>& PoseNames = PoseAsset->GetPoseNames();
+	ExtractionContext.PoseCurves.Add(FPoseCurve(0, PoseNames[0].UID, 1.0f));
+	FAnimationPoseData SecondOutAnimData(OutPose, OutCurve, OutAttributes);
+	PoseAsset->GetAnimationPose(SecondOutAnimData, ExtractionContext);
+	OutPose = OutAnimData.GetPose();
+	OutCurve = OutAnimData.GetCurve();
+#else
+	OutCurve.UIDList = AnimCurves.UIDList;
+
+	ExtractionContext.PoseCurves.Add(1.0f);
+	PoseAsset->GetAnimationPose(OutPose, OutCurve, ExtractionContext);
+#endif
+
+	const TArray<FTransform, FAnimStackAllocator>& ArrayPoseBoneTransform = OutPose.GetBones();
+	const TArray<FBoneIndexType>& ArrayPoseBoneIndices = OutPose.GetBoneContainer().GetBoneIndicesArray();
+
+	for (int32 ArrayIndex = 0; ArrayIndex < ArrayPoseBoneIndices.Num(); ArrayIndex++)
+	{
+		int32 const& PoseBoneIndex = ArrayPoseBoneIndices[ArrayIndex];
+
+		checkSlow(PoseBoneIndex != INDEX_NONE);
+
+		FTransform CumulativePoseTransform;
+		int32 ParentIndex = PoseBoneIndex;
+
+		while (ParentIndex > 0)
+		{
+			for (int32 IndicesIndex = 0; IndicesIndex < ArrayPoseBoneIndices.Num(); ++IndicesIndex)
+			{
+				const uint16 BoneIndex = ArrayPoseBoneIndices[IndicesIndex];
+				if (BoneIndex == ParentIndex)
+				{
+					CumulativePoseTransform = CumulativePoseTransform * ArrayPoseBoneTransform[IndicesIndex];
+					break;
+				} 
+				else if (BoneIndex > ParentIndex)
+				{
+					break; // Break for, ArrayPoseBoneIndices is sorted.
+				}
+			}
+
+			ParentIndex = RefSkeletalMesh->GetRefSkeleton().GetParentIndex(ParentIndex);
+		}
+
+		FString Name = RefSkeletalMesh->GetRefSkeleton().GetBoneName(PoseBoneIndex).ToString();
+		int32 SMBoneIndex = SkeletalMeshComponent->GetBoneIndex(FName(*Name));
+		FTransform BoneToComponentTranform = SkeletalMeshComponent->GetEditableComponentSpaceTransforms()[SMBoneIndex];
+		FTransform TransformToAdd = BoneToComponentTranform.Inverse() * CumulativePoseTransform;
+
+		OutArrayBoneName.Add(Name);
+		OutArrayTransform.Add(TransformToAdd);
+	}
+
+}
+
+#undef LOCTEXT_NAMESPACE
+
+PRAGMA_ENABLE_DEPRECATION_WARNINGS
diff --git a/Engine/Plugins/Experimental/Mutable/Source/CustomizableObjectEditor/Private/Nodes/CustomizableObjectNodeColorArithmeticOp.cpp b/Engine/Plugins/Experimental/Mutable/Source/CustomizableObjectEditor/Private/Nodes/CustomizableObjectNodeColorArithmeticOp.cpp
new file mode 100644
--- /dev/null
+++ b/Engine/Plugins/Experimental/Mutable/Source/CustomizableObjectEditor/Private/Nodes/CustomizableObjectNodeColorArithmeticOp.cpp
@@ -0,0 +1,80 @@
+// Copyright Epic Games, Inc. All Rights Reserved.
+
+#include "Nodes/CustomizableObjectNodeColorArithmeticOp.h"
+#include "CustomizableObjectEditorModule.h"
+
+#define LOCTEXT_NAMESPACE "CustomizableObjectEditor"
+
+
+UCustomizableObjectNodeColorArithmeticOp::UCustomizableObjectNodeColorArithmeticOp()
+	: Super()
+{
+	Operation = EColorArithmeticOperation::E_Add;
+}
+
+
+void UCustomizableObjectNodeColorArithmeticOp::PostEditChangeProperty(FPropertyChangedEvent& PropertyChangedEvent)
+{
+	//UProperty* PropertyThatChanged = PropertyChangedEvent.Property;
+	//if ( PropertyThatChanged && PropertyThatChanged->GetName() == TEXT("NumLODs") )
+	{
+	}
+
+	Super::PostEditChangeProperty(PropertyChangedEvent);
+}
+
+
+void UCustomizableObjectNodeColorArithmeticOp::AllocateDefaultPins(UCustomizableObjectNodeRemapPins* RemapPins)
+{
+	const UEdGraphSchema_CustomizableObject* Schema = GetDefault<UEdGraphSchema_CustomizableObject>();
+
+	UEdGraphPin* XPin = CustomCreatePin(EGPD_Input, Schema->PC_Color, FName("A"));
+	XPin->bDefaultValueIsIgnored = true;
+	XPin->PinFriendlyName = FText();
+
+	FName PinName = TEXT("B");
+	UEdGraphPin* YPin = CustomCreatePin(EGPD_Input, Schema->PC_Color, PinName);
+	YPin->bDefaultValueIsIgnored = true;
+	YPin->PinFriendlyName = FText();
+
+	PinName = TEXT("Result");
+	UEdGraphPin* ColorPin = CustomCreatePin(EGPD_Output, Schema->PC_Color,PinName);
+	ColorPin->bDefaultValueIsIgnored = true;
+	ColorPin->PinFriendlyName = FText();
+}
+
+
+FText UCustomizableObjectNodeColorArithmeticOp::GetNodeTitle(ENodeTitleType::Type TitleType) const
+{
+	if (TitleType == ENodeTitleType::ListView)
+	{
+		return LOCTEXT("Color_Arithmetic_Operation", "Color Arithmetic Operation");
+	}
+
+	const UEnum* EnumPtr = FindObject<UEnum>(nullptr, TEXT("/Script/CustomizableObjectEditor.EColorArithmeticOperation"), true);
+
+	if (!EnumPtr)
+	{
+		return FText::FromString(FString("Color Operation"));
+	}
+
+	const int32 index = EnumPtr->GetIndexByValue((int32)Operation);
+	return EnumPtr->GetDisplayNameTextByIndex(index);
+}
+
+
+FLinearColor UCustomizableObjectNodeColorArithmeticOp::GetNodeTitleColor() const
+{
+	const UEdGraphSchema_CustomizableObject* Schema = GetDefault<UEdGraphSchema_CustomizableObject>();
+	return Schema->GetPinTypeColor(Schema->PC_Color);
+}
+
+
+FText UCustomizableObjectNodeColorArithmeticOp::GetTooltipText() const
+{
+	return LOCTEXT("Color_Arithmetic_Tooltip", "Perform an arithmetic operation between two colors on a per-component basis.");
+}
+
+
+#undef LOCTEXT_NAMESPACE
+
diff --git a/Engine/Plugins/Experimental/Mutable/Source/CustomizableObjectEditor/Private/Nodes/CustomizableObjectNodeColorFromFloats.cpp b/Engine/Plugins/Experimental/Mutable/Source/CustomizableObjectEditor/Private/Nodes/CustomizableObjectNodeColorFromFloats.cpp
new file mode 100644
--- /dev/null
+++ b/Engine/Plugins/Experimental/Mutable/Source/CustomizableObjectEditor/Private/Nodes/CustomizableObjectNodeColorFromFloats.cpp
@@ -0,0 +1,63 @@
+// Copyright Epic Games, Inc. All Rights Reserved.
+
+#include "Nodes/CustomizableObjectNodeColorFromFloats.h"
+#include "CustomizableObjectEditorModule.h"
+
+#define LOCTEXT_NAMESPACE "CustomizableObjectEditor"
+
+
+UCustomizableObjectNodeColorFromFloats::UCustomizableObjectNodeColorFromFloats()
+	: Super()
+{
+
+}
+
+
+void UCustomizableObjectNodeColorFromFloats::PostEditChangeProperty(FPropertyChangedEvent& PropertyChangedEvent)
+{
+	Super::PostEditChangeProperty(PropertyChangedEvent);
+}
+
+
+void UCustomizableObjectNodeColorFromFloats::AllocateDefaultPins(UCustomizableObjectNodeRemapPins* RemapPins)
+{
+	const UEdGraphSchema_CustomizableObject* Schema = GetDefault<UEdGraphSchema_CustomizableObject>();
+
+	UEdGraphPin* RPin = CustomCreatePin(EGPD_Input, Schema->PC_Float, FName("R"));
+	RPin->bDefaultValueIsIgnored = true;
+
+	UEdGraphPin* GPin = CustomCreatePin(EGPD_Input, Schema->PC_Float, FName("G"));
+	GPin->bDefaultValueIsIgnored = true;
+
+	UEdGraphPin* BPin = CustomCreatePin(EGPD_Input, Schema->PC_Float, FName("B"));
+	BPin->bDefaultValueIsIgnored = true;
+
+	UEdGraphPin* APin = CustomCreatePin(EGPD_Input, Schema->PC_Float, FName("A"));
+	APin->bDefaultValueIsIgnored = true;
+
+	UEdGraphPin* ColorPin = CustomCreatePin(EGPD_Output, Schema->PC_Color, FName("Color"));
+	ColorPin->bDefaultValueIsIgnored = true;
+}
+
+
+FText UCustomizableObjectNodeColorFromFloats::GetNodeTitle(ENodeTitleType::Type TitleType) const
+{
+	return LOCTEXT("Color_From_Floats", "Color From Floats");
+}
+
+
+FLinearColor UCustomizableObjectNodeColorFromFloats::GetNodeTitleColor() const
+{
+	const UEdGraphSchema_CustomizableObject* Schema = GetDefault<UEdGraphSchema_CustomizableObject>();
+	return Schema->GetPinTypeColor(Schema->PC_Color);
+}
+
+
+FText UCustomizableObjectNodeColorFromFloats::GetTooltipText() const
+{
+	return LOCTEXT("Color_From_Floats_Tooltip", "Defines a color from its four RGBA components.");
+}
+
+
+#undef LOCTEXT_NAMESPACE
+
diff --git a/Engine/Plugins/Experimental/Mutable/Source/CustomizableObjectEditor/Private/Nodes/CustomizableObjectNodeColorParameter.cpp b/Engine/Plugins/Experimental/Mutable/Source/CustomizableObjectEditor/Private/Nodes/CustomizableObjectNodeColorParameter.cpp
new file mode 100644
--- /dev/null
+++ b/Engine/Plugins/Experimental/Mutable/Source/CustomizableObjectEditor/Private/Nodes/CustomizableObjectNodeColorParameter.cpp
@@ -0,0 +1,38 @@
+// Copyright Epic Games, Inc. All Rights Reserved.
+
+#include "Nodes/CustomizableObjectNodeColorParameter.h"
+#include "CustomizableObjectEditorModule.h"
+
+#define LOCTEXT_NAMESPACE "CustomizableObjectEditor"
+
+
+void UCustomizableObjectNodeColorParameter::AllocateDefaultPins(UCustomizableObjectNodeRemapPins* RemapPins)
+{
+	const UEdGraphSchema_CustomizableObject* Schema = GetDefault<UEdGraphSchema_CustomizableObject>();
+
+	UEdGraphPin* ValuePin = CustomCreatePin(EGPD_Output, Schema->PC_Color, FName("Value"));
+	ValuePin->bDefaultValueIsIgnored = true;
+}
+
+
+FText UCustomizableObjectNodeColorParameter::GetNodeTitle(ENodeTitleType::Type TitleType) const
+{
+	return LOCTEXT("Color Parameter", "Color Parameter");
+}
+
+
+FLinearColor UCustomizableObjectNodeColorParameter::GetNodeTitleColor() const
+{
+	const UEdGraphSchema_CustomizableObject* Schema = GetDefault<UEdGraphSchema_CustomizableObject>();
+	return Schema->GetPinTypeColor(Schema->PC_Color);
+}
+
+
+FText UCustomizableObjectNodeColorParameter::GetTooltipText() const
+{
+	return LOCTEXT("Color_Parameter_Tooltip", "Expose a runtime modifiable color parameter from the Customizable Object.");
+}
+
+
+#undef LOCTEXT_NAMESPACE
+
diff --git a/Engine/Plugins/Experimental/Mutable/Source/CustomizableObjectEditor/Private/Nodes/CustomizableObjectNodeColorSwitch.cpp b/Engine/Plugins/Experimental/Mutable/Source/CustomizableObjectEditor/Private/Nodes/CustomizableObjectNodeColorSwitch.cpp
new file mode 100644
--- /dev/null
+++ b/Engine/Plugins/Experimental/Mutable/Source/CustomizableObjectEditor/Private/Nodes/CustomizableObjectNodeColorSwitch.cpp
@@ -0,0 +1,40 @@
+// Copyright Epic Games, Inc. All Rights Reserved.
+
+#include "Nodes/CustomizableObjectNodeColorSwitch.h"
+
+#include "EdGraphSchema_CustomizableObject.h"
+
+#define LOCTEXT_NAMESPACE "CustomizableObjectEditor"
+
+
+void UCustomizableObjectNodeColorSwitch::Serialize(FArchive& Ar)
+{
+	Super::Serialize(Ar);
+
+	if (Ar.CustomVer(FCustomizableObjectCustomVersion::GUID) < FCustomizableObjectCustomVersion::BugPinsSwitch)
+	{
+		OutputPinReference = FindPin(TEXT("Color"));	
+	}
+}
+
+
+FString UCustomizableObjectNodeColorSwitch::GetOutputPinName() const
+{
+	return TEXT("Color");
+}
+	
+
+FName UCustomizableObjectNodeColorSwitch::GetCategory() const
+{
+	return UEdGraphSchema_CustomizableObject::PC_Color;
+}
+
+
+FString UCustomizableObjectNodeColorSwitch::GetPinPrefix() const
+{
+	return TEXT("Color ");
+}
+
+
+#undef LOCTEXT_NAMESPACE
+
diff --git a/Engine/Plugins/Experimental/Mutable/Source/CustomizableObjectEditor/Private/Nodes/CustomizableObjectNodeColorVariation.cpp b/Engine/Plugins/Experimental/Mutable/Source/CustomizableObjectEditor/Private/Nodes/CustomizableObjectNodeColorVariation.cpp
new file mode 100644
--- /dev/null
+++ b/Engine/Plugins/Experimental/Mutable/Source/CustomizableObjectEditor/Private/Nodes/CustomizableObjectNodeColorVariation.cpp
@@ -0,0 +1,73 @@
+// Copyright Epic Games, Inc. All Rights Reserved.
+
+#include "Nodes/CustomizableObjectNodeColorVariation.h"
+#include "CustomizableObjectEditorModule.h"
+
+#include "UnrealEditorPortabilityHelpers.h"
+
+#define LOCTEXT_NAMESPACE "CustomizableObjectEditor"
+
+
+UCustomizableObjectNodeColorVariation::UCustomizableObjectNodeColorVariation()
+	: Super()
+{
+
+}
+
+
+void UCustomizableObjectNodeColorVariation::PostEditChangeProperty(FPropertyChangedEvent& PropertyChangedEvent)
+{
+	FProperty* PropertyThatChanged = PropertyChangedEvent.Property;
+	if ( PropertyThatChanged  )
+	{
+		ReconstructNode();
+	}
+
+	Super::PostEditChangeProperty(PropertyChangedEvent);
+}
+
+
+void UCustomizableObjectNodeColorVariation::AllocateDefaultPins(UCustomizableObjectNodeRemapPins* RemapPins)
+{
+	const UEdGraphSchema_CustomizableObject* Schema = GetDefault<UEdGraphSchema_CustomizableObject>();
+
+	FString RemapIndex = Helper_GetPinName(OutputPin());
+
+	UEdGraphPin* OutputPin = CustomCreatePin(EGPD_Output, Schema->PC_Color, FName(TEXT("Color")));
+	OutputPin->bDefaultValueIsIgnored = true;
+
+	for ( int VariationIndex = Variations.Num() - 1; VariationIndex >= 0; --VariationIndex)
+	{
+		FString PinName = FString::Printf( TEXT("Variation %d"), VariationIndex);
+		UEdGraphPin* VariationPin = CustomCreatePin(EGPD_Input, Schema->PC_Color, FName(*PinName), false);
+		VariationPin->bDefaultValueIsIgnored = true;
+
+		FString FriendlyName = FString::Printf(TEXT("Variation %d [%s]"), VariationIndex, *Variations[VariationIndex].Tag);
+		VariationPin->PinFriendlyName = FText::FromString(FriendlyName);
+	}
+
+	UEdGraphPin* DefaultVariation = CustomCreatePin(EGPD_Input, Schema->PC_Color, FName(TEXT("Default")), false);
+	DefaultVariation->bDefaultValueIsIgnored = true;
+}
+
+
+FText UCustomizableObjectNodeColorVariation::GetNodeTitle(ENodeTitleType::Type TitleType) const
+{
+	return LOCTEXT("Color_Variation", "Color Variation");
+}
+
+
+FLinearColor UCustomizableObjectNodeColorVariation::GetNodeTitleColor() const
+{
+	const UEdGraphSchema_CustomizableObject* Schema = GetDefault<UEdGraphSchema_CustomizableObject>();
+	return Schema->GetPinTypeColor(Schema->PC_Color);
+}
+
+
+FText UCustomizableObjectNodeColorVariation::GetTooltipText() const
+{
+	return LOCTEXT("Color_Variation_Tooltip", "Select a color depending on what tags are active.");
+}
+
+#undef LOCTEXT_NAMESPACE
+
diff --git a/Engine/Plugins/Experimental/Mutable/Source/CustomizableObjectEditor/Private/Nodes/CustomizableObjectNodeCopyMaterial.cpp b/Engine/Plugins/Experimental/Mutable/Source/CustomizableObjectEditor/Private/Nodes/CustomizableObjectNodeCopyMaterial.cpp
new file mode 100644
--- /dev/null
+++ b/Engine/Plugins/Experimental/Mutable/Source/CustomizableObjectEditor/Private/Nodes/CustomizableObjectNodeCopyMaterial.cpp
@@ -0,0 +1,147 @@
+// Copyright Epic Games, Inc. All Rights Reserved.
+
+#include "Nodes/CustomizableObjectNodeCopyMaterial.h"
+#include "Nodes/CustomizableObjectNodeExternalPin.h"
+#include "Nodes/CustomizableObjectNodeObject.h"
+#include "EdGraphSchema_CustomizableObject.h"
+#include "GenerateMutableSource/GenerateMutableSource.h"
+#include "GraphTraversal.h"
+
+#define LOCTEXT_NAMESPACE "CustomizableObjectEditor"
+
+
+/** Mesh input pin key */
+const static FString MeshPinName = TEXT("Mesh_Input_Pin");
+
+/** Material input pin key */
+const static FString MaterialPinName = TEXT("Material_Input_Pin");
+
+FText UCustomizableObjectNodeCopyMaterial::GetNodeTitle(ENodeTitleType::Type TitleType) const
+{
+	return LOCTEXT("Copy_Material", "Copy Material");
+}
+
+void UCustomizableObjectNodeCopyMaterial::AllocateDefaultPins(UCustomizableObjectNodeRemapPins* RemapPins)
+{
+	if (UCustomizableObjectNodeMaterialRemapPinsByName* RemapPinsByNameCustom = Cast<UCustomizableObjectNodeMaterialRemapPinsByName>(RemapPins))
+	{
+		RemapPinsByNameCustom->Node = this;
+	}
+	
+	const UEdGraphSchema_CustomizableObject* Schema = GetDefault<UEdGraphSchema_CustomizableObject>();
+
+	// Input pins
+	FString PinFriendlyName = TEXT("Mesh");
+	UEdGraphPin* Pin = CustomCreatePin(EGPD_Input, Schema->PC_Mesh, FName(MeshPinName));
+	Pin->PinFriendlyName = FText::FromString(PinFriendlyName);
+	Pin->bDefaultValueIsIgnored = true;
+
+	PinFriendlyName = TEXT("Base Material");
+	Pin = CustomCreatePin(EGPD_Input, Schema->PC_Material, FName(MaterialPinName));
+	Pin->PinFriendlyName = FText::FromString(PinFriendlyName);
+	Pin->bDefaultValueIsIgnored = true;
+
+	// Output pins
+	PinFriendlyName = TEXT("Material");
+	FString PinName = PinFriendlyName + FString(TEXT("_Output_Pin"));
+	Pin = CustomCreatePin(EGPD_Output, Schema->PC_Material, FName(*PinName));
+	Pin->PinFriendlyName = FText::FromString(PinFriendlyName);
+}
+
+UEdGraphPin* UCustomizableObjectNodeCopyMaterial::GetMeshPin() const
+{
+	return FindPin(MeshPinName);
+}
+
+UEdGraphPin* UCustomizableObjectNodeCopyMaterial::GetMaterialPin() const
+{
+	return FindPin(MaterialPinName);
+}
+
+UCustomizableObjectNodeSkeletalMesh* UCustomizableObjectNodeCopyMaterial::GetMeshNode() const
+{
+	UCustomizableObjectNodeSkeletalMesh* Result = nullptr;
+
+	UEdGraphPin* MaterialPin = GetMeshPin();
+	if (const UEdGraphPin* ConnectedPin = FollowInputPin(*MaterialPin))
+	{
+		const UEdGraphPin* SourceMeshPin = FindMeshBaseSource(*ConnectedPin, false);
+		if (SourceMeshPin)
+		{
+			Result = Cast<UCustomizableObjectNodeSkeletalMesh>(SourceMeshPin->GetOwningNode());
+		}
+	}
+
+	return Result;
+}
+
+
+UCustomizableObjectNodeMaterial* UCustomizableObjectNodeCopyMaterial::GetMaterialNode() const
+{
+	UCustomizableObjectNodeMaterial* Result = nullptr;
+
+	UEdGraphPin* MaterialPin = GetMaterialPin();
+	if (const UEdGraphPin* ConnectedPin = FollowInputPin(*MaterialPin))
+	{
+		return Cast<UCustomizableObjectNodeMaterial>(ConnectedPin->GetOwningNode());
+	}
+
+	return Result;
+}
+
+bool UCustomizableObjectNodeCopyMaterial::CanConnect(const UEdGraphPin* InOwnedInputPin, const UEdGraphPin* InOutputPin, bool& bOutIsOtherNodeBlocklisted, bool& bOutArePinsCompatible) const
+{
+	if (!Super::CanConnect(InOwnedInputPin, InOutputPin, bOutIsOtherNodeBlocklisted, bOutIsOtherNodeBlocklisted))
+	{
+		return false;
+	}
+
+	if (InOwnedInputPin == GetMaterialPin())
+	{
+		const UEdGraphNode* OuputPinOwningNode = InOutputPin->GetOwningNode();
+		return (OuputPinOwningNode->IsA(UCustomizableObjectNodeMaterial::StaticClass()) && !OuputPinOwningNode->IsA(UCustomizableObjectNodeCopyMaterial::StaticClass()))
+			|| OuputPinOwningNode->IsA(UCustomizableObjectNodeExternalPin::StaticClass());
+	}
+
+	return true;
+}
+
+bool UCustomizableObjectNodeCopyMaterial::ShouldBreakExistingConnections(const UEdGraphPin* InputPin, const UEdGraphPin* OutputPin) const
+{
+	return true;
+}
+
+bool UCustomizableObjectNodeCopyMaterial::IsNodeOutDatedAndNeedsRefresh()
+{
+	return false;
+}
+
+bool UCustomizableObjectNodeCopyMaterial::ProvidesCustomPinRelevancyTest() const
+{
+	return true;
+}
+
+bool UCustomizableObjectNodeCopyMaterial::IsPinRelevant(const UEdGraphPin* Pin) const
+{
+	const UEdGraphNode* Node = Pin->GetOwningNode();
+
+	if (Pin->Direction == EGPD_Output)
+	{
+		return (Node->IsA(UCustomizableObjectNodeMaterial::StaticClass()) && !Node->IsA(UCustomizableObjectNodeCopyMaterial::StaticClass())) ||
+			(Node->IsA(UCustomizableObjectNodeExternalPin::StaticClass()) && Pin->PinType.PinCategory == UEdGraphSchema_CustomizableObject::PC_Material) ||
+			Pin->PinType.PinCategory == UEdGraphSchema_CustomizableObject::PC_Mesh;
+	}
+	else
+	{
+		return Node->IsA(UCustomizableObjectNodeObject::StaticClass()) ||
+			(Node->IsA(UCustomizableObjectNodeExposePin::StaticClass()) && Pin->PinType.PinCategory == UEdGraphSchema_CustomizableObject::PC_Material);
+			
+	}
+}
+
+FText UCustomizableObjectNodeCopyMaterial::GetTooltipText() const
+{
+	return LOCTEXT("CopyMaterial_Tooltip", "Copies a Customizable Object material.\nDuplicates all Material node input pins and properties except for the Mesh input pin.");
+}
+
+#undef LOCTEXT_NAMESPACE
diff --git a/Engine/Plugins/Experimental/Mutable/Source/CustomizableObjectEditor/Private/Nodes/CustomizableObjectNodeCopyMaterialDetails.cpp b/Engine/Plugins/Experimental/Mutable/Source/CustomizableObjectEditor/Private/Nodes/CustomizableObjectNodeCopyMaterialDetails.cpp
new file mode 100644
--- /dev/null
+++ b/Engine/Plugins/Experimental/Mutable/Source/CustomizableObjectEditor/Private/Nodes/CustomizableObjectNodeCopyMaterialDetails.cpp
@@ -0,0 +1,21 @@
+// Copyright Epic Games, Inc. All Rights Reserved.
+
+#include "CustomizableObjectNodeCopyMaterialDetails.h"
+
+#include "DetailLayoutBuilder.h"
+
+TSharedRef<IDetailCustomization> FCustomizableObjectNodeCopyMaterialDetails::MakeInstance()
+{
+	return MakeShareable(new FCustomizableObjectNodeCopyMaterialDetails);
+}
+
+void FCustomizableObjectNodeCopyMaterialDetails::CustomizeDetails(IDetailLayoutBuilder& DetailBuilder)
+{
+	TArray<FName> Names;
+	DetailBuilder.GetCategoryNames(Names);
+
+	for (const FName& Name : Names)
+	{
+		DetailBuilder.HideCategory(Name);
+	}
+}
diff --git a/Engine/Plugins/Experimental/Mutable/Source/CustomizableObjectEditor/Private/Nodes/CustomizableObjectNodeCopyMaterialDetails.h b/Engine/Plugins/Experimental/Mutable/Source/CustomizableObjectEditor/Private/Nodes/CustomizableObjectNodeCopyMaterialDetails.h
new file mode 100644
--- /dev/null
+++ b/Engine/Plugins/Experimental/Mutable/Source/CustomizableObjectEditor/Private/Nodes/CustomizableObjectNodeCopyMaterialDetails.h
@@ -0,0 +1,15 @@
+// Copyright Epic Games, Inc. All Rights Reserved.
+
+#pragma once
+
+#include "IDetailCustomization.h"
+
+/** Copy Material node details panel. Hides all properties from the inheret Material node. */
+class FCustomizableObjectNodeCopyMaterialDetails : public IDetailCustomization
+{
+public:
+	static TSharedRef<IDetailCustomization> MakeInstance();
+
+	/** Hides details copied from CustomizableObjectNodeMaterial. */
+	virtual void CustomizeDetails(IDetailLayoutBuilder& DetailBuilder) override;
+};
\ No newline at end of file
diff --git a/Engine/Plugins/Experimental/Mutable/Source/CustomizableObjectEditor/Private/Nodes/CustomizableObjectNodeEditMaterial.cpp b/Engine/Plugins/Experimental/Mutable/Source/CustomizableObjectEditor/Private/Nodes/CustomizableObjectNodeEditMaterial.cpp
new file mode 100644
--- /dev/null
+++ b/Engine/Plugins/Experimental/Mutable/Source/CustomizableObjectEditor/Private/Nodes/CustomizableObjectNodeEditMaterial.cpp
@@ -0,0 +1,289 @@
+// Copyright Epic Games, Inc. All Rights Reserved.
+
+#include "Nodes/CustomizableObjectNodeEditMaterial.h"
+#include "CustomizableObjectEditorModule.h"
+#include "EdGraphSchema_CustomizableObject.h"
+#include "Nodes/CustomizableObjectNodeExtendMaterial.h"
+#include "Nodes/CustomizableObjectNodeMaterial.h"
+#include "CustomizableObjectLayout.h"
+
+#define LOCTEXT_NAMESPACE "CustomizableObjectEditor"
+
+
+void UCustomizableObjectNodeEditMaterial::AllocateDefaultPins(UCustomizableObjectNodeRemapPins* RemapPins)
+{
+	const UEdGraphSchema_CustomizableObject* Schema = GetDefault<UEdGraphSchema_CustomizableObject>();
+	
+	if (const UCustomizableObjectNodeMaterial* ParentMaterialNode = GetParentMaterialNodeIfPath())
+	{
+		const int32 NumImages = ParentMaterialNode->GetNumParameters(EMaterialParameterType::Texture);
+		for (int32 ImageIndex = 0; ImageIndex < NumImages; ++ImageIndex)
+		{
+			if (ParentMaterialNode->IsImageMutableMode(ImageIndex))
+			{
+				UCustomizableObjectNodeEditMaterialPinEditImageData* PinEditImageData = NewObject<UCustomizableObjectNodeEditMaterialPinEditImageData>(this);
+				PinEditImageData->ImageId = ParentMaterialNode->GetParameterId(EMaterialParameterType::Texture, ImageIndex);
+			
+				const FName ImageName = ParentMaterialNode->GetParameterName(EMaterialParameterType::Texture, ImageIndex);
+				UEdGraphPin* PinImage = CustomCreatePin(EGPD_Input, Schema->PC_Image, ImageName, PinEditImageData);
+				PinImage->bHidden = true;
+				PinImage->bDefaultValueIsIgnored = true;
+
+				PinsParameter.Add(PinEditImageData->ImageId, FEdGraphPinReference(PinImage));
+
+				FString PinMaskName = ImageName.ToString() + " Mask";
+				UEdGraphPin* PinMask = CustomCreatePin(EGPD_Input, Schema->PC_Image, *PinMaskName);
+				PinMask->bHidden = true;
+				PinMask->bDefaultValueIsIgnored = true;
+			
+				PinEditImageData->PinMask = FEdGraphPinReference(PinMask);
+			}
+		}
+	}
+
+	CustomCreatePin(EGPD_Output, Schema->PC_Material, TEXT("Material"));
+}
+
+
+void UCustomizableObjectNodeEditMaterial::RemapPins(const TMap<UEdGraphPin*, UEdGraphPin*>& PinsToRemap)
+{
+	Super::RemapPins(PinsToRemap);
+
+	// Update mask pin references.
+	for (const UEdGraphPin* Pin : GetAllPins())
+	{
+		if (UCustomizableObjectNodeEditMaterialPinEditImageData* PinData = Cast<UCustomizableObjectNodeEditMaterialPinEditImageData>(GetPinData(*Pin)))
+		{
+			if (UEdGraphPin* const* Result = PinsToRemap.Find(PinData->PinMask.Get()))
+			{
+				PinData->PinMask = FEdGraphPinReference(*Result);
+			}
+		}
+	}
+}
+
+
+const UEdGraphPin* UCustomizableObjectNodeEditMaterial::GetUsedImageMaskPin(const FGuid& ImageId) const
+{
+	if (const UEdGraphPin* Pin = GetUsedImagePin(ImageId))
+	{
+		return GetPinData<UCustomizableObjectNodeEditMaterialPinEditImageData>(*Pin).PinMask.Get();
+	}
+
+	return nullptr;
+}
+
+
+bool UCustomizableObjectNodeEditMaterial::IsSingleOutputNode() const
+{
+	return true;
+}
+
+
+bool UCustomizableObjectNodeEditMaterial::CustomRemovePin(UEdGraphPin& Pin)
+{
+	CustomRemovePinWork(Pin);
+	
+	return Super::CustomRemovePin(Pin);
+}
+
+
+UEdGraphPin* UCustomizableObjectNodeEditMaterial::OutputPin() const
+{
+	return Super::OutputPin();
+}
+
+
+void UCustomizableObjectNodeEditMaterial::SetParentNode(UCustomizableObject* Object, FGuid NodeId)
+{
+	PreSetParentNodeWork(Object, NodeId);
+
+	Super::SetParentNode(Object, NodeId);
+	
+	PostSetParentNodeWork(Object, NodeId);
+	
+	BlockIds.Empty();
+}
+
+
+void UCustomizableObjectNodeEditMaterial::BackwardsCompatibleFixup()
+{
+	Super::BackwardsCompatibleFixup();
+
+	const int32 CustomizableObjectCustomVersion = GetLinkerCustomVersion(FCustomizableObjectCustomVersion::GUID);
+
+	// Convert deprecated node index list to the node id list.
+	if (CustomizableObjectCustomVersion < FCustomizableObjectCustomVersion::PostLoadToCustomVersion
+		&& BlockIds.Num() < Blocks_DEPRECATED.Num())
+	{
+		UCustomizableObjectNodeMaterialBase* ParentMaterialNode = GetParentMaterialNode();
+
+		TArray<UCustomizableObjectLayout*> Layouts = ParentMaterialNode->GetLayouts();
+
+		if (!Layouts.IsValidIndex(ParentLayoutIndex))
+		{
+			UE_LOG(LogMutable, Warning, TEXT("[%s] UCustomizableObjectNodeEditMaterial refers to an invalid texture layout index %d. Parent node has %d layouts."), 
+				*GetOutermost()->GetName(), ParentLayoutIndex, Layouts.Num());
+		}
+		else
+		{
+			UCustomizableObjectLayout* Layout = Layouts[ParentLayoutIndex];
+
+			if (Cast<UCustomizableObjectNodeMaterial>(ParentMaterialNode))
+			{
+				for (int IndexIndex = BlockIds.Num(); IndexIndex < Blocks_DEPRECATED.Num(); ++IndexIndex)
+				{
+					int BlockIndex = Blocks_DEPRECATED[IndexIndex];
+					if (Layout->Blocks.IsValidIndex(BlockIndex) )
+					{
+						const FGuid Id = Layout->Blocks[BlockIndex].Id;
+						if (Id.IsValid())
+						{
+							BlockIds.Add(Id);
+						}
+						else
+						{
+							UE_LOG(LogMutable, Warning, TEXT("[%s] UCustomizableObjectNodeEditMaterial refers to an valid layout block %d but that block doesn't have an id."),
+								*GetOutermost()->GetName(), BlockIndex );
+						}
+					}
+					else
+					{
+						UE_LOG(LogMutable, Warning, TEXT("[%s] UCustomizableObjectNodeEditMaterial refers to an invalid layout block index %d. Parent node has %d blocks."),
+							*GetOutermost()->GetName(), BlockIndex, Layout->Blocks.Num());
+					}
+				}
+			}
+		}
+	}
+	
+	if (CustomizableObjectCustomVersion < FCustomizableObjectCustomVersion::AutomaticNodeMaterial)
+	{
+		if (const UCustomizableObjectNodeMaterial* ParentMaterial = GetParentMaterialNode())
+		{
+			for (const FCustomizableObjectNodeEditMaterialImage& Image : Images_DEPRECATED)
+			{
+				const UEdGraphPin* ImagePin = FindPin(Image.Name);
+				const UEdGraphPin* PinMask = FindPin(Image.Name + " Mask");
+				if (!ImagePin || !PinMask)
+				{
+					continue;
+				}
+				
+				UCustomizableObjectNodeEditMaterialPinEditImageData*  PinEditImageData = NewObject<UCustomizableObjectNodeEditMaterialPinEditImageData>(this);
+				PinEditImageData->ImageId = FGuid::NewGuid();
+				PinEditImageData->PinMask = PinMask;
+				
+				// Search for the Image Id the Edit pin was referring to.
+				const int32 NumImages = ParentMaterial->GetNumParameters(EMaterialParameterType::Texture);
+				for (int32 ImageIndex = 0; ImageIndex < NumImages; ++ImageIndex)
+				{
+					if (ParentMaterial->GetParameterName(EMaterialParameterType::Texture, ImageIndex).ToString() == Image.Name)
+					{
+						PinEditImageData->ImageId = ParentMaterial->GetParameterId(EMaterialParameterType::Texture, ImageIndex);
+						break;
+					}
+				}
+				
+				AddPinData(*ImagePin, *PinEditImageData);
+			}
+		}
+		
+		Images_DEPRECATED.Empty();
+		ReconstructNode();
+	}
+
+	// Fill PinsParameter.
+	if (CustomizableObjectCustomVersion < FCustomizableObjectCustomVersion::AutomaticNodeMaterialPerformanceBug) // || CustomizableObjectCustomVersion < FCustomizableObjectCustomVersion::AutomaticNodeMaterialPerformance
+	{
+		for (const UEdGraphPin* Pin : GetAllNonOrphanPins())
+		{
+			if (const UCustomizableObjectNodeEditMaterialPinEditImageData* PinData = Cast<UCustomizableObjectNodeEditMaterialPinEditImageData>(GetPinData(*Pin)))
+			{
+				PinsParameter.Add(PinData->ImageId, FEdGraphPinReference(Pin));
+			}
+		}
+	}
+	
+	if (CustomizableObjectCustomVersion < FCustomizableObjectCustomVersion::EditMaterialOnlyMutableModeParameters)
+	{
+		ReconstructNode();
+	}
+}
+
+
+void UCustomizableObjectNodeEditMaterial::PostBackwardsCompatibleFixup()
+{
+	Super::PostBackwardsCompatibleFixup();
+	
+	PostBackwardsCompatibleFixupWork();
+}
+
+
+void UCustomizableObjectNodeEditMaterial::Serialize(FArchive& Ar)
+{
+	Super::Serialize(Ar);
+
+	Ar.UsingCustomVersion(FCustomizableObjectCustomVersion::GUID);
+}
+
+
+FText UCustomizableObjectNodeEditMaterial::GetNodeTitle(ENodeTitleType::Type TitleType) const
+{
+	return LOCTEXT("Edit_Material", "Edit Material");
+}
+
+
+FLinearColor UCustomizableObjectNodeEditMaterial::GetNodeTitleColor() const
+{
+	const UEdGraphSchema_CustomizableObject* Schema = GetDefault<UEdGraphSchema_CustomizableObject>();
+	return Schema->GetPinTypeColor(Schema->PC_Material);
+}
+
+
+bool UCustomizableObjectNodeEditMaterial::IsNodeOutDatedAndNeedsRefresh()
+{
+	return IsNodeOutDatedAndNeedsRefreshWork();
+}
+
+
+FString UCustomizableObjectNodeEditMaterial::GetRefreshMessage() const
+{
+	return "Source material has changed, texture channels might have been added, removed or renamed. Please refresh the parent material node to reflect those changes.";
+}
+
+
+FText UCustomizableObjectNodeEditMaterial::GetTooltipText() const
+{
+	return LOCTEXT("Edit_Material_Tooltip", "Modify the texture parameters of an ancestor's material partially or completely.");
+}
+
+
+void UCustomizableObjectNodeEditMaterial::PinConnectionListChanged(UEdGraphPin* Pin)
+{
+	Super::PinConnectionListChanged(Pin);
+	
+	PinConnectionListChangedWork(Pin);
+}
+
+
+UCustomizableObjectNode& UCustomizableObjectNodeEditMaterial::GetNode()
+{
+	return *this;
+}
+
+
+TMap<FGuid, FEdGraphPinReference>& UCustomizableObjectNodeEditMaterial::GetPinsParameter()
+{
+	return PinsParameter;	
+}
+
+
+FCustomizableObjectNodeParentedMaterial& UCustomizableObjectNodeEditMaterial::GetNodeParentedMaterial()
+{
+	return *this;	
+}
+
+
+#undef LOCTEXT_NAMESPACE
+
diff --git a/Engine/Plugins/Experimental/Mutable/Source/CustomizableObjectEditor/Private/Nodes/CustomizableObjectNodeEditMaterialBase.cpp b/Engine/Plugins/Experimental/Mutable/Source/CustomizableObjectEditor/Private/Nodes/CustomizableObjectNodeEditMaterialBase.cpp
new file mode 100644
--- /dev/null
+++ b/Engine/Plugins/Experimental/Mutable/Source/CustomizableObjectEditor/Private/Nodes/CustomizableObjectNodeEditMaterialBase.cpp
@@ -0,0 +1,65 @@
+// Copyright Epic Games, Inc. All Rights Reserved.
+
+#include "Nodes/CustomizableObjectNodeEditMaterialBase.h"
+#include "CustomizableObjectEditorModule.h"
+
+
+void UCustomizableObjectNodeEditMaterialBase::BeginPostDuplicate(bool bDuplicateForPIE)
+{
+	Super::BeginPostDuplicate(bDuplicateForPIE);
+
+	if (ParentMaterialObject)
+	{
+		if (UCustomizableObjectGraph* CEdGraph = Cast<UCustomizableObjectGraph>(GetGraph()))
+		{
+			ParentMaterialNodeId = CEdGraph->RequestNotificationForNodeIdChange(ParentMaterialNodeId, NodeGuid);
+		}
+	}
+}
+
+
+void UCustomizableObjectNodeEditMaterialBase::UpdateReferencedNodeId(const FGuid& NewGuid)
+{
+	if (ParentMaterialObject)
+	{
+		ParentMaterialNodeId = NewGuid;
+	}
+}
+
+
+UCustomizableObjectNode& UCustomizableObjectNodeEditMaterialBase::GetNode()
+{
+	return *this;
+}
+
+
+void UCustomizableObjectNodeEditMaterialBase::SetParentNode(UCustomizableObject* Object, FGuid NodeId)
+{
+	FCustomizableObjectNodeParentedMaterial::SetParentNode(Object, NodeId);
+	
+	// Update layout grid widget.
+	TSharedPtr<ICustomizableObjectEditor> Editor = GetGraphEditor();
+	if (Editor.IsValid())
+	{
+		Editor->UpdateGraphNodeProperties();
+	}
+}
+
+void UCustomizableObjectNodeEditMaterialBase::SaveParentNode(UCustomizableObject* Object, FGuid NodeId)
+{
+	ParentMaterialObject = Object;
+	ParentMaterialNodeId = NodeId;
+}
+
+
+FGuid UCustomizableObjectNodeEditMaterialBase::GetParentNodeId() const
+{
+	return ParentMaterialNodeId;
+}
+
+
+UCustomizableObject* UCustomizableObjectNodeEditMaterialBase::GetParentObject() const
+{
+	return ParentMaterialObject;
+}
+
diff --git a/Engine/Plugins/Experimental/Mutable/Source/CustomizableObjectEditor/Private/Nodes/CustomizableObjectNodeEditMaterialBaseDetails.cpp b/Engine/Plugins/Experimental/Mutable/Source/CustomizableObjectEditor/Private/Nodes/CustomizableObjectNodeEditMaterialBaseDetails.cpp
new file mode 100644
--- /dev/null
+++ b/Engine/Plugins/Experimental/Mutable/Source/CustomizableObjectEditor/Private/Nodes/CustomizableObjectNodeEditMaterialBaseDetails.cpp
@@ -0,0 +1,146 @@
+// Copyright Epic Games, Inc. All Rights Reserved.
+
+#include "CustomizableObjectNodeEditMaterialBaseDetails.h"
+#include "CustomizableObjectEditorModule.h"
+#include "PropertyCustomizationHelpers.h"
+#include "CustomizableObjectEditorUtilities.h"
+
+#include "DetailLayoutBuilder.h"
+#include "DetailCategoryBuilder.h"
+#include "LevelEditor.h"
+#include "LevelEditorActions.h"
+
+#include "Widgets/Text/STextBlock.h"
+#include "Widgets/Input/STextComboBox.h"
+#include "Templates/SharedPointer.h"
+#include "UObject/WeakObjectPtr.h"
+
+#include "Engine/SkeletalMesh.h"
+#include "Engine/StaticMesh.h"
+
+#include "Runtime/Launch/Resources/Version.h"
+
+
+#define LOCTEXT_NAMESPACE "CustomizableObjectDetails"
+
+
+TSharedRef<IDetailCustomization> FCustomizableObjectNodeEditMaterialBaseDetails::MakeInstance()
+{
+	return MakeShareable(new FCustomizableObjectNodeEditMaterialBaseDetails);
+}
+
+
+void FCustomizableObjectNodeEditMaterialBaseDetails::CustomizeDetails(IDetailLayoutBuilder& DetailBuilder)
+{
+	FCustomizableObjectNodeParentedMaterialDetails::CustomizeDetails(DetailBuilder);
+
+	NodeEditMaterialBase = nullptr;
+
+	const TArray<TWeakObjectPtr<UObject>>& SelectedObjects = DetailBuilder.GetDetailsView()->GetSelectedObjects();
+	if (SelectedObjects.Num())
+	{
+		NodeEditMaterialBase = Cast<UCustomizableObjectNodeEditMaterialBase>(SelectedObjects[0].Get());
+	}
+
+	IDetailCategoryBuilder& BlocksCategory = DetailBuilder.EditCategory("Parent");
+
+	if (NodeEditMaterialBase)
+	{
+		// Add all the materials to the combobox
+		TSharedPtr<FString> LayoutToSelect = GenerateLayoutComboboxOptions();
+
+		TSharedRef<IPropertyHandle> LayoutProperty = DetailBuilder.GetProperty("ParentLayoutIndex");
+
+		BlocksCategory.AddCustomRow(LOCTEXT("FCustomizableObjectNodeEditMaterialBaseDetails", "Blocks"))
+		[
+			SNew(SHorizontalBox)
+			+ SHorizontalBox::Slot()
+			.AutoWidth()
+			.Padding(0.0f, 5.0f, 6.0f, 0.0f)
+			[
+				SNew(STextBlock)
+				.Text(LOCTEXT("LayoutText", "Layout: "))
+			]
+
+			+ SHorizontalBox::Slot()
+			[
+				SNew(SProperty, LayoutProperty)
+				.ShouldDisplayName(false)
+				.CustomWidget()
+				[
+					SNew(SBorder)
+					.BorderImage(FAppStyle::GetBrush("NoBorder"))
+					.Padding(FMargin(0.0f, 0.0f, 10.0f, 0.0f))
+					[
+						SAssignNew(LayoutComboBox, STextComboBox)
+						.OptionsSource(&LayoutOptionNames)
+						.InitiallySelectedItem(LayoutToSelect)
+						.OnSelectionChanged(this, &FCustomizableObjectNodeEditMaterialBaseDetails::OnLayoutComboBoxSelectionChanged, LayoutProperty)
+					]
+				]
+			]
+		];
+	}
+	else
+	{
+		BlocksCategory.AddCustomRow(LOCTEXT("FCustomizableObjectNodeEditMaterialBaseDetails", "Node"))
+		[
+			SNew(STextBlock)
+			.Text(LOCTEXT("Node not found", "Node not found"))
+		];
+	}
+}
+
+
+TSharedPtr<FString> FCustomizableObjectNodeEditMaterialBaseDetails::GenerateLayoutComboboxOptions()
+{
+	LayoutOptionReferences.Empty();
+	LayoutOptionNames.Empty();
+
+	TSharedPtr<FString> ItemToSelect = nullptr;
+
+	UCustomizableObjectNodeMaterialBase* ParentMaterialNode = NodeEditMaterialBase->GetParentMaterialNode();
+
+	if (ParentMaterialNode && ParentMaterialNode->GetLayouts().Num())
+	{
+		for (int32 i = 0; i < ParentMaterialNode->GetLayouts().Num(); ++i)
+		{
+			LayoutOptionReferences.Add(i);
+			LayoutOptionNames.Add(MakeShareable(new FString(FString::FromInt(i))));
+
+			if (i == NodeEditMaterialBase->ParentLayoutIndex)
+			{
+				ItemToSelect = LayoutOptionNames.Last();
+			}
+		}
+	}
+
+	return ItemToSelect;
+}
+
+
+void FCustomizableObjectNodeEditMaterialBaseDetails::OnLayoutComboBoxSelectionChanged(TSharedPtr<FString> Selection, ESelectInfo::Type SelectInfo, TSharedRef<IPropertyHandle> LayoutProperty)
+{
+	check(LayoutOptionNames.Num() == LayoutOptionReferences.Num());
+	
+	for (int OptionIndex = 0; OptionIndex < LayoutOptionNames.Num(); ++OptionIndex)
+	{
+		if (LayoutOptionNames[OptionIndex] == Selection)
+		{
+			LayoutProperty->SetValue(LayoutOptionReferences[OptionIndex]);
+			break;	
+		}
+	}
+}
+
+
+void FCustomizableObjectNodeEditMaterialBaseDetails::OnParentComboBoxSelectionChanged(TSharedPtr<FString> Selection, ESelectInfo::Type SelectInfo, TSharedRef<IPropertyHandle> ParentProperty)
+{
+	FCustomizableObjectNodeParentedMaterialDetails::OnParentComboBoxSelectionChanged(Selection, SelectInfo, ParentProperty);
+
+	const TSharedPtr<FString> LayoutToSelect = GenerateLayoutComboboxOptions();
+	LayoutComboBox->SetSelectedItem(LayoutToSelect);
+}
+
+
+#undef LOCTEXT_NAMESPACE
diff --git a/Engine/Plugins/Experimental/Mutable/Source/CustomizableObjectEditor/Private/Nodes/CustomizableObjectNodeEditMaterialBaseDetails.h b/Engine/Plugins/Experimental/Mutable/Source/CustomizableObjectEditor/Private/Nodes/CustomizableObjectNodeEditMaterialBaseDetails.h
new file mode 100644
--- /dev/null
+++ b/Engine/Plugins/Experimental/Mutable/Source/CustomizableObjectEditor/Private/Nodes/CustomizableObjectNodeEditMaterialBaseDetails.h
@@ -0,0 +1,41 @@
+// Copyright Epic Games, Inc. All Rights Reserved.
+
+#pragma once
+
+#include "CustomizableObjectNodeParentedMaterialDetails.h"
+
+#include "Templates/SharedPointer.h"
+#include "UObject/WeakObjectPtr.h"
+#include "Misc/Guid.h"
+#include "Types/SlateEnums.h"
+#include "PropertyHandle.h"
+#include "CustomizableObjectEditorModule.h"
+#include "Nodes/CustomizableObjectNodeEditMaterialBase.h"
+
+class FCustomizableObjectNodeEditMaterialBaseDetails : public FCustomizableObjectNodeParentedMaterialDetails
+{
+public:
+	/** Makes a new instance of this detail layout class for a specific detail view requesting it */
+	static TSharedRef<IDetailCustomization> MakeInstance();
+
+	// ILayoutDetails interface
+	void CustomizeDetails(IDetailLayoutBuilder& DetailBuilder) override;
+
+	// FCustomizableObjectNodeParentMaterialDetails interface
+protected:
+	void OnParentComboBoxSelectionChanged(TSharedPtr<FString> Selection, ESelectInfo::Type SelectInfo, TSharedRef<IPropertyHandle> ParentProperty) override;
+
+private:
+	TSharedPtr<STextComboBox> LayoutComboBox;
+	
+	TArray<int> LayoutOptionReferences;
+	TArray<TSharedPtr<FString>> LayoutOptionNames;
+	
+	UCustomizableObjectNodeEditMaterialBase* NodeEditMaterialBase;
+
+	// Prepares combobox options for the Layout
+	TSharedPtr<FString> GenerateLayoutComboboxOptions();
+
+	// Callback when we select a different Layout
+	void OnLayoutComboBoxSelectionChanged(TSharedPtr<FString> Selection, ESelectInfo::Type SelectInfo, TSharedRef<IPropertyHandle> LayoutProperty);
+};
diff --git a/Engine/Plugins/Experimental/Mutable/Source/CustomizableObjectEditor/Private/Nodes/CustomizableObjectNodeEnumParameter.cpp b/Engine/Plugins/Experimental/Mutable/Source/CustomizableObjectEditor/Private/Nodes/CustomizableObjectNodeEnumParameter.cpp
new file mode 100644
--- /dev/null
+++ b/Engine/Plugins/Experimental/Mutable/Source/CustomizableObjectEditor/Private/Nodes/CustomizableObjectNodeEnumParameter.cpp
@@ -0,0 +1,52 @@
+// Copyright Epic Games, Inc. All Rights Reserved.
+
+#include "Nodes/CustomizableObjectNodeEnumParameter.h"
+#include "CustomizableObjectEditorModule.h"
+
+#define LOCTEXT_NAMESPACE "CustomizableObjectEditor"
+
+
+void UCustomizableObjectNodeEnumParameter::PostEditChangeProperty(FPropertyChangedEvent& PropertyChangedEvent)
+{
+	FProperty* PropertyThatChanged = PropertyChangedEvent.Property;
+	if ( PropertyThatChanged && PropertyThatChanged->GetName() == TEXT("Values") )
+	{
+		ReconstructNode();
+	}
+
+	Super::PostEditChangeProperty(PropertyChangedEvent);
+}
+
+
+void UCustomizableObjectNodeEnumParameter::AllocateDefaultPins(UCustomizableObjectNodeRemapPins* RemapPins)
+{
+	const UEdGraphSchema_CustomizableObject* Schema = GetDefault<UEdGraphSchema_CustomizableObject>();
+
+	UEdGraphPin* ValuePin = CustomCreatePin(EGPD_Output, Schema->PC_Enum, FName("Value"));
+	ValuePin->bDefaultValueIsIgnored = true;
+}
+
+
+FText UCustomizableObjectNodeEnumParameter::GetNodeTitle(ENodeTitleType::Type TitleType) const
+{
+	return LOCTEXT("Enum_Parameter", "Enum Parameter");
+}
+
+
+FLinearColor UCustomizableObjectNodeEnumParameter::GetNodeTitleColor() const
+{
+	const UEdGraphSchema_CustomizableObject* Schema = GetDefault<UEdGraphSchema_CustomizableObject>();
+	return Schema->GetPinTypeColor(Schema->PC_Enum);
+}
+
+
+FText UCustomizableObjectNodeEnumParameter::GetTooltipText() const
+{
+	return LOCTEXT("Enum_Parameter_Tooltip",
+		"Exposes and defines a parameter offering multiple choices to modify the Customizable Object.\nAlso defines a default one among them. \nIt's abstract, does not define what type those options refer to.");
+}
+
+
+
+#undef LOCTEXT_NAMESPACE
+
diff --git a/Engine/Plugins/Experimental/Mutable/Source/CustomizableObjectEditor/Private/Nodes/CustomizableObjectNodeExposePin.cpp b/Engine/Plugins/Experimental/Mutable/Source/CustomizableObjectEditor/Private/Nodes/CustomizableObjectNodeExposePin.cpp
new file mode 100644
--- /dev/null
+++ b/Engine/Plugins/Experimental/Mutable/Source/CustomizableObjectEditor/Private/Nodes/CustomizableObjectNodeExposePin.cpp
@@ -0,0 +1,74 @@
+// Copyright Epic Games, Inc. All Rights Reserved.
+
+#include "Nodes/CustomizableObjectNodeExposePin.h"
+#include "CustomizableObjectEditorModule.h"
+#include "CustomizableObject.h"
+
+#include "UnrealEditorPortabilityHelpers.h"
+
+#define LOCTEXT_NAMESPACE "CustomizableObjectEditor"
+
+
+void UCustomizableObjectNodeExposePin::PostEditChangeProperty(FPropertyChangedEvent& PropertyChangedEvent)
+{
+	FProperty* PropertyThatChanged = PropertyChangedEvent.Property;
+	if (PropertyThatChanged && PropertyThatChanged->GetName() == TEXT("Name"))
+	{
+		OnNameChangedDelegate.Broadcast();
+	}
+}
+
+
+void UCustomizableObjectNodeExposePin::AllocateDefaultPins(UCustomizableObjectNodeRemapPins* RemapPins)
+{
+	const bool bIsArrayPinCategory = PinType == UEdGraphSchema_CustomizableObject::PC_GroupProjector;
+	UEdGraphPin* InputPin = CustomCreatePin(EGPD_Input, PinType, FName(TEXT("Object")), bIsArrayPinCategory);
+}
+
+
+FText UCustomizableObjectNodeExposePin::GetNodeTitle(ENodeTitleType::Type TitleType) const
+{
+	const FText PinCategoryName = UEdGraphSchema_CustomizableObject::GetPinCategoryName(PinType);
+
+	if (TitleType == ENodeTitleType::ListView)
+	{
+		return FText::Format(LOCTEXT("Expose_Pin_Title_ListView", "Export {0} Pin"), PinCategoryName);
+	}
+	else
+	{
+		return FText::Format(LOCTEXT("Expose_Pin_Title", "{0}\nExport {1} Pin"), FText::FromString(Name), PinCategoryName);
+	}
+}
+
+
+FLinearColor UCustomizableObjectNodeExposePin::GetNodeTitleColor() const
+{
+	const UEdGraphSchema_CustomizableObject* Schema = GetDefault<UEdGraphSchema_CustomizableObject>();
+	return Schema->GetPinTypeColor(PinType);
+}
+
+
+FText UCustomizableObjectNodeExposePin::GetTooltipText() const
+{
+	return LOCTEXT("Expose_Pin_Tooltip", "Exposes a value to the rest of its Customizable Object hierarchy.");
+}
+
+bool UCustomizableObjectNodeExposePin::CanConnect(const UEdGraphPin* InOwnedInputPin, const UEdGraphPin* InOutputPin, bool& bOutIsOtherNodeBlocklisted, bool& bOutArePinsCompatible) const
+{
+	// Chech the pin types do match
+	bOutArePinsCompatible = Super::CanConnect(InOwnedInputPin, InOutputPin, bOutIsOtherNodeBlocklisted, bOutArePinsCompatible);
+
+	// Check the type of the other node to make sure it is not one we do not want to allow the connection with
+	bOutIsOtherNodeBlocklisted = Cast<UCustomizableObjectNodeExternalPin>(InOutputPin->GetOwningNode()) != nullptr;
+
+	return bOutArePinsCompatible && !bOutIsOtherNodeBlocklisted;
+}
+
+
+FString UCustomizableObjectNodeExposePin::GetNodeName() const
+{
+	return Name;
+}
+
+
+#undef LOCTEXT_NAMESPACE
diff --git a/Engine/Plugins/Experimental/Mutable/Source/CustomizableObjectEditor/Private/Nodes/CustomizableObjectNodeExtendMaterial.cpp b/Engine/Plugins/Experimental/Mutable/Source/CustomizableObjectEditor/Private/Nodes/CustomizableObjectNodeExtendMaterial.cpp
new file mode 100644
--- /dev/null
+++ b/Engine/Plugins/Experimental/Mutable/Source/CustomizableObjectEditor/Private/Nodes/CustomizableObjectNodeExtendMaterial.cpp
@@ -0,0 +1,215 @@
+// Copyright Epic Games, Inc. All Rights Reserved.
+
+#include "Nodes/CustomizableObjectNodeExtendMaterial.h"
+#include "CustomizableObjectEditorModule.h"
+
+#define LOCTEXT_NAMESPACE "CustomizableObjectEditor"
+
+
+void UCustomizableObjectNodeExtendMaterial::Serialize(FArchive& Ar)
+{
+	Super::Serialize(Ar);
+	
+	Ar.UsingCustomVersion(FCustomizableObjectCustomVersion::GUID);
+}
+
+
+void UCustomizableObjectNodeExtendMaterial::BackwardsCompatibleFixup()
+{
+	Super::PostBackwardsCompatibleFixup();
+
+	const int32 CustomizableObjectCustomVersion = GetLinkerCustomVersion(FCustomizableObjectCustomVersion::GUID);
+	
+	if (CustomizableObjectCustomVersion < FCustomizableObjectCustomVersion::ExtendMaterialRemoveImages)
+	{
+		if (const UCustomizableObjectNodeMaterial* ParentMaterial = GetParentMaterialNode())
+		{
+			for (const FCustomizableObjectNodeExtendMaterialImage& Image : Images_DEPRECATED)
+			{
+				const UEdGraphPin* ImagePin = FindPin(Image.Name);
+				if (!ImagePin)
+				{
+					continue;
+				}
+
+				FGuid ImageId = FGuid::NewGuid();
+
+				// Search for the Image Id the Extend pin was referring to.
+				const int32 NumImages = ParentMaterial->GetNumParameters(EMaterialParameterType::Texture);
+				for (int32 ImageIndex = 0; ImageIndex < NumImages; ++ImageIndex)
+				{
+					if (ParentMaterial->GetParameterName(EMaterialParameterType::Texture, ImageIndex).ToString() == Image.Name)
+					{
+						ImageId = ParentMaterial->GetParameterId(EMaterialParameterType::Texture, ImageIndex);
+						break;
+					}
+				}
+                
+				PinsParameter.Add(ImageId, FEdGraphPinReference(ImagePin));                
+			}
+		}
+		
+		Images_DEPRECATED.Empty();
+		ReconstructNode();
+	}
+}
+
+void UCustomizableObjectNodeExtendMaterial::PostBackwardsCompatibleFixup()
+{
+	Super::PostBackwardsCompatibleFixup();
+	
+	PostBackwardsCompatibleFixupWork();
+}
+
+
+void UCustomizableObjectNodeExtendMaterial::AllocateDefaultPins(UCustomizableObjectNodeRemapPins* RemapPins)
+{
+	const UEdGraphSchema_CustomizableObject* Schema = GetDefault<UEdGraphSchema_CustomizableObject>();
+
+	UEdGraphPin* AddMeshPin = CustomCreatePin(EGPD_Input, Schema->PC_Mesh, FName("Add Mesh") );
+	AddMeshPin->bDefaultValueIsIgnored = true;
+
+	// Begin texture pins
+	UCustomizableObjectNodeMaterial* ParentMaterialNode = GetParentMaterialNodeIfPath();
+	
+	// In case of a NodeCopyMaterial get the real NodeMaterial
+	if (const UCustomizableObjectNodeCopyMaterial* ParentMaterialCopyNode = Cast<UCustomizableObjectNodeCopyMaterial>(ParentMaterialNode))
+	{
+		ParentMaterialNode = ParentMaterialCopyNode->GetMaterialNode();
+	}
+
+	if (ParentMaterialNode)
+	{
+		const int32 NumImages = ParentMaterialNode->GetNumParameters(EMaterialParameterType::Texture);
+		for (int32 ImageIndex = 0; ImageIndex < NumImages; ++ImageIndex)
+		{
+			if (ParentMaterialNode->IsImageMutableMode(ImageIndex))
+			{
+				const FName ImageName = ParentMaterialNode->GetParameterName(EMaterialParameterType::Texture, ImageIndex);
+				UEdGraphPin* PinImage = CustomCreatePin(EGPD_Input, Schema->PC_Image, ImageName);
+				PinImage->bDefaultValueIsIgnored = true;
+
+				const FGuid ImageId = ParentMaterialNode->GetParameterId(EMaterialParameterType::Texture, ImageIndex);
+				
+				PinsParameter.Add(ImageId, FEdGraphPinReference(PinImage));
+			}
+		}
+	}
+	// End texture pins
+
+	CustomCreatePin(EGPD_Output, Schema->PC_Material, FName("Material"));
+}
+
+
+FText UCustomizableObjectNodeExtendMaterial::GetNodeTitle(ENodeTitleType::Type TitleType) const
+{
+	return LOCTEXT("Extend_Material", "Extend Material");
+}
+
+
+FLinearColor UCustomizableObjectNodeExtendMaterial::GetNodeTitleColor() const
+{
+	const UEdGraphSchema_CustomizableObject* Schema = GetDefault<UEdGraphSchema_CustomizableObject>();
+	return Schema->GetPinTypeColor(Schema->PC_Material);
+}
+
+
+void UCustomizableObjectNodeExtendMaterial::PinConnectionListChanged(UEdGraphPin* Pin)
+{
+	Super::PinConnectionListChanged(Pin);
+
+	PinConnectionListChangedWork(Pin);
+}
+
+
+bool UCustomizableObjectNodeExtendMaterial::IsNodeOutDatedAndNeedsRefresh()
+{
+	return IsNodeOutDatedAndNeedsRefreshWork();
+}
+
+
+FString UCustomizableObjectNodeExtendMaterial::GetRefreshMessage() const
+{
+	return "Source material has changed, texture channels might have been added, removed or renamed. Please refresh the parent material node to reflect those changes.";
+}
+
+
+void UCustomizableObjectNodeExtendMaterial::SaveParentNode(UCustomizableObject* Object, FGuid NodeId)
+{
+	ParentMaterialObject = Object;
+	ParentMaterialNodeId = NodeId;
+}
+
+
+UEdGraphPin* UCustomizableObjectNodeExtendMaterial::AddMeshPin() const
+{
+	return FindPin(TEXT("Add Mesh"));
+}
+
+
+UEdGraphPin* UCustomizableObjectNodeExtendMaterial::OutputPin() const
+{
+	return FindPin(TEXT("Material"));
+}
+
+void UCustomizableObjectNodeExtendMaterial::SetParentNode(UCustomizableObject* Object, FGuid NodeId)
+{
+	PreSetParentNodeWork(Object, NodeId);
+
+	ICustomizableObjectNodeParentedNode::SetParentNode(Object, NodeId);
+	
+	PostSetParentNodeWork(Object, NodeId);
+}
+
+
+FText UCustomizableObjectNodeExtendMaterial::GetTooltipText() const
+{
+	return LOCTEXT("Extend_Material_Tooltip", "Extend an ancestor's material: add a new mesh section, and add its corresponding texture to the ancestor's material texture parameters.");
+}
+
+
+bool UCustomizableObjectNodeExtendMaterial::IsSingleOutputNode() const
+{
+	return true;
+}
+
+
+bool UCustomizableObjectNodeExtendMaterial::CustomRemovePin(UEdGraphPin& Pin)
+{
+	CustomRemovePinWork(Pin);
+	
+	return Super::CustomRemovePin(Pin);
+}
+
+
+UCustomizableObjectNode& UCustomizableObjectNodeExtendMaterial::GetNode()
+{
+	return *this;
+}
+
+
+TMap<FGuid, FEdGraphPinReference>& UCustomizableObjectNodeExtendMaterial::GetPinsParameter()
+{
+	return PinsParameter;
+}
+
+
+FGuid UCustomizableObjectNodeExtendMaterial::GetParentNodeId() const
+{
+	return ParentMaterialNodeId;
+}
+
+
+UCustomizableObject* UCustomizableObjectNodeExtendMaterial::GetParentObject() const
+{
+	return ParentMaterialObject;
+}
+
+
+FCustomizableObjectNodeParentedMaterial& UCustomizableObjectNodeExtendMaterial::GetNodeParentedMaterial()
+{
+	return *this;
+}
+
+
+#undef LOCTEXT_NAMESPACE
diff --git a/Engine/Plugins/Experimental/Mutable/Source/CustomizableObjectEditor/Private/Nodes/CustomizableObjectNodeExternalPin.cpp b/Engine/Plugins/Experimental/Mutable/Source/CustomizableObjectEditor/Private/Nodes/CustomizableObjectNodeExternalPin.cpp
new file mode 100644
--- /dev/null
+++ b/Engine/Plugins/Experimental/Mutable/Source/CustomizableObjectEditor/Private/Nodes/CustomizableObjectNodeExternalPin.cpp
@@ -0,0 +1,206 @@
+// Copyright Epic Games, Inc. All Rights Reserved.
+
+#include "Nodes/CustomizableObjectNodeExternalPin.h"
+#include "CustomizableObjectEditorModule.h"
+#include "CustomizableObject.h"
+#include "CustomizableObjectPin.h"
+
+#define LOCTEXT_NAMESPACE "CustomizableObjectEditor"
+
+
+void UCustomizableObjectNodeRemapPinsCustomExternalPin::RemapPins(const TArray<UEdGraphPin*>& OldPins, const TArray<UEdGraphPin*>& NewPins, TMap<UEdGraphPin*, UEdGraphPin*>& PinsToRemap, TArray<UEdGraphPin*>& PinsToOrphan)
+{
+	if (OldPins.Num()) 
+	{
+		UEdGraphPin* const OldPin = OldPins[0];
+		if (NewPins.Num())
+		{
+			PinsToRemap.Add(OldPin, NewPins[0]);
+		}
+
+		if (!Node->GetNodeExposePin())
+		{
+			PinsToOrphan.Add(OldPin);
+		}
+	}
+}
+
+
+void UCustomizableObjectNodeExternalPin::Serialize(FArchive& Ar)
+{
+	Super::Serialize(Ar);
+
+	Ar.UsingCustomVersion(FCustomizableObjectCustomVersion::GUID);
+}
+
+
+void UCustomizableObjectNodeExternalPin::BackwardsCompatibleFixup()
+{
+	Super::BackwardsCompatibleFixup();
+	
+	// Get the pin type from the actual pin.
+	if (GetLinkerCustomVersion(FCustomizableObjectCustomVersion::GUID) < FCustomizableObjectCustomVersion::BeforeCustomVersionWasAdded)
+	{
+		if (PinType.IsNone())
+		{
+			PinType = Pins[0]->PinType.PinCategory;
+		}
+	}
+}
+
+
+void UCustomizableObjectNodeExternalPin::PostBackwardsCompatibleFixup()
+{
+	Super::PostBackwardsCompatibleFixup();
+
+	if (UCustomizableObjectNodeExposePin* NodeExposePin = GetNodeExposePin())
+	{
+		OnNameChangedDelegateHandle = NodeExposePin->OnNameChangedDelegate.AddUObject(this, &Super::ReconstructNode);
+		DestroyNodeDelegateHandle = NodeExposePin->DestroyNodeDelegate.AddUObject(this, &Super::ReconstructNode);
+	}
+
+	// Reconstruct the node since the NodeExposePin pin name may have changed while not loaded.
+	ReconstructNode();
+}
+
+
+void UCustomizableObjectNodeExternalPin::SetExternalObjectNodeId(FGuid Guid)
+{
+	if (UCustomizableObjectNodeExposePin* NodeExposePin = GetNodeExposePin())
+	{
+		NodeExposePin->OnNameChangedDelegate.Remove(OnNameChangedDelegateHandle);
+		NodeExposePin->DestroyNodeDelegate.Remove(DestroyNodeDelegateHandle);
+	}
+
+	ExternalObjectNodeId = Guid;
+
+	if (UCustomizableObjectNodeExposePin* NodeExposePin = GetNodeExposePin())
+	{
+		OnNameChangedDelegateHandle = NodeExposePin->OnNameChangedDelegate.AddUObject(this, &Super::ReconstructNode);
+		DestroyNodeDelegateHandle = NodeExposePin->DestroyNodeDelegate.AddUObject(this, &Super::ReconstructNode);
+	}
+
+	ReconstructNode();
+}
+
+
+UEdGraphPin* UCustomizableObjectNodeExternalPin::GetExternalPin() const
+{
+	const TArray<UEdGraphPin*> NonOrphanPins = GetAllNonOrphanPins();
+	if (NonOrphanPins.Num())
+	{
+		return NonOrphanPins[0];
+	}
+	else
+	{
+		return nullptr;
+	}
+}
+
+
+UCustomizableObjectNodeExposePin* UCustomizableObjectNodeExternalPin::GetNodeExposePin() const
+{
+	return GetCustomizableObjectExternalNode<UCustomizableObjectNodeExposePin>(ExternalObject, ExternalObjectNodeId);
+}
+
+
+UEdGraphPin* UCustomizableObjectNodeExternalPin::CreateExternalPin(const UCustomizableObjectNodeExposePin* NodeExposePin)
+{
+	FName PinName = NodeExposePin ? FName(NodeExposePin->GetNodeName()) : FName("Object");
+	const bool bIsArrayPinCategory = PinType == UEdGraphSchema_CustomizableObject::PC_GroupProjector;
+
+	return CustomCreatePin(EGPD_Output, PinType, PinName, bIsArrayPinCategory);
+}
+
+
+void UCustomizableObjectNodeExternalPin::BeginConstruct()
+{
+	// Create a pin already orphaned since the node does not have a linked Expose Pin node yet.
+	UEdGraphPin* Pin = CreateExternalPin(GetNodeExposePin());
+	OrphanPin(*Pin);
+}
+
+
+void UCustomizableObjectNodeExternalPin::AllocateDefaultPins(UCustomizableObjectNodeRemapPins* RemapPins)
+{
+	const UEdGraphSchema_CustomizableObject* Schema = GetDefault<UEdGraphSchema_CustomizableObject>();
+
+	// Pass information to the remap pins action context.
+	if (UCustomizableObjectNodeRemapPinsCustomExternalPin* RemapPinsCustom = Cast<UCustomizableObjectNodeRemapPinsCustomExternalPin>(RemapPins))
+	{
+		RemapPinsCustom->Node = this;
+	}
+
+	if (UCustomizableObjectNodeExposePin* NodeExposePin = GetNodeExposePin())
+	{
+		CreateExternalPin(NodeExposePin);
+	}
+}
+
+
+FText UCustomizableObjectNodeExternalPin::GetNodeTitle(ENodeTitleType::Type TitleType) const
+{
+	const FText PinTypeName = UEdGraphSchema_CustomizableObject::GetPinCategoryName(PinType);
+
+	if (TitleType == ENodeTitleType::ListView || !ExternalObject)
+	{
+		return FText::Format(LOCTEXT("External_Pin_Title", "Import {0} Pin"), PinTypeName);
+	}
+	else
+	{
+		return FText::Format(LOCTEXT("External_Pin_Title", "{0}\nImport {1} Pin"), FText::FromString(ExternalObject->GetName()), PinTypeName);
+	}
+}
+
+
+FLinearColor UCustomizableObjectNodeExternalPin::GetNodeTitleColor() const
+{
+	const UEdGraphSchema_CustomizableObject* Schema = GetDefault<UEdGraphSchema_CustomizableObject>();
+	return Schema->GetPinTypeColor(PinType);
+}
+
+
+FText UCustomizableObjectNodeExternalPin::GetTooltipText() const
+{
+	return LOCTEXT("Import_Pin_Tooltip", "Make use of a value defined elsewhere in this Customizable Object hierarchy.");
+}
+
+bool UCustomizableObjectNodeExternalPin::CanConnect(const UEdGraphPin* InOwnedInputPin, const UEdGraphPin* InOutputPin, bool& bOutIsOtherNodeBlocklisted, bool& bOutArePinsCompatible) const
+{
+	// Chech the pin types do match
+	bOutArePinsCompatible = Super::CanConnect(InOwnedInputPin, InOutputPin, bOutIsOtherNodeBlocklisted, bOutArePinsCompatible);
+
+	// Check the type of the other node to make sure it is not one we do not want to allow the connection with
+	bOutIsOtherNodeBlocklisted = Cast<UCustomizableObjectNodeExposePin>(InOutputPin->GetOwningNode()) != nullptr;
+
+	return bOutArePinsCompatible && !bOutIsOtherNodeBlocklisted;
+}
+
+
+void UCustomizableObjectNodeExternalPin::BeginPostDuplicate(bool bDuplicateForPIE)
+{
+	Super::BeginPostDuplicate(bDuplicateForPIE);
+
+	if (ExternalObjectNodeId.IsValid())
+	{
+		if (UCustomizableObjectGraph* CEdGraph = Cast<UCustomizableObjectGraph>(GetGraph()))
+		{
+			ExternalObjectNodeId = CEdGraph->RequestNotificationForNodeIdChange(ExternalObjectNodeId, NodeGuid);
+		}
+	}
+}
+
+
+void UCustomizableObjectNodeExternalPin::UpdateReferencedNodeId(const FGuid& NewGuid)
+{
+	ExternalObjectNodeId = NewGuid;
+}
+
+
+UCustomizableObjectNodeRemapPins* UCustomizableObjectNodeExternalPin::CreateRemapPinsDefault() const
+{
+	return NewObject<UCustomizableObjectNodeRemapPinsCustomExternalPin>();
+}
+
+
+#undef LOCTEXT_NAMESPACE
diff --git a/Engine/Plugins/Experimental/Mutable/Source/CustomizableObjectEditor/Private/Nodes/CustomizableObjectNodeExternalPinDetails.h b/Engine/Plugins/Experimental/Mutable/Source/CustomizableObjectEditor/Private/Nodes/CustomizableObjectNodeExternalPinDetails.h
new file mode 100644
--- /dev/null
+++ b/Engine/Plugins/Experimental/Mutable/Source/CustomizableObjectEditor/Private/Nodes/CustomizableObjectNodeExternalPinDetails.h
@@ -0,0 +1,34 @@
+// Copyright Epic Games, Inc. All Rights Reserved.
+
+#pragma once
+
+#include "IDetailCustomization.h"
+#include "Types/SlateEnums.h"
+#include "Templates/SharedPointer.h"
+#include "UObject/WeakObjectPtr.h"
+#include "PropertyHandle.h"
+#include "AssetRegistry/AssetData.h"
+#include "Widgets/Input/STextComboBox.h"
+#include "Nodes/CustomizableObjectNodeExternalPin.h"
+
+
+class FCustomizableObjectNodeExternalPinDetails : public IDetailCustomization
+{
+public:
+	/** Makes a new instance of this detail layout class for a specific detail view requesting it */
+	static TSharedRef<IDetailCustomization> MakeInstance();
+
+	/** ILayoutDetails interface */
+	virtual void CustomizeDetails( IDetailLayoutBuilder& DetailBuilder ) override;
+
+private:
+	IDetailLayoutBuilder* DetailBuilderPtr = nullptr;
+
+	UCustomizableObjectNodeExternalPin* Node = nullptr;
+	
+	TArray<TSharedPtr<FString>> GroupNodeComboBoxOptions;
+
+	void ParentObjectSelectionChanged(const FAssetData& AssetData);
+
+	void OnGroupNodeComboBoxSelectionChanged(TSharedPtr<FString> Selection, ESelectInfo::Type SelectInfo);
+};
diff --git a/Engine/Plugins/Experimental/Mutable/Source/CustomizableObjectEditor/Private/Nodes/CustomizableObjectNodeFloatParameter.cpp b/Engine/Plugins/Experimental/Mutable/Source/CustomizableObjectEditor/Private/Nodes/CustomizableObjectNodeFloatParameter.cpp
new file mode 100644
--- /dev/null
+++ b/Engine/Plugins/Experimental/Mutable/Source/CustomizableObjectEditor/Private/Nodes/CustomizableObjectNodeFloatParameter.cpp
@@ -0,0 +1,115 @@
+// Copyright Epic Games, Inc. All Rights Reserved.
+
+#include "Nodes/CustomizableObjectNodeFloatParameter.h"
+#include "CustomizableObjectEditorModule.h"
+#include "UnrealEditorPortabilityHelpers.h"
+
+#define LOCTEXT_NAMESPACE "CustomizableObjectEditor"
+
+
+const UEdGraphPin* UCustomizableObjectNodeFloatParameter::GetDescriptionImagePin(int32 Index) const
+{
+	return DescriptionImagePinsReferences[Index].Get();
+}
+
+
+void UCustomizableObjectNodeFloatParameter::PostEditChangeProperty(FPropertyChangedEvent& PropertyChangedEvent)
+{
+	FProperty* PropertyThatChanged = PropertyChangedEvent.Property;
+	if ( PropertyThatChanged && (PropertyThatChanged->GetName() == TEXT("DescriptionImage") || PropertyThatChanged->GetName() == TEXT("Name")) )
+	{
+		ReconstructNode();
+	}
+
+	Super::PostEditChangeProperty(PropertyChangedEvent);
+}
+
+
+void UCustomizableObjectNodeFloatParameter::Serialize(FArchive& Ar)
+{
+	Super::Serialize(Ar);
+
+	Ar.UsingCustomVersion(FCustomizableObjectCustomVersion::GUID);
+
+	if (Ar.CustomVer(FCustomizableObjectCustomVersion::GUID) < FCustomizableObjectCustomVersion::PinsNamesImageToTexture)
+	{
+		for (int32 Index = 0; Index < DescriptionImage.Num(); ++Index)
+		{
+			DescriptionImagePinsReferences.Add(FindPin(FString::Printf(TEXT("Description Image %d"), Index)));
+		}
+	}
+}
+
+
+void UCustomizableObjectNodeFloatParameter::AllocateDefaultPins(UCustomizableObjectNodeRemapPins* RemapPins)
+{
+	const UEdGraphSchema_CustomizableObject* Schema = GetDefault<UEdGraphSchema_CustomizableObject>();
+
+	UEdGraphPin* ValuePin = CustomCreatePin(EGPD_Output, Schema->PC_Float, FName("Value"));
+	ValuePin->bDefaultValueIsIgnored = true;
+
+	DescriptionImagePinsReferences.Reset();
+	
+	for (int32 Index = 0; Index <  DescriptionImage.Num(); ++Index )
+	{
+		FString PinName = DescriptionImage[Index].Name.IsEmpty() ?
+			FString::Printf(TEXT("Description Texture %d"), Index) :
+			DescriptionImage[Index].Name;
+		
+		UEdGraphPin* DescriptionImagePin = CustomCreatePin(EGPD_Input, Schema->PC_Image, FName(*PinName));
+		DescriptionImagePin->bDefaultValueIsIgnored = true;
+
+		DescriptionImagePinsReferences.Add(DescriptionImagePin);
+	}
+}
+
+
+bool UCustomizableObjectNodeFloatParameter::IsAffectedByLOD() const
+{
+	return false;
+}
+
+
+int32 UCustomizableObjectNodeFloatParameter::GetNumDescriptionImage() const
+{
+	return DescriptionImagePinsReferences.Num();
+}
+
+
+FText UCustomizableObjectNodeFloatParameter::GetNodeTitle(ENodeTitleType::Type TitleType) const
+{
+	return LOCTEXT("Float_Parameter", "Float Parameter");
+}
+
+
+FLinearColor UCustomizableObjectNodeFloatParameter::GetNodeTitleColor() const
+{
+	const UEdGraphSchema_CustomizableObject* Schema = GetDefault<UEdGraphSchema_CustomizableObject>();
+	return Schema->GetPinTypeColor(Schema->PC_Float);
+}
+
+
+void UCustomizableObjectNodeFloatParameter::PostPasteNode()
+{
+	TArray<UEdGraphPin*> InputPins;
+	GetInputPins(InputPins);
+
+	DescriptionImage.SetNum(InputPins.Num());
+	const int32 MaxIndex = InputPins.Num();
+	for (int32 i = 0; i < MaxIndex; ++i)
+	{
+		const UEdGraphPin* DescriptionImagePin = InputPins[i];
+		if (DescriptionImagePin->LinkedTo.Num())
+		{
+			DescriptionImage[i].Name = Helper_GetPinName(InputPins[i]);
+		}
+	}
+}
+
+
+FText UCustomizableObjectNodeFloatParameter::GetTooltipText() const
+{
+	return LOCTEXT("Float_Parameter_Tooltip", "Expose a numeric parameter from the Customizable Object that can be modified at runtime.");
+}
+
+#undef LOCTEXT_NAMESPACE
diff --git a/Engine/Plugins/Experimental/Mutable/Source/CustomizableObjectEditor/Private/Nodes/CustomizableObjectNodeFloatSwitch.cpp b/Engine/Plugins/Experimental/Mutable/Source/CustomizableObjectEditor/Private/Nodes/CustomizableObjectNodeFloatSwitch.cpp
new file mode 100644
--- /dev/null
+++ b/Engine/Plugins/Experimental/Mutable/Source/CustomizableObjectEditor/Private/Nodes/CustomizableObjectNodeFloatSwitch.cpp
@@ -0,0 +1,40 @@
+// Copyright Epic Games, Inc. All Rights Reserved.
+
+#include "Nodes/CustomizableObjectNodeFloatSwitch.h"
+
+#include "EdGraphSchema_CustomizableObject.h"
+
+#define LOCTEXT_NAMESPACE "CustomizableObjectEditor"
+
+
+void UCustomizableObjectNodeFloatSwitch::Serialize(FArchive& Ar)
+{
+	Super::Serialize(Ar);
+
+	if (Ar.CustomVer(FCustomizableObjectCustomVersion::GUID) < FCustomizableObjectCustomVersion::BugPinsSwitch)
+	{
+		OutputPinReference = FindPin(TEXT("Float"));	
+	}
+}
+
+
+FString UCustomizableObjectNodeFloatSwitch::GetOutputPinName() const
+{
+	return TEXT("Float");
+}
+
+
+FName UCustomizableObjectNodeFloatSwitch::GetCategory() const
+{
+	return UEdGraphSchema_CustomizableObject::PC_Float;
+}
+
+
+FString UCustomizableObjectNodeFloatSwitch::GetPinPrefix() const
+{
+	return TEXT("Float ");
+}
+
+
+#undef LOCTEXT_NAMESPACE
+
diff --git a/Engine/Plugins/Experimental/Mutable/Source/CustomizableObjectEditor/Private/Nodes/CustomizableObjectNodeFloatVariation.cpp b/Engine/Plugins/Experimental/Mutable/Source/CustomizableObjectEditor/Private/Nodes/CustomizableObjectNodeFloatVariation.cpp
new file mode 100644
--- /dev/null
+++ b/Engine/Plugins/Experimental/Mutable/Source/CustomizableObjectEditor/Private/Nodes/CustomizableObjectNodeFloatVariation.cpp
@@ -0,0 +1,73 @@
+// Copyright Epic Games, Inc. All Rights Reserved.
+
+#include "Nodes/CustomizableObjectNodeFloatVariation.h"
+#include "CustomizableObjectEditorModule.h"
+
+#include "UnrealEditorPortabilityHelpers.h"
+
+#define LOCTEXT_NAMESPACE "CustomizableObjectEditor"
+
+
+UCustomizableObjectNodeFloatVariation::UCustomizableObjectNodeFloatVariation()
+	: Super()
+{
+
+}
+
+
+void UCustomizableObjectNodeFloatVariation::PostEditChangeProperty(FPropertyChangedEvent& PropertyChangedEvent)
+{
+	FProperty* PropertyThatChanged = PropertyChangedEvent.Property;
+	if ( PropertyThatChanged  )
+	{
+		ReconstructNode();
+	}
+
+	Super::PostEditChangeProperty(PropertyChangedEvent);
+}
+
+
+void UCustomizableObjectNodeFloatVariation::AllocateDefaultPins(UCustomizableObjectNodeRemapPins* RemapPins)
+{
+	const UEdGraphSchema_CustomizableObject* Schema = GetDefault<UEdGraphSchema_CustomizableObject>();
+
+	UEdGraphPin* OutputPin = CustomCreatePin(EGPD_Output, Schema->PC_Float, FName(TEXT("Float")));
+	OutputPin->bDefaultValueIsIgnored = true;
+
+	for ( int VariationIndex = Variations.Num() - 1; VariationIndex >= 0; --VariationIndex)
+	{
+		FString RemapIndex = Helper_GetPinName(VariationPin(VariationIndex));
+
+		FString PinName = FString::Printf( TEXT("Variation %d"), VariationIndex);
+		UEdGraphPin* VariationPin = CustomCreatePin(EGPD_Input, Schema->PC_Float, FName(*PinName), false);
+		VariationPin->bDefaultValueIsIgnored = true;
+
+		FString FriendlyName = FString::Printf(TEXT("Variation %d [%s]"), VariationIndex, *Variations[VariationIndex].Tag);
+		VariationPin->PinFriendlyName = FText::FromString(FriendlyName);
+	}
+
+	UEdGraphPin* DefaultVariation = CustomCreatePin(EGPD_Input, Schema->PC_Float, FName(TEXT("Default")), false);
+	DefaultVariation->bDefaultValueIsIgnored = true;
+}
+
+
+FText UCustomizableObjectNodeFloatVariation::GetNodeTitle(ENodeTitleType::Type TitleType) const
+{
+	return LOCTEXT("Float_Variation", "Float Variation");
+}
+
+
+FLinearColor UCustomizableObjectNodeFloatVariation::GetNodeTitleColor() const
+{
+	const UEdGraphSchema_CustomizableObject* Schema = GetDefault<UEdGraphSchema_CustomizableObject>();
+	return Schema->GetPinTypeColor(Schema->PC_Float);
+}
+
+
+FText UCustomizableObjectNodeFloatVariation::GetTooltipText() const
+{
+	return LOCTEXT("Float_Variation_Tooltip", "Select a float depending on what tags are active.");
+}
+
+#undef LOCTEXT_NAMESPACE
+
diff --git a/Engine/Plugins/Experimental/Mutable/Source/CustomizableObjectEditor/Private/Nodes/CustomizableObjectNodeGroupProjectorParameter.cpp b/Engine/Plugins/Experimental/Mutable/Source/CustomizableObjectEditor/Private/Nodes/CustomizableObjectNodeGroupProjectorParameter.cpp
new file mode 100644
--- /dev/null
+++ b/Engine/Plugins/Experimental/Mutable/Source/CustomizableObjectEditor/Private/Nodes/CustomizableObjectNodeGroupProjectorParameter.cpp
@@ -0,0 +1,135 @@
+// Copyright Epic Games, Inc. All Rights Reserved.
+
+#include "Nodes/CustomizableObjectNodeGroupProjectorParameter.h"
+#include "UnrealEditorPortabilityHelpers.h"
+#include "EdGraphSchema_CustomizableObject.h"
+
+#define LOCTEXT_NAMESPACE "CustomizableObjectEditor"
+
+
+UCustomizableObjectNodeGroupProjectorParameter::UCustomizableObjectNodeGroupProjectorParameter()
+	: Super()
+{
+
+}
+
+
+TArray<FGroupProjectorParameterImage> UCustomizableObjectNodeGroupProjectorParameter::GetOptionImagesFromTable() const
+{
+	TArray<FGroupProjectorParameterImage> ArrayResult;
+
+	if (OptionImagesDataTable == nullptr)
+	{
+		return ArrayResult;
+	}
+
+	TArray<FName> ArrayRowName = OptionImagesDataTable->GetRowNames();
+
+	FProperty* PropertyTexturePath = OptionImagesDataTable->FindTableProperty(DataTableTextureColumnName);
+
+	if (PropertyTexturePath == nullptr)
+	{
+		UE_LOG(LogMutable, Warning, TEXT("WARNING: No column found with texture path information to load projection textures"));
+		return ArrayResult;
+	}
+
+	int32 NameIndex = 0;
+
+	for (auto RowIt = Helper_GetRowMap(OptionImagesDataTable).CreateConstIterator(); RowIt; ++RowIt)
+	{
+		uint8* RowData = RowIt.Value();
+		FString PropertyValue(TEXT(""));
+		PropertyTexturePath->ExportText_InContainer(0, PropertyValue, RowData, RowData, nullptr, PPF_None);
+		UTexture2D* Texture = LoadObject<UTexture2D>(nullptr, *PropertyValue, nullptr);
+
+		if (Texture == nullptr)
+		{
+			UE_LOG(LogMutable, Warning, TEXT("WARNING: Unable to load texture %s"), *PropertyValue);
+		}
+		else
+		{
+			FGroupProjectorParameterImage GroupProjectorParameterImage;
+			GroupProjectorParameterImage.OptionName = ArrayRowName[NameIndex].ToString();
+			GroupProjectorParameterImage.OptionImage = Texture;
+			ArrayResult.Add(GroupProjectorParameterImage);
+		}
+
+		NameIndex++;
+	}
+
+	return ArrayResult;
+}
+
+
+TArray<FGroupProjectorParameterImage> UCustomizableObjectNodeGroupProjectorParameter::GetFinalOptionImagesNoRepeat() const
+{
+	TArray<FGroupProjectorParameterImage> ArrayDataTable = GetOptionImagesFromTable();
+
+	for (int32 i = 0; i < OptionImages.Num(); ++i)
+	{
+		bool AlreadyAdded = false;
+		for (int32 j = 0; j < ArrayDataTable.Num(); ++j)
+		{
+			if (OptionImages[i].OptionName == ArrayDataTable[j].OptionName)
+			{
+				AlreadyAdded = true;
+				break;
+			}
+		}
+
+		if (!AlreadyAdded)
+		{
+			ArrayDataTable.Add(OptionImages[i]);
+		}
+	}
+
+	return ArrayDataTable;
+}
+
+
+void UCustomizableObjectNodeGroupProjectorParameter::Serialize(FArchive& Ar)
+{
+	Super::Serialize(Ar);
+
+	Ar.UsingCustomVersion(FCustomizableObjectCustomVersion::GUID);
+
+	UEdGraphPin* projectorPin = ProjectorPin();
+	if (Ar.CustomVer(FCustomizableObjectCustomVersion::GUID) < FCustomizableObjectCustomVersion::GroupProjectorPinTypeAdded
+		&& projectorPin
+		&& projectorPin->PinType.PinCategory == UEdGraphSchema_CustomizableObject::PC_Projector)
+	{
+		projectorPin->PinType.PinCategory = UEdGraphSchema_CustomizableObject::PC_GroupProjector;
+	}
+}
+
+
+FText UCustomizableObjectNodeGroupProjectorParameter::GetNodeTitle(ENodeTitleType::Type TitleType) const
+{
+	return LOCTEXT("Group_Projector_Parameter", "Group Projector Parameter");
+}
+
+
+FLinearColor UCustomizableObjectNodeGroupProjectorParameter::GetNodeTitleColor() const
+{
+	const UEdGraphSchema_CustomizableObject* Schema = GetDefault<UEdGraphSchema_CustomizableObject>();
+	return Schema->GetPinTypeColor(Schema->PC_GroupProjector);
+}
+
+
+FText UCustomizableObjectNodeGroupProjectorParameter::GetTooltipText() const
+{
+	return LOCTEXT("Group_Projector_Parameter_Tooltip", "Projects one or many textures to all children in the group it's connected to. It modifies only the materials that define a specific material asset texture parameter.");
+}
+
+
+void UCustomizableObjectNodeGroupProjectorParameter::AllocateDefaultPins(UCustomizableObjectNodeRemapPins* RemapPins)
+{
+	const UEdGraphSchema_CustomizableObject* Schema = GetDefault<UEdGraphSchema_CustomizableObject>();
+
+	FString PinName = TEXT("Value");
+	UEdGraphPin* ValuePin = CustomCreatePin(EGPD_Output, Schema->PC_GroupProjector, FName(*PinName));
+	ValuePin->bDefaultValueIsIgnored = true;
+}
+
+
+#undef LOCTEXT_NAMESPACE
diff --git a/Engine/Plugins/Experimental/Mutable/Source/CustomizableObjectEditor/Private/Nodes/CustomizableObjectNodeLayoutBlocks.cpp b/Engine/Plugins/Experimental/Mutable/Source/CustomizableObjectEditor/Private/Nodes/CustomizableObjectNodeLayoutBlocks.cpp
new file mode 100644
--- /dev/null
+++ b/Engine/Plugins/Experimental/Mutable/Source/CustomizableObjectEditor/Private/Nodes/CustomizableObjectNodeLayoutBlocks.cpp
@@ -0,0 +1,193 @@
+// Copyright Epic Games, Inc. All Rights Reserved.
+
+#include "Nodes/CustomizableObjectNodeLayoutBlocks.h"
+#include "CustomizableObjectEditorModule.h"
+#include "GraphTraversal.h"
+
+
+#define LOCTEXT_NAMESPACE "CustomizableObjectEditor"
+
+
+UCustomizableObjectNodeLayoutBlocks::UCustomizableObjectNodeLayoutBlocks()
+	: Super()
+{
+	Layout = CreateDefaultSubobject<UCustomizableObjectLayout>(FName("CustomizableObjectLayout"));
+}
+
+
+void UCustomizableObjectNodeLayoutBlocks::PostBackwardsCompatibleFixup()
+{
+	Super::PostBackwardsCompatibleFixup();
+
+	if (UEdGraphPin* Pin = OutputPin())
+	{
+		LinkPostEditChangePropertyDelegate(Pin);
+	}
+
+	// Make sure we have IDs for all blocks. Generate fake ids for old blocks so that it is deterministic
+	// and new references to these blocks work even if this node is not saved.
+	for (int BlockIndex=0; BlockIndex< Layout->Blocks.Num(); ++BlockIndex)
+	{
+		FCustomizableObjectLayoutBlock& Block = Layout->Blocks[BlockIndex];
+		if (!Block.Id.IsValid())
+		{
+			FGuid FakeId;
+			FakeId[0] = BlockIndex + 1;
+			FakeId[1] = BlockIndex + 1;
+			FakeId[2] = BlockIndex + 1;
+			FakeId[3] = BlockIndex + 1;
+			Block.Id = FakeId;
+		}
+	}
+}
+
+
+void UCustomizableObjectNodeLayoutBlocks::AllocateDefaultPins(UCustomizableObjectNodeRemapPins* RemapPins)
+{
+	const UEdGraphSchema_CustomizableObject* Schema = GetDefault<UEdGraphSchema_CustomizableObject>();
+
+	UEdGraphPin* OutputPin = CustomCreatePin(EGPD_Output, Schema->PC_Layout, FName("Layout"));
+	OutputPin->bDefaultValueIsIgnored = true;
+}
+
+
+FText UCustomizableObjectNodeLayoutBlocks::GetNodeTitle(ENodeTitleType::Type TitleType) const
+{
+	return LOCTEXT("Mesh_Layout", "Mesh Layout");
+}
+
+
+FLinearColor UCustomizableObjectNodeLayoutBlocks::GetNodeTitleColor() const
+{
+	const UEdGraphSchema_CustomizableObject* Schema = GetDefault<UEdGraphSchema_CustomizableObject>();
+	return Schema->GetPinTypeColor(Schema->PC_Layout);
+}
+
+
+void UCustomizableObjectNodeLayoutBlocks::PinConnectionListChanged(UEdGraphPin* Pin)
+{
+	if (Layout && Pin == OutputPin())
+	{
+		LinkPostEditChangePropertyDelegate(Pin);
+	}
+}
+
+
+void UCustomizableObjectNodeLayoutBlocks::AddAttachedErrorData(const FAttachedErrorDataView& AttachedErrorData)
+{
+	if (Layout)
+	{
+		Layout->UnassignedUVs.Emplace();
+
+		const int32 UnassignedUVsCount = AttachedErrorData.UnassignedUVs.Num() / 2;
+		for (int32 Index = 0; Index < UnassignedUVsCount; ++Index)
+		{
+			Layout->UnassignedUVs.Last().Emplace(AttachedErrorData.UnassignedUVs[Index * 2 + 0],
+				AttachedErrorData.UnassignedUVs[Index * 2 + 1]);
+		}
+	}
+}
+
+
+void UCustomizableObjectNodeLayoutBlocks::ResetAttachedErrorData()
+{
+	if (Layout)
+	{
+		Layout->UnassignedUVs.Empty();
+	}
+}
+
+
+FText UCustomizableObjectNodeLayoutBlocks::GetTooltipText() const
+{
+	return LOCTEXT("Layout_Block_Tooltip", "Defines of how the material UV layout is cut into logical pieces.\nWhen destined to a Material Node, it also defines the resolution per layout block of that material for any extension that is made to that material.");
+
+}
+
+
+bool UCustomizableObjectNodeLayoutBlocks::IsSingleOutputNode() const
+{
+	return true;
+}
+
+
+void UCustomizableObjectNodeLayoutBlocks::MeshPostEditChangeProperty(UObject* Node, FPropertyChangedEvent& FPropertyChangedEvent)
+{
+	SetLayoutSkeletalMesh();
+}
+
+
+void UCustomizableObjectNodeLayoutBlocks::LinkPostEditChangePropertyDelegate(UEdGraphPin* Pin)
+{
+	if (LastMeshNodeConnected.IsValid())
+	{
+		LastMeshNodeConnected->PostEditChangePropertyRegularDelegate.AddUObject(this, &UCustomizableObjectNodeLayoutBlocks::MeshPostEditChangeProperty);
+	}
+
+	if (const UEdGraphPin* ConnectedPin = FollowOutputPin(*Pin))
+	{
+		UEdGraphNode* MeshNode = ConnectedPin->GetOwningNode();
+
+		if (MeshNode->IsA(UCustomizableObjectNodeStaticMesh::StaticClass()) || MeshNode->IsA(UCustomizableObjectNodeSkeletalMesh::StaticClass()))
+		{
+			LastMeshNodeConnected = Cast<UCustomizableObjectNode>(MeshNode);
+			LastMeshNodeConnected->PostEditChangePropertyRegularDelegate.AddUObject(this, &UCustomizableObjectNodeLayoutBlocks::MeshPostEditChangeProperty);
+
+			SetLayoutSkeletalMesh();
+		}
+	}
+}
+
+
+void UCustomizableObjectNodeLayoutBlocks::SetLayoutSkeletalMesh()
+{
+	if (const UEdGraphPin* ConnectedPin = FollowOutputPin(*OutputPin()))
+	{
+		if (const UCustomizableObjectNodeMesh* MeshNode = Cast<UCustomizableObjectNodeMesh>(LastMeshNodeConnected))
+		{
+			int32 LODIndex;
+			int32 SectionIndex;
+			int32 LayoutIndex;
+			MeshNode->GetPinSection(*ConnectedPin, LODIndex, SectionIndex, LayoutIndex);
+
+			Layout->SetLayout(MeshNode->GetMesh(), LODIndex, SectionIndex, LayoutIndex);
+		}
+	}
+}
+
+
+void UCustomizableObjectNodeLayoutBlocks::PostPasteNode()
+{
+	Super::PostPasteNode();
+	
+	Layout->SetLayout(nullptr, -1, -1, -1);
+}
+
+
+void UCustomizableObjectNodeLayoutBlocks::Serialize(FArchive& Ar)
+{
+	Super::Serialize(Ar);
+	
+	Ar.UsingCustomVersion(FCustomizableObjectCustomVersion::GUID);
+	
+	if (Ar.CustomVer(FCustomizableObjectCustomVersion::GUID) < FCustomizableObjectCustomVersion::LayoutClassAdded)
+	{
+		Layout->SetGridSize(GridSize_DEPRECATED);
+		Layout->SetMaxGridSize(MaxGridSize_DEPRECATED);
+		Layout->Blocks = Blocks_DEPRECATED;
+		Layout->SetPackingStrategy(PackingStrategy_DEPRECATED);
+
+		if (Layout->GetGridSize() == FIntPoint::ZeroValue)
+		{
+			Layout->SetGridSize(FIntPoint(4));
+		}
+
+		if(Layout->GetMaxGridSize() == FIntPoint::ZeroValue)
+		{
+			Layout->SetMaxGridSize(FIntPoint(4));
+		}
+	}
+}
+
+#undef LOCTEXT_NAMESPACE
+
diff --git a/Engine/Plugins/Experimental/Mutable/Source/CustomizableObjectEditor/Private/Nodes/CustomizableObjectNodeLayoutBlocksDetails.cpp b/Engine/Plugins/Experimental/Mutable/Source/CustomizableObjectEditor/Private/Nodes/CustomizableObjectNodeLayoutBlocksDetails.cpp
new file mode 100644
--- /dev/null
+++ b/Engine/Plugins/Experimental/Mutable/Source/CustomizableObjectEditor/Private/Nodes/CustomizableObjectNodeLayoutBlocksDetails.cpp
@@ -0,0 +1,121 @@
+// Copyright Epic Games, Inc. All Rights Reserved.
+
+#include "CustomizableObjectNodeLayoutBlocksDetails.h"
+#include "CustomizableObjectEditorModule.h"
+#include "PropertyCustomizationHelpers.h"
+
+#include "DetailLayoutBuilder.h"
+#include "DetailCategoryBuilder.h"
+#include "LevelEditor.h"
+#include "LevelEditorActions.h"
+
+#include "Widgets/Text/STextBlock.h"
+
+#include "Runtime/Launch/Resources/Version.h"
+
+
+#define LOCTEXT_NAMESPACE "CustomizableObjectDetails"
+
+
+TSharedRef<IDetailCustomization> FCustomizableObjectNodeLayoutBlocksDetails::MakeInstance()
+{
+	return MakeShareable( new FCustomizableObjectNodeLayoutBlocksDetails );
+}
+
+
+void FCustomizableObjectNodeLayoutBlocksDetails::CustomizeDetails( IDetailLayoutBuilder& DetailBuilder )
+{
+	Node = 0;
+#if (ENGINE_MAJOR_VERSION == 4 && ENGINE_MINOR_VERSION >= 18) || ENGINE_MAJOR_VERSION >= 5
+	const IDetailsView* DetailsView = DetailBuilder.GetDetailsView();
+	if (DetailsView->GetSelectedObjects().Num())
+	{
+		Node = Cast<UCustomizableObjectNodeLayoutBlocks>(DetailsView->GetSelectedObjects()[0].Get());
+	}
+#else
+	const IDetailsView& DetailsView = DetailBuilder.GetDetailsView();
+	if ( DetailsView.GetSelectedObjects().Num() )
+	{
+		Node = Cast<UCustomizableObjectNodeLayoutBlocks>( DetailsView.GetSelectedObjects()[0].Get() );
+	}
+#endif
+
+	IDetailCategoryBuilder& BlocksCategory = DetailBuilder.EditCategory( "Blocks" );
+
+	if (Node)
+	{
+		BlocksCategory.AddCustomRow( LOCTEXT("FCustomizableObjectNodeLayoutBlocksDetails", "Blocks") )
+		[
+			SNew( STextBlock )
+			.Text(  LOCTEXT("Use the node preview panel to set up the blocks.","Use the node preview panel to set up the blocks.") )
+		];
+	}
+	else
+	{
+		BlocksCategory.AddCustomRow( LOCTEXT("FCustomizableObjectNodeLayoutBlocksDetails", "Node") )
+		[
+			SNew( STextBlock )
+			.Text( LOCTEXT( "Node not found", "Node not found" ) )
+		];
+	}
+
+}
+
+
+FIntPoint FCustomizableObjectNodeLayoutBlocksDetails::GetGridSize() const
+{
+	if (Node->Layout)
+	{
+		return Node->Layout->GetGridSize();
+	}
+
+	return FIntPoint(0);
+}
+
+
+void FCustomizableObjectNodeLayoutBlocksDetails::OnBlockChanged( int BlockIndex, FIntRect Block )
+{
+	if (Node->Layout)
+	{
+		Node->Layout->Blocks[BlockIndex].Min = Block.Min;
+		Node->Layout->Blocks[BlockIndex].Max = Block.Max;
+		Node->PostEditChange();
+	}
+}
+
+
+TArray<FIntRect> FCustomizableObjectNodeLayoutBlocksDetails::GetBlocks() const
+{
+	TArray<FIntRect> Blocks;
+
+	if (Node->Layout)
+	{
+		Blocks.SetNum(Node->Layout->Blocks.Num());
+
+		for (int BlockIndex = 0; BlockIndex < Node->Layout->Blocks.Num(); ++BlockIndex)
+		{
+			Blocks[BlockIndex] = FIntRect(Node->Layout->Blocks[BlockIndex].Min, Node->Layout->Blocks[BlockIndex].Max);
+		}
+	}
+
+	return Blocks;
+}
+
+
+void FCustomizableObjectNodeLayoutBlocksDetails::OnGridComboBoxSelectionChanged(TSharedPtr<FString> Selection, ESelectInfo::Type SelectInfo)
+{
+	/*
+	if (CustomInstance)
+	{
+		mu::Model* Model = CustomInstance->CustomizableObject->GetModel();
+
+		CustomInstance->PreEditChange(NULL);
+		CustomInstance->State = Model->FindState( TCHAR_TO_ANSI(**Selection) );
+		CustomInstance->PostEditChange();
+
+		ResetParamBox();
+	}
+	*/
+}
+
+#undef LOCTEXT_NAMESPACE
diff --git a/Engine/Plugins/Experimental/Mutable/Source/CustomizableObjectEditor/Private/Nodes/CustomizableObjectNodeLayoutBlocksDetails.h b/Engine/Plugins/Experimental/Mutable/Source/CustomizableObjectEditor/Private/Nodes/CustomizableObjectNodeLayoutBlocksDetails.h
new file mode 100644
--- /dev/null
+++ b/Engine/Plugins/Experimental/Mutable/Source/CustomizableObjectEditor/Private/Nodes/CustomizableObjectNodeLayoutBlocksDetails.h
@@ -0,0 +1,33 @@
+// Copyright Epic Games, Inc. All Rights Reserved.
+
+#pragma once
+
+
+#include "IDetailCustomization.h"
+#include "Templates/SharedPointer.h"
+#include "Types/SlateEnums.h"
+#include "UObject/WeakObjectPtr.h"
+
+
+class FCustomizableObjectNodeLayoutBlocksDetails : public IDetailCustomization
+{
+public:
+	/** Makes a new instance of this detail layout class for a specific detail view requesting it */
+	static TSharedRef<IDetailCustomization> MakeInstance();
+
+	/** ILayoutDetails interface */
+	virtual void CustomizeDetails( IDetailLayoutBuilder& DetailBuilder ) override;
+
+
+private:
+
+	class UCustomizableObjectNodeLayoutBlocks* Node;
+	TArray< TSharedPtr<FString> > GridComboOptions;
+
+	void OnGridComboBoxSelectionChanged(TSharedPtr<FString> Selection, ESelectInfo::Type SelectInfo);
+	void OnBlockChanged( int BlockIndex, FIntRect Block );
+
+	FIntPoint GetGridSize() const;
+	TArray<FIntRect> GetBlocks() const;
+
+};
\ No newline at end of file
