[mutable] Move Mutable plugin from the restricted folder to the experimental plugins folder.

diff --git a/Engine/Plugins/Experimental/Mutable/Source/CustomizableObjectEditor/Private/Nodes/CustomizableObjectNodeMaterial.cpp b/Engine/Plugins/Experimental/Mutable/Source/CustomizableObjectEditor/Private/Nodes/CustomizableObjectNodeMaterial.cpp
new file mode 100644
--- /dev/null
+++ b/Engine/Plugins/Experimental/Mutable/Source/CustomizableObjectEditor/Private/Nodes/CustomizableObjectNodeMaterial.cpp
@@ -0,0 +1,1166 @@
+// Copyright Epic Games, Inc. All Rights Reserved.
+
+#include "Nodes/CustomizableObjectNodeMaterial.h"
+#include "Materials/Material.h"
+#include "Materials/MaterialFunctionInterface.h"
+#include "CustomizableObjectCompiler.h"
+#include "EdGraphSchema_CustomizableObject.h"
+#include "GraphTraversal.h"
+#include "ObjectEditorUtils.h"
+#include "SPinViewerNodeMaterialPinImageDetails.h"
+#include "Materials/MaterialExpressionTextureCoordinate.h"
+#include "Nodes/CustomizableObjectNodeSkeletalMesh.h"
+#include "Nodes/CustomizableObjectNodeStaticMesh.h"
+#include "Nodes/SCustomizableObjectNodeMaterial.h"
+#include "Nodes/CustomizableObjectNodeTable.h"
+#include "Nodes/CustomizableObjectNodeCopyMaterial.h"
+#include "Slate/Private/Widgets/Views/SListPanel.h"
+#include "Materials/MaterialInstance.h"
+#include "MaterialCachedData.h"
+
+
+#define LOCTEXT_NAMESPACE "CustomizableObjectEditor"
+
+
+const TArray<EMaterialParameterType> UCustomizableObjectNodeMaterial::ParameterTypes = {
+	EMaterialParameterType::Texture,
+	EMaterialParameterType::Vector,
+	EMaterialParameterType::Scalar };
+
+
+bool UCustomizableObjectNodeMaterialRemapPinsByName::Equal(const UEdGraphPin& OldPin, const UEdGraphPin& NewPin) const
+{
+	const UCustomizableObjectNodeMaterialPinDataParameter* PinDataOldPin = Cast<UCustomizableObjectNodeMaterialPinDataParameter>(Node->GetPinData(OldPin));
+	const UCustomizableObjectNodeMaterialPinDataParameter* PinDataNewPin = Cast<UCustomizableObjectNodeMaterialPinDataParameter>(Node->GetPinData(NewPin));
+	if (PinDataOldPin && PinDataNewPin)
+	{
+		return PinDataOldPin->ParameterId == PinDataNewPin->ParameterId;
+	}
+	else
+	{
+		return Super::Equal(OldPin, NewPin);
+	}
+}
+
+
+void UCustomizableObjectNodeMaterialRemapPinsByName::RemapPins(const TArray<UEdGraphPin*>& OldPins, const TArray<UEdGraphPin*>& NewPins, TMap<UEdGraphPin*, UEdGraphPin*>& PinsToRemap, TArray<UEdGraphPin*>& PinsToOrphan)
+{
+	for (UEdGraphPin* OldPin : OldPins)
+	{
+		bool bFound = false;
+
+		for (UEdGraphPin* NewPin : NewPins)
+		{
+			if (Equal(*OldPin, *NewPin))
+			{
+				bFound = true;
+
+				PinsToRemap.Add(OldPin, NewPin);
+				break;
+			}
+		}
+
+		if (!bFound && (OldPin->LinkedTo.Num() || HasSavedPinData(*OldPin)))
+		{
+			PinsToOrphan.Add(OldPin);
+		}
+	}
+}
+
+
+bool UCustomizableObjectNodeMaterialRemapPinsByName::HasSavedPinData(const UEdGraphPin &Pin) const
+{
+	if (const UCustomizableObjectNodeMaterialPinDataParameter* PinData = Cast<UCustomizableObjectNodeMaterialPinDataParameter>(Node->GetPinData(Pin)))
+	{
+		return !PinData->IsDefault();
+	}
+	else
+	{
+		return false;
+	}
+}
+
+
+bool UCustomizableObjectNodeMaterialPinDataImage::IsDefault() const
+{
+	return PinMode == EPinMode::Default &&
+		UVLayout == UV_LAYOUT_DEFAULT &&
+		!ReferenceTexture;
+}
+
+
+EPinMode UCustomizableObjectNodeMaterialPinDataImage::GetPinMode() const
+{
+	return PinMode;
+}
+
+
+void UCustomizableObjectNodeMaterialPinDataImage::SetPinMode(const EPinMode InPinMode)
+{
+	FObjectEditorUtils::SetPropertyValue(this, GET_MEMBER_NAME_CHECKED(UCustomizableObjectNodeMaterialPinDataImage, PinMode), InPinMode);
+}
+
+
+void UCustomizableObjectNodeMaterialPinDataImage::Init(UCustomizableObjectNodeMaterial& InNodeMaterial)
+{
+	NodeMaterial = &InNodeMaterial;
+}
+
+
+FName UCustomizableObjectNodeMaterial::GetPinName(const EMaterialParameterType Type, const int32 ParameterIndex) const
+{
+	const FString ParameterName = GetParameterName(Type, ParameterIndex).ToString();
+	return FName(GetParameterLayerIndex(Type, ParameterIndex) != INDEX_NONE ? GetParameterLayerName(Type, ParameterIndex).ToString() + " - " + ParameterName : ParameterName);
+}
+
+
+int32 UCustomizableObjectNodeMaterial::GetExpressionTextureCoordinate(UMaterial* Material, const FGuid& ImageId)
+{
+	if (const UMaterialExpressionTextureSample* TextureSample = Material->FindExpressionByGUID<UMaterialExpressionTextureSample>(ImageId))
+	{
+		if (const UMaterialExpressionTextureCoordinate* TextureCoords = Cast<UMaterialExpressionTextureCoordinate>(TextureSample->Coordinates.Expression))
+		{
+			return TextureCoords->CoordinateIndex;
+		}
+		else
+		{
+			return 0;
+		}
+	}
+	else
+	{
+		return -1;
+	}
+}
+
+EPinMode UCustomizableObjectNodeMaterial::NodePinModeToImagePinMode(const ENodePinMode NodePinMode)
+{
+	switch (NodePinMode)
+	{
+	case ENodePinMode::Mutable:
+		return EPinMode::Mutable;
+	case ENodePinMode::Passthrough:
+		return EPinMode::Passthrough;
+	default:
+		check(false); // Missing case.
+		return EPinMode::Mutable;
+	}
+}
+
+
+EPinMode UCustomizableObjectNodeMaterial::GetImagePinMode(const UEdGraphPin& Pin) const
+{
+	if (FollowInputPin(Pin))
+	{
+		return EPinMode::Mutable;
+	}
+	else if (const UEdGraphPin* MaterialAssetPin = GetMaterialAssetPin())
+	{
+		if (const UEdGraphPin* ConnectedPin = FollowInputPin(*MaterialAssetPin))
+		{
+			if (const UCustomizableObjectNodeTable* TableNode = Cast<UCustomizableObjectNodeTable>(ConnectedPin->GetOwningNode()))
+			{
+				const FGuid ParameterGuid = GetPinData<UCustomizableObjectNodeMaterialPinDataImage>(Pin).ParameterId;
+				if (TableNode->ForceImageMutableMode(ConnectedPin, ParameterGuid))
+				{
+					return EPinMode::Mutable;
+				}
+			}
+		}
+	}
+	
+	switch (GetPinData<UCustomizableObjectNodeMaterialPinDataImage>(Pin).GetPinMode())
+	{
+	case EPinMode::Default:
+		return NodePinModeToImagePinMode(TextureParametersMode);
+	case EPinMode::Mutable:
+		return EPinMode::Mutable;
+	case EPinMode::Passthrough:
+		return EPinMode::Passthrough;
+	default:
+		check(false); // Missing case.
+		return EPinMode::Mutable;		
+	}
+}
+
+
+bool UCustomizableObjectNodeMaterialPinDataParameter::IsDefault() const
+{
+	return true;
+}
+
+
+/** Translates a given EPinMode to FText. */
+FText EPinModeToText(const EPinMode PinMode)
+{
+	switch (PinMode)
+	{
+	case EPinMode::Default:
+		return LOCTEXT("EPinModeDefault", "Node Defined");
+
+	case EPinMode::Mutable:
+		return LOCTEXT("EPinModeMutable", "Mutable");
+
+	case EPinMode::Passthrough:
+		return LOCTEXT("EPinModePassthrough", "Passthrough");
+
+	default:
+		check(false); // Missing case.
+		return FText();
+	}
+}
+
+
+void UCustomizableObjectNodeMaterialPinDataImage::PostEditChangeProperty(FPropertyChangedEvent& PropertyChangedEvent)
+{
+	Super::PostEditChangeProperty(PropertyChangedEvent);
+
+	if (const FProperty* PropertyThatChanged = PropertyChangedEvent.Property)
+	{
+		if (PropertyThatChanged->GetFName() == GET_MEMBER_NAME_CHECKED(UCustomizableObjectNodeMaterialPinDataImage, PinMode))
+		{
+			NodeMaterial->UpdateImagePinMode(ParameterId);
+		}
+	}
+}
+
+void UCustomizableObjectNodeMaterial::AllocateDefaultPins(UCustomizableObjectNodeRemapPins* RemapPins)
+{
+	if (UCustomizableObjectNodeMaterialRemapPinsByName* RemapPinsByNameCustom = Cast<UCustomizableObjectNodeMaterialRemapPinsByName>(RemapPins))
+	{
+		RemapPinsByNameCustom->Node = this;
+	}
+	
+	const UEdGraphSchema_CustomizableObject* Schema = GetDefault<UEdGraphSchema_CustomizableObject>();
+
+	{
+		const FString PinFriendlyName = TEXT("Mesh");
+		const FString PinName = PinFriendlyName + FString(TEXT("_Input_Pin"));
+		UEdGraphPin* MeshPin = CustomCreatePin(EGPD_Input, Schema->PC_Mesh, FName(*PinName));
+		MeshPin->PinFriendlyName = FText::FromString(PinFriendlyName);
+		MeshPin->bDefaultValueIsIgnored = true;
+	}
+	
+	{
+		const FString PinFriendlyName = TEXT("Table Material");
+		const FString PinName = PinFriendlyName + FString(TEXT("_Input_Pin"));
+		UEdGraphPin* TableMaterialPin = CustomCreatePin(EGPD_Input, Schema->PC_MaterialAsset, FName(*PinName));
+		TableMaterialPin->PinFriendlyName = FText::FromString(PinFriendlyName);
+		TableMaterialPin->bDefaultValueIsIgnored = true;
+		TableMaterialPin->PinToolTip = "Pin for a Material from a Table Node";
+	}
+
+	for (const EMaterialParameterType Type : ParameterTypes)
+	{
+		AllocateDefaultParameterPins(Type);
+	}
+
+	{
+		const FString PinFriendlyName = TEXT("Material");
+		const FString PinName = PinFriendlyName + FString(TEXT("_Output_Pin"));
+		UEdGraphPin* OutputPin = CustomCreatePin(EGPD_Output, Schema->PC_Material, FName(*PinName));
+		OutputPin->PinFriendlyName = FText::FromString(PinFriendlyName);
+	}
+}
+
+
+bool UCustomizableObjectNodeMaterial::CanPinBeHidden(const UEdGraphPin& Pin) const
+{
+	return Super::CanPinBeHidden(Pin) &&
+		Pin.Direction == EGPD_Input && 
+		Pin.PinType.PinCategory != UEdGraphSchema_CustomizableObject::PC_Mesh;
+}
+
+
+UCustomizableObjectNodeRemapPinsByName* UCustomizableObjectNodeMaterial::CreateRemapPinsDefault() const
+{
+	return NewObject<UCustomizableObjectNodeMaterialRemapPinsByName>();
+}
+
+
+void UCustomizableObjectNodeMaterial::BackwardsCompatibleFixup()
+{
+	Super::BackwardsCompatibleFixup();
+
+	const int32 CustomizableObjectCustomVersion = GetLinkerCustomVersion(FCustomizableObjectCustomVersion::GUID);
+	
+	if (CustomizableObjectCustomVersion < FCustomizableObjectCustomVersion::AutomaticNodeMaterial)
+	{
+		for (const FCustomizableObjectNodeMaterialImage& Image : Images_DEPRECATED)
+		{
+			const FString OldPinName = (Image.LayerIndex == -1 ? Image.Name : Image.PinName) + FString(TEXT("_Input_Image"));
+			UEdGraphPin* OldPin = FindPin(OldPinName);
+			if (!OldPin) // If we can not find a pin it means that the data was corrupted (old Image array and pins where not synchronized).
+			{
+				continue;
+			}
+
+			UCustomizableObjectNodeMaterialPinDataImage* PinData = NewObject<UCustomizableObjectNodeMaterialPinDataImage>(this);
+			PinData->ParameterId = FGuid::NewGuid();
+			PinData->ReferenceTexture = Image.ReferenceTexture;
+
+			// Find referenced Material Parameter
+			const int32 NumParameters = GetNumParameters(EMaterialParameterType::Texture);
+			for (int32 ParameterIndex = 0; ParameterIndex < NumParameters; ++ParameterIndex)
+			{
+				if (GetParameterName(EMaterialParameterType::Texture, ParameterIndex).ToString() == Image.Name)
+				{
+					PinData->ParameterId = GetParameterId(EMaterialParameterType::Texture, ParameterIndex);
+
+					if (const int32 UVLayout = GetImageUVLayout(ParameterIndex); Image.UVLayout != UVLayout)
+					{
+						PinData->UVLayout = UVLayout;
+					}
+					break;
+				}
+			}
+			
+			AddPinData(*OldPin, *PinData);
+		}
+
+		for (const FCustomizableObjectNodeMaterialVector& Vector : VectorParams_DEPRECATED)
+		{
+			const FString OldPinName = (Vector.LayerIndex == -1 ? Vector.Name : Vector.PinName) + FString(TEXT("_Input_Vector"));
+			const UEdGraphPin* OldPin = FindPin(OldPinName);
+			if (!OldPin) // If we can not find a pin it means that the data was corrupted (old ScalarParams array and pins where not synchronized).
+			{
+				continue;
+			}
+			
+			UCustomizableObjectNodeMaterialPinDataVector* PinData = NewObject<UCustomizableObjectNodeMaterialPinDataVector>(this);
+			PinData->ParameterId = FGuid::NewGuid();
+			
+			// Find referenced Material Parameter
+			const int32 NumParameters = GetNumParameters(EMaterialParameterType::Vector);
+			for (int32 ParameterIndex = 0; ParameterIndex < NumParameters; ++ParameterIndex)
+			{
+				if (GetParameterName(EMaterialParameterType::Vector, ParameterIndex).ToString() == Vector.Name)
+				{
+					PinData->ParameterId = GetParameterId(EMaterialParameterType::Vector, ParameterIndex);
+					break;
+				}
+			}
+			
+			AddPinData(*OldPin, *PinData);
+		}
+
+		for (const FCustomizableObjectNodeMaterialScalar& Scalar : ScalarParams_DEPRECATED)
+		{
+			const FString OldPinName = (Scalar.LayerIndex == -1 ? Scalar.Name : Scalar.PinName) + FString(TEXT("_Input_Scalar"));
+			const UEdGraphPin* OldPin = FindPin(OldPinName);
+			if (!OldPin) // If we can not find a pin it means that the data was corrupted (old ScalarParams array and pins where not synchronized).
+			{
+				continue;
+			}
+			
+			UCustomizableObjectNodeMaterialPinDataScalar* PinData = NewObject<UCustomizableObjectNodeMaterialPinDataScalar>(this);
+			PinData->ParameterId = FGuid::NewGuid();
+
+			// Find referenced Material Parameter
+			const int32 NumParameters = GetNumParameters(EMaterialParameterType::Scalar);
+			for (int32 ParameterIndex = 0; ParameterIndex < NumParameters; ++ParameterIndex)
+			{
+				if (GetParameterName(EMaterialParameterType::Scalar, ParameterIndex).ToString() == Scalar.Name)
+				{
+					PinData->ParameterId = GetParameterId(EMaterialParameterType::Scalar, ParameterIndex);
+					break;
+				}
+			}
+			
+			AddPinData(*OldPin, *PinData);
+		}
+
+		// Check if there are still pins which where not present in the Images, ScalarParams and ScalarParams arrays.
+		for (const UEdGraphPin* Pin : GetAllNonOrphanPins())
+		{
+			if ((Pin->PinType.PinCategory == UEdGraphSchema_CustomizableObject::PC_Image ||
+ 				 Pin->PinType.PinCategory == UEdGraphSchema_CustomizableObject::PC_Color ||
+				 Pin->PinType.PinCategory == UEdGraphSchema_CustomizableObject::PC_Float) &&
+				!GetPinData(*Pin))
+			{
+				UCustomizableObjectNodeMaterialPinDataParameter* PinData = [&](UObject* Outer) -> UCustomizableObjectNodeMaterialPinDataParameter*
+				{
+					if (Pin->PinType.PinCategory == UEdGraphSchema_CustomizableObject::PC_Image)
+					{
+						return NewObject<UCustomizableObjectNodeMaterialPinDataImage>(Outer);	
+					}
+					else if (Pin->PinType.PinCategory == UEdGraphSchema_CustomizableObject::PC_Color)
+					{
+						return NewObject<UCustomizableObjectNodeMaterialPinDataVector>(Outer);
+					}
+					else if (Pin->PinType.PinCategory == UEdGraphSchema_CustomizableObject::PC_Float)
+					{
+						return NewObject<UCustomizableObjectNodeMaterialPinDataScalar>(Outer);	
+					}
+					else
+					{
+						check(false); // Parameter type not contemplated.
+						return nullptr;
+					}
+				}(this);
+				
+				PinData->ParameterId = FGuid::NewGuid();
+				
+				AddPinData(*Pin, *PinData);
+			}
+		}
+		
+		Images_DEPRECATED.Empty();
+		VectorParams_DEPRECATED.Empty();
+		ScalarParams_DEPRECATED.Empty();
+		
+		Super::ReconstructNode(); // Super required to avoid ambiguous call compilation error.
+	}
+
+	// Fill PinsParameter.
+	if (CustomizableObjectCustomVersion < FCustomizableObjectCustomVersion::AutomaticNodeMaterialPerformanceBug) // || CustomizableObjectCustomVersion < FCustomizableObjectCustomVersion::AutomaticNodeMaterialPerformance
+	{
+		for (const UEdGraphPin* Pin : GetAllNonOrphanPins())
+		{
+			if (const UCustomizableObjectNodeMaterialPinDataParameter* PinData = Cast<UCustomizableObjectNodeMaterialPinDataParameter>(GetPinData(*Pin)))
+			{
+				PinsParameter.Add(PinData->ParameterId, FEdGraphPinReference(Pin));
+			}
+		}
+	}
+
+	if (CustomizableObjectCustomVersion < FCustomizableObjectCustomVersion::AutomaticNodeMaterialUXImprovements)
+	{
+		TextureParametersMode = bDefaultPinModeMutable_DEPRECATED ?
+			ENodePinMode::Mutable :
+			ENodePinMode::Passthrough;
+	}
+	
+	if (CustomizableObjectCustomVersion < FCustomizableObjectCustomVersion::ExtendMaterialOnlyMutableModeParameters)
+	{
+		const uint32 NumTextureParameters = GetNumParameters(EMaterialParameterType::Texture);
+		for (uint32 ImageIndex = 0; ImageIndex < NumTextureParameters; ++ImageIndex)
+		{
+			if (const UEdGraphPin* Pin = GetParameterPin(EMaterialParameterType::Texture, ImageIndex))
+			{
+				UCustomizableObjectNodeMaterialPinDataImage* PinDataImage = Cast<UCustomizableObjectNodeMaterialPinDataImage>(GetPinData(*Pin));
+				PinDataImage->NodeMaterial = this;
+				
+				PinsImagePinMode.Add(Pin->PinId, GetImagePinMode(*Pin));
+			}
+		}
+	}
+	
+	if (CustomizableObjectCustomVersion < FCustomizableObjectCustomVersion::ExtendMaterialOnlyMutableModeParametersBug)
+	{
+		for (const UEdGraphPin* Pin : GetAllPins())
+		{
+			if (UCustomizableObjectNodeMaterialPinDataImage* const PinDataImage = Cast<UCustomizableObjectNodeMaterialPinDataImage>(GetPinData(*Pin)))
+			{
+				PinDataImage->NodeMaterial = this;
+				
+				PinsImagePinMode.Add(Pin->PinId, GetImagePinMode(*Pin));
+			}
+		}
+	}
+	
+	if (CustomizableObjectCustomVersion < FCustomizableObjectCustomVersion::NodeMaterialAddTablePin)
+	{
+		Super::ReconstructNode();
+	}
+}
+
+void UCustomizableObjectNodeMaterial::PostBackwardsCompatibleFixup()
+{
+	Super::PostBackwardsCompatibleFixup();
+	
+	UpdateAllImagesPinMode(); // Could be changed when not loaded.
+}
+
+
+void UCustomizableObjectNodeMaterial::Serialize(FArchive& Ar)
+{
+	Super::Serialize(Ar);
+
+	Ar.UsingCustomVersion(FCustomizableObjectCustomVersion::GUID);
+}
+
+
+void UCustomizableObjectNodeMaterial::PostLoad()
+{
+	Super::PostLoad();
+	
+	if (Material)
+	{
+		Material->ConditionalPostLoad(); // Make sure the Material has been fully loaded.
+	}
+}
+
+
+void UCustomizableObjectNodeMaterial::PostEditChangeProperty(FPropertyChangedEvent& PropertyChangedEvent)
+{
+	Super::PostEditChangeProperty(PropertyChangedEvent);
+
+	if (const FProperty* PropertyThatChanged = PropertyChangedEvent.Property)
+	{
+		if (PropertyThatChanged->GetFName() == GET_MEMBER_NAME_CHECKED(UCustomizableObjectNodeMaterial, Material))
+		{
+			Super::ReconstructNode(); // Super required to avoid ambiguous call compilation error.
+		}
+		else if (PropertyThatChanged->GetFName() == GET_MEMBER_NAME_CHECKED(UCustomizableObjectNodeMaterial, TextureParametersMode))
+		{
+			UpdateAllImagesPinMode();
+		}
+	}
+}
+
+
+FText UCustomizableObjectNodeMaterial::GetNodeTitle(ENodeTitleType::Type TitleType) const
+{
+	if (Material)
+	{
+		FFormatNamedArguments Args;
+		Args.Add(TEXT("MaterialName"), FText::FromString(Material->GetName()));
+
+		return FText::Format(LOCTEXT("Material_Title", "{MaterialName}\nMaterial"), Args);
+	}
+	else
+	{
+		return LOCTEXT("Material", "Material");
+	}
+}
+
+
+FLinearColor UCustomizableObjectNodeMaterial::GetNodeTitleColor() const
+{
+	const UEdGraphSchema_CustomizableObject* Schema = GetDefault<UEdGraphSchema_CustomizableObject>();
+	return Schema->GetPinTypeColor(Schema->PC_Material);
+}
+
+
+UEdGraphPin* UCustomizableObjectNodeMaterial::OutputPin() const
+{
+	FString PinFriendlyName = TEXT("Material");
+	FString PinName = PinFriendlyName + FString(TEXT("_Output_Pin"));
+
+	UEdGraphPin* Pin = FindPin(PinName);
+
+	if (Pin)
+	{
+		return Pin;
+	}
+	else
+	{
+		return FindPin(TEXT("Material"));
+	}
+}
+
+
+UEdGraphPin* UCustomizableObjectNodeMaterial::GetMeshPin() const
+{
+	FString PinFriendlyName = TEXT("Mesh");
+	FString PinName = PinFriendlyName + FString(TEXT("_Input_Pin"));
+
+	UEdGraphPin* Pin = FindPin(PinName);
+
+	if(Pin)
+	{
+		return Pin;
+	}
+	else
+	{
+		return FindPin(TEXT("Mesh"));
+	}
+}
+
+
+UEdGraphPin* UCustomizableObjectNodeMaterial::GetMaterialAssetPin() const
+{
+	FString PinFriendlyName = TEXT("Table Material");
+	FString PinName = PinFriendlyName + FString(TEXT("_Input_Pin"));
+
+	UEdGraphPin* Pin = FindPin(PinName);
+
+	if (Pin)
+	{
+		return Pin;
+	}
+	else
+	{
+		return FindPin(TEXT("Table Material"));
+	}
+}
+
+
+bool UCustomizableObjectNodeMaterial::IsImageMutableMode(const int32 ImageIndex) const
+{
+	if (const UEdGraphPin* Pin = GetParameterPin(EMaterialParameterType::Texture, ImageIndex))
+	{
+		return IsImageMutableMode(*Pin);
+	}
+	else
+	{
+		return NodePinModeToImagePinMode(TextureParametersMode) == EPinMode::Mutable;
+	}
+}
+
+
+bool UCustomizableObjectNodeMaterial::IsImageMutableMode(const UEdGraphPin& Pin) const
+{
+	return PinsImagePinMode[Pin.PinId] == EPinMode::Mutable;
+}
+
+
+void UCustomizableObjectNodeMaterial::UpdateImagePinMode(const FGuid ParameterId)
+{
+	UpdateImagePinMode(*PinsParameter[ParameterId].Get());
+}
+
+
+void UCustomizableObjectNodeMaterial::UpdateImagePinMode(const UEdGraphPin& Pin)
+{
+	EPinMode* PinMode = PinsImagePinMode.Find(Pin.PinId);
+	const EPinMode OldPinMode = *PinMode;
+
+	*PinMode = GetImagePinMode(Pin);
+	
+	if (*PinMode != OldPinMode)
+	{
+		PostImagePinModeChangedDelegate.Broadcast();
+	}
+}
+
+
+void UCustomizableObjectNodeMaterial::UpdateAllImagesPinMode()
+{
+	bool bChanged = false;
+	
+	const int32 NumImages = GetNumParameters(EMaterialParameterType::Texture);
+	for (int32 ImageIndex = 0; ImageIndex < NumImages; ++ImageIndex)
+	{
+		if (const UEdGraphPin* Pin = GetParameterPin(EMaterialParameterType::Texture, ImageIndex))
+		{
+			EPinMode* PinMode = PinsImagePinMode.Find(Pin->PinId);
+			const EPinMode OldPinMode = *PinMode;
+
+			*PinMode = GetImagePinMode(*Pin);
+
+			bChanged = bChanged || (*PinMode != OldPinMode);
+		}
+	}
+
+	if (bChanged)
+	{
+		PostImagePinModeChangedDelegate.Broadcast();
+	}
+}
+
+
+UTexture2D* UCustomizableObjectNodeMaterial::GetImageReferenceTexture(const int32 ImageIndex) const
+{
+	if (const UEdGraphPin* Pin = GetParameterPin(EMaterialParameterType::Texture, ImageIndex))
+	{
+		return GetPinData<UCustomizableObjectNodeMaterialPinDataImage>(*Pin).ReferenceTexture;
+	}
+	else
+	{
+		return nullptr;
+	}
+}
+
+
+UTexture2D* UCustomizableObjectNodeMaterial::GetImageValue(const int32 ImageIndex) const
+{
+	FName TextureName = GetParameterName(EMaterialParameterType::Texture, ImageIndex);
+	
+	UTexture* Texture;
+	Material->GetTextureParameterValue(TextureName, Texture);
+
+	return Cast<UTexture2D>(Texture);
+}
+
+
+TArray<UCustomizableObjectLayout*> UCustomizableObjectNodeMaterial::GetLayouts()
+{
+	TArray<UCustomizableObjectLayout*> Result;
+
+	if (UEdGraphPin* MeshPin = GetMeshPin())
+	{
+		if (const UEdGraphPin* ConnectedPin = FollowInputPin(*MeshPin))
+		{
+			if (const UEdGraphPin* SourceMeshPin = FindMeshBaseSource(*ConnectedPin, false))
+			{
+				if (const UCustomizableObjectNodeSkeletalMesh* MeshNode = Cast<UCustomizableObjectNodeSkeletalMesh>(SourceMeshPin->GetOwningNode()))
+				{
+					Result = MeshNode->GetLayouts(SourceMeshPin);
+				}
+				else if (const UCustomizableObjectNodeTable* TableNode = Cast<UCustomizableObjectNodeTable>(SourceMeshPin->GetOwningNode()))
+				{
+					Result = TableNode->GetLayouts(SourceMeshPin);
+				}
+			}
+		}
+	}
+
+	return Result;
+}
+
+
+int32 UCustomizableObjectNodeMaterial::GetImageUVLayout(const int32 ImageIndex) const
+{
+	if (const UEdGraphPin* Pin = GetParameterPin(EMaterialParameterType::Texture, ImageIndex))
+	{
+		const UCustomizableObjectNodeMaterialPinDataImage& PinData = GetPinData<UCustomizableObjectNodeMaterialPinDataImage>(*Pin);
+		if (PinData.UVLayout != UCustomizableObjectNodeMaterialPinDataImage::UV_LAYOUT_DEFAULT)
+		{
+			return PinData.UVLayout;
+		}
+	}
+
+	// Else
+	const FGuid ImageId = GetParameterId(EMaterialParameterType::Texture, ImageIndex);
+	
+	if (int32 TextureCoordinate = GetExpressionTextureCoordinate(Material->GetMaterial(), ImageId); TextureCoordinate >= 0)
+	{
+		return TextureCoordinate;
+	}
+	else
+	{
+		FMaterialLayersFunctions Layers;
+		Material->GetMaterialLayers(Layers);
+
+		TArray<TArray<TObjectPtr<UMaterialFunctionInterface>>*> MaterialFunctionInterfaces;
+		MaterialFunctionInterfaces.SetNumUninitialized(2);
+		MaterialFunctionInterfaces[0] = &Layers.Layers;
+		MaterialFunctionInterfaces[1] = &Layers.Blends;
+
+		for (const TArray<TObjectPtr<UMaterialFunctionInterface>>* MaterialFunctionInterface : MaterialFunctionInterfaces)
+		{
+			for (const TObjectPtr<UMaterialFunctionInterface>& Layer : *MaterialFunctionInterface)
+			{
+				if (TextureCoordinate = GetExpressionTextureCoordinate(Layer->GetPreviewMaterial()->GetMaterial(), ImageId); TextureCoordinate >= 0)
+				{
+					return TextureCoordinate;
+				}	
+			}
+		}
+	}
+	
+	check(false); // Texture Parameter not found.
+	return 0;
+}
+
+
+int32 UCustomizableObjectNodeMaterial::GetNumParameters(const EMaterialParameterType Type) const
+{
+	if (Material)
+	{
+		return Material->GetCachedExpressionData().GetParameterTypeEntry(Type).ParameterInfoSet.Num();
+	}
+	else
+	{
+		return 0;
+	}
+}
+
+
+FGuid UCustomizableObjectNodeMaterial::GetParameterId(const EMaterialParameterType Type, const int32 ParameterIndex) const
+{
+	const FMaterialCachedExpressionData& Data = Material->GetCachedExpressionData();
+
+	if (Data.EditorOnlyData)
+	{
+		if (Data.EditorOnlyData->EditorEntries[(int32)Type].EditorInfo.Num() != 0)
+		{
+			return Data.EditorOnlyData->EditorEntries[(int32)Type].EditorInfo[ParameterIndex].ExpressionGuid;
+		}
+	}
+	
+	return FGuid();
+}
+
+
+FName UCustomizableObjectNodeMaterial::GetParameterName(const EMaterialParameterType Type, const int32 ParameterIndex) const
+{
+	const FMaterialCachedParameterEntry& Entry = Material->GetCachedExpressionData().GetParameterTypeEntry(Type);
+
+	for (TSet<FMaterialParameterInfo>::TConstIterator It(Entry.ParameterInfoSet); It; ++It)
+	{
+		const int32 IteratorIndex = It.GetId().AsInteger();
+
+		if (IteratorIndex == ParameterIndex)
+		{
+			return (*It).Name;
+		}
+	}
+	
+	// The parameter should exist
+	check(false);
+
+	return FName();
+}
+
+
+int32 UCustomizableObjectNodeMaterial::GetParameterLayerIndex(const EMaterialParameterType Type, const int32 ParameterIndex) const
+{
+	const FMaterialCachedParameterEntry& Entry = Material->GetCachedExpressionData().GetParameterTypeEntry(Type);
+
+	for (TSet<FMaterialParameterInfo>::TConstIterator It(Entry.ParameterInfoSet); It; ++It)
+	{
+		const int32 IteratorIndex = It.GetId().AsInteger();
+
+		if (IteratorIndex == ParameterIndex)
+		{
+			return (*It).Index;
+		}
+	}
+
+	// The parameter should exist
+	check(false);
+
+	return -1;
+}
+
+
+FText UCustomizableObjectNodeMaterial::GetParameterLayerName(const EMaterialParameterType Type, const int32 ParameterIndex) const
+{
+	const FMaterialCachedParameterEntry& Entry = Material->GetCachedExpressionData().GetParameterTypeEntry(Type);
+
+	for (TSet<FMaterialParameterInfo>::TConstIterator It(Entry.ParameterInfoSet); It; ++It)
+	{
+		const int32 IteratorIndex = It.GetId().AsInteger();
+
+		if (IteratorIndex == ParameterIndex)
+		{
+			FMaterialLayersFunctions LayersValue;
+			Material->GetMaterialLayers(LayersValue);
+
+			return LayersValue.EditorOnly.LayerNames[(*It).Index];
+		}
+	}
+	
+	// The parameter should exist
+	check(false);
+
+	return FText();
+}
+
+
+bool UCustomizableObjectNodeMaterial::HasParameter(const FGuid& ParameterId) const
+{
+	for (const EMaterialParameterType Type : ParameterTypes)
+	{
+		const FMaterialCachedExpressionData& Data = Material->GetCachedExpressionData();
+		const FMaterialCachedParameterEntry& Entry = Data.GetParameterTypeEntry(Type);
+
+		if (!Data.EditorOnlyData || Data.EditorOnlyData->EditorEntries[(int32)Type].EditorInfo.IsEmpty())
+		{
+			continue;
+		}
+			
+		for (TSet<FMaterialParameterInfo>::TConstIterator It(Entry.ParameterInfoSet); It; ++It)
+		{
+			const int32 IteratorIndex = It.GetId().AsInteger();
+			const FGuid& ParamGUid = Data.EditorOnlyData->EditorEntries[(int32)Type].EditorInfo[IteratorIndex].ExpressionGuid;
+			
+			if (ParamGUid == ParameterId)
+			{
+				return true;
+			}
+		}
+	}
+
+	return false;
+}
+
+
+const UEdGraphPin* UCustomizableObjectNodeMaterial::GetParameterPin(const EMaterialParameterType Type, const int32 ParameterIndex) const
+{
+	const FGuid ParameterId = GetParameterId(Type, ParameterIndex);
+	
+	if (const FEdGraphPinReference* PinReference = PinsParameter.Find(ParameterId))
+	{
+		return PinReference->Get();
+	}
+	else
+	{
+		return nullptr;
+	}
+}
+
+
+bool UCustomizableObjectNodeMaterial::IsNodeOutDatedAndNeedsRefresh()
+{
+	const bool bOutdated = RealMaterialDataHasChanged();
+
+	// Remove previous compilation warnings
+	if (!bOutdated && bHasCompilerMessage)
+	{
+		RemoveWarnings();
+		GetGraph()->NotifyGraphChanged();
+	}
+
+	return bOutdated;
+}
+
+
+FString UCustomizableObjectNodeMaterial::GetRefreshMessage() const
+{
+	return "Referenced material has changed, texture channels might have been added, removed or renamed. Please refresh the node material to reflect those changes.";
+}
+
+
+TSharedPtr<SWidget> UCustomizableObjectNodeMaterial::CustomizePinDetails(UEdGraphPin& Pin)
+{
+	if (UCustomizableObjectNodeMaterialPinDataImage* PinData = Cast<UCustomizableObjectNodeMaterialPinDataImage>(GetPinData(Pin)))
+	{
+		return SNew(SPinViewerNodeMaterialPinImageDetails).PinData(PinData);
+	}
+	else
+	{
+		return nullptr;
+	}
+}
+
+
+bool UCustomizableObjectNodeMaterial::CustomRemovePin(UEdGraphPin& Pin)
+{
+	if (const UCustomizableObjectNodeMaterialPinDataParameter* PinData = Cast<UCustomizableObjectNodeMaterialPinDataParameter>(GetPinData(Pin)))
+	{
+		if (const FEdGraphPinReference* Result = PinsParameter.Find(PinData->ParameterId);
+			Result &&
+			Result->Get() == &Pin)
+		{
+			PinsParameter.Remove(PinData->ParameterId);
+		}
+	}
+
+	PinsImagePinMode.Remove(Pin.PinId);
+	
+	return Super::CustomRemovePin(Pin);
+}
+
+void UCustomizableObjectNodeMaterial::ReconstructNode(UCustomizableObjectNodeRemapPins* RemapPinsMode)
+{
+	Super::ReconstructNode(RemapPinsMode);
+
+	UpdateAllImagesPinMode();
+}
+
+
+bool UCustomizableObjectNodeMaterial::RealMaterialDataHasChanged() const
+{
+	for (const UEdGraphPin* Pin : GetAllNonOrphanPins())
+	{
+		if (const UCustomizableObjectNodeMaterialPinDataParameter* PinData = Cast<UCustomizableObjectNodeMaterialPinDataParameter>(GetPinData(*Pin)))
+		{
+			if (!HasParameter(PinData->ParameterId) &&
+				(FollowInputPin(*Pin) || !PinData->IsDefault()))
+			{
+				return true;
+			}
+		}
+	}
+
+	return false;
+}
+
+
+bool UCustomizableObjectNodeMaterial::IsPinRelevant(const UEdGraphPin* Pin) const
+{
+	const UEdGraphSchema_CustomizableObject* Schema = GetDefault<UEdGraphSchema_CustomizableObject>();
+
+	if (Pin->Direction == EEdGraphPinDirection::EGPD_Output)
+	{	
+		return Pin->PinType.PinCategory == Schema->PC_Mesh;
+	}
+	else if (Pin->Direction == EEdGraphPinDirection::EGPD_Input)
+	{	
+		return Pin->PinType.PinCategory == Schema->PC_Material;
+	}
+	else
+	{
+		return false;
+	}
+}
+
+
+FText UCustomizableObjectNodeMaterial::GetTooltipText() const
+{
+	return LOCTEXT("Material_Tooltip", "Defines a Customizable Object material.\nConsists of a mesh section, a material asset to paint it, and the runtime modifiable inputs to the material asset parameters.");
+}
+
+
+TSharedPtr<SGraphNode> UCustomizableObjectNodeMaterial::CreateVisualWidget()
+{
+	return SNew(SCustomizableObjectNodeMaterial, this);
+}
+
+
+UCustomizableObjectNodeMaterialPinDataParameter* UCustomizableObjectNodeMaterial::CreatePinData(const EMaterialParameterType Type, const int32 ParameterIndex)
+{
+	UCustomizableObjectNodeMaterialPinDataParameter* PinData = nullptr;
+
+	switch (Type)
+	{
+	case EMaterialParameterType::Texture:
+		{
+			UCustomizableObjectNodeMaterialPinDataImage* PinDataImage = NewObject<UCustomizableObjectNodeMaterialPinDataImage>(this);
+			PinDataImage->Init(*this);
+		
+			PinData = PinDataImage;
+			break;
+		}
+	case EMaterialParameterType::Vector:
+		{
+			PinData = NewObject<UCustomizableObjectNodeMaterialPinDataVector>(this);	
+			break;
+		}
+
+	case EMaterialParameterType::Scalar:
+		{
+			PinData = NewObject<UCustomizableObjectNodeMaterialPinDataScalar>(this);	
+			break;
+		}
+
+	default:
+		check(false); // Parameter type not contemplated.
+	}
+
+	PinData->ParameterId = GetParameterId(Type, ParameterIndex);
+
+	return PinData;
+}
+
+void UCustomizableObjectNodeMaterial::AllocateDefaultParameterPins(const EMaterialParameterType Type)
+{
+	const FName PinCategory = UEdGraphSchema_CustomizableObject::GetPinCategory(Type);
+
+	const int32 NumParameters = GetNumParameters(Type);
+	for (int32 ParameterIndex = 0; ParameterIndex < NumParameters; ++ParameterIndex)
+	{
+		UCustomizableObjectNodeMaterialPinDataParameter* PinData = CreatePinData(Type, ParameterIndex);
+		
+		const FName PinName = GetPinName(Type, ParameterIndex);
+		UEdGraphPin* Pin = CustomCreatePin(EGPD_Input, PinCategory, PinName, PinData);
+		Pin->bHidden = true;
+		Pin->bDefaultValueIsIgnored = true;
+
+		PinsParameter.Add(PinData->ParameterId, FEdGraphPinReference(Pin));
+
+		if (Type == EMaterialParameterType::Texture)
+		{
+			PinsImagePinMode.Add(Pin->PinId, GetImagePinMode(*Pin));
+		}
+	}
+}
+
+
+void UCustomizableObjectNodeMaterial::SetDefaultMaterial()
+{
+	if (const UEdGraphPin* MeshPin = GetMeshPin();
+		!Material && MeshPin)
+	{
+		if (const UEdGraphPin* LinkedMeshPin = FollowInputPin(*MeshPin))
+		{
+			UEdGraphNode* LinkedMeshNode = LinkedMeshPin->GetOwningNode();
+	
+			if (const UCustomizableObjectNodeSkeletalMesh* NodeSkeletalMesh = Cast<UCustomizableObjectNodeSkeletalMesh>(LinkedMeshNode))
+			{
+				Material = NodeSkeletalMesh->GetMaterialFor(LinkedMeshPin);
+				if (Material)
+				{
+					Super::ReconstructNode(); // Super required to avoid ambiguous call compilation error.
+				}
+			}
+			else if (const UCustomizableObjectNodeStaticMesh* NodeStaticMesh = Cast<UCustomizableObjectNodeStaticMesh>(LinkedMeshNode))
+			{
+				Material = NodeStaticMesh->GetMaterialFor(LinkedMeshPin);
+				if (Material)
+				{
+					Super::ReconstructNode();
+				}
+			}
+		}
+	}
+}
+
+
+void UCustomizableObjectNodeMaterial::PinConnectionListChanged(UEdGraphPin* Pin)
+{
+	if (Pin == GetMeshPin())
+	{
+		if (LastMeshNodeConnected.IsValid())
+		{
+			LastMeshNodeConnected->PostEditChangePropertyDelegate.RemoveDynamic(this, &UCustomizableObjectNodeMaterial::MeshPostEditChangeProperty);
+		}
+
+		if (const UEdGraphPin* ConnectedPin = FollowInputPin(*Pin))
+		{
+			UEdGraphNode* MeshNode = ConnectedPin->GetOwningNode();
+
+			if (MeshNode->IsA(UCustomizableObjectNodeStaticMesh::StaticClass()) || MeshNode->IsA(UCustomizableObjectNodeSkeletalMesh::StaticClass()))
+			{
+				SetDefaultMaterial();
+
+				LastMeshNodeConnected = Cast<UCustomizableObjectNode>(MeshNode);
+				LastMeshNodeConnected->PostEditChangePropertyDelegate.AddUniqueDynamic(this, &UCustomizableObjectNodeMaterial::MeshPostEditChangeProperty);
+			}
+		}
+	}
+	else if (Cast<UCustomizableObjectNodeMaterialPinDataImage>(GetPinData(*Pin)))
+	{
+		UpdateImagePinMode(*Pin);
+	}
+	else if (Pin == GetMaterialAssetPin())
+	{
+		UpdateAllImagesPinMode();	
+	}
+}
+
+
+void UCustomizableObjectNodeMaterial::PostPasteNode()
+{
+	Super::PostPasteNode();
+	SetDefaultMaterial();
+}
+
+
+void UCustomizableObjectNodeMaterial::MeshPostEditChangeProperty(FPostEditChangePropertyDelegateParameters& Parameters)
+{
+	if (const UEdGraphPin* MeshPin = FindPin(TEXT("Mesh")))
+	{
+		if (const UEdGraphPin* ConnectedPin = FollowInputPin(*MeshPin);
+			ConnectedPin && ConnectedPin->GetOwningNode() == Parameters.Node)
+		{
+			SetDefaultMaterial();
+		}
+		else if (UCustomizableObjectNode* MeshNode = Cast<UCustomizableObjectNode>(Parameters.Node))
+		{
+			MeshNode->PostEditChangePropertyDelegate.RemoveDynamic(this, &UCustomizableObjectNodeMaterial::MeshPostEditChangeProperty);
+		}
+	}
+}
+
+
+void UCustomizableObjectNodeMaterial::BreakExistingConnectionsPostConnection(UEdGraphPin* InputPin, UEdGraphPin* OutputPin)
+{
+	if (OutputPin->GetOwningNode() == this) // Case: NodeMaterial (OutputPin) --> (InputPin) Other Node
+	{
+		// Allow multiple NodeCopyMaterial nodes to be connected
+		if (InputPin->GetOwningNode()->IsA(UCustomizableObjectNodeCopyMaterial::StaticClass()))
+		{
+			return;
+		}
+
+		// Remove the previous NodeMaterial
+		UEdGraphPin* RemovePin = nullptr;
+
+		for (UEdGraphPin* Pin : OutputPin->LinkedTo)
+		{
+			if (Pin != InputPin && !Pin->GetOwningNode()->IsA(UCustomizableObjectNodeCopyMaterial::StaticClass())) {
+				RemovePin = Pin;
+				break;
+			}
+		}
+
+		OutputPin->BreakLinkTo(RemovePin); // Can not be called inside the range for loop. Can be called with a null value.
+	}
+}
+
+#undef LOCTEXT_NAMESPACE
diff --git a/Engine/Plugins/Experimental/Mutable/Source/CustomizableObjectEditor/Private/Nodes/CustomizableObjectNodeMaterialBase.cpp b/Engine/Plugins/Experimental/Mutable/Source/CustomizableObjectEditor/Private/Nodes/CustomizableObjectNodeMaterialBase.cpp
new file mode 100644
--- /dev/null
+++ b/Engine/Plugins/Experimental/Mutable/Source/CustomizableObjectEditor/Private/Nodes/CustomizableObjectNodeMaterialBase.cpp
@@ -0,0 +1,10 @@
+// Copyright Epic Games, Inc. All Rights Reserved.
+
+#include "Nodes/CustomizableObjectNodeMaterialBase.h"
+#include "CustomizableObjectEditorModule.h"
+
+
+bool UCustomizableObjectNodeMaterialBase::ShouldBreakExistingConnections(const UEdGraphPin* InputPin, const UEdGraphPin* OutputPin) const
+{
+	return false;
+}
diff --git a/Engine/Plugins/Experimental/Mutable/Source/CustomizableObjectEditor/Private/Nodes/CustomizableObjectNodeMaterialVariation.cpp b/Engine/Plugins/Experimental/Mutable/Source/CustomizableObjectEditor/Private/Nodes/CustomizableObjectNodeMaterialVariation.cpp
new file mode 100644
--- /dev/null
+++ b/Engine/Plugins/Experimental/Mutable/Source/CustomizableObjectEditor/Private/Nodes/CustomizableObjectNodeMaterialVariation.cpp
@@ -0,0 +1,75 @@
+// Copyright Epic Games, Inc. All Rights Reserved.
+
+#include "Nodes/CustomizableObjectNodeMaterialVariation.h"
+#include "CustomizableObjectEditorModule.h"
+
+#include "UnrealEditorPortabilityHelpers.h"
+
+#define LOCTEXT_NAMESPACE "CustomizableObjectEditor"
+
+
+UCustomizableObjectNodeMaterialVariation::UCustomizableObjectNodeMaterialVariation()
+	: Super()
+{
+
+}
+
+
+void UCustomizableObjectNodeMaterialVariation::PostEditChangeProperty(FPropertyChangedEvent& PropertyChangedEvent)
+{
+	FProperty* PropertyThatChanged = PropertyChangedEvent.Property;
+	if ( PropertyThatChanged 
+		//&& PropertyThatChanged->GetName() == TEXT("Variations")
+		)
+	{
+		ReconstructNode();
+	}
+
+	Super::PostEditChangeProperty(PropertyChangedEvent);
+}
+
+
+void UCustomizableObjectNodeMaterialVariation::AllocateDefaultPins(UCustomizableObjectNodeRemapPins* RemapPins)
+{
+	const UEdGraphSchema_CustomizableObject* Schema = GetDefault<UEdGraphSchema_CustomizableObject>();
+
+	FString PinName = TEXT("Material");
+	UEdGraphPin* OutputPin = CustomCreatePin(EGPD_Output, Schema->PC_Material, FName(*PinName));
+	OutputPin->bDefaultValueIsIgnored = true;
+
+	for ( int VariationIndex = Variations.Num() - 1; VariationIndex >= 0; --VariationIndex)
+	{
+		PinName = FString::Printf( TEXT("Variation %d"), VariationIndex);
+		UEdGraphPin* VariationPin = CustomCreatePin(EGPD_Input, Schema->PC_Material, FName(*PinName), true);
+		VariationPin->bDefaultValueIsIgnored = true;
+
+		FString FriendlyName = FString::Printf(TEXT("Variation %d [%s]"), VariationIndex, *Variations[VariationIndex].Tag);
+		VariationPin->PinFriendlyName = FText::FromString(FriendlyName);
+	}
+
+	PinName = TEXT("Default");
+	UEdGraphPin* DefaultVariation = CustomCreatePin(EGPD_Input, Schema->PC_Material, FName(*PinName), true);
+	DefaultVariation->bDefaultValueIsIgnored = true;
+}
+
+
+FText UCustomizableObjectNodeMaterialVariation::GetNodeTitle(ENodeTitleType::Type TitleType) const
+{
+	return LOCTEXT("Material_Variation", "Material Variation");
+}
+
+
+FLinearColor UCustomizableObjectNodeMaterialVariation::GetNodeTitleColor() const
+{
+	const UEdGraphSchema_CustomizableObject* Schema = GetDefault<UEdGraphSchema_CustomizableObject>();
+	return Schema->GetPinTypeColor(Schema->PC_Material);
+}
+
+
+FText UCustomizableObjectNodeMaterialVariation::GetTooltipText() const
+{
+	return LOCTEXT("Material_Variation_Tooltip", "Changes the materials given depending on what tags or state is active.");
+}
+
+#undef LOCTEXT_NAMESPACE
+
diff --git a/Engine/Plugins/Experimental/Mutable/Source/CustomizableObjectEditor/Private/Nodes/CustomizableObjectNodeMesh.cpp b/Engine/Plugins/Experimental/Mutable/Source/CustomizableObjectEditor/Private/Nodes/CustomizableObjectNodeMesh.cpp
new file mode 100644
--- /dev/null
+++ b/Engine/Plugins/Experimental/Mutable/Source/CustomizableObjectEditor/Private/Nodes/CustomizableObjectNodeMesh.cpp
@@ -0,0 +1,5 @@
+// Copyright Epic Games, Inc. All Rights Reserved.
+
+#include "Nodes/CustomizableObjectNodeMesh.h"
+
+
diff --git a/Engine/Plugins/Experimental/Mutable/Source/CustomizableObjectEditor/Private/Nodes/CustomizableObjectNodeMeshClipDeform.cpp b/Engine/Plugins/Experimental/Mutable/Source/CustomizableObjectEditor/Private/Nodes/CustomizableObjectNodeMeshClipDeform.cpp
new file mode 100644
--- /dev/null
+++ b/Engine/Plugins/Experimental/Mutable/Source/CustomizableObjectEditor/Private/Nodes/CustomizableObjectNodeMeshClipDeform.cpp
@@ -0,0 +1,43 @@
+// Copyright Epic Games, Inc. All Rights Reserved.
+
+#include "Nodes/CustomizableObjectNodeMeshClipDeform.h"
+#include "CustomizableObjectEditorModule.h"
+#include "Toolkits/ToolkitManager.h"
+#include "UnrealEditorPortabilityHelpers.h"
+
+#define LOCTEXT_NAMESPACE "CustomizableObjectEditor"
+
+
+UCustomizableObjectNodeMeshClipDeform::UCustomizableObjectNodeMeshClipDeform()
+	: Super()
+{
+}
+
+void UCustomizableObjectNodeMeshClipDeform::AllocateDefaultPins(UCustomizableObjectNodeRemapPins* RemapPins)
+{
+	const UEdGraphSchema_CustomizableObject* Schema = GetDefault<UEdGraphSchema_CustomizableObject>();
+	
+	UEdGraphPin* ClipMeshPin = CustomCreatePin(EGPD_Input, Schema->PC_Mesh, FName("Clip Shape"));
+	ClipMeshPin->bDefaultValueIsIgnored = true;
+	
+	UEdGraphPin* OutputPin = CustomCreatePin(EGPD_Output, Schema->PC_Material, FName("Material"));
+	ClipMeshPin->bDefaultValueIsIgnored = true;
+}
+
+FText UCustomizableObjectNodeMeshClipDeform::GetNodeTitle(ENodeTitleType::Type TitleType) const
+{
+	return LOCTEXT("Clip_Deform_Mesh", "Clip Deform Mesh");
+}
+
+FLinearColor UCustomizableObjectNodeMeshClipDeform::GetNodeTitleColor() const
+{
+	const UEdGraphSchema_CustomizableObject* Schema = GetDefault<UEdGraphSchema_CustomizableObject>();
+	return Schema->GetPinTypeColor(Schema->PC_Material);
+}
+
+FText UCustomizableObjectNodeMeshClipDeform::GetTooltipText() const
+{
+	return LOCTEXT("Clip_Deform_Tooltip", "Defines a clip with mesh deformation based on a shape mesh and blend weights.");
+
+}
+#undef LOCTEXT_NAMESPACE
diff --git a/Engine/Plugins/Experimental/Mutable/Source/CustomizableObjectEditor/Private/Nodes/CustomizableObjectNodeMeshClipMorph.cpp b/Engine/Plugins/Experimental/Mutable/Source/CustomizableObjectEditor/Private/Nodes/CustomizableObjectNodeMeshClipMorph.cpp
new file mode 100644
--- /dev/null
+++ b/Engine/Plugins/Experimental/Mutable/Source/CustomizableObjectEditor/Private/Nodes/CustomizableObjectNodeMeshClipMorph.cpp
@@ -0,0 +1,170 @@
+// Copyright Epic Games, Inc. All Rights Reserved.
+
+#include "Nodes/CustomizableObjectNodeMeshClipMorph.h"
+#include "CustomizableObjectEditorModule.h"
+#include "Toolkits/ToolkitManager.h"
+#include "UnrealEditorPortabilityHelpers.h"
+
+#define LOCTEXT_NAMESPACE "CustomizableObjectEditor"
+
+
+UCustomizableObjectNodeMeshClipMorph::UCustomizableObjectNodeMeshClipMorph()
+	: Super()
+{
+	StartOffset = FVector::ZeroVector;
+	bLocalStartOffset = true;
+	B = 0.f;
+	Radius = 8.f;
+	Radius2 = 4.f;
+	RotationAngle = 0.f;
+	Exponent = 1.f;
+
+	Origin = FVector::ZeroVector;
+	Normal = -FVector::UpVector;
+	MaxEffectRadius = -1.f;
+
+	bUpdateViewportWidget = true;
+}
+
+
+FVector UCustomizableObjectNodeMeshClipMorph::GetOriginWithOffset() const
+{
+	FVector NewOrigin;
+
+	if (bLocalStartOffset)
+	{
+		FVector XAxis, YAxis, ZAxis;
+		FindLocalAxes(XAxis, YAxis, ZAxis);
+
+		NewOrigin = Origin + StartOffset.X * XAxis + StartOffset.Y * YAxis + StartOffset.Z * ZAxis;
+	}
+	else
+	{
+		NewOrigin = Origin + StartOffset;
+	}
+
+	return NewOrigin;
+}
+
+void UCustomizableObjectNodeMeshClipMorph::FindLocalAxes(FVector& XAxis, FVector& YAxis, FVector& ZAxis) const
+{
+	YAxis = FVector(0.f, 1.f, 0.f);
+
+	if (FMath::Abs(FVector::DotProduct(Normal, YAxis)) > 0.95f)
+	{
+		YAxis = FVector(0.f, 0.f, 1.f);
+	}
+
+	XAxis = FVector::CrossProduct(Normal, YAxis);
+	XAxis = XAxis.RotateAngleAxis(RotationAngle, Normal);
+	YAxis = FVector::CrossProduct(Normal, XAxis);
+	ZAxis = Normal;
+
+	XAxis.Normalize();
+	YAxis.Normalize();
+}
+
+
+void UCustomizableObjectNodeMeshClipMorph::ChangeStartOffsetTransform()
+{
+	// Local Offset
+	FVector XAxis, YAxis, ZAxis;
+	FindLocalAxes(XAxis, YAxis, ZAxis);
+
+	if (bLocalStartOffset)
+	{
+		StartOffset = FVector(FVector::DotProduct(StartOffset, XAxis), FVector::DotProduct(StartOffset, YAxis),
+							   FVector::DotProduct(StartOffset, ZAxis));
+	}
+	else
+	{
+		StartOffset = StartOffset.X * XAxis + StartOffset.Y * YAxis + StartOffset.Z * ZAxis;
+	}
+}
+
+
+void UCustomizableObjectNodeMeshClipMorph::PostEditChangeProperty(FPropertyChangedEvent & PropertyChangedEvent)
+{
+	FProperty* PropertyThatChanged = PropertyChangedEvent.Property;
+
+	const FName PropertyName = Helper_GetPropertyName(PropertyChangedEvent);
+
+	if (PropertyName == "bLocalStartOffset")
+	{
+		ChangeStartOffsetTransform();
+	}
+
+	bUpdateViewportWidget = true;
+
+	Super::PostEditChangeProperty(PropertyChangedEvent);
+}
+
+void UCustomizableObjectNodeMeshClipMorph::Serialize(FArchive& Ar)
+{
+	Super::Serialize(Ar);
+
+	Ar.UsingCustomVersion(FCustomizableObjectCustomVersion::GUID);
+
+	if (Ar.CustomVer(FCustomizableObjectCustomVersion::GUID) < FCustomizableObjectCustomVersion::PostLoadToCustomVersion
+		&& bLocalStartOffset)
+	{
+		// Previous Offset
+		FVector Tangent, Binormal;
+		Origin.FindBestAxisVectors(Tangent, Binormal);
+		FVector OldOffset = StartOffset.X * Tangent + StartOffset.Y * Binormal + StartOffset.Z * Normal;
+
+		// Local Offset
+		FVector XAxis, YAxis, ZAxis;
+		FindLocalAxes(XAxis, YAxis, ZAxis);
+
+		StartOffset = FVector(FVector::DotProduct(OldOffset, XAxis), FVector::DotProduct(OldOffset, YAxis),
+							  FVector::DotProduct(OldOffset, ZAxis));
+	}
+}
+
+
+void UCustomizableObjectNodeMeshClipMorph::AllocateDefaultPins(UCustomizableObjectNodeRemapPins* RemapPins)
+{
+	const UEdGraphSchema_CustomizableObject* Schema = GetDefault<UEdGraphSchema_CustomizableObject>();
+
+	UEdGraphPin* OutputPin = CustomCreatePin(EGPD_Output, Schema->PC_Material, FName("Material"));
+}
+
+
+FText UCustomizableObjectNodeMeshClipMorph::GetNodeTitle(ENodeTitleType::Type TitleType) const
+{
+	return LOCTEXT("Clip_Morph_Mesh", "Clip Morph Mesh");
+}
+
+
+FLinearColor UCustomizableObjectNodeMeshClipMorph::GetNodeTitleColor() const
+{
+	const UEdGraphSchema_CustomizableObject* Schema = GetDefault<UEdGraphSchema_CustomizableObject>();
+	return Schema->GetPinTypeColor(Schema->PC_Material);
+}
+
+
+void UCustomizableObjectNodeMeshClipMorph::PinConnectionListChanged(UEdGraphPin* Pin)
+{
+	Super::PinConnectionListChanged(Pin);
+
+	if (Pin == OutputPin())
+	{
+		TSharedPtr<ICustomizableObjectEditor> Editor = GetGraphEditor();
+
+		if (Editor.IsValid())
+		{
+			Editor->UpdateGraphNodeProperties();
+		}
+	}
+}
+
+
+
+FText UCustomizableObjectNodeMeshClipMorph::GetTooltipText() const
+{
+	return LOCTEXT("Clip_Mesh_Morph_Tooltip", "Defines a cutting plane on a bone to cut tagged Materials that go past it, while morphing the mesh after the cut to blend in more naturally.\nIt only cuts and morphs mesh that receives some influence of that bone or other descendant bones.");
+}
+
+
+#undef LOCTEXT_NAMESPACE
diff --git a/Engine/Plugins/Experimental/Mutable/Source/CustomizableObjectEditor/Private/Nodes/CustomizableObjectNodeMeshClipMorphDetails.cpp b/Engine/Plugins/Experimental/Mutable/Source/CustomizableObjectEditor/Private/Nodes/CustomizableObjectNodeMeshClipMorphDetails.cpp
new file mode 100644
--- /dev/null
+++ b/Engine/Plugins/Experimental/Mutable/Source/CustomizableObjectEditor/Private/Nodes/CustomizableObjectNodeMeshClipMorphDetails.cpp
@@ -0,0 +1,364 @@
+// Copyright Epic Games, Inc. All Rights Reserved.
+
+#include "CustomizableObjectNodeMeshClipMorphDetails.h"
+#include "CustomizableObjectEditorModule.h"
+#include "PropertyCustomizationHelpers.h"
+#include "CustomizableObjectEditorUtilities.h"
+#include "CustomizableObject.h"
+#include "Nodes/CustomizableObjectNodeObjectGroup.h"
+#include "CustomizableObjectCompiler.h"
+#include "GraphTraversal.h"
+
+#include "DetailLayoutBuilder.h"
+#include "DetailCategoryBuilder.h"
+#include "LevelEditor.h"
+#include "LevelEditorActions.h"
+
+#include "Engine/SkeletalMesh.h"
+
+#include "Widgets/Layout/SBorder.h"
+#include "Widgets/Text/STextBlock.h"
+#include "Widgets/Input/STextComboBox.h"
+#include "Widgets/Input/SCheckBox.h"
+
+#include "Runtime/Launch/Resources/Version.h"
+
+
+#define LOCTEXT_NAMESPACE "MeshClipMorphDetails"
+
+
+TSharedRef<IDetailCustomization> FCustomizableObjectNodeMeshClipMorphDetails::MakeInstance()
+{
+	return MakeShareable(new FCustomizableObjectNodeMeshClipMorphDetails);
+}
+
+
+void FCustomizableObjectNodeMeshClipMorphDetails::CustomizeDetails(IDetailLayoutBuilder& DetailBuilder)
+{
+	Node = nullptr;
+	DetailBuilderPtr = &DetailBuilder;
+
+#if (ENGINE_MAJOR_VERSION == 4 && ENGINE_MINOR_VERSION >= 18) || ENGINE_MAJOR_VERSION >= 5
+	const IDetailsView* DetailsView = DetailBuilder.GetDetailsView();
+	if (DetailsView->GetSelectedObjects().Num())
+	{
+		Node = Cast<UCustomizableObjectNodeMeshClipMorph>(DetailsView->GetSelectedObjects()[0].Get());
+	}
+#else
+	const IDetailsView& DetailsView = DetailBuilder.GetDetailsView();
+	if (DetailsView.GetSelectedObjects().Num())
+	{
+		Node = Cast<UCustomizableObjectNodeMeshClipMorph>(DetailsView.GetSelectedObjects()[0].Get());
+	}
+#endif
+
+	IDetailCategoryBuilder& BlocksCategory = DetailBuilder.EditCategory("MeshToClipAndMorph");
+	DetailBuilder.HideProperty("BoneName");
+
+	IDetailCategoryBuilder& MeshClipParametersCategory = DetailBuilder.EditCategory("MeshClipParameters");
+	DetailBuilder.HideProperty("bInvertNormal");
+
+	TSharedPtr<IPropertyHandle> ReferenceSkeletonIndexPropertyHandle = DetailBuilder.GetProperty("ReferenceSkeletonIndex");
+	const FSimpleDelegate OnReferenceSkeletonIndexChangedDelegate = 
+		FSimpleDelegate::CreateSP(this, &FCustomizableObjectNodeMeshClipMorphDetails::OnReferenceSkeletonIndexChanged);
+	ReferenceSkeletonIndexPropertyHandle->SetOnPropertyValueChanged(OnReferenceSkeletonIndexChangedDelegate);
+
+	if (Node)
+	{
+		SkeletalMesh = nullptr;
+		TSharedPtr<FString> BoneToSelect;
+
+		UCustomizableObject* CustomizableObject = Cast<UCustomizableObject>(Node->GetCustomizableObjectGraph()->GetOuter());
+
+		if (CustomizableObject)
+		{
+			if (CustomizableObject->bIsChildObject)
+			{
+				bool bMultipleBaseObjects;
+				UCustomizableObjectNodeObject* RootNode = GetRootNode(CustomizableObject, bMultipleBaseObjects);
+
+				if (RootNode && !bMultipleBaseObjects)
+				{
+					TArray<UCustomizableObject*> VisitedObjects;
+					UCustomizableObject* Parent = GetFullGraphRootObject(RootNode, VisitedObjects);
+
+					if (Parent)
+					{
+						SkeletalMesh = VisitedObjects.Last()->GetRefSkeletalMesh(Node->ReferenceSkeletonIndex);
+					}
+				}
+			}
+			else
+			{
+				SkeletalMesh = CustomizableObject->GetRefSkeletalMesh(Node->ReferenceSkeletonIndex);
+			}
+		}
+
+		if (SkeletalMesh)
+		{
+			BoneComboOptions.Empty();
+				
+			for (int32 i = 0; i < SkeletalMesh->GetRefSkeleton().GetNum(); ++i)
+			{
+				FName BoneName = SkeletalMesh->GetRefSkeleton().GetBoneName(i);
+				BoneComboOptions.Add(MakeShareable(new FString(BoneName.ToString())));
+
+				if (BoneName == Node->BoneName)
+				{
+					BoneToSelect = BoneComboOptions.Last();
+				}
+			}
+
+            BoneComboOptions.Sort(CompareNames);
+
+			// Add them to the parent combo box
+			TSharedRef<IPropertyHandle> BoneProperty = DetailBuilder.GetProperty("BoneName");
+
+			BlocksCategory.AddCustomRow(LOCTEXT("FCustomizableObjectNodeMeshClipMorphDetails", "Bone Name"))
+			[
+				SNew(SProperty, BoneProperty)
+				.ShouldDisplayName(false)
+				.CustomWidget()
+				[
+					SNew(SBorder)
+					.BorderImage(FAppStyle::GetBrush("NoBorder"))
+					.Padding(FMargin(0.0f, 0.0f, 10.0f, 0.0f))
+					.HAlign(EHorizontalAlignment::HAlign_Fill)
+					[
+						SNew(SHorizontalBox)
+						+SHorizontalBox::Slot().HAlign(HAlign_Left).VAlign(VAlign_Center)
+						[
+							SNew(STextBlock)
+							.Text( LOCTEXT("FCustomizableObjectNodeRemoveMeshDetails", "Bone Name"))
+							.Font(IDetailLayoutBuilder::GetDetailFont())
+						]
+						+SHorizontalBox::Slot().HAlign(HAlign_Fill)
+						[
+							SNew(STextComboBox)
+							.OptionsSource(&BoneComboOptions)
+							.InitiallySelectedItem(BoneToSelect)
+							.OnSelectionChanged(this, &FCustomizableObjectNodeMeshClipMorphDetails::OnBoneComboBoxSelectionChanged, BoneProperty)
+							.Font(IDetailLayoutBuilder::GetDetailFont())
+						]
+					]
+				]
+			];
+
+			MeshClipParametersCategory.AddCustomRow(LOCTEXT("FCustomizableObjectNodeMeshClipMorphDetails", "Invert plane normal"))
+			[
+				SNew(SProperty, BoneProperty)
+				.ShouldDisplayName(false)
+				.CustomWidget()
+				[
+					SNew(SBorder)
+					.BorderImage(FAppStyle::GetBrush("NoBorder"))
+					.Padding(FMargin(0.0f, 0.0f, 10.0f, 0.0f))
+					.HAlign(EHorizontalAlignment::HAlign_Fill)
+					[
+						SNew(SHorizontalBox)
+						+SHorizontalBox::Slot().HAlign(HAlign_Left).VAlign(VAlign_Center)
+						[
+							SNew(STextBlock)
+							.Text( LOCTEXT("FCustomizableObjectNodeMeshClipMorphDetails", "Invert plane normal"))
+						]
+						+SHorizontalBox::Slot().HAlign(HAlign_Left)
+						[
+							SNew(SCheckBox)
+							.OnCheckStateChanged(this, &FCustomizableObjectNodeMeshClipMorphDetails::OnInvertNormalCheckboxChanged)
+							.IsChecked(this, &FCustomizableObjectNodeMeshClipMorphDetails::GetInvertNormalCheckBoxState)
+							.ToolTipText(LOCTEXT("InvertNormal", "Invert normal direction of the clip plane"))
+						]
+					]
+				]
+			];
+		}
+	}
+	else
+	{
+		BlocksCategory.AddCustomRow(LOCTEXT("Node", "Node"))
+		[
+			SNew(STextBlock)
+			.Text(LOCTEXT("Node not found", "Node not found"))
+		];
+	}
+}
+
+
+FVector FindBoneLocation(int32 BoneIndex, USkeletalMesh* SkeletalMesh)
+{
+	const TArray<FTransform>& BoneArray = SkeletalMesh->GetRefSkeleton().GetRefBonePose();
+
+	int32 ParentIndex = BoneIndex;
+	FVector Location = FVector::ZeroVector;
+
+	while (ParentIndex >= 0)
+	{
+		Location = BoneArray[ParentIndex].TransformPosition(Location);
+		ParentIndex = SkeletalMesh->GetRefSkeleton().GetParentIndex(ParentIndex);
+	}
+
+	return Location;
+}
+
+
+void FCustomizableObjectNodeMeshClipMorphDetails::OnBoneComboBoxSelectionChanged(TSharedPtr<FString> Selection, ESelectInfo::Type SelectInfo, TSharedRef<IPropertyHandle> BoneProperty)
+{
+	for (int OptionIndex = 0; OptionIndex < BoneComboOptions.Num(); ++OptionIndex)
+	{
+		if (BoneComboOptions[OptionIndex] == Selection)
+		{
+			//TArray<FVector> AffectedBoneLocations;
+
+			//if (SkeletalMesh)
+			//{
+			//	const TArray<FTransform>& BoneArray = SkeletalMesh->RefSkeleton.GetRefBonePose();
+			//	int32 BoneIndex = SkeletalMesh->RefSkeleton.FindBoneIndex(FName(**Selection));
+			//	
+			//	if (BoneIndex >= 0)
+			//	{
+			//		// Get tranform from skeleton root to selected bone
+			//		int32 ParentIndex = BoneIndex;
+			//		FTransform RootToBoneTransform = FTransform::Identity;
+
+			//		while (ParentIndex >= 0)
+			//		{
+			//			RootToBoneTransform = RootToBoneTransform * BoneArray[ParentIndex];
+			//			ParentIndex = SkeletalMesh->RefSkeleton.GetParentIndex(ParentIndex);
+			//		}
+
+			//		// Now that we have the transform from root to the selected bone, go down from this bone to the end of the bone chain hanging from it building a bounding box
+			//		int32 CurrentBone = BoneIndex;
+
+			//		FVector SelectedBoneLocationWithTransform = RootToBoneTransform.TransformPosition(FVector::ZeroVector);
+			//		FVector SelectedBoneLocationFromSelection = BoneLocations[OptionIndex];
+			//		check(FVector::DistSquared(SelectedBoneLocationWithTransform, SelectedBoneLocationFromSelection) < 0.01f);
+
+			//		while (CurrentBone >= 0)
+			//		{
+			//			FVector CurrentBoneLocation = RootToBoneTransform.TransformPosition(FVector::ZeroVector);
+			//			AffectedBoneLocations.Add(CurrentBoneLocation);
+
+			//			check(FVector::DistSquared(CurrentBoneLocation, FindBoneLocation(CurrentBone, SkeletalMesh)) < 0.01f);
+
+			//			int32 NextBone = -1;
+
+			//			for (int32 i = 0; i < SkeletalMesh->RefSkeleton.GetNum(); ++i)
+			//			{
+			//				if (SkeletalMesh->RefSkeleton.GetParentIndex(i) == CurrentBone)
+			//				{
+			//					NextBone = i;
+			//					break;
+			//				}
+			//			}
+
+			//			CurrentBone = NextBone;
+
+			//			if (CurrentBone >= 0)
+			//			{
+			//				RootToBoneTransform = BoneArray[CurrentBone] * RootToBoneTransform;
+			//			}
+			//		}
+			//	}
+			//}
+
+			FVector Location = FVector::ZeroVector;
+			FVector Direction = FVector::ForwardVector;
+
+			if (SkeletalMesh)
+			{
+				const TArray<FTransform>& BoneArray = SkeletalMesh->GetRefSkeleton().GetRefBonePose();
+				int32 ParentIndex = SkeletalMesh->GetRefSkeleton().FindBoneIndex(FName(**Selection));
+
+				FVector ChildLocation = FVector::ForwardVector;
+
+				for (int32 i = 0; i < SkeletalMesh->GetRefSkeleton().GetNum(); ++i)
+				{
+					if (SkeletalMesh->GetRefSkeleton().GetParentIndex(i) == ParentIndex)
+					{
+						ChildLocation = BoneArray[i].TransformPosition(FVector::ZeroVector);
+						break;
+					}
+				}
+
+				//Direction = ChildLocation;
+
+				while (ParentIndex >= 0)
+				{
+					Location = BoneArray[ParentIndex].TransformPosition(Location);
+					ChildLocation = BoneArray[ParentIndex].TransformPosition(ChildLocation);
+					//Direction = BoneArray[ParentIndex].TransformVector(Direction);
+					ParentIndex = SkeletalMesh->GetRefSkeleton().GetParentIndex(ParentIndex);
+				}
+
+				Direction = (ChildLocation - Location).GetSafeNormal();
+			}
+			 
+			BoneProperty->SetValue(*BoneComboOptions[OptionIndex].Get());
+			
+			if (Node)
+			{
+				Node->Origin = Location;
+				Node->Normal = Direction;
+			}
+
+			return;
+		}
+	}
+
+	BoneProperty->SetValue(FString());
+}
+
+
+void FCustomizableObjectNodeMeshClipMorphDetails::OnInvertNormalCheckboxChanged(ECheckBoxState CheckBoxState)
+{
+	if (Node == nullptr)
+	{
+		return;
+	}
+
+	switch (CheckBoxState)
+	{
+		case ECheckBoxState::Checked:
+		{
+			Node->bInvertNormal = true;
+			Node->Normal *= -1.0f;
+			break;
+		}
+		case ECheckBoxState::Unchecked:
+		{
+			Node->bInvertNormal = false;
+			Node->Normal *= -1.0f;
+			break;
+		}
+	}
+
+	if (Node->bLocalStartOffset)
+	{
+		Node->StartOffset.Z *= -1;
+		Node->StartOffset.X *= -1;
+	}
+}
+
+
+ECheckBoxState FCustomizableObjectNodeMeshClipMorphDetails::GetInvertNormalCheckBoxState() const
+{
+	if (Node == nullptr)
+	{
+		return ECheckBoxState::Unchecked;
+	}
+
+	return Node->bInvertNormal ? ECheckBoxState::Checked : ECheckBoxState::Unchecked;
+}
+
+
+void FCustomizableObjectNodeMeshClipMorphDetails::OnReferenceSkeletonIndexChanged()
+{
+	if (Node)
+	{
+		Node->BoneName = FName();
+	}
+	DetailBuilderPtr->ForceRefreshDetails();
+}
+
+
+#undef LOCTEXT_NAMESPACE
diff --git a/Engine/Plugins/Experimental/Mutable/Source/CustomizableObjectEditor/Private/Nodes/CustomizableObjectNodeMeshClipMorphDetails.h b/Engine/Plugins/Experimental/Mutable/Source/CustomizableObjectEditor/Private/Nodes/CustomizableObjectNodeMeshClipMorphDetails.h
new file mode 100644
--- /dev/null
+++ b/Engine/Plugins/Experimental/Mutable/Source/CustomizableObjectEditor/Private/Nodes/CustomizableObjectNodeMeshClipMorphDetails.h
@@ -0,0 +1,36 @@
+// Copyright Epic Games, Inc. All Rights Reserved.
+
+#pragma once
+
+#include "IDetailCustomization.h"
+#include "Types/SlateEnums.h"
+#include "Templates/SharedPointer.h"
+#include "UObject/WeakObjectPtr.h"
+#include "Misc/Guid.h"
+#include "PropertyHandle.h"
+
+
+class FCustomizableObjectNodeMeshClipMorphDetails : public IDetailCustomization
+{
+public:
+	/** Makes a new instance of this detail layout class for a specific detail view requesting it */
+	static TSharedRef<IDetailCustomization> MakeInstance();
+
+	/** ILayoutDetails interface */
+	virtual void CustomizeDetails( IDetailLayoutBuilder& DetailBuilder ) override;
+
+private:
+	class UCustomizableObjectNodeMeshClipMorph* Node;
+	class IDetailLayoutBuilder* DetailBuilderPtr = nullptr;
+	TArray< TSharedPtr<FString> > BoneComboOptions;
+
+	UPROPERTY()
+	class USkeletalMesh* SkeletalMesh;
+
+	void OnBoneComboBoxSelectionChanged(TSharedPtr<FString> Selection, ESelectInfo::Type SelectInfo, TSharedRef<IPropertyHandle> BoneProperty);
+
+	void OnInvertNormalCheckboxChanged(ECheckBoxState CheckBoxState);
+	ECheckBoxState GetInvertNormalCheckBoxState() const;
+
+	void OnReferenceSkeletonIndexChanged();
+};
diff --git a/Engine/Plugins/Experimental/Mutable/Source/CustomizableObjectEditor/Private/Nodes/CustomizableObjectNodeMeshClipWithMesh.cpp b/Engine/Plugins/Experimental/Mutable/Source/CustomizableObjectEditor/Private/Nodes/CustomizableObjectNodeMeshClipWithMesh.cpp
new file mode 100644
--- /dev/null
+++ b/Engine/Plugins/Experimental/Mutable/Source/CustomizableObjectEditor/Private/Nodes/CustomizableObjectNodeMeshClipWithMesh.cpp
@@ -0,0 +1,93 @@
+// Copyright Epic Games, Inc. All Rights Reserved.
+
+#include "Nodes/CustomizableObjectNodeMeshClipWithMesh.h"
+#include "CustomizableObjectEditorModule.h"
+#include "Nodes/CustomizableObjectNodeRemoveMesh.h"
+#include "Toolkits/ToolkitManager.h"
+
+#define LOCTEXT_NAMESPACE "CustomizableObjectEditor"
+
+
+UCustomizableObjectNodeMeshClipWithMesh::UCustomizableObjectNodeMeshClipWithMesh()
+	: Super()
+	, CustomizableObjectToClipWith(nullptr)
+{
+	// We allow this node to be shared
+	Transform = FTransform::Identity;
+
+	//When initialize we don't use materials neither tags
+	bUseMaterials = false;
+	bUseTags = false;
+}
+
+
+void UCustomizableObjectNodeMeshClipWithMesh::BeginPostDuplicate(bool bDuplicateForPIE)
+{
+	Super::BeginPostDuplicate(bDuplicateForPIE);
+
+	if (UCustomizableObjectGraph* Graph = Cast<UCustomizableObjectGraph>(GetGraph()))
+	{
+		TArray<FGuid> NotifiedNodes;
+		for (const FGuid& NodeToClipWithId : ArrayMaterialNodeToClipWithID)
+		{
+			FGuid NodeToClipAlreadyUpdated = Graph->RequestNotificationForNodeIdChange(NodeToClipWithId, NodeGuid);
+
+			if (NodeToClipWithId != NodeToClipAlreadyUpdated)
+			{
+				NotifiedNodes.Add(NodeToClipAlreadyUpdated);
+			}
+		}
+
+		ArrayMaterialNodeToClipWithID = NotifiedNodes;
+	}
+}
+
+
+void UCustomizableObjectNodeMeshClipWithMesh::AllocateDefaultPins(UCustomizableObjectNodeRemapPins* RemapPins)
+{
+	const UEdGraphSchema_CustomizableObject* Schema = GetDefault<UEdGraphSchema_CustomizableObject>();
+
+	UEdGraphPin* ClipMeshPin = CustomCreatePin(EGPD_Input, Schema->PC_Mesh, FName("Clip Mesh"));
+	ClipMeshPin->bDefaultValueIsIgnored = true;
+
+	UEdGraphPin* OutputPin_p = CustomCreatePin(EGPD_Output, Schema->PC_Material, FName("Material"));
+}
+
+
+FText UCustomizableObjectNodeMeshClipWithMesh::GetNodeTitle(ENodeTitleType::Type TitleType) const
+{
+	return LOCTEXT("Clip_Mesh_With_Mesh", "Clip Mesh With Mesh");
+}
+
+
+FLinearColor UCustomizableObjectNodeMeshClipWithMesh::GetNodeTitleColor() const
+{
+	const UEdGraphSchema_CustomizableObject* Schema = GetDefault<UEdGraphSchema_CustomizableObject>();
+	return Schema->GetPinTypeColor(Schema->PC_Material);
+}
+
+
+void UCustomizableObjectNodeMeshClipWithMesh::PinConnectionListChanged(UEdGraphPin* Pin)
+{
+	Super::PinConnectionListChanged(Pin);
+
+	TSharedPtr<ICustomizableObjectEditor> Editor = GetGraphEditor();
+
+	if (Editor.IsValid())
+	{
+		Editor->UpdateGraphNodeProperties();
+	}
+}
+
+void UCustomizableObjectNodeMeshClipWithMesh::UpdateReferencedNodeId(const FGuid& NewGuid)
+{
+	ArrayMaterialNodeToClipWithID.Add(NewGuid);
+}
+
+
+FText UCustomizableObjectNodeMeshClipWithMesh::GetTooltipText() const
+{
+	return LOCTEXT("Clip_Mesh_Mesh_Tooltip", "Removes the part of a material that is completely enclosed in a mesh volume.\nIt only removes the faces that fall completely inside the cutting volume, along with the vertices and edges that define only faces that are deleted.");
+}
+
+#undef LOCTEXT_NAMESPACE
diff --git a/Engine/Plugins/Experimental/Mutable/Source/CustomizableObjectEditor/Private/Nodes/CustomizableObjectNodeMeshClipWithMeshDetails.cpp b/Engine/Plugins/Experimental/Mutable/Source/CustomizableObjectEditor/Private/Nodes/CustomizableObjectNodeMeshClipWithMeshDetails.cpp
new file mode 100644
--- /dev/null
+++ b/Engine/Plugins/Experimental/Mutable/Source/CustomizableObjectEditor/Private/Nodes/CustomizableObjectNodeMeshClipWithMeshDetails.cpp
@@ -0,0 +1,472 @@
+// Copyright Epic Games, Inc. All Rights Reserved.
+
+#include "CustomizableObjectNodeMeshClipWithMeshDetails.h"
+#include "CustomizableObjectEditorModule.h"
+#include "PropertyCustomizationHelpers.h"
+#include "CustomizableObjectEditorUtilities.h"
+#include "PropertyHandle.h"
+#include "CustomizableObject.h"
+
+#include "DetailLayoutBuilder.h"
+#include "DetailCategoryBuilder.h"
+#include "GraphTraversal.h"
+#include "LevelEditor.h"
+#include "LevelEditorActions.h"
+#include "ScopedTransaction.h"
+
+#include "Widgets/Images/SImage.h"
+#include "Widgets/Text/STextBlock.h"
+#include "Widgets/Input/STextComboBox.h"
+#include "Widgets/Input/SEditableTextBox.h"
+#include "Widgets/Input/SButton.h"
+#include "AssetRegistry/AssetData.h"
+
+#include "Runtime/Launch/Resources/Version.h"
+
+
+#define LOCTEXT_NAMESPACE "CustomizableObjectNodeMeshClipWithMeshDetails"
+
+
+TSharedRef<IDetailCustomization> FCustomizableObjectNodeMeshClipWithMeshDetails::MakeInstance()
+{
+	return MakeShareable(new FCustomizableObjectNodeMeshClipWithMeshDetails);
+}
+
+
+void FCustomizableObjectNodeMeshClipWithMeshDetails::CustomizeDetails(IDetailLayoutBuilder& DetailBuilder)
+{
+	Node = nullptr;
+	DetailBuilderPtr = &DetailBuilder;
+
+#if (ENGINE_MAJOR_VERSION == 4 && ENGINE_MINOR_VERSION >= 18) || ENGINE_MAJOR_VERSION >= 5
+	const IDetailsView* DetailsView = DetailBuilder.GetDetailsView();
+	if (DetailsView->GetSelectedObjects().Num())
+	{
+		Node = Cast<UCustomizableObjectNodeMeshClipWithMesh>(DetailsView->GetSelectedObjects()[0].Get());
+	}
+#else
+	const IDetailsView& DetailsView = DetailBuilder.GetDetailsView();
+	if (DetailsView.GetSelectedObjects().Num())
+	{
+		Node = Cast<UCustomizableObjectNodeMeshClipWithMesh>(DetailsView.GetSelectedObjects()[0].Get());
+	}
+#endif
+
+	IDetailCategoryBuilder& CustomizableObjectToClipCategory = DetailBuilder.EditCategory("CustomizableObjectToClip");
+
+
+	if(Node != nullptr)
+	{
+		SelectedCO = Node->CustomizableObjectToClipWith;
+
+		TSharedRef<IPropertyHandle> ParentObjectProperty = DetailBuilder.GetProperty("ParentObject");
+
+		// Hidding unnecessary properties
+		DetailBuilder.HideProperty("CustomizableObjectToClipWith");
+		DetailBuilder.HideProperty("ArrayMaterialNodeToClipWithID");
+		DetailBuilder.HideProperty("Tags");
+
+		if (const UEdGraphPin* ConnectedPin = FollowInputPin(*Node->ClipMeshPin());
+			ConnectedPin && !Cast<UCustomizableObjectNodeStaticMesh>(ConnectedPin->GetOwningNode()))
+		{
+			DetailBuilder.HideProperty("Transform");
+		}
+
+		ArrayMaterialNodeOption.Empty();
+		ArrayMaterialNodeName.Empty();
+
+		TArray<FGuid> ArrayMaterialNodeGuid;
+		TSharedPtr<FString> SelectedMethod;
+
+		CustomizableObjectToClipCategory.AddCustomRow(LOCTEXT("FCustomizableObjectNodeMeshClipWithMeshDetails", "Blocks"))
+		[
+			SNew(SHorizontalBox)
+			+ SHorizontalBox::Slot()
+			.Padding(FMargin(0.0f, 3.0f, 0.0f, 10.0f))
+			.FillWidth(0.15f)
+			[
+				SNew(STextBlock)
+				.Text(LOCTEXT("CustomizableObjectText", "Object:"))
+			]
+			+ SHorizontalBox::Slot()
+			.FillWidth(0.85f)
+			[
+				SNew(SObjectPropertyEntryBox)
+				.AllowedClass(UCustomizableObject::StaticClass())
+				.OnObjectChanged(this, &FCustomizableObjectNodeMeshClipWithMeshDetails::ClipperCustomizableObjectSelectionChanged)
+				.ObjectPath(SelectedCO != nullptr ? SelectedCO->GetPathName() : "")
+				.ForceVolatile(true)
+			]
+		];
+
+		SelectAMethod = MakeShareable(new FString("Select a Method"));
+
+		// Filling ComboBox options
+		ClippingMethods.Empty();
+		ClippingMethods.Add(MakeShareable(new FString("Material")));
+		ClippingMethods.Add(MakeShareable(new FString("Tags")));
+		ClippingMethods.Add(MakeShareable(new FString("Material & Tags")));
+
+		// Initializing the inital value of the ComboBox
+		SelectedMethod = SetInitialClippingMethod();
+		
+		// ComboBox to select the Clipping method (material or/and tags)
+		CustomizableObjectToClipCategory.AddCustomRow(LOCTEXT("FCustomizableObjectNodeMeshClipWithMeshDetails", "Blocks"))
+		[
+			SNew(SHorizontalBox)
+			+SHorizontalBox::Slot()
+			.AutoWidth()
+			.Padding(FMargin(0.0f, 8.0f, 0.0f, 5.0f))
+			[
+				SNew(STextBlock)
+				.Text(LOCTEXT("ClippingMethod", "Clipping Method:"))
+			]
+			
+			+ SHorizontalBox::Slot()
+			.AutoWidth()
+			.Padding(FMargin(4.0f, 5.0f, 0.0f, 5.0f))
+			[
+				SNew(STextComboBox)
+				.OptionsSource(&ClippingMethods)
+				.InitiallySelectedItem(SelectedMethod)
+				.OnSelectionChanged(this, &FCustomizableObjectNodeMeshClipWithMeshDetails::OnClippingMethodComboBoxSelectionChanged)
+			]
+		];
+
+
+		// Filling materials ComboBox with all the possible materials
+		int32 NumElementBeforeAdd;
+		int32 NumElementAfterAdd;
+		TSharedPtr<FString> ItemToSelect;
+		TArray<UCustomizableObjectNodeMaterial*> ArrayMaterialNodes;
+		if (SelectedCO != nullptr)
+		{
+			SelectedCO->Source->GetNodesOfClass<UCustomizableObjectNodeMaterial>(ArrayMaterialNodes);
+		}
+
+		for (const UCustomizableObjectNodeMaterial* MaterialNode : ArrayMaterialNodes)
+		{
+			if (MaterialNode->Material)
+			{
+				NumElementBeforeAdd = ArrayMaterialNodeName.Num();
+				ArrayMaterialNodeName.AddUnique(MaterialNode->Material->GetName());
+				NumElementAfterAdd = ArrayMaterialNodeName.Num();
+
+				if (NumElementBeforeAdd != NumElementAfterAdd)
+				{
+					ArrayMaterialNodeGuid.Add(MaterialNode->NodeGuid);
+				}
+			}
+		}
+
+		// Adding a none option
+		ArrayMaterialNodeOption.Add(MakeShareable(new FString("None")));
+		ItemToSelect = ArrayMaterialNodeOption.Last();
+
+		// Initializing the Selected material if it was already selected
+		const int32 MaxIndex = ArrayMaterialNodeName.Num();
+		for (int32 i = 0; i < MaxIndex; ++i)
+		{
+			ArrayMaterialNodeOption.Add(MakeShareable(new FString(ArrayMaterialNodeName[i])));
+
+			if (Node->ArrayMaterialNodeToClipWithID.Find(ArrayMaterialNodeGuid[i]) != INDEX_NONE)
+			{
+				ItemToSelect = ArrayMaterialNodeOption.Last();
+			}
+		}
+
+        ArrayMaterialNodeOption.Sort(CompareNames);
+		
+		// Creating a TagView widget for each tag of the node
+		CreateTagView();
+
+		// Widget that shows the Tags widget or/and the materials ComboBox
+		CustomizableObjectToClipCategory.AddCustomRow(LOCTEXT("FCustomizableObjectNodeMeshClipWithMeshDetails", "Blocks"))
+		[
+			SNew(SVerticalBox)
+			+SVerticalBox::Slot()
+			.AutoHeight()
+			[
+				SAssignNew(MaterialsSelector, SHorizontalBox)
+				+ SHorizontalBox::Slot()
+				.FillWidth(0.15f)
+				.Padding(0.0f, 3.0f, 0.0f, 0.0f)
+				[
+					SNew(STextBlock)
+					.Text(LOCTEXT("MaterialsText", "Material:"))
+				]
+				+ SHorizontalBox::Slot()
+				.FillWidth(0.85f)
+				.Padding(2.0f, 0.0f, 0.0f, 0.0f)
+				[
+					SNew(STextComboBox)
+					.OptionsSource(&ArrayMaterialNodeOption)
+					.InitiallySelectedItem(ItemToSelect)
+					.OnSelectionChanged(this, &FCustomizableObjectNodeMeshClipWithMeshDetails::OnMeshClipWithMeshNodeComboBoxSelectionChanged)
+				]
+			]
+
+			+ SVerticalBox::Slot()
+			.AutoHeight()
+			.Padding(1.0f,3.0f,0.0f,0.0f)
+			[
+				TagView.ToSharedRef()
+			]
+		];
+
+		// Hidding the widgets if necessary
+		MaterialsSelector->SetVisibility(Node->bUseMaterials ? EVisibility::Visible : EVisibility::Collapsed);
+		TagView->SetVisibility(Node->bUseTags ? EVisibility::Visible : EVisibility::Collapsed);
+
+	}
+	else
+	{
+		CustomizableObjectToClipCategory.AddCustomRow(LOCTEXT("Node", "Node"))
+			[
+				SNew(STextBlock)
+				.Text(LOCTEXT("Node not found", "Node not found"))
+			];
+	}
+}
+
+
+void FCustomizableObjectNodeMeshClipWithMeshDetails::ClipperCustomizableObjectSelectionChanged(const FAssetData& AssetData)
+{
+	UCustomizableObject* NewClipperCO = Cast<UCustomizableObject>(AssetData.GetAsset());
+
+	if (Node != nullptr)
+	{
+		Node->CustomizableObjectToClipWith = NewClipperCO;
+		Node->ArrayMaterialNodeToClipWithID.Empty();
+	}
+
+	if (DetailBuilderPtr)
+	{
+		DetailBuilderPtr->ForceRefreshDetails();
+	}
+}
+
+
+TSharedPtr<FString> FCustomizableObjectNodeMeshClipWithMeshDetails::SetInitialClippingMethod()
+{
+	TSharedPtr<FString> InitialMethod = SelectAMethod;
+
+	if (!Node->bUseMaterials && !Node->bUseTags)
+	{
+		if (Node->ArrayMaterialNodeToClipWithID.Num() > 0)
+		{
+			InitialMethod = ClippingMethods[0];
+			Node->bUseMaterials = true;
+		}
+
+		if (Node->Tags.Num() > 0)
+		{
+			InitialMethod = ClippingMethods[1];
+			Node->bUseTags = true;
+		}
+
+		if (Node->bUseMaterials && Node->bUseTags)
+		{
+			InitialMethod = ClippingMethods[2];
+		}
+	}
+	else
+	{
+		if (Node->bUseMaterials && !Node->bUseTags)
+		{
+			InitialMethod = ClippingMethods[0];
+		}
+		else if (!Node->bUseMaterials && Node->bUseTags)
+		{
+			InitialMethod = ClippingMethods[1];
+		}
+		else
+		{
+			InitialMethod = ClippingMethods[2];
+		}
+	}
+
+	return InitialMethod;
+}
+
+
+void FCustomizableObjectNodeMeshClipWithMeshDetails::OnClippingMethodComboBoxSelectionChanged(TSharedPtr<FString> Selection, ESelectInfo::Type SelectInfo)
+{
+	if (Selection.IsValid())
+	{
+		if (*Selection == "Tags")
+		{
+			Node->bUseMaterials = false;
+			Node->bUseTags = true;
+			Node->ArrayMaterialNodeToClipWithID.Empty();
+
+			DetailBuilderPtr->ForceRefreshDetails();
+		}
+		else if(*Selection == "Material")
+		{
+			Node->bUseMaterials = true;
+			Node->bUseTags = false;
+			Node->Tags.Empty();
+
+			DetailBuilderPtr->ForceRefreshDetails();
+		}
+		else
+		{
+			Node->bUseMaterials = true;
+			Node->bUseTags = true;
+		}
+
+		TagView->SetVisibility(Node->bUseTags ? EVisibility::Visible : EVisibility::Collapsed);
+		MaterialsSelector->SetVisibility(Node->bUseMaterials ? EVisibility::Visible : EVisibility::Collapsed);
+	}
+}
+
+
+void FCustomizableObjectNodeMeshClipWithMeshDetails::OnMeshClipWithMeshNodeComboBoxSelectionChanged(TSharedPtr<FString> Selection, ESelectInfo::Type SelectInfo)
+{
+	if (Selection.IsValid() && (SelectedCO != nullptr))
+	{
+		TArray<UCustomizableObjectNodeMaterial*> ArrayMaterialNode;
+		SelectedCO->Source->GetNodesOfClass<UCustomizableObjectNodeMaterial>(ArrayMaterialNode);
+
+		const FScopedTransaction Transaction(LOCTEXT("ChangedClipMeshWithMeshMaterialTransaction", "Changed Clip Mesh With Mesh Material"));
+		Node->Modify();
+		Node->ArrayMaterialNodeToClipWithID.Empty();
+		for (const UCustomizableObjectNodeMaterial* MaterialNode : ArrayMaterialNode)
+		{
+			if (MaterialNode && MaterialNode->Material && *Selection == MaterialNode->Material->GetName())
+			{
+				Node->ArrayMaterialNodeToClipWithID.Add(MaterialNode->NodeGuid);
+			}
+		}
+	}
+}
+
+
+FReply FCustomizableObjectNodeMeshClipWithMeshDetails::OnAddTagPressed()
+{
+	if (Node)
+	{
+		Node->Tags.Add("");
+		DetailBuilderPtr->ForceRefreshDetails();
+
+		return FReply::Handled();
+	}
+
+	return FReply::Unhandled();
+}
+
+
+void FCustomizableObjectNodeMeshClipWithMeshDetails::CreateTagView()
+{
+	TagView = SNew(SVerticalBox);
+	TagView->AddSlot().AutoHeight()
+	[
+		SNew(SHorizontalBox)
+		+ SHorizontalBox::Slot()
+		.AutoWidth()
+		.Padding(0.0f, 3.0f, 0.0f, 0.0f)
+		[
+			SNew(STextBlock)
+			.Text(LOCTEXT("ClippingMethodTags", "Tags"))
+		]
+		+ SHorizontalBox::Slot()
+		.AutoWidth()
+		.Padding(5.0f, 0.0f, 0.0f, 0.0f)
+		[
+			SNew(SButton)
+			.OnClicked(this, &FCustomizableObjectNodeMeshClipWithMeshDetails::OnAddTagPressed)
+			.ToolTipText(LOCTEXT("AddNewTag", "Add New Tag"))
+			[
+				SNew(SImage)
+				.Image(FAppStyle::GetBrush(TEXT("Plus")))	
+			]
+		]
+	];
+
+	for (int32 i = 0; i < Node->Tags.Num(); ++i)
+	{
+		TagView->AddSlot()
+		.Padding(10.0f,5.0f,0.0f,0.0f)
+		[
+			SNew(STagView)
+			.TagValue(Node->Tags[i])
+			.TagIndex(i)
+			.Node(Node)
+			.DetailBuilderPtr(DetailBuilderPtr)
+		];
+	}
+}
+
+
+// Tag View Widget ---------------------------------------------------------
+
+void STagView::Construct(const FArguments& InArgs)
+{
+	TagValue = InArgs._TagValue;
+	TagIndex = InArgs._TagIndex;
+	Node = InArgs._Node;
+	DetailBuilderPtr = InArgs._DetailBuilderPtr;
+
+	FString TagName = FString::Printf(TEXT("Tag %d"), TagIndex);
+
+	ChildSlot
+	[
+		SNew(SHorizontalBox)
+		+ SHorizontalBox::Slot()
+		.AutoWidth()
+		[
+			SNew(STextBlock)
+			.Text(FText::FromString(TagName))
+		]
+
+		+ SHorizontalBox::Slot()
+		.AutoWidth()
+		.Padding(5.0f,0.0f,0.0f,0.0f)
+		[
+			SNew(SEditableTextBox)
+			.Text(FText::FromString(TagValue))
+			.MinDesiredWidth(50.0f)
+			.HintText(FText::FromString("New Tag"))
+			.OnTextCommitted(this, &STagView::OnTextBoxTextCommitted)
+		]
+		
+		+ SHorizontalBox::Slot()
+		.AutoWidth()
+		.Padding(5.0f, 0.0f, 0.0f, 0.0f)
+		[
+			SNew(SButton)
+			.OnClicked(this, &STagView::DeleteTag)
+			.ToolTipText(LOCTEXT("DeleteTag", "Delete Tag"))
+			[
+				SNew(SImage)
+				.Image(FAppStyle::GetBrush(TEXT("Cross")))
+			]
+		]
+	];
+
+}
+
+
+FReply STagView::DeleteTag()
+{
+	// Remove a tag and redraw the details view
+	Node->Tags.RemoveAt(TagIndex);
+	DetailBuilderPtr->ForceRefreshDetails();
+
+	return FReply::Handled();
+}
+
+
+void STagView::OnTextBoxTextCommitted(const FText& NewText, ETextCommit::Type InTextCommit)
+{
+	// Save the content when the widget loses the focus or we press enter
+	if (InTextCommit == ETextCommit::OnEnter || InTextCommit == ETextCommit::OnUserMovedFocus)
+	{
+		TagValue = NewText.ToString();
+		Node->Tags[TagIndex] = TagValue;
+	}
+}
+
+
+#undef LOCTEXT_NAMESPACE
diff --git a/Engine/Plugins/Experimental/Mutable/Source/CustomizableObjectEditor/Private/Nodes/CustomizableObjectNodeMeshClipWithMeshDetails.h b/Engine/Plugins/Experimental/Mutable/Source/CustomizableObjectEditor/Private/Nodes/CustomizableObjectNodeMeshClipWithMeshDetails.h
new file mode 100644
--- /dev/null
+++ b/Engine/Plugins/Experimental/Mutable/Source/CustomizableObjectEditor/Private/Nodes/CustomizableObjectNodeMeshClipWithMeshDetails.h
@@ -0,0 +1,104 @@
+// Copyright Epic Games, Inc. All Rights Reserved.
+
+#pragma once
+
+#include "IDetailCustomization.h"
+#include "Templates/SharedPointer.h"
+#include "UObject/WeakObjectPtr.h"
+#include "Types/SlateEnums.h"
+#include "Widgets/SCompoundWidget.h"
+#include "Widgets/SBoxPanel.h"
+#include "AssetRegistry/AssetData.h"
+#include "Nodes/CustomizableObjectNodeMeshClipWithMesh.h"
+
+// Widget that represents the array of tags of the clip with mesh node
+class STagView : public SCompoundWidget
+{
+public:
+
+	SLATE_BEGIN_ARGS(class STagView) {}
+	SLATE_ARGUMENT(FString, TagValue)
+	SLATE_ARGUMENT(int32, TagIndex)
+	SLATE_ARGUMENT(UCustomizableObjectNodeMeshClipWithMesh*, Node)
+	SLATE_ARGUMENT(IDetailLayoutBuilder*, DetailBuilderPtr)
+	SLATE_END_ARGS()
+
+	void Construct(const FArguments& InArgs);
+
+	// Delete a tag of the array
+	FReply DeleteTag();
+
+	// Callback for the text box. Modifies the content of a tag.
+	void OnTextBoxTextCommitted(const FText& NewText, ETextCommit::Type InTextCommit);
+
+public:
+
+	// Pointer to the builder passed by CustomizeDetails method
+	IDetailLayoutBuilder* DetailBuilderPtr = nullptr;
+
+	// Pointer to the clip with mesh node
+	UCustomizableObjectNodeMeshClipWithMesh* Node;
+
+	// Value of the tag
+	FString TagValue;
+
+	// Index of the tag inside the array of tags
+	int32 TagIndex;
+};
+
+
+class FCustomizableObjectNodeMeshClipWithMeshDetails : public IDetailCustomization
+{
+public:
+	/** Makes a new instance of this detail layout class for a specific detail view requesting it */
+	static TSharedRef<IDetailCustomization> MakeInstance();
+
+	/** ILayoutDetails interface */
+	virtual void CustomizeDetails( IDetailLayoutBuilder& DetailBuilder ) override;
+
+	/** Callback when new CO is selected */
+	void ClipperCustomizableObjectSelectionChanged(const FAssetData& AssetData);
+
+private:
+	/** Callback when the clipping method is selected */
+	void OnClippingMethodComboBoxSelectionChanged(TSharedPtr<FString> Selection, ESelectInfo::Type SelectInfo);
+
+	/** Callback when new material name is selected */
+	void OnMeshClipWithMeshNodeComboBoxSelectionChanged(TSharedPtr<FString> Selection, ESelectInfo::Type SelectInfo);
+
+	/** Creates the array of widgets tags */
+	void CreateTagView();
+
+	/** Initializes the text of the method combobox and the view to display */
+	TSharedPtr<FString> SetInitialClippingMethod();
+
+	/** The node for which details are being customized */
+	class UCustomizableObjectNodeMeshClipWithMesh* Node;
+
+	/** Pointer to the builder passed by CustomizeDetails method */
+	IDetailLayoutBuilder* DetailBuilderPtr = nullptr;
+
+	/** References to the names of all possible methods to select the material */
+	TArray<TSharedPtr<FString>> ClippingMethods;
+
+	/** References to the names of all possible selectable material nodes in SelectedCO */
+	TArray<TSharedPtr<FString>> ArrayMaterialNodeOption;
+
+	/** To cover all LODs from a single selection, only unique names of materials are considered */
+	TArray<FString> ArrayMaterialNodeName;
+
+	/** Selected CO in the details tab of the Node node, for which to select a material that will be clipped by Node node */
+	class UCustomizableObject* SelectedCO;
+
+	/** Selector with all the materials of the Customizable Object */
+	TSharedPtr< SHorizontalBox > MaterialsSelector;
+
+	/** ComboBox with all the materials of the Customizable Object */
+	TSharedPtr< SVerticalBox > TagView;
+
+	/** Pointer to a FString to indicate that a method must be selected */
+	TSharedPtr<FString> SelectAMethod;
+
+	/** Callback to add a tag to the arry of tags */
+	FReply OnAddTagPressed();
+};
diff --git a/Engine/Plugins/Experimental/Mutable/Source/CustomizableObjectEditor/Private/Nodes/CustomizableObjectNodeMeshGeometryOperation.cpp b/Engine/Plugins/Experimental/Mutable/Source/CustomizableObjectEditor/Private/Nodes/CustomizableObjectNodeMeshGeometryOperation.cpp
new file mode 100644
--- /dev/null
+++ b/Engine/Plugins/Experimental/Mutable/Source/CustomizableObjectEditor/Private/Nodes/CustomizableObjectNodeMeshGeometryOperation.cpp
@@ -0,0 +1,56 @@
+// Copyright Epic Games, Inc. All Rights Reserved.
+
+#include "Nodes/CustomizableObjectNodeMeshGeometryOperation.h"
+#include "CustomizableObjectEditorModule.h"
+
+#define LOCTEXT_NAMESPACE "CustomizableObjectEditor"
+
+
+UCustomizableObjectNodeMeshGeometryOperation::UCustomizableObjectNodeMeshGeometryOperation()
+	: Super()
+{
+}
+
+
+void UCustomizableObjectNodeMeshGeometryOperation::AllocateDefaultPins(UCustomizableObjectNodeRemapPins* RemapPins)
+{
+	const UEdGraphSchema_CustomizableObject* Schema = GetDefault<UEdGraphSchema_CustomizableObject>();
+
+	UEdGraphPin* OutputPin = CustomCreatePin(EGPD_Output, Schema->PC_Mesh, FName("Mesh"));
+	OutputPin->bDefaultValueIsIgnored = true;
+
+	UEdGraphPin* MeshPin = CustomCreatePin(EGPD_Input, Schema->PC_Mesh, FName("Mesh A"));
+	MeshPin->bDefaultValueIsIgnored = true;
+
+	MeshPin = CustomCreatePin(EGPD_Input, Schema->PC_Mesh, FName("Mesh B"));
+	MeshPin->bDefaultValueIsIgnored = true;
+
+	UEdGraphPin* FactorPin = CustomCreatePin(EGPD_Input, Schema->PC_Float, FName("Scalar A"));
+	FactorPin->bDefaultValueIsIgnored = true;
+
+	FactorPin = CustomCreatePin(EGPD_Input, Schema->PC_Float, FName("Scalar B"));
+	FactorPin->bDefaultValueIsIgnored = true;
+}
+
+
+FText UCustomizableObjectNodeMeshGeometryOperation::GetNodeTitle(ENodeTitleType::Type TitleType) const
+{
+	return LOCTEXT("Mesh_GeometryOperation", "Geometry Operation");
+}
+
+
+FLinearColor UCustomizableObjectNodeMeshGeometryOperation::GetNodeTitleColor() const
+{
+	const UEdGraphSchema_CustomizableObject* Schema = GetDefault<UEdGraphSchema_CustomizableObject>();
+	return Schema->GetPinTypeColor(Schema->PC_Mesh);
+}
+
+
+FText UCustomizableObjectNodeMeshGeometryOperation::GetTooltipText() const
+{
+	return LOCTEXT("Mesh_GeometryOperation_Tooltip", "Apply a Geometry operation on a mesh.");
+}
+
+
+
+#undef LOCTEXT_NAMESPACE
diff --git a/Engine/Plugins/Experimental/Mutable/Source/CustomizableObjectEditor/Private/Nodes/CustomizableObjectNodeMeshMorph.cpp b/Engine/Plugins/Experimental/Mutable/Source/CustomizableObjectEditor/Private/Nodes/CustomizableObjectNodeMeshMorph.cpp
new file mode 100644
--- /dev/null
+++ b/Engine/Plugins/Experimental/Mutable/Source/CustomizableObjectEditor/Private/Nodes/CustomizableObjectNodeMeshMorph.cpp
@@ -0,0 +1,131 @@
+// Copyright Epic Games, Inc. All Rights Reserved.
+
+#include "Nodes/CustomizableObjectNodeMeshMorph.h"
+#include "CustomizableObjectEditorModule.h"
+#include "GraphTraversal.h"
+
+#define LOCTEXT_NAMESPACE "CustomizableObjectEditor"
+
+
+UCustomizableObjectNodeMeshMorph::UCustomizableObjectNodeMeshMorph()
+	: Super()
+{
+
+}
+
+
+void UCustomizableObjectNodeMeshMorph::AllocateDefaultPins(UCustomizableObjectNodeRemapPins* RemapPins)
+{
+	const UEdGraphSchema_CustomizableObject* Schema = GetDefault<UEdGraphSchema_CustomizableObject>();
+
+	FString PinName = TEXT("Mesh");
+	UEdGraphPin* OutputPin = CustomCreatePin(EGPD_Output, Schema->PC_Mesh, FName(*PinName));
+	OutputPin->bDefaultValueIsIgnored = true;
+
+	PinName = TEXT("Mesh");
+	UEdGraphPin* MeshPin = CustomCreatePin(EGPD_Input, Schema->PC_Mesh, FName(*PinName));
+	MeshPin->bDefaultValueIsIgnored = true;
+
+	PinName = TEXT("Factor");
+	UEdGraphPin* FactorPin = CustomCreatePin(EGPD_Input, Schema->PC_Float, FName(*PinName));
+	FactorPin->bDefaultValueIsIgnored = true;
+}
+
+
+FText UCustomizableObjectNodeMeshMorph::GetNodeTitle(ENodeTitleType::Type TitleType) const
+{
+	return LOCTEXT("Mesh_Morph", "Mesh Morph");
+}
+
+
+FLinearColor UCustomizableObjectNodeMeshMorph::GetNodeTitleColor() const
+{
+	const UEdGraphSchema_CustomizableObject* Schema = GetDefault<UEdGraphSchema_CustomizableObject>();
+	return Schema->GetPinTypeColor(Schema->PC_Mesh);
+}
+
+
+UCustomizableObjectNodeSkeletalMesh* UCustomizableObjectNodeMeshMorph::GetSourceSkeletalMesh() const
+{
+	if (const UEdGraphPin* Pin = MeshPin())
+	{
+		if (const UEdGraphPin* ConnectedPin = FollowInputPin(*Pin))
+		{
+			UEdGraphNode* InMeshNode = ConnectedPin->GetOwningNode();
+			if (UCustomizableObjectNodeSkeletalMesh* SkeletalMeshNode = Cast<UCustomizableObjectNodeSkeletalMesh>(InMeshNode))
+			{
+				return SkeletalMeshNode;
+			}
+			else if (UCustomizableObjectNodeMeshMorph* MorphNode = Cast<UCustomizableObjectNodeMeshMorph>(InMeshNode))
+			{
+				return MorphNode->GetSourceSkeletalMesh();
+			}
+		}
+	}
+
+	return nullptr;
+}
+
+
+bool UCustomizableObjectNodeMeshMorph::IsNodeOutDatedAndNeedsRefresh()
+{
+	bool Result = false;
+
+	const UCustomizableObjectNodeSkeletalMesh* SkeletalMeshNode = GetSourceSkeletalMesh();
+
+	if (SkeletalMeshNode && SkeletalMeshNode->SkeletalMesh)
+	{
+		bool bOutdated = true;
+		for (int m = 0; m < SkeletalMeshNode->SkeletalMesh->GetMorphTargets().Num(); ++m)
+		{
+			FString MorphName = SkeletalMeshNode->SkeletalMesh->GetMorphTargets()[m]->GetName();
+			if (MorphTargetName == MorphName)
+			{
+				bOutdated = false;
+				break;
+			}
+		}
+		Result = bOutdated;
+	}
+
+	// Remove previous compilation warnings
+	if (!Result && bHasCompilerMessage)
+	{
+		RemoveWarnings();
+		GetGraph()->NotifyGraphChanged();
+	}
+
+	return Result;
+}
+
+
+FString UCustomizableObjectNodeMeshMorph::GetRefreshMessage() const
+{
+    return "Morph Target not found in the SkeletalMesh. Please Refresh Node and select a valid morph option.";
+}
+
+
+void UCustomizableObjectNodeMeshMorph::PinConnectionListChanged(UEdGraphPin* Pin)
+{
+	Super::PinConnectionListChanged(Pin);
+
+	if (Pin->Direction == EGPD_Input && Pin->PinName == "Mesh")
+	{
+		TSharedPtr<ICustomizableObjectEditor> Editor = GetGraphEditor();
+
+		if (Editor.IsValid())
+		{
+			Editor->UpdateGraphNodeProperties();
+		}
+	}
+}
+
+
+FText UCustomizableObjectNodeMeshMorph::GetTooltipText() const
+{
+	return LOCTEXT("Mesh_Morph_Tooltip", "Changes the weight of a mesh morph target.");
+}
+
+
+
+#undef LOCTEXT_NAMESPACE
diff --git a/Engine/Plugins/Experimental/Mutable/Source/CustomizableObjectEditor/Private/Nodes/CustomizableObjectNodeMeshMorphDetails.cpp b/Engine/Plugins/Experimental/Mutable/Source/CustomizableObjectEditor/Private/Nodes/CustomizableObjectNodeMeshMorphDetails.cpp
new file mode 100644
--- /dev/null
+++ b/Engine/Plugins/Experimental/Mutable/Source/CustomizableObjectEditor/Private/Nodes/CustomizableObjectNodeMeshMorphDetails.cpp
@@ -0,0 +1,129 @@
+// Copyright Epic Games, Inc. All Rights Reserved.
+
+#include "CustomizableObjectNodeMeshMorphDetails.h"
+#include "CustomizableObjectEditorModule.h"
+#include "../CustomizableObjectCompiler.h"
+#include "CustomizableObjectEditorUtilities.h"
+
+#include "PropertyCustomizationHelpers.h"
+#include "DetailLayoutBuilder.h"
+#include "DetailCategoryBuilder.h"
+#include "GraphTraversal.h"
+#include "LevelEditor.h"
+#include "LevelEditorActions.h"
+
+#include "Widgets/Layout/SBorder.h"
+#include "Widgets/Text/STextBlock.h"
+#include "Widgets/Input/STextComboBox.h"
+
+#include "Engine/SkeletalMesh.h"
+#include "Engine/StaticMesh.h"
+#include "Animation/MorphTarget.h"
+
+#include "BoneSelectionWidget.h"
+
+#include "Runtime/Launch/Resources/Version.h"
+
+
+#define LOCTEXT_NAMESPACE "CustomizableObjectDetails"
+
+TSharedRef<IDetailCustomization> FCustomizableObjectNodeMeshMorphDetails::MakeInstance()
+{
+	return MakeShareable( new FCustomizableObjectNodeMeshMorphDetails );
+}
+
+void FCustomizableObjectNodeMeshMorphDetails::CustomizeDetails( IDetailLayoutBuilder& DetailBuilder )
+{
+	const IDetailsView* DetailsView = DetailBuilder.GetDetailsView();
+	if ( DetailsView->GetSelectedObjects().Num() )
+	{
+		Node = Cast<UCustomizableObjectNodeMeshMorph>( DetailsView->GetSelectedObjects()[0].Get() );
+	}
+
+	IDetailCategoryBuilder& BlocksCategory = DetailBuilder.EditCategory( "Morph Target" );
+	//BlocksCategory.CategoryIcon( "ActorClassIcon.CustomizableObject" );
+	
+	MorphTargetComboOptions.Empty();
+
+	bool SourceMeshFound = false;
+	if (Node)
+	{
+		if (const UEdGraphPin* ConnectedPin = FollowInputPin(*Node->MeshPin()))
+		{
+			SourceMeshFound = true;
+			
+			TSharedPtr<FString> ItemToSelect;
+
+			const UEdGraphPin* BaseSourcePin = FindMeshBaseSource(*ConnectedPin, false );
+			if ( BaseSourcePin )
+			{
+				USkeletalMesh* SkeletalMesh = nullptr;
+				if ( const UCustomizableObjectNodeSkeletalMesh* TypedNodeSkeletalMesh = Cast<UCustomizableObjectNodeSkeletalMesh>(BaseSourcePin->GetOwningNode()) )
+				{
+					SkeletalMesh = TypedNodeSkeletalMesh->SkeletalMesh;
+				}
+				else if (const UCustomizableObjectNodeTable* TypedNodeTable = Cast<UCustomizableObjectNodeTable>(BaseSourcePin->GetOwningNode()))
+				{
+					SkeletalMesh = TypedNodeTable->GetColumnDefaultAssetByType<USkeletalMesh>(BaseSourcePin);
+				}
+
+				if (SkeletalMesh)
+				{
+					for (int m = 0; m < SkeletalMesh->GetMorphTargets().Num(); ++m)
+					{
+						FString MorphName = *SkeletalMesh->GetMorphTargets()[m]->GetName();
+						MorphTargetComboOptions.Add(MakeShareable(new FString(MorphName)));
+
+						if (Node->MorphTargetName == MorphName)
+						{
+							ItemToSelect = MorphTargetComboOptions.Last();
+						}
+					}
+				}
+			}
+
+			MorphTargetComboOptions.Sort(CompareNames);
+
+			TSharedRef<IPropertyHandle> MorphTargetNameProperty = DetailBuilder.GetProperty("MorphTargetName");
+
+			BlocksCategory.AddCustomRow( LOCTEXT("FCustomizableObjectNodeMeshMorphDetails", "Target") )
+			[
+				SNew(SProperty, MorphTargetNameProperty)
+				.ShouldDisplayName( false)
+				.CustomWidget()
+				[
+					SNew(SBorder)
+					.BorderImage(FAppStyle::GetBrush("NoBorder"))
+					.Padding(FMargin(0.0f, 0.0f, 10.0f, 0.0f))
+					[
+						SNew(STextComboBox)
+							.OptionsSource( &MorphTargetComboOptions )
+							.InitiallySelectedItem( ItemToSelect )
+							.OnSelectionChanged(this, &FCustomizableObjectNodeMeshMorphDetails::OnMorphTargetComboBoxSelectionChanged, MorphTargetNameProperty)
+					]
+				]
+			];
+		}
+	}
+	
+	if (!SourceMeshFound)
+	{
+		BlocksCategory.AddCustomRow( LOCTEXT("FCustomizableObjectNodeMeshMorphDetails", "Node") )
+		[
+			SNew( STextBlock )
+			.Text( LOCTEXT( "No source mesh found.", "No source mesh found." ) )			
+		];
+	}
+
+}
+
+void FCustomizableObjectNodeMeshMorphDetails::OnMorphTargetComboBoxSelectionChanged(TSharedPtr<FString> Selection, ESelectInfo::Type SelectInfo, TSharedRef<IPropertyHandle> ParentProperty)
+{	
+	if (Selection.IsValid())
+	{
+		ParentProperty->SetValue(*Selection);	
+	}
+}
+
+#undef LOCTEXT_NAMESPACE
+
diff --git a/Engine/Plugins/Experimental/Mutable/Source/CustomizableObjectEditor/Private/Nodes/CustomizableObjectNodeMeshMorphDetails.h b/Engine/Plugins/Experimental/Mutable/Source/CustomizableObjectEditor/Private/Nodes/CustomizableObjectNodeMeshMorphDetails.h
new file mode 100644
--- /dev/null
+++ b/Engine/Plugins/Experimental/Mutable/Source/CustomizableObjectEditor/Private/Nodes/CustomizableObjectNodeMeshMorphDetails.h
@@ -0,0 +1,26 @@
+// Copyright Epic Games, Inc. All Rights Reserved.
+
+#pragma once
+
+#include "IDetailCustomization.h"
+#include "Types/SlateEnums.h"
+#include "Templates/SharedPointer.h"
+#include "UObject/WeakObjectPtr.h"
+#include "PropertyHandle.h"
+
+class FCustomizableObjectNodeMeshMorphDetails : public IDetailCustomization
+{
+public:
+	/** Makes a new instance of this detail layout class for a specific detail view requesting it */
+	static TSharedRef<IDetailCustomization> MakeInstance();
+
+	/** ILayoutDetails interface */
+	virtual void CustomizeDetails( IDetailLayoutBuilder& DetailBuilder ) override;
+
+private:
+
+	class UCustomizableObjectNodeMeshMorph* Node = nullptr;
+	TArray< TSharedPtr<FString> > MorphTargetComboOptions;
+
+	void OnMorphTargetComboBoxSelectionChanged(TSharedPtr<FString> Selection, ESelectInfo::Type SelectInfo, TSharedRef<IPropertyHandle> Property);
+};
diff --git a/Engine/Plugins/Experimental/Mutable/Source/CustomizableObjectEditor/Private/Nodes/CustomizableObjectNodeMeshMorphStackApplication.cpp b/Engine/Plugins/Experimental/Mutable/Source/CustomizableObjectEditor/Private/Nodes/CustomizableObjectNodeMeshMorphStackApplication.cpp
new file mode 100644
--- /dev/null
+++ b/Engine/Plugins/Experimental/Mutable/Source/CustomizableObjectEditor/Private/Nodes/CustomizableObjectNodeMeshMorphStackApplication.cpp
@@ -0,0 +1,138 @@
+// Copyright Epic Games, Inc. All Rights Reserved.
+
+#include "Nodes/CustomizableObjectNodeMeshMorphStackApplication.h"
+#include "CustomizableObjectEditorModule.h"
+#include "GraphTraversal.h"
+
+
+#define LOCTEXT_NAMESPACE "CustomizableObjectEditor"
+
+
+
+UCustomizableObjectNodeMeshMorphStackApplication::UCustomizableObjectNodeMeshMorphStackApplication() : Super()
+{
+
+}
+
+
+void UCustomizableObjectNodeMeshMorphStackApplication::AllocateDefaultPins(UCustomizableObjectNodeRemapPins* RemapPins)
+{
+	const UEdGraphSchema_CustomizableObject* Schema = GetDefault<UEdGraphSchema_CustomizableObject>();
+
+	// Input pins
+	FString PinName = TEXT("InMesh");
+	UEdGraphPin* InputMeshPin = CustomCreatePin(EGPD_Input, Schema->PC_Mesh, FName(*PinName));
+
+	PinName = TEXT("Stack");
+	UEdGraphPin* StackPin = CustomCreatePin(EGPD_Input, Schema->PC_Stack, FName(*PinName));
+
+	// Output pins
+	PinName = TEXT("Result Mesh");
+	UEdGraphPin* OutputMeshPin = CustomCreatePin(EGPD_Output, Schema->PC_Mesh, FName(*PinName));
+
+}
+
+
+FText UCustomizableObjectNodeMeshMorphStackApplication::GetNodeTitle(ENodeTitleType::Type TittleType)const
+{
+	return LOCTEXT("Mesh_Morph_Stack_Application", "Mesh Morph Stack Application");
+}
+
+
+FLinearColor UCustomizableObjectNodeMeshMorphStackApplication::GetNodeTitleColor() const
+{
+	const UEdGraphSchema_CustomizableObject* Schema = GetDefault<UEdGraphSchema_CustomizableObject>();
+	return Schema->GetPinTypeColor(Schema->PC_Mesh);
+}
+
+
+FText UCustomizableObjectNodeMeshMorphStackApplication::GetTooltipText() const
+{
+	return LOCTEXT("Morph_Stack_Application_Tooltip","Applies a morph stack to a mesh");
+}
+
+
+void UCustomizableObjectNodeMeshMorphStackApplication::PinConnectionListChanged(UEdGraphPin* Pin)
+{
+	Super::PinConnectionListChanged(Pin);
+
+	if (Pin->Direction == EGPD_Input && Pin->PinName == "InMesh")
+	{
+		UpdateMorphList();
+
+		TSharedPtr<ICustomizableObjectEditor> Editor = GetGraphEditor();
+
+		if (Editor.IsValid())
+		{
+			Editor->UpdateGraphNodeProperties();
+			MarkForReconstruct();
+		}
+	}
+}
+
+
+void UCustomizableObjectNodeMeshMorphStackApplication::ReconstructNode(UCustomizableObjectNodeRemapPins* RemapPins)
+{
+	UpdateMorphList();
+
+	Super::ReconstructNode(RemapPins);
+}
+
+
+void UCustomizableObjectNodeMeshMorphStackApplication::UpdateMorphList()
+{
+	MorphNames.Empty();
+
+	if (UEdGraphPin* MeshPin = GetMeshPin())
+	{
+		if (const UEdGraphPin* ConnectedPin = FollowInputPin(*MeshPin))
+		{
+			if (UCustomizableObjectNodeSkeletalMesh* SkeletalMeshNode = Cast< UCustomizableObjectNodeSkeletalMesh >(ConnectedPin->GetOuter()))
+			{
+				USkeletalMesh* SkeletalMesh = SkeletalMeshNode->SkeletalMesh;
+
+				if (SkeletalMesh)
+				{
+					for (int32 i = 0; i < SkeletalMesh->GetMorphTargets().Num(); ++i)
+					{
+						MorphNames.Add(SkeletalMesh->GetMorphTargets()[i]->GetName());
+					}
+				}
+			}
+		}
+	}
+}
+
+
+bool UCustomizableObjectNodeMeshMorphStackApplication::IsNodeOutDatedAndNeedsRefresh()
+{
+	if (UEdGraphPin* MeshPin = GetMeshPin())
+	{
+		if (const UEdGraphPin* ConnectedPin = FollowInputPin(*MeshPin))
+		{
+			if (const UCustomizableObjectNodeSkeletalMesh* SkeletalMeshNode = Cast< UCustomizableObjectNodeSkeletalMesh >(ConnectedPin->GetOuter()))
+			{
+				if (USkeletalMesh* SkeletalMesh = SkeletalMeshNode->SkeletalMesh)
+				{
+					if (SkeletalMesh->GetMorphTargets().Num() != MorphNames.Num())
+					{
+						return true;
+					}
+
+					for (int32 i = 0; i < SkeletalMesh->GetMorphTargets().Num(); ++i)
+					{
+						if (!MorphNames.Contains(SkeletalMesh->GetMorphTargets()[i]->GetName()))
+						{
+							return true;
+						}
+					}
+				}
+			}
+		}
+	}
+
+	return false;
+}
+
+
+#undef LOCTEXT_NAMESPACE
