[mutable] Move Mutable plugin from the restricted folder to the experimental plugins folder.

diff --git a/Engine/Plugins/Experimental/Mutable/Source/CustomizableObjectEditor/Private/Nodes/CustomizableObjectNodeMeshMorphStackDefinition.cpp b/Engine/Plugins/Experimental/Mutable/Source/CustomizableObjectEditor/Private/Nodes/CustomizableObjectNodeMeshMorphStackDefinition.cpp
new file mode 100644
--- /dev/null
+++ b/Engine/Plugins/Experimental/Mutable/Source/CustomizableObjectEditor/Private/Nodes/CustomizableObjectNodeMeshMorphStackDefinition.cpp
@@ -0,0 +1,179 @@
+// Copyright Epic Games, Inc. All Rights Reserved.
+
+#include "Nodes/CustomizableObjectNodeMeshMorphStackDefinition.h"
+#include "Nodes/CustomizableObjectNodeSkeletalMesh.h"
+#include "CustomizableObjectEditorModule.h"
+#include "GraphTraversal.h"
+
+
+#define LOCTEXT_NAMESPACE "CustomizableObjectEditor"
+
+
+
+UCustomizableObjectNodeMeshMorphStackDefinition::UCustomizableObjectNodeMeshMorphStackDefinition() : Super()
+{
+
+}
+
+
+void UCustomizableObjectNodeMeshMorphStackDefinition::AllocateDefaultPins(UCustomizableObjectNodeRemapPins* RemapPins)
+{
+	const UEdGraphSchema_CustomizableObject* Schema = GetDefault<UEdGraphSchema_CustomizableObject>();
+
+	// Input pins
+	UEdGraphPin* MeshPin = CustomCreatePin(EGPD_Input, Schema->PC_Mesh, FName(TEXT("Mesh")));
+	MeshPin->bDefaultValueIsIgnored = true;
+
+	for (int32 i = 0; i < MorphNames.Num(); ++i)
+	{
+		FString PinName = MorphNames[i] + FString("_Morph");
+		UEdGraphPin* Pin = CustomCreatePin(EGPD_Input, Schema->PC_Float, FName(*PinName));
+		Pin->PinFriendlyName = FText::FromString(MorphNames[i]);
+		Pin->bDefaultValueIsIgnored = true;
+	}
+
+	// Output pins
+	UEdGraphPin* StackPin = CustomCreatePin(EGPD_Output, Schema->PC_Stack, FName(TEXT("Stack")));
+	StackPin->bDefaultValueIsIgnored = true;
+}
+
+
+FText UCustomizableObjectNodeMeshMorphStackDefinition::GetNodeTitle(ENodeTitleType::Type TittleType)const
+{
+	return LOCTEXT("Mesh_Morph_Stack_Definition", "Mesh Morph Stack Definition");
+}
+
+
+FLinearColor UCustomizableObjectNodeMeshMorphStackDefinition::GetNodeTitleColor() const
+{
+	const UEdGraphSchema_CustomizableObject* Schema = GetDefault<UEdGraphSchema_CustomizableObject>();
+	return Schema->GetPinTypeColor(Schema->PC_Mesh);
+}
+
+
+FText UCustomizableObjectNodeMeshMorphStackDefinition::GetTooltipText() const
+{
+	return LOCTEXT("Morph_Stack_Definition_Tooltip","Allows to stack morphs");
+}
+
+
+void UCustomizableObjectNodeMeshMorphStackDefinition::PinConnectionListChanged(UEdGraphPin* Pin)
+{
+	Super::PinConnectionListChanged(Pin);
+
+	if (Pin->Direction == EGPD_Input && Pin->PinName == "Mesh")
+	{
+		UpdateMorphList();
+
+		TSharedPtr<ICustomizableObjectEditor> Editor = GetGraphEditor();
+
+		if (Editor.IsValid())
+		{
+			Editor->UpdateGraphNodeProperties();
+			Super::ReconstructNode();
+		}
+	}
+}
+
+
+void UCustomizableObjectNodeMeshMorphStackDefinition::ReconstructNode(UCustomizableObjectNodeRemapPins* RemapPins)
+{
+	UpdateMorphList();
+
+	Super::ReconstructNode(RemapPins);
+}
+
+
+void UCustomizableObjectNodeMeshMorphStackDefinition::UpdateMorphList()
+{
+	MorphNames.Empty();
+
+	if (UEdGraphPin* MeshPin = GetMeshPin())
+	{
+		if (const UEdGraphPin* ConnectedPin = FollowInputPin(*MeshPin))
+		{
+			USkeletalMesh* SkeletalMesh = nullptr;
+
+			if (const UCustomizableObjectNodeSkeletalMesh* SkeletalMeshNode = Cast< UCustomizableObjectNodeSkeletalMesh >(ConnectedPin->GetOuter()))
+			{
+				SkeletalMesh = SkeletalMeshNode->SkeletalMesh;
+			}
+			else if (const UCustomizableObjectNodeTable* TableNode = Cast< UCustomizableObjectNodeTable >(ConnectedPin->GetOuter()))
+			{
+				SkeletalMesh = TableNode->GetColumnDefaultAssetByType<USkeletalMesh>(MeshPin->LinkedTo[0]);
+			}
+
+			if (SkeletalMesh)
+			{
+				for (int32 i = 0; i < SkeletalMesh->GetMorphTargets().Num(); ++i)
+				{
+					MorphNames.Add(SkeletalMesh->GetMorphTargets()[i]->GetName());
+				}
+			}
+		}
+	}
+}
+
+
+UEdGraphPin* UCustomizableObjectNodeMeshMorphStackDefinition::GetMorphPin(int32 Index) const
+{
+	FString PinName = MorphNames[Index] + FString("_Morph");
+
+	return FindPin(PinName);
+}
+
+
+int32 UCustomizableObjectNodeMeshMorphStackDefinition::NextConnectedPin(int32 Index, TArray<FString> AvailableMorphs) const
+{
+	if (AvailableMorphs.Num())
+	{
+		for (int32 i = Index + 1; i < MorphNames.Num(); ++i)
+		{
+			FString PinName = MorphNames[i] + FString("_Morph");
+
+			UEdGraphPin* NextPin = FindPin(PinName);
+
+			// Checking if the node contains another connected pin and if the morph of this pin existis in the skeletal mesh
+			if (NextPin && FollowInputPin(*NextPin) && AvailableMorphs.Contains(MorphNames[i]))
+			{
+				return i;
+			}
+		}
+	}
+
+	return -1;
+}
+
+
+bool UCustomizableObjectNodeMeshMorphStackDefinition::IsNodeOutDatedAndNeedsRefresh()
+{
+	if (UEdGraphPin* MeshPin = GetMeshPin())
+	{
+		if (const UEdGraphPin* ConnectedPin = FollowInputPin(*MeshPin))
+		{
+			if (const UCustomizableObjectNodeSkeletalMesh* SkeletalMeshNode = Cast< UCustomizableObjectNodeSkeletalMesh >(ConnectedPin->GetOuter()))
+			{
+				if (USkeletalMesh* SkeletalMesh = SkeletalMeshNode->SkeletalMesh)
+				{
+					if(SkeletalMesh->GetMorphTargets().Num() != MorphNames.Num())
+					{
+						return true;
+					}
+
+					for (int32 i = 0; i < SkeletalMesh->GetMorphTargets().Num(); ++i)
+					{
+						if (!MorphNames.Contains(SkeletalMesh->GetMorphTargets()[i]->GetName()))
+						{
+							return true;
+						}
+					}
+				}
+			}
+		}
+	}
+
+	return false;
+}
+
+
+#undef LOCTEXT_NAMESPACE
diff --git a/Engine/Plugins/Experimental/Mutable/Source/CustomizableObjectEditor/Private/Nodes/CustomizableObjectNodeMeshReshape.cpp b/Engine/Plugins/Experimental/Mutable/Source/CustomizableObjectEditor/Private/Nodes/CustomizableObjectNodeMeshReshape.cpp
new file mode 100644
--- /dev/null
+++ b/Engine/Plugins/Experimental/Mutable/Source/CustomizableObjectEditor/Private/Nodes/CustomizableObjectNodeMeshReshape.cpp
@@ -0,0 +1,60 @@
+// Copyright Epic Games, Inc. All Rights Reserved.
+
+#include "Nodes/CustomizableObjectNodeMeshReshape.h"
+#include "CustomizableObjectEditorModule.h"
+#include "DetailLayoutBuilder.h"
+#include "GraphTraversal.h"
+
+#define LOCTEXT_NAMESPACE "CustomizableObjectEditor"
+
+
+UCustomizableObjectNodeMeshReshape::UCustomizableObjectNodeMeshReshape()
+	: Super()
+{
+}
+
+
+void UCustomizableObjectNodeMeshReshape::AllocateDefaultPins(UCustomizableObjectNodeRemapPins* RemapPins)
+{
+	const UEdGraphSchema_CustomizableObject* Schema = GetDefault<UEdGraphSchema_CustomizableObject>();
+
+	UEdGraphPin* OutputPin = CustomCreatePin(EGPD_Output, Schema->PC_Mesh, FName("Mesh"));
+	OutputPin->bDefaultValueIsIgnored = true;
+
+	UEdGraphPin* MeshPin = CustomCreatePin(EGPD_Input, Schema->PC_Mesh, FName("Base Mesh"));
+	MeshPin->bDefaultValueIsIgnored = true;
+
+	MeshPin = CustomCreatePin(EGPD_Input, Schema->PC_Mesh, FName("Base Shape"));
+	MeshPin->bDefaultValueIsIgnored = true;
+
+	MeshPin = CustomCreatePin(EGPD_Input, Schema->PC_Mesh, FName("Target Shape"));
+	MeshPin->bDefaultValueIsIgnored = true;
+}
+
+
+FText UCustomizableObjectNodeMeshReshape::GetNodeTitle(ENodeTitleType::Type TitleType) const
+{
+	return LOCTEXT("Mesh_Reshape", "Mesh Reshape");
+}
+
+
+FLinearColor UCustomizableObjectNodeMeshReshape::GetNodeTitleColor() const
+{
+	const UEdGraphSchema_CustomizableObject* Schema = GetDefault<UEdGraphSchema_CustomizableObject>();
+	return Schema->GetPinTypeColor(Schema->PC_Mesh);
+}
+
+
+FText UCustomizableObjectNodeMeshReshape::GetTooltipText() const
+{
+	return LOCTEXT("Mesh_Reshape_Tooltip", "Apply a mesh reshape on a mesh.");
+}
+
+
+FString UCustomizableObjectNodeMeshReshape::GetRefreshMessage() const
+{
+	return "One or more bones selected to deform do not exist in the current reference mesh skeleton!";
+}
+
+
+#undef LOCTEXT_NAMESPACE
diff --git a/Engine/Plugins/Experimental/Mutable/Source/CustomizableObjectEditor/Private/Nodes/CustomizableObjectNodeMeshReshapeCommon.cpp b/Engine/Plugins/Experimental/Mutable/Source/CustomizableObjectEditor/Private/Nodes/CustomizableObjectNodeMeshReshapeCommon.cpp
new file mode 100644
--- /dev/null
+++ b/Engine/Plugins/Experimental/Mutable/Source/CustomizableObjectEditor/Private/Nodes/CustomizableObjectNodeMeshReshapeCommon.cpp
@@ -0,0 +1,3 @@
+// Copyright Epic Games, Inc. All Rights Reserved.
+
+#include "Nodes/CustomizableObjectNodeMeshReshapeCommon.h"
diff --git a/Engine/Plugins/Experimental/Mutable/Source/CustomizableObjectEditor/Private/Nodes/CustomizableObjectNodeMeshReshapeDetails.cpp b/Engine/Plugins/Experimental/Mutable/Source/CustomizableObjectEditor/Private/Nodes/CustomizableObjectNodeMeshReshapeDetails.cpp
new file mode 100644
--- /dev/null
+++ b/Engine/Plugins/Experimental/Mutable/Source/CustomizableObjectEditor/Private/Nodes/CustomizableObjectNodeMeshReshapeDetails.cpp
@@ -0,0 +1,40 @@
+// Copyright Epic Games, Inc. All Rights Reserved.
+
+#include "CustomizableObjectNodeMeshReshapeDetails.h"
+#include "Nodes/CustomizableObjectNodeMeshReshape.h"
+#include "Nodes/CustomizableObjectNodeSkeletalMesh.h"
+#include "Nodes/CustomizableObjectNodeTable.h"
+#include "GenerateMutableSource/GenerateMutableSource.h"
+
+#include "DetailLayoutBuilder.h"
+#include "DetailCategoryBuilder.h"
+#include "PropertyCustomizationHelpers.h"
+
+#include "Templates/SharedPointer.h"
+#include "UObject/WeakObjectPtr.h"
+#include "Engine/SkeletalMesh.h"
+
+#include "Widgets/Text/STextBlock.h"
+#include "BoneSelectionWidget.h"
+#include "GraphTraversal.h"
+
+#define LOCTEXT_NAMESPACE "CustomizableObjectDetails"
+
+
+TSharedRef<IDetailCustomization> FCustomizableObjectNodeMeshReshapeDetails::MakeInstance()
+{
+	return MakeShareable(new FCustomizableObjectNodeMeshReshapeDetails);
+}
+
+
+void FCustomizableObjectNodeMeshReshapeDetails::CustomizeDetails(IDetailLayoutBuilder& DetailBuilder)
+{
+}
+
+
+void FCustomizableObjectNodeMeshReshapeDetails::CustomizeDetails(const TSharedPtr<IDetailLayoutBuilder>& DetailBuilder)
+{
+}
+
+
+#undef LOCTEXT_NAMESPACE
diff --git a/Engine/Plugins/Experimental/Mutable/Source/CustomizableObjectEditor/Private/Nodes/CustomizableObjectNodeMeshReshapeSelectionDetails.cpp b/Engine/Plugins/Experimental/Mutable/Source/CustomizableObjectEditor/Private/Nodes/CustomizableObjectNodeMeshReshapeSelectionDetails.cpp
new file mode 100644
--- /dev/null
+++ b/Engine/Plugins/Experimental/Mutable/Source/CustomizableObjectEditor/Private/Nodes/CustomizableObjectNodeMeshReshapeSelectionDetails.cpp
@@ -0,0 +1,199 @@
+// Copyright Epic Games, Inc. All Rights Reserved.
+
+#include "CustomizableObjectNodeMeshReshapeSelectionDetails.h"
+
+#include "CustomizableObjectEditorModule.h"
+#include "../CustomizableObjectCompiler.h"
+#include "CustomizableObjectEditorUtilities.h"
+
+#include "PropertyCustomizationHelpers.h"
+#include "DetailLayoutBuilder.h"
+#include "DetailCategoryBuilder.h"
+#include "GraphTraversal.h"
+#include "LevelEditor.h"
+#include "LevelEditorActions.h"
+
+#include "Widgets/Layout/SBorder.h"
+#include "Widgets/Text/STextBlock.h"
+#include "Widgets/Input/STextComboBox.h"
+
+#include "Engine/SkeletalMesh.h"
+#include "Engine/StaticMesh.h"
+#include "Animation/MorphTarget.h"
+
+#include "BoneSelectionWidget.h"
+
+#include "Runtime/Launch/Resources/Version.h"
+
+
+#define LOCTEXT_NAMESPACE "CustomizableObjectDetails"
+
+
+TSharedRef<IPropertyTypeCustomization> FMeshReshapeBonesReferenceCustomization::MakeInstance()
+{
+	return MakeShareable(new FMeshReshapeBonesReferenceCustomization());
+}
+
+void FMeshReshapeBonesReferenceCustomization::CustomizeHeader(TSharedRef<IPropertyHandle> StructPropertyHandle, class FDetailWidgetRow& HeaderRow, IPropertyTypeCustomizationUtils& StructCustomizationUtils)
+{
+	constexpr bool bRecurse = false;
+	BoneNameProperty = StructPropertyHandle->GetChildHandle(GET_MEMBER_NAME_CHECKED(FMeshReshapeBoneReference, BoneName), bRecurse);
+
+	TSharedPtr<IPropertyHandle> Parent = StructPropertyHandle->GetParentHandle();
+	if (Parent)
+	{
+		ParentArrayProperty = Parent->AsArray();
+	}
+
+	if (BoneNameProperty->IsValidHandle())
+	{
+		HeaderRow
+		.NameContent()
+		[
+			StructPropertyHandle->CreatePropertyNameWidget()
+		]
+		.ValueContent()
+		.MaxDesiredWidth(0.0f)
+		[
+			SAssignNew(SelectionWidget, SBoneSelectionWidget)
+			.ToolTipText(StructPropertyHandle->GetToolTipText())
+			.OnBoneSelectionChanged(this, &FMeshReshapeBonesReferenceCustomization::OnBoneSelectionChanged)
+			.OnGetSelectedBone(this, &FMeshReshapeBonesReferenceCustomization::GetSelectedBone)
+			.OnGetReferenceSkeleton(this, &FMeshReshapeBonesReferenceCustomization::GetReferenceSkeleton)
+		];
+
+		DefaultColor = SelectionWidget->GetColorAndOpacity();
+	}
+}
+
+USkeletalMesh* FMeshReshapeBonesReferenceCustomization::GetSkeletalMesh(TSharedPtr<IPropertyHandle> BoneNamePropertyHandle) const 
+{
+	if (!BoneNamePropertyHandle)
+	{
+		return nullptr;
+	}
+
+	TArray<UObject*> Objects;
+	BoneNamePropertyHandle->GetOuterObjects(Objects);
+
+	USkeletalMesh* TargetSkeletalMesh = nullptr;
+	
+	for (const UObject* Object : Objects)
+	{
+		UEdGraphPin* MeshPin = nullptr;
+
+		if (const UCustomizableObjectNodeMeshMorph* NodeMorph = Cast<UCustomizableObjectNodeMeshMorph>(Object))
+		{
+			MeshPin = NodeMorph->MeshPin();
+		}
+		else if (const UCustomizableObjectNodeMeshReshape* NodeReshape = Cast<UCustomizableObjectNodeMeshReshape>(Object))
+		{
+			MeshPin = NodeReshape->BaseMeshPin();
+		}
+		else
+		{
+			checkf(false, TEXT("node not implemented"));
+		}
+
+		if (MeshPin)
+		{
+			if (const UEdGraphPin* ConnectedPin = FollowInputPin(*MeshPin))
+			{
+				if (const UEdGraphPin* SourceMeshPin = FindMeshBaseSource(*ConnectedPin, false))
+				{
+					if (const UCustomizableObjectNodeSkeletalMesh* SkeletalMeshNode = Cast<UCustomizableObjectNodeSkeletalMesh>(SourceMeshPin->GetOwningNode()))
+					{
+						return SkeletalMeshNode->SkeletalMesh;
+					}
+					else if (UCustomizableObjectNodeTable* TableNode = Cast<UCustomizableObjectNodeTable>(SourceMeshPin->GetOwningNode()))
+					{
+						return TableNode->GetColumnDefaultAssetByType<USkeletalMesh>(SourceMeshPin);
+					}
+				}
+			}
+		}
+	}
+
+	return nullptr;
+}
+
+FName FMeshReshapeBonesReferenceCustomization::GetSelectedBone(bool& bMultipleValues) const
+{
+	if (BoneNameProperty.IsValid())
+	{
+		FString OutText;
+
+		FPropertyAccess::Result Result = BoneNameProperty->GetValueAsFormattedString(OutText);
+
+		return FName(*OutText);
+	}
+
+	return FName();
+}
+
+const struct FReferenceSkeleton& FMeshReshapeBonesReferenceCustomization::GetReferenceSkeleton() const
+{
+	static FReferenceSkeleton DummySkeleton;
+
+	USkeletalMesh* SkeletalMesh = GetSkeletalMesh(BoneNameProperty);
+	return SkeletalMesh ? SkeletalMesh->GetRefSkeleton() : DummySkeleton;
+}
+
+void FMeshReshapeBonesReferenceCustomization::OnBoneSelectionChanged(FName Name)
+{
+	FMeshReshapeBoneReference NewBone;
+	NewBone.BoneName = Name;
+
+	if (!ParentArrayProperty)
+	{
+		BoneNameProperty->SetValue(Name);
+		return;
+	}
+
+	uint32 NumParentElems = 0;
+	FPropertyAccess::Result Result = ParentArrayProperty->GetNumElements(NumParentElems);
+	check(Result == FPropertyAccess::Result::Success);
+
+	// Find if already set.
+	bool BoneNameFound = false;
+	for (uint32 I = 0; I < NumParentElems; ++I)
+	{
+		TSharedPtr<IPropertyHandle> ElemProperty = ParentArrayProperty->GetElement(I);
+		check(ElemProperty);
+	
+		constexpr bool bRecurse = false;	
+		TSharedPtr<IPropertyHandle> ElemBoneNameProperty = 
+				ElemProperty->GetChildHandle(GET_MEMBER_NAME_CHECKED(FMeshReshapeBoneReference, BoneName), bRecurse);
+
+		check(ElemBoneNameProperty);
+
+		FName BoneNameValue;
+		Result = ElemBoneNameProperty->GetValue(BoneNameValue);
+
+		check(Result == FPropertyAccess::Result::Success);
+
+		if (BoneNameValue == NewBone.BoneName)
+		{
+			BoneNameFound = true;
+			break;
+		}	
+	}
+
+
+	// TODO: This behavior does not seem to be correct, if the option is repeated the value is not set
+	// only logging a warning to the console, which might not be open and could make dificult to see what 
+	// is the problem. 
+	BoneNameProperty->SetValue(Name);
+	if (BoneNameFound)
+	{
+		const FLinearColor WarningColor = FLinearColor::Yellow;
+		SelectionWidget->SetColorAndOpacity(WarningColor);
+	}
+	else
+	{
+		SelectionWidget->SetColorAndOpacity(DefaultColor);
+	}
+}
+
+#undef LOCTEXT_NAMESPACE
+
diff --git a/Engine/Plugins/Experimental/Mutable/Source/CustomizableObjectEditor/Private/Nodes/CustomizableObjectNodeMeshReshapeSelectionDetails.h b/Engine/Plugins/Experimental/Mutable/Source/CustomizableObjectEditor/Private/Nodes/CustomizableObjectNodeMeshReshapeSelectionDetails.h
new file mode 100644
--- /dev/null
+++ b/Engine/Plugins/Experimental/Mutable/Source/CustomizableObjectEditor/Private/Nodes/CustomizableObjectNodeMeshReshapeSelectionDetails.h
@@ -0,0 +1,42 @@
+// Copyright Epic Games, Inc. All Rights Reserved.
+
+#pragma once
+
+#include "IDetailCustomization.h"
+#include "Types/SlateEnums.h"
+#include "Templates/SharedPointer.h"
+#include "UObject/WeakObjectPtr.h"
+#include "PropertyHandle.h"
+#include "BoneSelectionWidget.h"
+
+class FMeshReshapeBonesReferenceCustomization : public IPropertyTypeCustomization
+{
+public:
+	static TSharedRef<IPropertyTypeCustomization> MakeInstance();
+
+	// IPropertyTypeCustomization interface
+	void CustomizeHeader(TSharedRef<IPropertyHandle> StructPropertyHandle, class FDetailWidgetRow& HeaderRow, IPropertyTypeCustomizationUtils& StructCustomizationUtils) override;
+	void CustomizeChildren(TSharedRef<IPropertyHandle> StructPropertyHandle, class IDetailChildrenBuilder& ChildBuilder, IPropertyTypeCustomizationUtils& StructCustomizationUtils) override {};
+	
+	// Gets the Reference Skeleton from the target skeleton
+	const struct FReferenceSkeleton& GetReferenceSkeleton() const;
+	// Gets the current selected bone in the combo button
+	FName GetSelectedBone(bool& bMultipleValues) const;
+	// OnSelectionChanged callback
+	void OnBoneSelectionChanged(FName Name);
+
+private:
+
+	class USkeletalMesh* GetSkeletalMesh(TSharedPtr<IPropertyHandle> BoneNamePropertyHandle) const; 
+
+private:
+
+	// Property to change after bone has been picked
+	TSharedPtr<IPropertyHandle> BoneNameProperty;
+
+	// Parent propery, as array to identify repeated values.
+	TSharedPtr<IPropertyHandleArray> ParentArrayProperty;
+
+	TSharedPtr<SBoneSelectionWidget> SelectionWidget;
+	FLinearColor DefaultColor;
+};
diff --git a/Engine/Plugins/Experimental/Mutable/Source/CustomizableObjectEditor/Private/Nodes/CustomizableObjectNodeMeshSwitch.cpp b/Engine/Plugins/Experimental/Mutable/Source/CustomizableObjectEditor/Private/Nodes/CustomizableObjectNodeMeshSwitch.cpp
new file mode 100644
--- /dev/null
+++ b/Engine/Plugins/Experimental/Mutable/Source/CustomizableObjectEditor/Private/Nodes/CustomizableObjectNodeMeshSwitch.cpp
@@ -0,0 +1,42 @@
+// Copyright Epic Games, Inc. All Rights Reserved.
+
+#include "Nodes/CustomizableObjectNodeMeshSwitch.h"
+
+#include "EdGraphSchema_CustomizableObject.h"
+
+#define LOCTEXT_NAMESPACE "CustomizableObjectEditor"
+
+
+void UCustomizableObjectNodeMeshSwitch::Serialize(FArchive& Ar)
+{
+	Super::Serialize(Ar);
+
+	Ar.UsingCustomVersion(FCustomizableObjectCustomVersion::GUID);
+
+	if (Ar.CustomVer(FCustomizableObjectCustomVersion::GUID) < FCustomizableObjectCustomVersion::BugPinsSwitch)
+	{
+		OutputPinReference = FindPin(TEXT("Mesh"));	
+	}
+}
+
+
+FString UCustomizableObjectNodeMeshSwitch::GetOutputPinName() const
+{
+	return TEXT("Mesh");
+}
+
+
+FName UCustomizableObjectNodeMeshSwitch::GetCategory() const
+{
+	return UEdGraphSchema_CustomizableObject::PC_Mesh;
+}
+
+
+FString UCustomizableObjectNodeMeshSwitch::GetPinPrefix() const
+{
+	return TEXT("Mesh ");
+}
+
+
+#undef LOCTEXT_NAMESPACE
+
diff --git a/Engine/Plugins/Experimental/Mutable/Source/CustomizableObjectEditor/Private/Nodes/CustomizableObjectNodeMeshVariation.cpp b/Engine/Plugins/Experimental/Mutable/Source/CustomizableObjectEditor/Private/Nodes/CustomizableObjectNodeMeshVariation.cpp
new file mode 100644
--- /dev/null
+++ b/Engine/Plugins/Experimental/Mutable/Source/CustomizableObjectEditor/Private/Nodes/CustomizableObjectNodeMeshVariation.cpp
@@ -0,0 +1,71 @@
+// Copyright Epic Games, Inc. All Rights Reserved.
+
+#include "Nodes/CustomizableObjectNodeMeshVariation.h"
+#include "CustomizableObjectEditorModule.h"
+
+#include "UnrealEditorPortabilityHelpers.h"
+
+#define LOCTEXT_NAMESPACE "CustomizableObjectEditor"
+
+
+UCustomizableObjectNodeMeshVariation::UCustomizableObjectNodeMeshVariation()
+	: Super()
+{
+
+}
+
+
+void UCustomizableObjectNodeMeshVariation::PostEditChangeProperty(FPropertyChangedEvent& PropertyChangedEvent)
+{
+	FProperty* PropertyThatChanged = PropertyChangedEvent.Property;
+	if ( PropertyThatChanged  )
+	{
+		ReconstructNode();
+	}
+
+	Super::PostEditChangeProperty(PropertyChangedEvent);
+}
+
+
+void UCustomizableObjectNodeMeshVariation::AllocateDefaultPins(UCustomizableObjectNodeRemapPins* RemapPins)
+{
+	const UEdGraphSchema_CustomizableObject* Schema = GetDefault<UEdGraphSchema_CustomizableObject>();
+	
+	UEdGraphPin* OutputPin = CustomCreatePin(EGPD_Output, Schema->PC_Mesh, FName(TEXT("Mesh")));
+	OutputPin->bDefaultValueIsIgnored = true;
+
+	for ( int VariationIndex = Variations.Num() - 1; VariationIndex >= 0; --VariationIndex)
+	{
+		FString PinName = FString::Printf( TEXT("Variation %d"), VariationIndex);
+		UEdGraphPin* VariationPin = CustomCreatePin(EGPD_Input, Schema->PC_Mesh, FName(*PinName), false);
+		VariationPin->bDefaultValueIsIgnored = true;
+
+		FString FriendlyName = FString::Printf(TEXT("Variation %d [%s]"), VariationIndex, *Variations[VariationIndex].Tag);
+		VariationPin->PinFriendlyName = FText::FromString(FriendlyName);
+	}
+
+	UEdGraphPin* DefaultVariation = CustomCreatePin(EGPD_Input, Schema->PC_Mesh, FName(TEXT("Default")), false);
+	DefaultVariation->bDefaultValueIsIgnored = true;
+}
+
+
+FText UCustomizableObjectNodeMeshVariation::GetNodeTitle(ENodeTitleType::Type TitleType) const
+{
+	return LOCTEXT("Mesh_Variation", "Mesh Variation");
+}
+
+
+FLinearColor UCustomizableObjectNodeMeshVariation::GetNodeTitleColor() const
+{
+	const UEdGraphSchema_CustomizableObject* Schema = GetDefault<UEdGraphSchema_CustomizableObject>();
+	return Schema->GetPinTypeColor(Schema->PC_Mesh);
+}
+
+
+FText UCustomizableObjectNodeMeshVariation::GetTooltipText() const
+{
+	return LOCTEXT("Mesh_Variation_Tooltip", "Select a mesh depending on what tags are active.");
+}
+
+#undef LOCTEXT_NAMESPACE
+
diff --git a/Engine/Plugins/Experimental/Mutable/Source/CustomizableObjectEditor/Private/Nodes/CustomizableObjectNodeModifierBase.cpp b/Engine/Plugins/Experimental/Mutable/Source/CustomizableObjectEditor/Private/Nodes/CustomizableObjectNodeModifierBase.cpp
new file mode 100644
--- /dev/null
+++ b/Engine/Plugins/Experimental/Mutable/Source/CustomizableObjectEditor/Private/Nodes/CustomizableObjectNodeModifierBase.cpp
@@ -0,0 +1,10 @@
+// Copyright Epic Games, Inc. All Rights Reserved.
+
+#include "Nodes/CustomizableObjectNodeModifierBase.h"
+#include "CustomizableObjectEditorModule.h"
+
+
+UCustomizableObjectNodeModifierBase::UCustomizableObjectNodeModifierBase()
+	: Super()
+{
+}
diff --git a/Engine/Plugins/Experimental/Mutable/Source/CustomizableObjectEditor/Private/Nodes/CustomizableObjectNodeMorphMaterial.cpp b/Engine/Plugins/Experimental/Mutable/Source/CustomizableObjectEditor/Private/Nodes/CustomizableObjectNodeMorphMaterial.cpp
new file mode 100644
--- /dev/null
+++ b/Engine/Plugins/Experimental/Mutable/Source/CustomizableObjectEditor/Private/Nodes/CustomizableObjectNodeMorphMaterial.cpp
@@ -0,0 +1,146 @@
+// Copyright Epic Games, Inc. All Rights Reserved.
+
+#include "Nodes/CustomizableObjectNodeMorphMaterial.h"
+#include "CustomizableObjectEditorModule.h"
+#include "GraphTraversal.h"
+
+#define LOCTEXT_NAMESPACE "CustomizableObjectEditor"
+
+
+void UCustomizableObjectNodeMorphMaterial::AllocateDefaultPins(UCustomizableObjectNodeRemapPins* RemapPins)
+{
+	const UEdGraphSchema_CustomizableObject* Schema = GetDefault<UEdGraphSchema_CustomizableObject>();
+
+	//if (!MorphPin)
+	//{
+	//	MorphPin = CreatePin(EGPD_Input, Schema->PC_Mesh, FString(), NULL, false, false, TEXT("Morph") );
+	//	MorphPin->bDefaultValueIsIgnored = true;
+	//}
+	//else
+	//{
+	//	Pins.Add( MorphPin );
+	//}
+
+	FString PinName = TEXT("Material");
+	UEdGraphPin* MaterialPin = CustomCreatePin(EGPD_Output, Schema->PC_Material, FName(PinName));
+	MaterialPin->bDefaultValueIsIgnored = true;
+
+	PinName = TEXT("Factor");
+	UEdGraphPin* FactorPin = CustomCreatePin(EGPD_Input, Schema->PC_Float, FName(*PinName));
+	FactorPin->bDefaultValueIsIgnored = true;
+}
+
+
+FText UCustomizableObjectNodeMorphMaterial::GetNodeTitle(ENodeTitleType::Type TitleType) const
+{
+	return LOCTEXT("Morph_Material", "Morph Material");
+}
+
+
+FLinearColor UCustomizableObjectNodeMorphMaterial::GetNodeTitleColor() const
+{
+	const UEdGraphSchema_CustomizableObject* Schema = GetDefault<UEdGraphSchema_CustomizableObject>();
+	return Schema->GetPinTypeColor(Schema->PC_Material);
+}
+
+
+bool UCustomizableObjectNodeMorphMaterial::IsNodeOutDatedAndNeedsRefresh()
+{
+	bool Result = false;
+
+	if (const UCustomizableObjectNodeMaterial* CustomizableObjectNodeMaterial = Cast<UCustomizableObjectNodeMaterial>(GetParentMaterialNode()))
+	{
+		if (const UEdGraphPin* BaseSourcePin = FollowInputPin(*CustomizableObjectNodeMaterial->GetMeshPin()))
+		{
+			if (const UCustomizableObjectNodeSkeletalMesh* TypedSourceNode = Cast<UCustomizableObjectNodeSkeletalMesh>(BaseSourcePin->GetOwningNode()))
+			{
+				if (const USkeletalMesh* SkeletalMesh = TypedSourceNode->SkeletalMesh)
+				{
+					bool bOutdated = true;
+					for (int m = 0; m < SkeletalMesh->GetMorphTargets().Num(); ++m)
+					{
+						if (MorphTargetName == *SkeletalMesh->GetMorphTargets()[m]->GetName())
+						{
+							bOutdated = false;
+							break;
+						}
+					}
+					Result = bOutdated;
+				}
+			}
+		}
+	}
+
+	// Remove previous compilation warnings
+	if (!Result && bHasCompilerMessage)
+	{
+		RemoveWarnings();
+		GetGraph()->NotifyGraphChanged();
+	}
+
+    return Result;
+}
+
+
+FString UCustomizableObjectNodeMorphMaterial::GetRefreshMessage() const
+{
+	return "Morph Target not found in the SkeletalMesh. Please Refresh Node and select a valid morph option.";
+}
+
+
+FText UCustomizableObjectNodeMorphMaterial::GetTooltipText() const
+{
+	return LOCTEXT("Morph_Material_Tooltip", "Fully activate one morph of a parent's material.");
+}
+
+
+void UCustomizableObjectNodeMorphMaterial::BackwardsCompatibleFixup()
+{
+	Super::BackwardsCompatibleFixup();
+
+	const int32 CustomizableObjectCustomVersion = GetLinkerCustomVersion(FCustomizableObjectCustomVersion::GUID);
+
+	if (CustomizableObjectCustomVersion < FCustomizableObjectCustomVersion::MorphMaterialAddFactorPin)
+	{
+		ReconstructNode();
+	}
+}
+
+
+void UCustomizableObjectNodeMorphMaterial::PinConnectionListChanged(UEdGraphPin* Pin)
+{
+	Super::PinConnectionListChanged(Pin);
+
+	if (Pin->PinName == "Material")
+	{
+		TSharedPtr<ICustomizableObjectEditor> Editor = GetGraphEditor();
+
+		if (Editor.IsValid())
+		{
+			MorphTargetName = "";
+			Editor->UpdateGraphNodeProperties();
+		}
+	}
+}
+
+
+void UCustomizableObjectNodeMorphMaterial::PostEditChangeProperty(FPropertyChangedEvent& PropertyChangedEvent)
+{
+	Super::PostEditChangeProperty(PropertyChangedEvent);
+	
+	TSharedPtr<ICustomizableObjectEditor> Editor = GetGraphEditor();
+	
+	if (Editor.IsValid())
+	{
+		Editor->UpdateGraphNodeProperties();
+	}
+}
+
+
+bool UCustomizableObjectNodeMorphMaterial::IsSingleOutputNode() const
+{
+	return true;
+}
+
+
+#undef LOCTEXT_NAMESPACE
diff --git a/Engine/Plugins/Experimental/Mutable/Source/CustomizableObjectEditor/Private/Nodes/CustomizableObjectNodeMorphMaterialDetails.cpp b/Engine/Plugins/Experimental/Mutable/Source/CustomizableObjectEditor/Private/Nodes/CustomizableObjectNodeMorphMaterialDetails.cpp
new file mode 100644
--- /dev/null
+++ b/Engine/Plugins/Experimental/Mutable/Source/CustomizableObjectEditor/Private/Nodes/CustomizableObjectNodeMorphMaterialDetails.cpp
@@ -0,0 +1,146 @@
+// Copyright Epic Games, Inc. All Rights Reserved.
+
+#include "CustomizableObjectNodeMorphMaterialDetails.h"
+#include "CustomizableObjectEditorModule.h"
+#include "../CustomizableObjectCompiler.h"
+#include "PropertyCustomizationHelpers.h"
+#include "CustomizableObjectEditorUtilities.h"
+
+#include "LevelEditor.h"
+#include "LevelEditorActions.h"
+
+#include "DetailLayoutBuilder.h"
+#include "DetailCategoryBuilder.h"
+#include "GraphTraversal.h"
+#include "LevelEditor.h"
+#include "LevelEditorActions.h"
+
+#include "Engine/SkeletalMesh.h"
+#include "Engine/StaticMesh.h"
+#include "Animation/MorphTarget.h"
+
+#include "Widgets/Text/STextBlock.h"
+#include "Widgets/Input/STextComboBox.h"
+#include "Templates/SharedPointer.h"
+#include "UObject/WeakObjectPtr.h"
+
+#include "Runtime/Launch/Resources/Version.h"
+
+
+#define LOCTEXT_NAMESPACE "CustomizableObjectDetails"
+
+
+TSharedRef<IDetailCustomization> FCustomizableObjectNodeMorphMaterialDetails::MakeInstance()
+{
+	return MakeShareable( new FCustomizableObjectNodeMorphMaterialDetails );
+}
+
+
+void FCustomizableObjectNodeMorphMaterialDetails::CustomizeDetails( IDetailLayoutBuilder& DetailBuilder )
+{
+
+#if (ENGINE_MAJOR_VERSION == 4 && ENGINE_MINOR_VERSION >= 18) || ENGINE_MAJOR_VERSION >= 5
+	const IDetailsView* DetailsView = DetailBuilder.GetDetailsView();
+	if ( DetailsView->GetSelectedObjects().Num() )
+	{
+		Node = Cast<UCustomizableObjectNodeMorphMaterial>( DetailsView->GetSelectedObjects()[0].Get() );
+	}
+#else
+	const IDetailsView& DetailsView = DetailBuilder.GetDetailsView();
+	if (DetailsView.GetSelectedObjects().Num())
+	{
+		Node = Cast<UCustomizableObjectNodeMorphMaterial>(DetailsView.GetSelectedObjects()[0].Get());
+	}
+#endif
+
+	IDetailCategoryBuilder& BlocksCategory = DetailBuilder.EditCategory( "Customizable Object" );
+	//BlocksCategory.CategoryIcon( "ActorClassIcon.CustomizableObject" );
+
+	MorphTargetComboOptions.Empty();
+
+	if (Node)
+	{
+		// Morph target selection
+		TSharedPtr<FString> ItemToSelect;
+
+		UCustomizableObjectNodeMaterialBase* ParentMaterialNode = Node->GetParentMaterialNode();
+		if (ParentMaterialNode)
+		{
+			const UEdGraphPin* BaseSourcePin = FindMeshBaseSource(*ParentMaterialNode->OutputPin(), false );
+			if ( BaseSourcePin )
+			{
+				if ( const UCustomizableObjectNodeSkeletalMesh* TypedSourceNode = Cast<UCustomizableObjectNodeSkeletalMesh>(BaseSourcePin->GetOwningNode()) )
+				{
+					for (int m = 0; m < TypedSourceNode->SkeletalMesh->GetMorphTargets().Num(); ++m )
+					{
+						FString MorphName = *TypedSourceNode->SkeletalMesh->GetMorphTargets()[m]->GetName();
+						MorphTargetComboOptions.Add( MakeShareable( new FString( MorphName ) ) );
+
+						if (Node->MorphTargetName == MorphName)
+						{
+							ItemToSelect = MorphTargetComboOptions.Last();
+						}
+					}
+				}
+			}
+		}
+
+        MorphTargetComboOptions.Sort(CompareNames);
+
+		TSharedRef<IPropertyHandle> MorphTargetNameProperty = DetailBuilder.GetProperty("MorphTargetName");
+		BlocksCategory.AddCustomRow( LOCTEXT("FCustomizableObjectNodeMorphMaterialDetails", "Target") )
+		[
+			SNew(SProperty, MorphTargetNameProperty)
+			.ShouldDisplayName( false)
+			.CustomWidget()
+			[
+				SNew(SBorder)
+				.BorderImage(FAppStyle::GetBrush("NoBorder"))
+				.Padding(FMargin(0.0f, 0.0f, 10.0f, 0.0f))
+				[
+					SNew( SHorizontalBox )
+
+					+ SHorizontalBox::Slot()
+						.FillWidth(10.0f)
+						.VAlign(VAlign_Center)
+					[
+						SNew( STextBlock )
+							.Text( LOCTEXT("Morph Target","Morph Target") )
+					]
+
+					+ SHorizontalBox::Slot()
+						.FillWidth(10.0f)
+						.HAlign(HAlign_Fill)
+						.VAlign(VAlign_Center)
+					[
+						SNew(STextComboBox)
+							.OptionsSource( &MorphTargetComboOptions )
+							.InitiallySelectedItem( ItemToSelect )
+							.OnSelectionChanged(this, &FCustomizableObjectNodeMorphMaterialDetails::OnMorphTargetComboBoxSelectionChanged, MorphTargetNameProperty)
+					]
+				]
+			]
+		];
+	}
+	else
+	{
+		BlocksCategory.AddCustomRow( LOCTEXT("FCustomizableObjectNodeMorphMaterialDetails", "Node") )
+		[
+			SNew( STextBlock )
+			.Text( LOCTEXT( "Node not found", "Node not found" ) )
+		];
+	}
+
+}
+
+
+void FCustomizableObjectNodeMorphMaterialDetails::OnMorphTargetComboBoxSelectionChanged(TSharedPtr<FString> Selection, ESelectInfo::Type SelectInfo, TSharedRef<IPropertyHandle> ParentProperty)
+{	
+	if (Selection.IsValid())
+	{
+		ParentProperty->SetValue(*Selection);	
+	}
+}
+
+
+#undef LOCTEXT_NAMESPACE
diff --git a/Engine/Plugins/Experimental/Mutable/Source/CustomizableObjectEditor/Private/Nodes/CustomizableObjectNodeMorphMaterialDetails.h b/Engine/Plugins/Experimental/Mutable/Source/CustomizableObjectEditor/Private/Nodes/CustomizableObjectNodeMorphMaterialDetails.h
new file mode 100644
--- /dev/null
+++ b/Engine/Plugins/Experimental/Mutable/Source/CustomizableObjectEditor/Private/Nodes/CustomizableObjectNodeMorphMaterialDetails.h
@@ -0,0 +1,35 @@
+// Copyright Epic Games, Inc. All Rights Reserved.
+
+#pragma once
+
+#include "IDetailCustomization.h"
+
+#include "Templates/SharedPointer.h"
+#include "UObject/WeakObjectPtr.h"
+#include "Misc/Guid.h"
+#include "Types/SlateEnums.h"
+#include "PropertyHandle.h"
+#include "CustomizableObjectEditorModule.h"
+
+class FCustomizableObjectNodeMorphMaterialDetails : public IDetailCustomization
+{
+public:
+	// Makes a new instance of this detail layout class for a specific detail view requesting it 
+	static TSharedRef<IDetailCustomization> MakeInstance();
+
+	// ILayoutDetails interface
+	void CustomizeDetails( IDetailLayoutBuilder& DetailBuilder ) override;
+
+
+private:
+
+	class UCustomizableObjectNodeMorphMaterial* Node = nullptr;
+
+	TArray< TSharedPtr<FString> > MorphTargetComboOptions;
+
+	void OnParentComboBoxSelectionChanged(TSharedPtr<FString> Selection, ESelectInfo::Type SelectInfo, TSharedRef<IPropertyHandle> ParentProperty);
+	void OnMorphTargetComboBoxSelectionChanged(TSharedPtr<FString> Selection, ESelectInfo::Type SelectInfo, TSharedRef<IPropertyHandle> Property);
+
+	TSharedPtr<FString> PrepareComboboxSelection(const int LODIndex, TArray<UCustomizableObjectNodeObject*>& ParentObjectNodes);
+
+};
diff --git a/Engine/Plugins/Experimental/Mutable/Source/CustomizableObjectEditor/Private/Nodes/CustomizableObjectNodeObject.cpp b/Engine/Plugins/Experimental/Mutable/Source/CustomizableObjectEditor/Private/Nodes/CustomizableObjectNodeObject.cpp
new file mode 100644
--- /dev/null
+++ b/Engine/Plugins/Experimental/Mutable/Source/CustomizableObjectEditor/Private/Nodes/CustomizableObjectNodeObject.cpp
@@ -0,0 +1,258 @@
+// Copyright Epic Games, Inc. All Rights Reserved.
+
+#include "Nodes/CustomizableObjectNodeObject.h"
+#include "CustomizableObjectEditorModule.h"
+#include "../CustomizableObjectCompiler.h"
+#include "CustomizableObject.h"
+#include "GraphTraversal.h"
+#include "Logging/MessageLog.h"
+
+#define LOCTEXT_NAMESPACE "CustomizableObjectEditor"
+
+
+UCustomizableObjectNodeObject::UCustomizableObjectNodeObject()
+	: Super()
+{
+	bIsBase = true;
+
+	ObjectName = "Unnamed Object";
+	NumLODs = 1;
+	Identifier.Invalidate();
+
+	if (!Identifier.IsValid())
+	{
+		Identifier = FGuid::NewGuid();
+		IdentifierVerification = Identifier;
+	}
+}
+
+
+void UCustomizableObjectNodeObject::PostEditChangeProperty(FPropertyChangedEvent& PropertyChangedEvent)
+{
+	if (!Identifier.IsValid())
+	{
+		Identifier = FGuid::NewGuid();
+	}
+
+	// Update the cached flag in the main object
+	UCustomizableObject* CustomizableObject = Cast<UCustomizableObject>( GetCustomizableObjectGraph()->GetOuter() );
+	if (CustomizableObject)
+	{
+		CustomizableObject->bIsChildObject = ParentObject != nullptr;
+	}
+
+	FProperty* PropertyThatChanged = PropertyChangedEvent.Property;
+	if ( PropertyThatChanged && PropertyThatChanged->GetName() == TEXT("NumLODs"))
+	{
+		NumLODs = FMath::Clamp(NumLODs, 1, 64);
+
+		ReconstructNode();
+	}
+	
+	Super::PostEditChangeProperty(PropertyChangedEvent);
+}
+
+
+void UCustomizableObjectNodeObject::AllocateDefaultPins(UCustomizableObjectNodeRemapPins* RemapPins)
+{
+	const UEdGraphSchema_CustomizableObject* Schema = GetDefault<UEdGraphSchema_CustomizableObject>();
+
+	for (int32 i = 0; i < NumLODs; ++i)
+	{
+		FString LODName = FString::Printf( TEXT("LOD %d "), i );
+
+		UEdGraphPin* Pin = CustomCreatePin(EGPD_Input, Schema->PC_Material, FName(*LODName), true);
+		Pin->bDefaultValueIsIgnored = true;
+	}
+
+	UEdGraphPin* ChildrenPin = CustomCreatePin(EGPD_Input, Schema->PC_Object, FName("Children"), true );
+	ChildrenPin->bDefaultValueIsIgnored = true;
+
+	UEdGraphPin* OutputPin = CustomCreatePin(EGPD_Output, Schema->PC_Object, FName("Object"));
+
+	if (bIsBase)
+	{
+		OutputPin->bHidden = true;
+	}
+}
+
+
+FText UCustomizableObjectNodeObject::GetNodeTitle(ENodeTitleType::Type TitleType) const
+{
+	FFormatNamedArguments Args;	
+	Args.Add(TEXT("ObjectName"), FText::FromString(ObjectName) );
+
+
+	if (TitleType == ENodeTitleType::ListView)
+	{
+		if (!bIsBase)
+		{						
+			return FText::Format(LOCTEXT("Child_Object_Title_List_Deprecated", "{ObjectName} - Child Object (Deprecated)"), Args);
+		}
+		else
+		{
+			return FText::Format(LOCTEXT("Base_Object_Title_List", "{ObjectName} - Base Object"), Args);
+		}
+	}
+	else
+	{
+		if (!bIsBase)
+		{
+			return FText::Format(LOCTEXT("Child_Object_Title_Deprecated", "{ObjectName}\nChild Object (Deprecated)"), Args);
+		}
+		else
+		{
+			return FText::Format(LOCTEXT("Base_Object_Title", "{ObjectName}\nBase Object"), Args);
+		}
+	}
+}
+
+
+FLinearColor UCustomizableObjectNodeObject::GetNodeTitleColor() const
+{
+	const UEdGraphSchema_CustomizableObject* Schema = GetDefault<UEdGraphSchema_CustomizableObject>();
+	return Schema->GetPinTypeColor(Schema->PC_Object);
+}
+
+
+void UCustomizableObjectNodeObject::PrepareForCopying()
+{
+	FText Msg(LOCTEXT("Cannot copy object node","There can only be one Customizable Object Node Object element per graph") );
+	FMessageLog MessageLog("Mutable");
+	MessageLog.Notify(Msg, EMessageSeverity::Info, true);
+}
+
+
+int32 UCustomizableObjectNodeObject::GetLOD(UEdGraphPin* Pin) const
+{
+	for (int32 LOD = 0; LOD < GetNumLODPins(); ++LOD)
+	{
+		if (Pin == LODPin(LOD))
+		{
+			return LOD;
+		}
+	}
+
+	return -1;
+}
+
+
+bool UCustomizableObjectNodeObject::CanUserDeleteNode() const
+{
+	return !bIsBase;
+}
+
+
+bool UCustomizableObjectNodeObject::CanDuplicateNode() const
+{
+	return !bIsBase;
+}
+
+
+TArray<UCustomizableObjectNodeMaterial*> UCustomizableObjectNodeObject::GetMaterialNodes(const int LOD) const
+{
+	TArray<UCustomizableObjectNodeMaterial*> Result;
+
+	TQueue<UEdGraphNode*> PotentialCustomizableNodeObjects;
+
+	for (const UEdGraphPin* LinkedPin : FollowInputPinArray(*LODPin(LOD)))
+	{
+		PotentialCustomizableNodeObjects.Enqueue(LinkedPin->GetOwningNode());
+	}
+
+	UEdGraphNode* CurrentElement;
+	while (PotentialCustomizableNodeObjects.Dequeue(CurrentElement))
+	{
+		if (UCustomizableObjectNodeMaterial* CurrentMaterialNode = Cast<UCustomizableObjectNodeMaterial>(CurrentElement))
+		{
+			Result.Add(CurrentMaterialNode);
+		} 
+		else if (UCustomizableObjectNodeMaterialVariation* CurrentMaterialVariationNode = Cast<UCustomizableObjectNodeMaterialVariation>(CurrentElement))
+		{
+			// Case of material variation. It's not a material, but a node that further references any material, add all its inputs that could be a material
+			for (int numMaterialPin = 0; numMaterialPin < CurrentMaterialVariationNode->GetNumVariations(); ++numMaterialPin)
+			{
+				const UEdGraphPin* VariationPin = CurrentMaterialVariationNode->VariationPin(numMaterialPin);
+				for (const UEdGraphPin* LinkedPin : FollowInputPinArray(*VariationPin))
+				{
+					PotentialCustomizableNodeObjects.Enqueue(LinkedPin->GetOwningNode());
+				}
+			}
+
+			const UEdGraphPin* DefaultPin = CurrentMaterialVariationNode->DefaultPin();
+			for (const UEdGraphPin* LinkedPin : FollowInputPinArray(*DefaultPin))
+			{
+				PotentialCustomizableNodeObjects.Enqueue(LinkedPin->GetOwningNode());
+			}
+		} 
+	}
+
+	return Result;
+}
+
+
+void UCustomizableObjectNodeObject::PostBackwardsCompatibleFixup()
+{
+	Super::PostBackwardsCompatibleFixup();
+
+	const bool bJustBuiltCoFlag = !GetAllNonOrphanPins().Num() && !NodeGuid.IsValid();
+	if (!bJustBuiltCoFlag && Identifier == IdentifierVerification)
+	{
+		FCustomizableObjectEditorLogger::CreateLog(LOCTEXT("ResaveNode","Please re-save this Customizable Object to avoid binary differences when packaging"))
+		.Severity(EMessageSeverity::Warning)
+		.Node(*this)
+		.Log();
+	}
+}
+
+void UCustomizableObjectNodeObject::PostPasteNode()
+{
+	Super::PostPasteNode();
+
+	Identifier = FGuid::NewGuid();
+}
+
+void UCustomizableObjectNodeObject::PostDuplicate(bool bDuplicateForPIE)
+{
+	Super::PostDuplicate(bDuplicateForPIE);
+
+	Identifier = FGuid::NewGuid();
+}
+
+void UCustomizableObjectNodeObject::SetParentObject(UCustomizableObject* CustomizableParentObject)
+{
+	if (CustomizableParentObject != GetGraphEditor()->GetCustomizableObject())
+	{
+		ParentObject = CustomizableParentObject;
+
+		// Update the cached flag in the main object
+		UCustomizableObject* CustomizableObject = Cast<UCustomizableObject>(GetCustomizableObjectGraph()->GetOuter());
+		if (CustomizableObject)
+		{
+			CustomizableObject->bIsChildObject = ParentObject != nullptr;
+
+			TSharedPtr<ICustomizableObjectEditor> Editor = GetGraphEditor();
+
+			if (Editor.IsValid())
+			{
+				Editor->UpdateObjectProperties();
+			}
+		}
+	}
+}
+
+
+FText UCustomizableObjectNodeObject::GetTooltipText() const
+{
+	return LOCTEXT("Base_Object_Tooltip",
+	"As root object: Defines a customizable object root, its basic properties and its relationship with descendant Customizable Objects.\n\nAs a child object: Defines a Customizable Object children outside of the parent asset, to ease organization of medium and large\nCustomizable Objects. (Functionally equivalent to the Child Object Node.)");
+}
+
+
+bool UCustomizableObjectNodeObject::IsSingleOutputNode() const
+{
+	return true;
+}
+
+
+#undef LOCTEXT_NAMESPACE
diff --git a/Engine/Plugins/Experimental/Mutable/Source/CustomizableObjectEditor/Private/Nodes/CustomizableObjectNodeObjectDetails.cpp b/Engine/Plugins/Experimental/Mutable/Source/CustomizableObjectEditor/Private/Nodes/CustomizableObjectNodeObjectDetails.cpp
new file mode 100644
--- /dev/null
+++ b/Engine/Plugins/Experimental/Mutable/Source/CustomizableObjectEditor/Private/Nodes/CustomizableObjectNodeObjectDetails.cpp
@@ -0,0 +1,261 @@
+// Copyright Epic Games, Inc. All Rights Reserved.
+
+#include "CustomizableObjectNodeObjectDetails.h"
+
+#include "CustomizableObjectEditorModule.h"
+#include "PropertyCustomizationHelpers.h"
+#include "CustomizableObjectEditorUtilities.h"
+#include "CustomizableObject.h"
+#include "Nodes/CustomizableObjectNodeObjectGroup.h"
+
+#include "ScopedTransaction.h"
+#include "DetailLayoutBuilder.h"
+#include "DetailCategoryBuilder.h"
+#include "LevelEditor.h"
+#include "LevelEditorActions.h"
+
+#include "Widgets/Input/SButton.h"
+#include "Widgets/Text/STextBlock.h"
+#include "Widgets/Input/STextComboBox.h"
+
+#include "Runtime/Launch/Resources/Version.h"
+
+#include "SCustomizableObjectNodeObjectStatesView.h"
+#include "SCustomizableObjectNodeObjectRTMorphTargetOverride.h"
+
+
+#define LOCTEXT_NAMESPACE "CustomizableObjectDetails"
+
+
+TSharedRef<IDetailCustomization> FCustomizableObjectNodeObjectDetails::MakeInstance()
+{
+	return MakeShareable( new FCustomizableObjectNodeObjectDetails );
+}
+
+
+void FCustomizableObjectNodeObjectDetails::CustomizeDetails( IDetailLayoutBuilder& DetailBuilder )
+{
+	Node = nullptr;
+	DetailBuilderPtr = &DetailBuilder;
+
+#if (ENGINE_MAJOR_VERSION == 4 && ENGINE_MINOR_VERSION >= 18) || ENGINE_MAJOR_VERSION >= 5
+	const IDetailsView* DetailsView = DetailBuilder.GetDetailsView();
+	if ( DetailsView->GetSelectedObjects().Num() )
+	{
+		Node = Cast<UCustomizableObjectNodeObject>( DetailsView->GetSelectedObjects()[0].Get() );
+	}
+#else
+	const IDetailsView& DetailsView = DetailBuilder.GetDetailsView();
+	if (DetailsView.GetSelectedObjects().Num())
+	{
+		Node = Cast<UCustomizableObjectNodeObject>(DetailsView.GetSelectedObjects()[0].Get());
+	}
+#endif
+
+	IDetailCategoryBuilder& StatesCategory = DetailBuilder.EditCategory("States");
+	IDetailCategoryBuilder& RuntimeParameters = DetailBuilder.EditCategory("States Runtime Parameters");
+	IDetailCategoryBuilder& ExternalCategory = DetailBuilder.EditCategory("AttachedToExternalObject");
+	IDetailCategoryBuilder& RealTimeMorphTargets = DetailBuilder.EditCategory("RealTime Morph Targets");
+	//StatesCategory.CategoryIcon( "ActorClassIcon.CustomizableObject" );
+
+	if (Node)
+	{
+		TSharedRef<IPropertyHandle> StatesProperty = DetailBuilder.GetProperty("States");
+		StatesProperty->SetOnPropertyValueChanged(FSimpleDelegate::CreateSP(this, &FCustomizableObjectNodeObjectDetails::OnStatesPropertyChanged));
+		
+		// Hinding all the runtime parameters arrays
+		for (int i = 0; i < Node->States.Num(); ++i)
+		{
+			FString VariablePath = FString::Printf(TEXT("States[%d].RuntimeParameters"), i);
+			DetailBuilder.HideProperty(*VariablePath);
+		}
+
+		//Adding StetesView Widget
+		RuntimeParameters.AddCustomRow(LOCTEXT("StatesRuntimeParameters", "Runtime Parameters"))
+		[
+			SAssignNew(StatesViewWidget,SCustomizableObjectNodeObjectSatesView)
+			.Node(Node)
+		];
+
+		TSharedRef<IPropertyHandle> ParentObjectProperty = DetailBuilder.GetProperty("ParentObject");
+
+		DetailBuilder.HideProperty("ParentObjectGroupId");
+		DetailBuilder.HideProperty("ParentObject");
+		GroupNodeComboOptions.Empty();
+
+		if (Node->bIsBase)
+		{
+			ExternalCategory.AddCustomRow(LOCTEXT("FCustomizableObjectNodeObjectDetails", "Blocks"))
+			[
+				SNew(SObjectPropertyEntryBox)
+				.AllowedClass(UCustomizableObject::StaticClass())
+				.OnObjectChanged(this, &FCustomizableObjectNodeObjectDetails::ParentObjectSelectionChanged)
+				.ObjectPath(Node->ParentObject->GetPathName())
+				.ForceVolatile(true)
+			];
+
+			if (Node->ParentObject)
+			{
+				TArray<UCustomizableObjectNodeObjectGroup*> GroupNodes;
+				Node->ParentObject->Source->GetNodesOfClass<UCustomizableObjectNodeObjectGroup>(GroupNodes);
+
+				TSharedPtr<FString> ItemToSelect;
+
+				for (UCustomizableObjectNodeObjectGroup* GroupNode : GroupNodes)
+				{
+					GroupNodeComboOptions.Add(MakeShareable(new FString(GroupNode->GroupName)));
+
+					if (Node->ParentObjectGroupId == GroupNode->NodeGuid)
+					{
+						ItemToSelect = GroupNodeComboOptions.Last();
+					}
+				}
+
+				if (!Node->ParentObjectGroupId.IsValid() && ParentComboOptions.Num() > 0)
+				{
+					ItemToSelect = GroupNodeComboOptions.Last();
+				}
+
+                GroupNodeComboOptions.Sort(CompareNames);
+
+				TSharedRef<IPropertyHandle> ParentProperty = DetailBuilder.GetProperty("ParentObjectGroupId");
+
+				ExternalCategory.AddCustomRow(LOCTEXT("FCustomizableObjectNodeObjectDetails", "Blocks"))
+				[
+					SNew(SProperty, ParentProperty)
+					.ShouldDisplayName(false)
+					.CustomWidget()
+					[
+						SNew(SBorder)
+						.BorderImage(FAppStyle::GetBrush("NoBorder"))
+						.Padding(FMargin(0.0f, 0.0f, 10.0f, 0.0f))
+						[
+							SNew(STextComboBox)
+							.OptionsSource(&GroupNodeComboOptions)
+							.InitiallySelectedItem(ItemToSelect)
+							.OnSelectionChanged(this, &FCustomizableObjectNodeObjectDetails::OnGroupNodeComboBoxSelectionChanged, ParentProperty)
+						]
+					]
+				];
+
+				DetailBuilder.HideProperty("NumMeshComponents");
+			}
+            else
+            {
+                RealTimeMorphTargets.AddCustomRow(
+                    LOCTEXT("FCustomizableObjectNodeObjectMorphTargetsDetails", 
+                            "Realtime Morph Targets Override "))
+                [
+                    SNew(SCustomizableObjectNodeSkeletalMeshRTMorphTargetOverride)
+                    .Node(Node)
+                ];
+			}
+
+		}
+		else
+		{
+			DetailBuilder.HideProperty("NumMeshComponents");
+		}
+
+		
+		/*
+		StatesCategory.AddWidget()
+		[
+			SNew( SVerticalBox )
+			+ SVerticalBox::Slot()
+			.Padding( 2.0f )
+			[
+				SNew( SFilterableDetail, NSLocEdL( "States", "States" ), &StatesCategory )
+				[
+					SAssignNew(StatesTree, STreeView<TSharedPtr< FStateDetailsNode > >)
+						//.Visibility(EVisibility::Collapsed)
+						.SelectionMode(ESelectionMode::Single)
+						.TreeItemsSource( &RootTreeItems )
+						// Called to child items for any given parent item
+						.OnGetChildren( this, &FCustomizableObjectDetails::OnGetChildrenForStateTree )
+						// Generates the actual widget for a tree item
+						.OnGenerateRow( this, &FCustomizableObjectDetails::OnGenerateRowForStateTree ) 
+
+						// Generates the right click menu.
+						//.OnContextMenuOpening(this, &SClassViewer::BuildMenuWidget)
+
+						// Find out when the user selects something in the tree
+						//.OnSelectionChanged( this, &SClassViewer::OnClassViewerSelectionChanged )
+
+						// Allow for some spacing between items with a larger item height.
+						.ItemHeight(20.0f)
+
+						.HeaderRow
+						(
+							SNew(SHeaderRow)
+							.Visibility(EVisibility::Collapsed)
+							+ SHeaderRow::Column(TEXT("State"))
+							.DefaultLabel(LocEdL("CustomizableObjectDetails","State","State"))
+						)
+				]
+			]
+		];
+		*/
+	}
+	else
+	{
+		StatesCategory.AddCustomRow( LOCTEXT("Node","Node") )
+		[
+			SNew( STextBlock )
+			.Text( LOCTEXT( "Node not found", "Node not found" ) )
+		];
+	}
+}
+
+
+void FCustomizableObjectNodeObjectDetails::ParentObjectSelectionChanged(const FAssetData & AssetData)
+{
+	if (Node)
+	{
+		UCustomizableObject* Parent = Cast<UCustomizableObject>(AssetData.GetAsset());
+		Node->SetParentObject(Parent);
+
+		// If set the parent to nullt, invalidate also the reference GUID
+		if (!Parent)
+		{
+			Node->ParentObjectGroupId.Invalidate();
+		}
+	}
+
+	if (DetailBuilderPtr)
+	{
+		DetailBuilderPtr->ForceRefreshDetails();
+	}
+}
+
+
+void FCustomizableObjectNodeObjectDetails::OnGroupNodeComboBoxSelectionChanged(TSharedPtr<FString> Selection, ESelectInfo::Type SelectInfo, TSharedRef<IPropertyHandle> ParentProperty)
+{
+	if (Selection.IsValid() && Node->ParentObject != nullptr)
+	{
+		TArray<UCustomizableObjectNodeObjectGroup*> GroupNodes;
+		Node->ParentObject->Source->GetNodesOfClass<UCustomizableObjectNodeObjectGroup>(GroupNodes);
+
+		for (UCustomizableObjectNodeObjectGroup* GroupNode : GroupNodes)
+		{
+			if (*Selection == GroupNode->GroupName)
+			{
+				const FScopedTransaction Transaction(LOCTEXT("ChangedAttachedToExternalObjectTransaction", "Changed Attached to External Object"));
+				Node->Modify();
+				Node->ParentObjectGroupId = GroupNode->NodeGuid;
+			}
+		}
+	}
+}
+
+
+void FCustomizableObjectNodeObjectDetails::OnStatesPropertyChanged()
+{
+	if (DetailBuilderPtr)
+	{
+		DetailBuilderPtr->ForceRefreshDetails();
+	}
+}
+
+#undef LOCTEXT_NAMESPACE
+
diff --git a/Engine/Plugins/Experimental/Mutable/Source/CustomizableObjectEditor/Private/Nodes/CustomizableObjectNodeObjectDetails.h b/Engine/Plugins/Experimental/Mutable/Source/CustomizableObjectEditor/Private/Nodes/CustomizableObjectNodeObjectDetails.h
new file mode 100644
--- /dev/null
+++ b/Engine/Plugins/Experimental/Mutable/Source/CustomizableObjectEditor/Private/Nodes/CustomizableObjectNodeObjectDetails.h
@@ -0,0 +1,44 @@
+// Copyright Epic Games, Inc. All Rights Reserved.
+
+#pragma once
+
+#include "IDetailCustomization.h"
+
+#include "Types/SlateEnums.h"
+#include "Templates/SharedPointer.h"
+#include "UObject/WeakObjectPtr.h"
+#include "PropertyHandle.h"
+#include "AssetRegistry/AssetData.h"
+
+
+class FCustomizableObjectNodeObjectDetails : public IDetailCustomization
+{
+public:
+	/** Makes a new instance of this detail layout class for a specific detail view requesting it */
+	static TSharedRef<IDetailCustomization> MakeInstance();
+
+	/** ILayoutDetails interface */
+	virtual void CustomizeDetails( IDetailLayoutBuilder& DetailBuilder ) override;
+
+	void ParentObjectSelectionChanged(const FAssetData& AssetData);
+
+private:
+
+	class UCustomizableObjectNodeObject* Node;
+	TArray< TSharedPtr<FString> > GroupNodeComboOptions;
+	IDetailLayoutBuilder* DetailBuilderPtr = nullptr;
+
+	TArray< TSharedPtr<FString> > ParentComboOptions;
+	TArray< class UCustomizableObjectNodeMaterial* > ParentOptionNode;
+
+	void OnGroupNodeComboBoxSelectionChanged(TSharedPtr<FString> Selection, ESelectInfo::Type SelectInfo, TSharedRef<IPropertyHandle> ParentProperty);
+
+	// Refreshes the Details View when the states variable has been updated
+	void OnStatesPropertyChanged();
+
+	FIntPoint GetGridSize() const;
+	TArray<FIntRect> GetBlocks() const;
+
+	// Widget to create the runtime parameters view
+	TSharedPtr< class SCustomizableObjectNodeObjectSatesView > StatesViewWidget;
+};
diff --git a/Engine/Plugins/Experimental/Mutable/Source/CustomizableObjectEditor/Private/Nodes/CustomizableObjectNodeParentedMaterial.cpp b/Engine/Plugins/Experimental/Mutable/Source/CustomizableObjectEditor/Private/Nodes/CustomizableObjectNodeParentedMaterial.cpp
new file mode 100644
--- /dev/null
+++ b/Engine/Plugins/Experimental/Mutable/Source/CustomizableObjectEditor/Private/Nodes/CustomizableObjectNodeParentedMaterial.cpp
@@ -0,0 +1,75 @@
+// Copyright Epic Games, Inc. All Rights Reserved.
+
+#include "Nodes/CustomizableObjectNodeParentedMaterial.h"
+#include "Nodes/CustomizableObjectNodeCopyMaterial.h"
+#include "Nodes/CustomizableObjectNodeObject.h"
+
+
+UCustomizableObjectNodeMaterial* FCustomizableObjectNodeParentedMaterial::GetParentMaterialNode() const
+{
+	if (UCustomizableObjectNode* Node = GetParentNode())
+	{
+		return CastChecked<UCustomizableObjectNodeMaterial>(Node);
+	}
+	else
+	{
+		return nullptr;	
+	}
+}
+
+TArray<UCustomizableObjectNodeMaterial*> FCustomizableObjectNodeParentedMaterial::GetPossibleParentMaterialNodes() const
+{
+	TArray<UCustomizableObjectNodeMaterial*> Result;
+	
+	const UCustomizableObjectNode& Node = GetNode();
+	const int32 LOD = Node.GetLOD();
+
+	if (LOD == -1)
+	{
+		return Result; // Early exit.
+	}
+	
+	TArray<UCustomizableObjectNodeObject*> ParentObjectNodes = Node.GetParentObjectNodes(LOD);
+	for (int32 Index = 1; Index < ParentObjectNodes.Num(); ++Index) // Skip first Object node since we do not support directly material siblings.
+	{
+		const UCustomizableObjectNodeObject* ParentObjectNode = ParentObjectNodes[Index];
+	
+		for (UCustomizableObjectNodeMaterial* MaterialNode : ParentObjectNode->GetMaterialNodes(LOD))
+		{
+			if (UCustomizableObjectNodeCopyMaterial* TypedCopyMaterialNode = Cast<UCustomizableObjectNodeCopyMaterial>(MaterialNode))
+			{
+				if (TypedCopyMaterialNode->GetMaterialNode())
+				{
+					Result.Add(TypedCopyMaterialNode);
+				}
+			}
+			else if (UCustomizableObjectNodeMaterial* TypedMaterialNode = Cast<UCustomizableObjectNodeMaterial>(MaterialNode))
+			{
+				Result.Add(TypedMaterialNode);
+			}
+		}
+	}
+	
+	return Result;
+}
+
+
+UCustomizableObjectNodeMaterial* FCustomizableObjectNodeParentedMaterial::GetParentMaterialNodeIfPath() const
+{
+	UCustomizableObjectNodeMaterial* ParentMaterialNode = GetParentMaterialNode();
+	
+	if (GetPossibleParentMaterialNodes().Contains(ParentMaterialNode)) // There is a path to the parent material
+	{
+		return ParentMaterialNode;
+	}
+	else
+	{
+		return nullptr;
+	}
+}
+
+
+const UCustomizableObjectNode& FCustomizableObjectNodeParentedMaterial::GetNode() const
+{
+	return const_cast<FCustomizableObjectNodeParentedMaterial*>(this)->GetNode();
+}
diff --git a/Engine/Plugins/Experimental/Mutable/Source/CustomizableObjectEditor/Private/Nodes/CustomizableObjectNodeParentedMaterialDetails.cpp b/Engine/Plugins/Experimental/Mutable/Source/CustomizableObjectEditor/Private/Nodes/CustomizableObjectNodeParentedMaterialDetails.cpp
new file mode 100644
--- /dev/null
+++ b/Engine/Plugins/Experimental/Mutable/Source/CustomizableObjectEditor/Private/Nodes/CustomizableObjectNodeParentedMaterialDetails.cpp
@@ -0,0 +1,210 @@
+// Copyright Epic Games, Inc. All Rights Reserved.
+
+#include "CustomizableObjectNodeParentedMaterialDetails.h"
+
+#include "CustomizableObjectEditorModule.h"
+#include "PropertyCustomizationHelpers.h"
+#include "CustomizableObjectEditorUtilities.h"
+
+#include "DetailLayoutBuilder.h"
+#include "DetailCategoryBuilder.h"
+#include "LevelEditor.h"
+#include "LevelEditorActions.h"
+#include "ScopedTransaction.h"
+
+#include "Widgets/Text/STextBlock.h"
+#include "Widgets/Input/STextComboBox.h"
+#include "Templates/SharedPointer.h"
+#include "UObject/WeakObjectPtr.h"
+
+#include "Engine/SkeletalMesh.h"
+#include "Engine/StaticMesh.h"
+
+#include "Runtime/Launch/Resources/Version.h"
+
+
+#define LOCTEXT_NAMESPACE "CustomizableObjectDetails"
+
+
+TSharedRef<IDetailCustomization> FCustomizableObjectNodeParentedMaterialDetails::MakeInstance()
+{
+	return MakeShareable(new FCustomizableObjectNodeParentedMaterialDetails);
+}
+
+
+void FCustomizableObjectNodeParentedMaterialDetails::CustomizeDetails(IDetailLayoutBuilder& DetailBuilder)
+{
+	Node = nullptr;
+	NodeParentedMaterial = nullptr;
+
+	const TArray<TWeakObjectPtr<UObject>>& SelectedObjects = DetailBuilder.GetDetailsView()->GetSelectedObjects();
+	if (SelectedObjects.Num())
+	{
+		Node = Cast<UCustomizableObjectNode>(SelectedObjects[0].Get());
+
+		if (UCustomizableObjectNodeEditMaterialBase* NodeEditMaterial = Cast<UCustomizableObjectNodeEditMaterialBase>(SelectedObjects[0].Get()))
+		{
+			NodeParentedMaterial = NodeEditMaterial;
+		}
+		else if (UCustomizableObjectNodeExtendMaterial* NodeExtendMaterial = Cast<UCustomizableObjectNodeExtendMaterial>(SelectedObjects[0].Get()))
+		{
+			NodeParentedMaterial = NodeExtendMaterial;
+		}
+	}
+
+	IDetailCategoryBuilder& BlocksCategory = DetailBuilder.EditCategory("Parent");
+
+
+	if (NodeParentedMaterial)
+	{
+		// Get parent materials
+		TArray<UCustomizableObjectNodeMaterial*> ParentMaterialBaseNodes = NodeParentedMaterial->GetPossibleParentMaterialNodes();
+
+		ParentMaterialOptionReferences.Empty();
+		for (UCustomizableObjectNodeMaterial* ParentMaterial : ParentMaterialBaseNodes)
+		{
+			UCustomizableObject* Object = Cast<UCustomizableObject>(ParentMaterial->GetGraph()->GetOuter());
+			FMaterialReference MaterialReference = { Object, ParentMaterial->NodeGuid };
+
+			ParentMaterialOptionReferences.Add(MaterialReference);
+		}
+
+		// Get option names
+		ParentMaterialOptionNames = GetComboBoxNames(ParentMaterialBaseNodes);
+		check(ParentMaterialOptionReferences.Num() == ParentMaterialOptionNames.Num());
+
+		// Sort items (selection sort)
+		const int Num = ParentMaterialOptionReferences.Num();
+		for (int32 I = 0; I < Num - 1; ++I)
+		{
+			int32 MinIndex = I;
+			for (int32 J = I + 1; J < Num;  J++)
+			{
+				if (*ParentMaterialOptionNames[I] < *ParentMaterialOptionNames[MinIndex])
+				{
+					MinIndex = J;
+				}
+			}
+
+			ParentMaterialOptionNames.Swap(MinIndex, I);
+			ParentMaterialOptionReferences.Swap(MinIndex, I);
+		}
+
+		// Select item
+		TSharedPtr<FString> SelectedItem;
+		for (int32 Option = 0; Option < ParentMaterialBaseNodes.Num(); ++Option)
+		{
+			if (ParentMaterialBaseNodes[Option] == NodeParentedMaterial->GetParentMaterialNode())
+			{
+				SelectedItem = ParentMaterialOptionNames[Option];
+				break;
+			}
+		}
+
+		TSharedRef<IPropertyHandle> ParentProperty = DetailBuilder.GetProperty("ParentMaterialObject");
+
+		BlocksCategory.AddCustomRow(LOCTEXT("FCustomizableObjectNodeParentedMaterialDetails", "Blocks"))
+		[
+			SNew(SHorizontalBox)
+			+ SHorizontalBox::Slot()
+			.AutoWidth()
+			.Padding(0.0f, 5.0f, 0.0f, 0.0f)
+			[
+				SNew(STextBlock)
+				.Text(LOCTEXT("MaterialText", "Material: "))
+			]
+			+ SHorizontalBox::Slot()
+			[
+				SNew(SProperty, ParentProperty)
+				.ShouldDisplayName(false)
+				.CustomWidget()
+				[
+					SNew(SBorder)
+					.BorderImage(FAppStyle::GetBrush("NoBorder"))
+					.Padding(FMargin(0.0f, 0.0f, 10.0f, 0.0f))
+					[
+						SNew(STextComboBox)
+						.OptionsSource(&ParentMaterialOptionNames)
+						.InitiallySelectedItem(SelectedItem)
+						.OnSelectionChanged(this, &FCustomizableObjectNodeParentedMaterialDetails::OnParentComboBoxSelectionChanged, ParentProperty)
+					]
+				]
+			]
+		];
+	}
+	else
+	{
+		BlocksCategory.AddCustomRow(LOCTEXT("FCustomizableObjectNodeEditMaterialBaseDetails", "Node"))
+		[
+			SNew(STextBlock)
+			.Text(LOCTEXT("Node not found", "Node not found"))
+		];
+	}
+}
+
+
+void FCustomizableObjectNodeParentedMaterialDetails::OnParentComboBoxSelectionChanged(TSharedPtr<FString> Selection, ESelectInfo::Type SelectInfo, TSharedRef<IPropertyHandle> ParentProperty)
+{
+	check(ParentMaterialOptionNames.Num() == ParentMaterialOptionReferences.Num());
+	
+	for (int32 OptionIndex = 0; OptionIndex < ParentMaterialOptionNames.Num(); ++OptionIndex)
+	{
+		if (ParentMaterialOptionNames[OptionIndex] == Selection)
+		{
+			const FScopedTransaction Transaction(LOCTEXT("ChangedParentMaterialTransaction", "Changed Parent Material"));
+			Node->Modify();
+			NodeParentedMaterial->SetParentNode(ParentMaterialOptionReferences[OptionIndex].Object.Get(), ParentMaterialOptionReferences[OptionIndex].Id);
+			break;
+		}
+	}
+}
+
+
+TArray<TSharedPtr<FString>> FCustomizableObjectNodeParentedMaterialDetails::GetComboBoxNames(const TArray<UCustomizableObjectNodeMaterial*>& ParentMaterialNodes) const
+{
+	TArray<TSharedPtr<FString>> Result;
+
+	for (UCustomizableObjectNodeMaterial* ParentMaterialNode : ParentMaterialNodes)
+	{
+		const FString ObjectName = ParentMaterialNode->GetParentObjectNodes(ParentMaterialNode->GetLOD())[0]->ObjectName; // Since ParentMaterialNode comes from ParentMaterialNodes, it will always have at least a ParentObjectNode.
+		FString Name = ObjectName + TEXT(" - ");
+
+		if (UCustomizableObjectNodeCopyMaterial* CopyMaterialNode = Cast<UCustomizableObjectNodeCopyMaterial>(ParentMaterialNode))
+		{
+			Name += FString(" [Copy Material] (") + GetComboBoxParentMaterialName(CopyMaterialNode) + FString(" + ");
+
+			const UCustomizableObjectNodeSkeletalMesh* NodeSkeletalMesh = CopyMaterialNode->GetMeshNode();
+			if (NodeSkeletalMesh && NodeSkeletalMesh->SkeletalMesh)
+			{
+				Name += NodeSkeletalMesh->SkeletalMesh->GetName();
+			}
+			else
+			{
+				Name += FString("[Missing Skeletal Mesh]");
+			}
+
+			Name += FString(")");
+		}
+		if (UCustomizableObjectNodeMaterial* MaterialNode = Cast<UCustomizableObjectNodeMaterial>(ParentMaterialNode))
+		{
+			Name += GetComboBoxParentMaterialName(MaterialNode);
+		}
+		else
+		{
+			check(false); // Error. All inputs must have a valid name.
+		}
+
+		Result.Add(MakeShareable(new FString(Name)));
+	}
+
+	return Result;
+}
+
+
+FString FCustomizableObjectNodeParentedMaterialDetails::GetComboBoxParentMaterialName(const UCustomizableObjectNodeMaterial* ParentMaterial) const
+{
+	return ParentMaterial->Material ? ParentMaterial->Material->GetName() : FString("[Missing material]");
+}
+
+
+#undef LOCTEXT_NAMESPACE
diff --git a/Engine/Plugins/Experimental/Mutable/Source/CustomizableObjectEditor/Private/Nodes/CustomizableObjectNodeParentedMaterialDetails.h b/Engine/Plugins/Experimental/Mutable/Source/CustomizableObjectEditor/Private/Nodes/CustomizableObjectNodeParentedMaterialDetails.h
new file mode 100644
--- /dev/null
+++ b/Engine/Plugins/Experimental/Mutable/Source/CustomizableObjectEditor/Private/Nodes/CustomizableObjectNodeParentedMaterialDetails.h
@@ -0,0 +1,48 @@
+// Copyright Epic Games, Inc. All Rights Reserved.
+
+#pragma once
+
+#include "IDetailCustomization.h"
+
+#include "CustomizableObject.h"	
+#include "Templates/SharedPointer.h"
+#include "UObject/WeakObjectPtr.h"
+#include "Misc/Guid.h"
+#include "Types/SlateEnums.h"
+#include "PropertyHandle.h"
+#include "CustomizableObjectEditorModule.h"
+#include "Nodes/CustomizableObjectNode.h"
+#include "Nodes/CustomizableObjectNodeParentedMaterial.h"
+
+class FCustomizableObjectNodeParentedMaterialDetails : public IDetailCustomization
+{
+public:
+	/** Makes a new instance of this detail layout class for a specific detail view requesting it */
+	static TSharedRef<IDetailCustomization> MakeInstance();
+
+	// ILayoutDetails interface
+	void CustomizeDetails( IDetailLayoutBuilder& DetailBuilder ) override;
+
+protected:
+	UCustomizableObjectNode* Node = nullptr;
+	FCustomizableObjectNodeParentedMaterial* NodeParentedMaterial = nullptr;
+
+	struct FMaterialReference
+	{
+		TWeakObjectPtr<UCustomizableObject> Object;
+		FGuid Id;
+	};
+	TArray<FMaterialReference> ParentMaterialOptionReferences;
+	TArray<TSharedPtr<FString>> ParentMaterialOptionNames;
+
+	/** Given a list of parent materials, get their ComboBox option entry names. 
+	 *
+	 * All ParentMaterialNodes pointers must be valid.
+	 */
+	TArray<TSharedPtr<FString>> GetComboBoxNames(const TArray<UCustomizableObjectNodeMaterial*>& ParentMaterialNodes) const;
+
+	/** Return a fornamted name for the parent Material. */
+	FString GetComboBoxParentMaterialName(const UCustomizableObjectNodeMaterial* ParentMaterial) const;
+
+	virtual void OnParentComboBoxSelectionChanged(TSharedPtr<FString> Selection, ESelectInfo::Type SelectInfo, TSharedRef<IPropertyHandle> ParentProperty);
+};
diff --git a/Engine/Plugins/Experimental/Mutable/Source/CustomizableObjectEditor/Private/Nodes/CustomizableObjectNodeParentedNode.cpp b/Engine/Plugins/Experimental/Mutable/Source/CustomizableObjectEditor/Private/Nodes/CustomizableObjectNodeParentedNode.cpp
new file mode 100644
--- /dev/null
+++ b/Engine/Plugins/Experimental/Mutable/Source/CustomizableObjectEditor/Private/Nodes/CustomizableObjectNodeParentedNode.cpp
@@ -0,0 +1,16 @@
+// Copyright Epic Games, Inc. All Rights Reserved.
+
+#include "Nodes/CustomizableObjectNodeParentedNode.h"
+
+
+void ICustomizableObjectNodeParentedNode::SetParentNode(UCustomizableObject* Object, FGuid NodeId)
+{
+	SaveParentNode(Object, NodeId);
+}
+
+
+UCustomizableObjectNode* ICustomizableObjectNodeParentedNode::GetParentNode() const
+{
+	return GetCustomizableObjectExternalNode<UCustomizableObjectNode>(GetParentObject(), GetParentNodeId());
+}
+
diff --git a/Engine/Plugins/Experimental/Mutable/Source/CustomizableObjectEditor/Private/Nodes/CustomizableObjectNodePinViewerDetails.cpp b/Engine/Plugins/Experimental/Mutable/Source/CustomizableObjectEditor/Private/Nodes/CustomizableObjectNodePinViewerDetails.cpp
new file mode 100644
--- /dev/null
+++ b/Engine/Plugins/Experimental/Mutable/Source/CustomizableObjectEditor/Private/Nodes/CustomizableObjectNodePinViewerDetails.cpp
@@ -0,0 +1,23 @@
+// Copyright Epic Games, Inc. All Rights Reserved.
+
+#include "CustomizableObjectNodePinViewerDetails.h"
+
+#include "DetailLayoutBuilder.h"
+
+#include "PinViewer/SPinViewer.h"
+
+#define LOCTEXT_NAMESPACE "CustomizableObjectEditor"
+
+
+TSharedRef<IDetailCustomization> FCustomizableObjectNodePinViewerDetails::MakeInstance()
+{
+	return MakeShareable(new FCustomizableObjectNodePinViewerDetails);
+}
+
+
+void FCustomizableObjectNodePinViewerDetails::CustomizeDetails(IDetailLayoutBuilder& DetailBuilder)
+{
+	PinViewerAttachToDetailCustomization(DetailBuilder);
+}
+
+#undef LOCTEXT_NAMESPACE
\ No newline at end of file
diff --git a/Engine/Plugins/Experimental/Mutable/Source/CustomizableObjectEditor/Private/Nodes/CustomizableObjectNodePinViewerDetails.h b/Engine/Plugins/Experimental/Mutable/Source/CustomizableObjectEditor/Private/Nodes/CustomizableObjectNodePinViewerDetails.h
new file mode 100644
--- /dev/null
+++ b/Engine/Plugins/Experimental/Mutable/Source/CustomizableObjectEditor/Private/Nodes/CustomizableObjectNodePinViewerDetails.h
@@ -0,0 +1,18 @@
+// Copyright Epic Games, Inc. All Rights Reserved.
+
+#pragma once
+
+#include "IDetailCustomization.h"
+#include "Templates/SharedPointer.h"
+
+
+/** Simple Detail Customization which includes the PinViewer. */
+class FCustomizableObjectNodePinViewerDetails : public IDetailCustomization
+{
+public:
+	/** Makes a new instance of this detail layout class for a specific detail view requesting it */
+	static TSharedRef<IDetailCustomization> MakeInstance();
+
+	// IDetailCustomization interface
+	virtual void CustomizeDetails(IDetailLayoutBuilder& DetailBuilder) override;
+};
diff --git a/Engine/Plugins/Experimental/Mutable/Source/CustomizableObjectEditor/Private/Nodes/CustomizableObjectNodeProjectorConstant.cpp b/Engine/Plugins/Experimental/Mutable/Source/CustomizableObjectEditor/Private/Nodes/CustomizableObjectNodeProjectorConstant.cpp
new file mode 100644
--- /dev/null
+++ b/Engine/Plugins/Experimental/Mutable/Source/CustomizableObjectEditor/Private/Nodes/CustomizableObjectNodeProjectorConstant.cpp
@@ -0,0 +1,84 @@
+// Copyright Epic Games, Inc. All Rights Reserved.
+
+#include "Nodes/CustomizableObjectNodeProjectorConstant.h"
+#include "Components/SkeletalMeshComponent.h"
+#include "CustomizableObjectEditorModule.h"
+#include "UnrealEditorPortabilityHelpers.h"
+
+#define LOCTEXT_NAMESPACE "CustomizableObjectEditor"
+
+
+UCustomizableObjectNodeProjectorConstant::UCustomizableObjectNodeProjectorConstant()
+	: Super()
+	, ProjectionType(ECustomizableObjectProjectorType::Planar)
+	, ProjectionAngle(360.0f)
+	, BoneComboBoxLocation(FVector::ZeroVector)
+	, BoneComboBoxForwardDirection(FVector::ZeroVector)
+	, BoneComboBoxUpDirection(FVector::ZeroVector)
+	, ParameterSetModified(-1)
+{
+
+}
+
+
+void UCustomizableObjectNodeProjectorConstant::PostEditChangeProperty(FPropertyChangedEvent& PropertyChangedEvent)
+{
+	FProperty* PropertyThatChanged = PropertyChangedEvent.Property;
+	//if ( PropertyThatChanged && PropertyThatChanged->GetName() == TEXT("NumLODs") )
+	{
+	}
+
+	const FName PropertyName = Helper_GetPropertyName(PropertyChangedEvent);
+
+	if (PropertyName == "ProjectionType")
+	{
+		ParameterSetModified = 0;
+	}
+	else if (PropertyName == "ProjectionAngle")
+	{
+		Value.Angle = FMath::DegreesToRadians(ProjectionAngle);
+		ParameterSetModified = 1;
+	}
+	else if (PropertyName == "ProjectorBone")
+	{
+		Value.Position = (FVector3f)BoneComboBoxLocation;
+		Value.Direction = (FVector3f)BoneComboBoxForwardDirection;
+		Value.Up = (FVector3f)BoneComboBoxUpDirection;
+		ParameterSetModified = 2;
+	}
+
+	Super::PostEditChangeProperty(PropertyChangedEvent);
+}
+
+
+void UCustomizableObjectNodeProjectorConstant::AllocateDefaultPins(UCustomizableObjectNodeRemapPins* RemapPins)
+{
+	const UEdGraphSchema_CustomizableObject* Schema = GetDefault<UEdGraphSchema_CustomizableObject>();
+
+	FString PinName = TEXT("Value");
+	UEdGraphPin* ValuePin = CustomCreatePin(EGPD_Output, Schema->PC_Projector, FName(*PinName));
+	ValuePin->bDefaultValueIsIgnored = true;
+}
+
+
+FText UCustomizableObjectNodeProjectorConstant::GetNodeTitle(ENodeTitleType::Type TitleType) const
+{
+	return LOCTEXT("Projector_Constant", "Projector Constant");
+}
+
+
+FLinearColor UCustomizableObjectNodeProjectorConstant::GetNodeTitleColor() const
+{
+	const UEdGraphSchema_CustomizableObject* Schema = GetDefault<UEdGraphSchema_CustomizableObject>();
+	return Schema->GetPinTypeColor(Schema->PC_Projector);
+}
+
+
+FText UCustomizableObjectNodeProjectorConstant::GetTooltipText() const
+{
+		return LOCTEXT("Projector_Constant_Tooltip",
+			"Defines a constant projector.It can't move, scale or rotate at runtime. The texture that is projected can still be changed, depending on the configuration of the other inputs of the texture project node that is connected to the projector constant.");
+}
+
+
+#undef LOCTEXT_NAMESPACE
diff --git a/Engine/Plugins/Experimental/Mutable/Source/CustomizableObjectEditor/Private/Nodes/CustomizableObjectNodeProjectorParameter.cpp b/Engine/Plugins/Experimental/Mutable/Source/CustomizableObjectEditor/Private/Nodes/CustomizableObjectNodeProjectorParameter.cpp
new file mode 100644
--- /dev/null
+++ b/Engine/Plugins/Experimental/Mutable/Source/CustomizableObjectEditor/Private/Nodes/CustomizableObjectNodeProjectorParameter.cpp
@@ -0,0 +1,65 @@
+// Copyright Epic Games, Inc. All Rights Reserved.
+
+#include "Nodes/CustomizableObjectNodeProjectorParameter.h"
+#include "Components/SkeletalMeshComponent.h"
+#include "CustomizableObjectEditorModule.h"
+#include "UnrealEditorPortabilityHelpers.h"
+
+#define LOCTEXT_NAMESPACE "CustomizableObjectEditor"
+
+void UCustomizableObjectNodeProjectorParameter::PostEditChangeProperty(FPropertyChangedEvent& PropertyChangedEvent)
+{
+	FProperty* PropertyThatChanged = PropertyChangedEvent.Property;
+
+	const FName PropertyName = Helper_GetPropertyName(PropertyChangedEvent);
+
+	if (PropertyName == "ProjectionType")
+	{
+		ParameterSetModified = 0;
+	}
+	else if (PropertyName == "ProjectionAngle")
+	{
+		DefaultValue.Angle = FMath::DegreesToRadians(ProjectionAngle);
+		ParameterSetModified = 1;
+	}
+	else if (PropertyName == "ProjectorBone")
+	{
+		DefaultValue.Position = (FVector3f)BoneComboBoxLocation;
+		DefaultValue.Direction = (FVector3f)BoneComboBoxForwardDirection;
+		DefaultValue.Up = (FVector3f)BoneComboBoxUpDirection;
+		ParameterSetModified = 2;
+	}
+
+	Super::PostEditChangeProperty(PropertyChangedEvent);
+}
+
+
+void UCustomizableObjectNodeProjectorParameter::AllocateDefaultPins(UCustomizableObjectNodeRemapPins* RemapPins)
+{
+	const UEdGraphSchema_CustomizableObject* Schema = GetDefault<UEdGraphSchema_CustomizableObject>();
+
+	FString PinName = TEXT("Value");
+	UEdGraphPin* ValuePin = CustomCreatePin(EGPD_Output, Schema->PC_Projector, FName(*PinName));
+	ValuePin->bDefaultValueIsIgnored = true;
+}
+
+
+FText UCustomizableObjectNodeProjectorParameter::GetNodeTitle(ENodeTitleType::Type TitleType) const
+{
+	return LOCTEXT("Projector_Parameter", "Projector Parameter");
+}
+
+
+FLinearColor UCustomizableObjectNodeProjectorParameter::GetNodeTitleColor() const
+{
+	const UEdGraphSchema_CustomizableObject* Schema = GetDefault<UEdGraphSchema_CustomizableObject>();
+	return Schema->GetPinTypeColor(Schema->PC_Projector);
+}
+
+
+FText UCustomizableObjectNodeProjectorParameter::GetTooltipText() const
+{
+	return LOCTEXT("Projector_Parameter_Tooltip", "Exposes a runtime modifiable projector parameter from the Customizable Object.");
+}
+
+#undef LOCTEXT_NAMESPACE
diff --git a/Engine/Plugins/Experimental/Mutable/Source/CustomizableObjectEditor/Private/Nodes/CustomizableObjectNodeProjectorParameterDetails.cpp b/Engine/Plugins/Experimental/Mutable/Source/CustomizableObjectEditor/Private/Nodes/CustomizableObjectNodeProjectorParameterDetails.cpp
new file mode 100644
--- /dev/null
+++ b/Engine/Plugins/Experimental/Mutable/Source/CustomizableObjectEditor/Private/Nodes/CustomizableObjectNodeProjectorParameterDetails.cpp
@@ -0,0 +1,335 @@
+// Copyright Epic Games, Inc. All Rights Reserved.
+
+#include "CustomizableObjectNodeProjectorParameterDetails.h"
+#include "CustomizableObjectEditorModule.h"
+#include "PropertyCustomizationHelpers.h"
+#include "CustomizableObjectEditorUtilities.h"
+
+#include "DetailLayoutBuilder.h"
+#include "DetailCategoryBuilder.h"
+#include "LevelEditor.h"
+#include "LevelEditorActions.h"
+
+#include "Widgets/Input/SButton.h"
+#include "Widgets/Text/STextBlock.h"
+
+#include "Runtime/Launch/Resources/Version.h"
+#include "Engine/SkeletalMesh.h"
+#include "Widgets/Input/STextComboBox.h"
+#include "Math/Vector4.h"
+#include "Math/Vector.h"
+
+#include "Runtime/Launch/Resources/Version.h"
+#if (ENGINE_MAJOR_VERSION == 4 && ENGINE_MINOR_VERSION > 22) || ENGINE_MAJOR_VERSION >= 5
+#include "HAL/PlatformApplicationMisc.h"
+#endif
+
+
+#define LOCTEXT_NAMESPACE "CustomizableObjectDetails"
+
+
+TSharedRef<IDetailCustomization> FCustomizableObjectNodeProjectorParameterDetails::MakeInstance()
+{
+	return MakeShareable( new FCustomizableObjectNodeProjectorParameterDetails );
+}
+
+
+void FCustomizableObjectNodeProjectorParameterDetails::CustomizeDetails( IDetailLayoutBuilder& DetailBuilder )
+{
+	NodeConstant = nullptr;
+	NodeParameter = nullptr;
+	DetailBuilderPtr = &DetailBuilder;
+
+#if (ENGINE_MAJOR_VERSION == 4 && ENGINE_MINOR_VERSION >= 18) || ENGINE_MAJOR_VERSION >= 5
+	const IDetailsView* DetailsView = DetailBuilder.GetDetailsView();
+	if ( DetailsView->GetSelectedObjects().Num() )
+	{
+		if (DetailsView->GetSelectedObjects()[0].Get()->IsA(UCustomizableObjectNodeProjectorConstant::StaticClass()))
+		{
+			NodeConstant = Cast<UCustomizableObjectNodeProjectorConstant>(DetailsView->GetSelectedObjects()[0].Get());
+		}
+		else if (DetailsView->GetSelectedObjects()[0].Get()->IsA(UCustomizableObjectNodeProjectorParameter::StaticClass()))
+		{
+			NodeParameter = Cast<UCustomizableObjectNodeProjectorParameter>(DetailsView->GetSelectedObjects()[0].Get());
+		}
+	}
+#else
+	const IDetailsView& DetailsView = DetailBuilder.GetDetailsView();
+	if (DetailsView.GetSelectedObjects().Num())
+	{
+		if (DetailsView.GetSelectedObjects()[0].Get()->IsA(UCustomizableObjectNodeProjectorConstant::StaticClass()))
+		{
+			NodeConstant = Cast<UCustomizableObjectNodeProjectorConstant>(DetailsView.GetSelectedObjects()[0].Get());
+		}
+		else if (DetailsView.GetSelectedObjects()[0].Get()->IsA(UCustomizableObjectNodeProjectorParameter::StaticClass()))
+		{
+			NodeParameter = Cast<UCustomizableObjectNodeProjectorParameter>(DetailsView.GetSelectedObjects()[0].Get());
+		}
+	}
+#endif
+
+	IDetailCategoryBuilder& BlocksCategory = DetailBuilder.EditCategory( "Clipboard" );
+	//BlocksCategory.CategoryIcon( "ActorClassIcon.CustomizableObject" );
+
+	if ((NodeConstant != nullptr) || (NodeParameter != nullptr))
+	{
+		BlocksCategory.AddCustomRow( LOCTEXT("FCustomizableObjectNodeProjectorParameterDetails", "Projector Data") )
+		[
+			SNew( SHorizontalBox )
+			+ SHorizontalBox::Slot()
+			[
+				SNew(SButton)
+				.ToolTipText( LOCTEXT( "Copy_Projector", "Copy projector location to clipboard.") )
+				.OnClicked( this, &FCustomizableObjectNodeProjectorParameterDetails::OnProjectorCopyPressed )
+				[
+					SNew(STextBlock)
+					.Text(LOCTEXT("Copy","Copy"))
+				]
+			]
+			+ SHorizontalBox::Slot()
+			[
+				SNew(SButton)
+				.ToolTipText( LOCTEXT( "Paste_Projector", "Paste projector location from clipboard.") )
+				.OnClicked( this, &FCustomizableObjectNodeProjectorParameterDetails::OnProjectorPastePressed )
+				[
+					SNew(STextBlock)
+					.Text(LOCTEXT("Paste","Paste"))
+				]
+			]
+		];
+	}
+	else
+	{
+		BlocksCategory.AddCustomRow( LOCTEXT("Node", "Node") )
+		[
+			SNew( STextBlock )
+			.Text( LOCTEXT( "Node not found", "Node not found" ) )
+		];
+	}
+
+	TSharedPtr<IPropertyHandle> ReferenceSkeletonIndexPropertyHandle = DetailBuilder.GetProperty("ReferenceSkeletonIndex");
+	const FSimpleDelegate OnReferenceSkeletonIndexChangedDelegate =
+		FSimpleDelegate::CreateSP(this, &FCustomizableObjectNodeProjectorParameterDetails::OnReferenceSkeletonIndexChanged);
+	ReferenceSkeletonIndexPropertyHandle->SetOnPropertyValueChanged(OnReferenceSkeletonIndexChangedDelegate);
+
+	IDetailCategoryBuilder& SnapToBoneCategory = DetailBuilder.EditCategory("ProjectorSnapToBone");
+	DetailBuilder.HideProperty("ProjectorBone");
+
+	SkeletalMesh = nullptr;
+	TSharedPtr<FString> BoneToSelect;
+
+	UCustomizableObject* CustomizableObject = nullptr;
+	FName ProjectorBoneName;
+	int32 ReferenceIndex = 0;
+
+	if (NodeConstant != nullptr)
+	{
+		CustomizableObject = Cast<UCustomizableObject>(NodeConstant->GetCustomizableObjectGraph()->GetOuter());
+		ProjectorBoneName = NodeConstant->ProjectorBone;
+		ReferenceIndex = NodeConstant->ReferenceSkeletonIndex;
+	}
+	else if (NodeParameter != nullptr)
+	{
+		CustomizableObject = Cast<UCustomizableObject>(NodeParameter->GetCustomizableObjectGraph()->GetOuter());
+		ProjectorBoneName = NodeParameter->ProjectorBone;
+		ReferenceIndex = NodeParameter->ReferenceSkeletonIndex;
+	}
+
+	//UCustomizableObject* CustomizableObject = Cast<UCustomizableObject>(Node->GetCustomizableObjectGraph()->GetOuter());
+
+	if (CustomizableObject)
+	{
+		SkeletalMesh = CustomizableObject->GetRefSkeletalMesh(ReferenceIndex);
+	}
+
+	if (SkeletalMesh)
+	{
+		BoneComboOptions.Empty();
+			
+		for (int32 i = 0; i < SkeletalMesh->GetRefSkeleton().GetNum(); ++i)
+		{
+			FName BoneName = SkeletalMesh->GetRefSkeleton().GetBoneName(i);
+			BoneComboOptions.Add(MakeShareable(new FString(BoneName.ToString())));
+
+			if (BoneName == ProjectorBoneName)
+			{
+				BoneToSelect = BoneComboOptions.Last();
+			}
+		}
+
+        BoneComboOptions.Sort(CompareNames);
+
+		// Add them to the parent combo box
+		TSharedRef<IPropertyHandle> BoneProperty = DetailBuilder.GetProperty("ProjectorBone");
+
+		SnapToBoneCategory.AddCustomRow(LOCTEXT("UCustomizableObjectNodeProjectorParameter", "Projector Bone"))
+		[
+			SNew(SProperty, BoneProperty)
+			.ShouldDisplayName(false)
+			.CustomWidget()
+			[
+				SNew(SBorder)
+				.BorderImage(FAppStyle::GetBrush("NoBorder"))
+				.Padding(FMargin(0.0f, 0.0f, 10.0f, 0.0f))
+				.HAlign(EHorizontalAlignment::HAlign_Fill)
+				[
+					SNew(SHorizontalBox)
+					+SHorizontalBox::Slot().HAlign(HAlign_Left).VAlign(VAlign_Center)
+					[
+						SNew(STextBlock)
+						.Text( LOCTEXT("FCustomizableObjectNodeRemoveMeshDetails", "Projector Bone"))
+						.Font(IDetailLayoutBuilder::GetDetailFont())
+					]
+					+SHorizontalBox::Slot().HAlign(HAlign_Fill)
+					[
+						SNew(STextComboBox)
+						.OptionsSource(&BoneComboOptions)
+						.InitiallySelectedItem(BoneToSelect)
+						.OnSelectionChanged(this, &FCustomizableObjectNodeProjectorParameterDetails::OnBoneComboBoxSelectionChanged, BoneProperty)
+						.Font(IDetailLayoutBuilder::GetDetailFont())
+					]
+				]
+			]
+		];
+	}
+}
+
+
+FReply FCustomizableObjectNodeProjectorParameterDetails::OnProjectorCopyPressed()
+{
+	FString ExportedText;
+
+	if (NodeParameter)
+	{
+		UScriptStruct* Struct = NodeParameter->DefaultValue.StaticStruct();
+        Struct->ExportText(ExportedText, &NodeParameter->DefaultValue, nullptr, nullptr, (PPF_ExportsNotFullyQualified | PPF_Copy | PPF_Delimited | PPF_IncludeTransient), nullptr);
+	}
+	else if(NodeConstant)
+	{
+		UScriptStruct* Struct = NodeConstant->Value.StaticStruct();
+        Struct->ExportText(ExportedText, &NodeConstant->Value, nullptr, nullptr, (PPF_ExportsNotFullyQualified | PPF_Copy | PPF_Delimited | PPF_IncludeTransient), nullptr);
+	}
+
+#if (ENGINE_MAJOR_VERSION == 4 && ENGINE_MINOR_VERSION > 22) || ENGINE_MAJOR_VERSION >= 5
+	FPlatformApplicationMisc::ClipboardCopy(*ExportedText);
+#else
+	FPlatformMisc::ClipboardCopy(*ExportedText);
+#endif
+
+	return FReply::Handled();
+}
+
+
+FReply FCustomizableObjectNodeProjectorParameterDetails::OnProjectorPastePressed()
+{
+	FString ClipText;
+
+#if (ENGINE_MAJOR_VERSION == 4 && ENGINE_MINOR_VERSION > 22) || ENGINE_MAJOR_VERSION >= 5
+	FPlatformApplicationMisc::ClipboardPaste(ClipText);
+#else
+	FPlatformMisc::ClipboardPaste(ClipText);
+#endif
+
+	TSharedPtr<ICustomizableObjectEditor> Editor = nullptr;
+
+	if (NodeParameter)
+	{
+		UScriptStruct* Struct = NodeParameter->DefaultValue.StaticStruct();
+		Struct->ImportText(*ClipText, &NodeParameter->DefaultValue, nullptr, 0, GLog, GetPathNameSafe(Struct));
+		NodeParameter->ParameterSetModified = 2;
+		Editor = NodeParameter->GetGraphEditor();
+	}
+	else if (NodeConstant)
+	{
+		UScriptStruct* Struct = NodeConstant->Value.StaticStruct();
+		Struct->ImportText(*ClipText, &NodeConstant->Value, nullptr, 0, GLog, GetPathNameSafe(Struct));
+		NodeConstant->ParameterSetModified = 2;
+		Editor = NodeConstant->GetGraphEditor();
+	}
+
+	if (Editor.IsValid())
+	{
+		Editor->UpdateGraphNodeProperties();
+	}
+
+	return FReply::Handled();
+}
+
+void FCustomizableObjectNodeProjectorParameterDetails::OnBoneComboBoxSelectionChanged(TSharedPtr<FString> Selection, ESelectInfo::Type SelectInfo, TSharedRef<IPropertyHandle> BoneProperty)
+{
+	for (int OptionIndex = 0; OptionIndex < BoneComboOptions.Num(); ++OptionIndex)
+	{
+		if (BoneComboOptions[OptionIndex] == Selection)
+		{
+			FVector Location = FVector::ZeroVector;
+			FVector Direction = FVector::ForwardVector;
+
+			if (SkeletalMesh)
+			{
+				const TArray<FTransform>& BoneArray = SkeletalMesh->GetRefSkeleton().GetRefBonePose();
+				int32 ParentIndex = SkeletalMesh->GetRefSkeleton().FindBoneIndex(FName(**Selection));
+
+				FVector ChildLocation = FVector::ForwardVector;
+
+				for (int32 i = 0; i < SkeletalMesh->GetRefSkeleton().GetNum(); ++i)
+				{
+					if (SkeletalMesh->GetRefSkeleton().GetParentIndex(i) == ParentIndex)
+					{
+						ChildLocation = BoneArray[i].TransformPosition(FVector::ZeroVector);
+						break;
+					}
+				}
+
+				while (ParentIndex >= 0)
+				{
+					Location = BoneArray[ParentIndex].TransformPosition(Location);
+					ChildLocation = BoneArray[ParentIndex].TransformPosition(ChildLocation);
+					ParentIndex = SkeletalMesh->GetRefSkeleton().GetParentIndex(ParentIndex);
+				}
+
+				Direction = (ChildLocation - Location).GetSafeNormal();
+
+				if ((Location != FVector::ZeroVector) && (Direction != FVector::ZeroVector))
+				{
+					FVector UpTemp = (FVector::DotProduct(Direction, FVector(0.0f, 0.0f, 1.0f)) > 0.99f) ? FVector(0.1f, 0.1f, 1.0f).GetSafeNormal() : FVector(0.0f, 0.0f, 1.0f);
+					FVector Right = FVector::CrossProduct(UpTemp, Direction);
+					FVector UpFinal = FVector::CrossProduct(Direction, Right);
+
+					if (NodeConstant != nullptr)
+					{
+						NodeConstant->BoneComboBoxLocation = Location;
+						NodeConstant->BoneComboBoxForwardDirection = Direction;
+						NodeConstant->BoneComboBoxUpDirection = UpFinal.GetSafeNormal();
+					}
+					else if (NodeParameter != nullptr)
+					{
+						NodeParameter->BoneComboBoxLocation = Location;
+						NodeParameter->BoneComboBoxForwardDirection = Direction;
+						NodeParameter->BoneComboBoxUpDirection = UpFinal.GetSafeNormal();
+					}
+				}
+			}
+
+			BoneProperty->SetValue(*BoneComboOptions[OptionIndex].Get());
+
+			return;
+		}
+	}
+}
+
+
+void FCustomizableObjectNodeProjectorParameterDetails::OnReferenceSkeletonIndexChanged()
+{
+	if (NodeConstant != nullptr)
+	{
+		NodeConstant->ProjectorBone = FName();
+	}
+	else if (NodeParameter != nullptr)
+	{
+		NodeParameter->ProjectorBone = FName();
+	}
+
+	DetailBuilderPtr->ForceRefreshDetails();
+}
+
+#undef LOCTEXT_NAMESPACE
diff --git a/Engine/Plugins/Experimental/Mutable/Source/CustomizableObjectEditor/Private/Nodes/CustomizableObjectNodeProjectorParameterDetails.h b/Engine/Plugins/Experimental/Mutable/Source/CustomizableObjectEditor/Private/Nodes/CustomizableObjectNodeProjectorParameterDetails.h
new file mode 100644
--- /dev/null
+++ b/Engine/Plugins/Experimental/Mutable/Source/CustomizableObjectEditor/Private/Nodes/CustomizableObjectNodeProjectorParameterDetails.h
@@ -0,0 +1,35 @@
+// Copyright Epic Games, Inc. All Rights Reserved.
+
+#pragma once
+
+#include "IDetailCustomization.h"
+#include "Templates/SharedPointer.h"
+#include "UObject/WeakObjectPtr.h"
+#include "Input/Reply.h"
+#include "PropertyHandle.h"
+
+
+class FCustomizableObjectNodeProjectorParameterDetails : public IDetailCustomization
+{
+public:
+	/** Makes a new instance of this detail layout class for a specific detail view requesting it */
+	static TSharedRef<IDetailCustomization> MakeInstance();
+
+	/** ILayoutDetails interface */
+	virtual void CustomizeDetails( IDetailLayoutBuilder& DetailBuilder ) override;
+
+
+private:
+	TArray< TSharedPtr<FString> > BoneComboOptions;
+
+	class UCustomizableObjectNodeProjectorConstant* NodeConstant;
+	class UCustomizableObjectNodeProjectorParameter* NodeParameter;
+	FReply OnProjectorCopyPressed();
+	FReply OnProjectorPastePressed();
+
+	class USkeletalMesh* SkeletalMesh;
+	class IDetailLayoutBuilder* DetailBuilderPtr;
+
+	void OnBoneComboBoxSelectionChanged(TSharedPtr<FString> Selection, ESelectInfo::Type SelectInfo, TSharedRef<IPropertyHandle> BoneProperty);
+	void OnReferenceSkeletonIndexChanged();
+};
\ No newline at end of file
diff --git a/Engine/Plugins/Experimental/Mutable/Source/CustomizableObjectEditor/Private/Nodes/CustomizableObjectNodeRemoveMesh.cpp b/Engine/Plugins/Experimental/Mutable/Source/CustomizableObjectEditor/Private/Nodes/CustomizableObjectNodeRemoveMesh.cpp
new file mode 100644
--- /dev/null
+++ b/Engine/Plugins/Experimental/Mutable/Source/CustomizableObjectEditor/Private/Nodes/CustomizableObjectNodeRemoveMesh.cpp
@@ -0,0 +1,60 @@
+// Copyright Epic Games, Inc. All Rights Reserved.
+
+#include "Nodes/CustomizableObjectNodeRemoveMesh.h"
+#include "CustomizableObjectEditorModule.h"
+
+#define LOCTEXT_NAMESPACE "CustomizableObjectEditor"
+
+
+void UCustomizableObjectNodeRemoveMesh::AllocateDefaultPins(UCustomizableObjectNodeRemapPins* RemapPins)
+{
+	const UEdGraphSchema_CustomizableObject* Schema = GetDefault<UEdGraphSchema_CustomizableObject>();
+
+	UEdGraphPin* RemoveMeshPin = CustomCreatePin(EGPD_Input, Schema->PC_Mesh, FName("Remove Mesh") );
+	RemoveMeshPin->bDefaultValueIsIgnored = true;
+
+	UEdGraphPin* OutputPin = CustomCreatePin(EGPD_Output, Schema->PC_Material, FName("Material"));
+}
+
+
+FText UCustomizableObjectNodeRemoveMesh::GetNodeTitle(ENodeTitleType::Type TitleType) const
+{
+	return LOCTEXT("Remove_Mesh", "Remove Mesh");
+}
+
+
+FLinearColor UCustomizableObjectNodeRemoveMesh::GetNodeTitleColor() const
+{
+	const UEdGraphSchema_CustomizableObject* Schema = GetDefault<UEdGraphSchema_CustomizableObject>();
+	return Schema->GetPinTypeColor(Schema->PC_Material);
+}
+
+
+void UCustomizableObjectNodeRemoveMesh::PinConnectionListChanged(UEdGraphPin* Pin)
+{
+	Super::PinConnectionListChanged(Pin);
+
+	if (Pin == OutputPin())
+	{
+		TSharedPtr<ICustomizableObjectEditor> Editor = GetGraphEditor();
+
+		if (Editor.IsValid())
+		{
+			Editor->UpdateGraphNodeProperties();
+		}
+	}
+}
+
+FText UCustomizableObjectNodeRemoveMesh::GetTooltipText() const
+{
+	return LOCTEXT("Remove_Mesh_Tooltip",
+	"Removes the faces of a material that are defined only by vertexes shared by the material and the input mesh.It also removes any vertex\nand edge that only define deleted faces, they are not left dangling. If the mesh removed covers at least all the faces included in one or\nmore layout blocs, those blocs are removed, freeing final texture layout space.");
+}
+
+bool UCustomizableObjectNodeRemoveMesh::IsSingleOutputNode() const
+{
+	return true;
+}
+
+
+#undef LOCTEXT_NAMESPACE
diff --git a/Engine/Plugins/Experimental/Mutable/Source/CustomizableObjectEditor/Private/Nodes/CustomizableObjectNodeRemoveMeshBlocks.cpp b/Engine/Plugins/Experimental/Mutable/Source/CustomizableObjectEditor/Private/Nodes/CustomizableObjectNodeRemoveMeshBlocks.cpp
new file mode 100644
--- /dev/null
+++ b/Engine/Plugins/Experimental/Mutable/Source/CustomizableObjectEditor/Private/Nodes/CustomizableObjectNodeRemoveMeshBlocks.cpp
@@ -0,0 +1,173 @@
+// Copyright Epic Games, Inc. All Rights Reserved.
+
+#include "Nodes/CustomizableObjectNodeRemoveMeshBlocks.h"
+#include "CustomizableObjectEditorModule.h"
+
+#define LOCTEXT_NAMESPACE "CustomizableObjectEditor"
+
+
+void UCustomizableObjectNodeRemoveMeshBlocks::BackwardsCompatibleFixup()
+{
+	Super::BackwardsCompatibleFixup();
+	
+	// Convert deprecated node index list to the node id list.
+	if (GetLinkerCustomVersion(FCustomizableObjectCustomVersion::GUID) < FCustomizableObjectCustomVersion::PostLoadToCustomVersion
+		&& BlockIds.Num() < Blocks_DEPRECATED.Num())
+	{
+		if (UCustomizableObjectNodeMaterialBase* ParentMaterialNode = GetParentMaterialNode())
+		{
+			TArray<UCustomizableObjectLayout*> Layouts = ParentMaterialNode->GetLayouts();
+
+			if (!Layouts.IsValidIndex(ParentLayoutIndex))
+			{
+				UE_LOG(LogMutable, Warning, TEXT("[%s] UCustomizableObjectNodeRemoveMeshBlocks refers to an invalid texture layout index %d. Parent node has %d layouts."),
+					*GetOutermost()->GetName(), ParentLayoutIndex, Layouts.Num());
+			}
+			else if (UCustomizableObjectNodeMaterial* ParentMaterial = Cast<UCustomizableObjectNodeMaterial>(ParentMaterialNode))
+			{
+				UCustomizableObjectLayout* Layout = Layouts[ParentLayoutIndex];
+
+				for (int IndexIndex = BlockIds.Num(); IndexIndex < Blocks_DEPRECATED.Num(); ++IndexIndex)
+				{
+					const int BlockIndex = Blocks_DEPRECATED[IndexIndex];
+					if (!Layout->Blocks.IsValidIndex(BlockIndex))
+					{
+						UE_LOG(LogMutable, Warning, TEXT("[%s] UCustomizableObjectNodeRemoveMeshBlocks refers to an invalid layout block index %d. Parent node has %d blocks."),
+							*GetOutermost()->GetName(), BlockIndex, Layout->Blocks.Num());
+
+						continue;
+					}
+					
+					const FGuid Id = Layout->Blocks[BlockIndex].Id;
+					if (!Id.IsValid())
+					{
+						UE_LOG(LogMutable, Warning, TEXT("[%s] UCustomizableObjectNodeRemoveMeshBlocks refers to an valid layout block %d but that block doesn't have an id."),
+							*GetOutermost()->GetName(), BlockIndex);
+
+						continue;
+					}
+
+					BlockIds.Add(Id);
+				}
+			}
+		}
+	}
+}
+
+
+void UCustomizableObjectNodeRemoveMeshBlocks::Serialize(FArchive& Ar)
+{
+	Super::Serialize(Ar);
+
+	Ar.UsingCustomVersion(FCustomizableObjectCustomVersion::GUID);
+}
+
+
+void UCustomizableObjectNodeRemoveMeshBlocks::PostEditChangeProperty(FPropertyChangedEvent& PropertyChangedEvent)
+{
+	FProperty* PropertyThatChanged = PropertyChangedEvent.Property;
+	if (PropertyThatChanged 
+		&& (PropertyThatChanged->GetName() == TEXT("ParentMaterialObject") || PropertyThatChanged->GetName() == TEXT("ParentLayoutIndex")))
+	{
+		// Reset selected blocks
+		BlockIds.Empty();
+		
+		TSharedPtr<ICustomizableObjectEditor> Editor = GetGraphEditor();
+
+		if (Editor.IsValid())
+		{
+			Editor->UpdateGraphNodeProperties();
+		}
+	}
+
+	Super::PostEditChangeProperty(PropertyChangedEvent);
+}
+
+void UCustomizableObjectNodeRemoveMeshBlocks::AllocateDefaultPins(UCustomizableObjectNodeRemapPins* RemapPins)
+{
+	const UEdGraphSchema_CustomizableObject* Schema = GetDefault<UEdGraphSchema_CustomizableObject>();
+
+	UEdGraphPin* OutputPin = CustomCreatePin(EGPD_Output, Schema->PC_Material, FName("Material"));
+}
+
+
+FText UCustomizableObjectNodeRemoveMeshBlocks::GetNodeTitle(ENodeTitleType::Type TitleType) const
+{
+	return LOCTEXT("Remove_Mesh_Blocks", "Remove Mesh Blocks");
+}
+
+
+FLinearColor UCustomizableObjectNodeRemoveMeshBlocks::GetNodeTitleColor() const
+{
+	const UEdGraphSchema_CustomizableObject* Schema = GetDefault<UEdGraphSchema_CustomizableObject>();
+	return Schema->GetPinTypeColor(Schema->PC_Material);
+}
+
+
+void UCustomizableObjectNodeRemoveMeshBlocks::PinConnectionListChanged(UEdGraphPin* Pin)
+{
+	Super::PinConnectionListChanged(Pin);
+
+	if (Pin == OutputPin())
+	{
+		TSharedPtr<ICustomizableObjectEditor> Editor = GetGraphEditor();
+
+		if (Editor.IsValid())
+		{
+			Editor->UpdateGraphNodeProperties();
+		}
+	}
+}
+
+bool UCustomizableObjectNodeRemoveMeshBlocks::IsNodeOutDatedAndNeedsRefresh()
+{
+	bool Result = false;
+
+	UCustomizableObjectNodeMaterialBase* ParentMaterialNode = GetParentMaterialNode();
+	if (ParentMaterialNode)
+	{
+		TArray<UCustomizableObjectLayout*> Layouts = ParentMaterialNode->GetLayouts();
+
+		if (!Layouts.IsValidIndex(ParentLayoutIndex))
+		{
+			Result = true;
+		}
+		else
+		{
+			UCustomizableObjectLayout* Layout = Layouts[ParentLayoutIndex];
+				
+			if (BlockIds.Num() > Layout->Blocks.Num())
+			{
+				Result = true;
+			}
+		}
+	}
+
+	// Remove previous compilation warnings
+	if (!Result && bHasCompilerMessage)
+	{
+		RemoveWarnings();
+		GetGraph()->NotifyGraphChanged();
+	}
+
+	return Result;
+}
+
+FString UCustomizableObjectNodeRemoveMeshBlocks::GetRefreshMessage() const
+{
+	return "Source Layout has changed, layout blocks might have changed. Please Refresh Node to reflect those changes.";
+}
+
+
+FText UCustomizableObjectNodeRemoveMeshBlocks::GetTooltipText() const
+{
+	return LOCTEXT("Remove_Mesh_Blocks_Tooltip", "Remove all the geometry in the chosen layout blocks from a material.");
+}
+
+
+bool UCustomizableObjectNodeRemoveMeshBlocks::IsSingleOutputNode() const
+{
+	return true;
+}
+
+#undef LOCTEXT_NAMESPACE
