[mutable] Move Mutable plugin from the restricted folder to the experimental plugins folder.

diff --git a/Engine/Plugins/Experimental/Mutable/Source/CustomizableObjectEditor/Private/Nodes/CustomizableObjectNodeSkeletalMesh.cpp b/Engine/Plugins/Experimental/Mutable/Source/CustomizableObjectEditor/Private/Nodes/CustomizableObjectNodeSkeletalMesh.cpp
new file mode 100644
--- /dev/null
+++ b/Engine/Plugins/Experimental/Mutable/Source/CustomizableObjectEditor/Private/Nodes/CustomizableObjectNodeSkeletalMesh.cpp
@@ -0,0 +1,828 @@
+// Copyright Epic Games, Inc. All Rights Reserved.
+
+#include "Nodes/CustomizableObjectNodeSkeletalMesh.h"
+#include "Nodes/CustomizableObjectNode.h"
+#include "CustomizableObjectEditorModule.h"
+#include "Engine/SkeletalMesh.h"
+#include "Materials/Material.h"
+#include "UnrealPortabilityHelpers.h"
+#include "UnrealEditorPortabilityHelpers.h"
+#include "NodeFactory.h"
+#include "ISinglePropertyView.h"
+#include "EditorFramework/ThumbnailInfo.h"
+#include "Widgets/Input/SCheckBox.h"
+#include "CustomizableObjectEditorStyle.h"
+#include "CustomizableObjectEditor.h"
+#include "GraphTraversal.h"
+#include "MutableUtils.h"
+#include "RemapPins/CustomizableObjectNodeRemapPinsByNameDefaultPin.h"
+
+
+#define LOCTEXT_NAMESPACE "CustomizableObjectEditor"
+
+
+/** Default node pin configuration pin name (node does not have an skeletal mesh). */
+static const TCHAR* SKELETAL_MESH_PIN_NAME = TEXT("Skeletal Mesh");
+
+
+void UCustomizableObjectNodeSkeletalMesh::PostEditChangeProperty(FPropertyChangedEvent& PropertyChangedEvent)
+{
+	Super::PostEditChangeProperty(PropertyChangedEvent);
+
+	FProperty* PropertyThatChanged = PropertyChangedEvent.Property;
+	if (PropertyThatChanged && PropertyThatChanged->GetName() == TEXT("SkeletalMesh"))
+	{
+		MarkForReconstruct();
+	}
+}
+
+
+void UCustomizableObjectNodeSkeletalMesh::Serialize(FArchive& Ar)
+{
+	Super::Serialize(Ar);
+
+	Ar.UsingCustomVersion(FCustomizableObjectCustomVersion::GUID);
+
+	if (Ar.CustomVer(FCustomizableObjectCustomVersion::GUID) < FCustomizableObjectCustomVersion::PostLoadToCustomVersion)
+	{
+		for (FCustomizableObjectNodeSkeletalMeshLOD& LOD : LODs)
+		{
+			for(FCustomizableObjectNodeSkeletalMeshMaterial& Material : LOD.Materials)
+			{
+				if (Material.MeshPin && !Material.MeshPinRef.Get())
+				{
+					UEdGraphPin* AuxPin = UEdGraphPin::FindPinCreatedFromDeprecatedPin(Material.MeshPin);
+					Material.MeshPinRef.SetPin(AuxPin);
+				}
+
+				if (!Material.LayoutPinsRef.Num())
+				{
+					if (Material.LayoutPins.Num())
+					{
+						for (UEdGraphPin_Deprecated* LayoutPin : Material.LayoutPins)
+						{
+							UEdGraphPin* AuxPin = UEdGraphPin::FindPinCreatedFromDeprecatedPin(LayoutPin);
+							FEdGraphPinReference AuxEdGraphPinReference(AuxPin);
+							Material.LayoutPinsRef.Add(AuxEdGraphPinReference);
+						}
+					}
+					else
+					{
+						FString MaterialLayoutName = Material.Name + " Layout";
+						for (UEdGraphPin* Pin : GetAllNonOrphanPins())
+						{
+							if (Pin
+								&& Pin->Direction == EEdGraphPinDirection::EGPD_Input
+								&& (MaterialLayoutName == Helper_GetPinName(Pin)
+									|| MaterialLayoutName == Pin->PinFriendlyName.ToString()))
+							{
+								FEdGraphPinReference AuxEdGraphPinReference(Pin);
+								Material.LayoutPinsRef.Add(AuxEdGraphPinReference);
+								break;
+							}
+						}
+					}
+				}
+
+				if (!Material.ImagePinsRef.Num())
+				{
+					for (UEdGraphPin_Deprecated* ImagePin : Material.ImagePins)
+					{
+						UEdGraphPin* AuxPin = UEdGraphPin::FindPinCreatedFromDeprecatedPin(ImagePin);
+						FEdGraphPinReference AuxEdGraphPinReference(AuxPin);
+						Material.ImagePinsRef.Add(AuxEdGraphPinReference);
+					}
+				}
+			}
+		}
+	}
+}
+
+
+void UCustomizableObjectNodeSkeletalMesh::AllocateDefaultPins(UCustomizableObjectNodeRemapPins* RemapPins)
+{
+	// Pass information to the remap pins action context
+	if (UCustomizableObjectNodeRemapPinsByNameDefaultPin* RemapPinsCustom = Cast<UCustomizableObjectNodeRemapPinsByNameDefaultPin>(RemapPins))
+	{
+		RemapPinsCustom->DefaultPin = DefaultPin.Get();
+	}
+	
+	const UEdGraphSchema_CustomizableObject* Schema = GetDefault<UEdGraphSchema_CustomizableObject>();
+
+	TArray<FCustomizableObjectNodeSkeletalMeshLOD> OldLODs(LODs);
+	LODs.Empty();
+
+	if (SkeletalMesh)
+	{
+		int NumLODs = Helper_GetLODInfoArray(SkeletalMesh).Num();
+		TMap<FString, TArray<int32>> MaterialNameUseCount; // Stores how many times the material name has been used per LOD
+		TMap<FString, TArray<int32>> ImageNameUseCount; // Stores how many times the image name has been used per LOD
+
+		if (const FSkeletalMeshModel* ImportedModel = Helper_GetImportedModel(SkeletalMesh))
+		{
+			for (int32 LODIndex = 0; LODIndex < NumLODs; ++LODIndex)
+			{
+				LODs.Add(FCustomizableObjectNodeSkeletalMeshLOD());
+
+				// Reset the material array and handle the actual materials from the mesh
+				int32 NumMaterials = ImportedModel->LODModels[LODIndex].Sections.Num();
+
+				FString OldNameForMap;
+
+				for (int32 MaterialIndex = 0; MaterialIndex < NumMaterials; ++MaterialIndex)
+				{
+					LODs[LODIndex].Materials.Add(FCustomizableObjectNodeSkeletalMeshMaterial());
+
+					UMaterialInterface* MaterialInterface = GetMaterialInterfaceFor(LODIndex, MaterialIndex, ImportedModel);
+
+					FString MaterialName = LOCTEXT("Unnamed Material", "Unnamed Material").ToString();
+					if (MaterialInterface)
+					{
+						MaterialName = MaterialInterface->GetName();
+						LODs[LODIndex].Materials[MaterialIndex].Name = MaterialName;
+					}
+
+					FString PrefixId; // Prefix is empty if only 1 LOD and 1 material
+					FString FriendlyPrefixId; // For retrocompatibility, only friendly names have all indexes
+
+					if (NumLODs > 1) // Add LOD ids to prefix if more than one LOD
+					{
+						FString Aux = FString::Printf(TEXT("LOD%d-"), LODIndex);
+						if (LODIndex > 0)
+						{
+							PrefixId.Append(Aux);
+						}
+						FriendlyPrefixId.Append(Aux);
+					}
+
+					int32 MaterialCurrentCount = 0;
+					if (TArray<int32>* MUsageCountP = MaterialNameUseCount.Find(MaterialName))
+					{
+						if (MUsageCountP->Num() > LODIndex)
+						{
+							(*MUsageCountP)[LODIndex]++;
+							MaterialCurrentCount = (*MUsageCountP)[LODIndex];
+						}
+						else
+						{
+							MUsageCountP->AddZeroed(LODIndex - (MUsageCountP->Num() - 1));
+							(*MUsageCountP)[LODIndex] = MaterialCurrentCount;
+						}
+					}
+					else
+					{
+						TArray<int32> NewMaterialNameUseCount;
+						NewMaterialNameUseCount.Add(MaterialCurrentCount);
+						MaterialNameUseCount.Add(MaterialName, NewMaterialNameUseCount);
+					}
+
+					if (MaterialCurrentCount > 0) // Add section ids to prefix if a section/material name is used more than once
+					{
+						FString Aux = FString::Printf(TEXT("%d-"), MaterialCurrentCount);
+						PrefixId.Append(Aux);
+					}
+
+					UEdGraphPin* Pin = nullptr;
+
+					{
+						// Mesh
+						FString MeshName = PrefixId + MaterialName + FString(" Mesh");
+						Pin = CustomCreatePin(EGPD_Output, Schema->PC_Mesh, FName(*MeshName));
+						Pin->PinFriendlyName = FText::FromString(MeshName);
+						Pin->PinToolTip = MeshName;
+						LODs[LODIndex].Materials[MaterialIndex].MeshPinRef.SetPin(Pin);
+
+						// Layout
+						if (ImportedModel->LODModels.Num() > LODIndex)
+						{
+							const uint32 numberOfUVLayouts = ImportedModel->LODModels[LODIndex].NumTexCoords;
+							for (uint32 LayoutIndex = 0; LayoutIndex < numberOfUVLayouts; ++LayoutIndex)
+							{
+								FString LayoutPrefix = "";
+								FString FriendlyLayoutPrefix = "";
+								if (LayoutIndex > 0)
+								{
+									LayoutPrefix = FString::Printf(TEXT("UV%d-"), LayoutIndex);
+								}
+								if (numberOfUVLayouts > 1)
+								{
+									FriendlyLayoutPrefix.Append(FString::Printf(TEXT("UV%d-"), LayoutIndex));
+								}
+								FString LayoutName = PrefixId + LayoutPrefix + MaterialName + FString(" Layout");
+								FString PinFriendlyName = FriendlyPrefixId + FriendlyLayoutPrefix + MaterialName + FString(" Layout");
+								FString PinTooltip = FString::Printf(TEXT("Material %d "), MaterialIndex) + PinFriendlyName;
+								Pin = CustomCreatePin(EGPD_Input, Schema->PC_Layout, FName(*LayoutName));
+								Pin->PinFriendlyName = FText::FromString(PinFriendlyName);
+								Pin->PinToolTip = PinTooltip;
+								Pin->bDefaultValueIsIgnored = true;
+								LODs[LODIndex].Materials[MaterialIndex].LayoutPinsRef.Add(FEdGraphPinReference(Pin));
+							}
+						}
+					}
+
+					// Images
+					if (MaterialInterface)
+					{
+						UMaterial* Material = MaterialInterface->GetMaterial();
+
+						TArray<FMaterialParameterInfo> ImageInfos;
+						TArray<FGuid> ImageIds;
+						Material->GetAllTextureParameterInfo(ImageInfos, ImageIds);
+
+						for (int32 ImageIndex = 0; ImageIndex < ImageInfos.Num(); ++ImageIndex)
+						{
+							FString ImageName = ImageInfos[ImageIndex].Name.ToString();
+							FString Sufix = FString::Printf(TEXT("-mat-%s"), *MaterialName);
+
+
+							int32 ImageCurrentCount = 0;
+							if (TArray<int32>* IUsageCountP = ImageNameUseCount.Find(ImageName))
+							{
+								if (IUsageCountP->Num() > LODIndex)
+								{
+									(*IUsageCountP)[LODIndex]++;
+									ImageCurrentCount = (*IUsageCountP)[LODIndex];
+								}
+								else
+								{
+									IUsageCountP->AddZeroed(LODIndex - (IUsageCountP->Num() - 1));
+									(*IUsageCountP)[LODIndex] = ImageCurrentCount;
+								}
+							}
+							else
+							{
+								TArray<int32> NewImageNameUseCount;
+								NewImageNameUseCount.Add(ImageCurrentCount);
+								ImageNameUseCount.Add(ImageName, NewImageNameUseCount);
+							}
+
+							FString PinName = PrefixId + ImageName + FString::Printf(TEXT("-mat-%s"), *MaterialName);
+							FString PinFriendlyImageName = ImageName;
+
+							if (ImageCurrentCount > 0)
+							{
+								PinName.Append(FString::Printf(TEXT("-%d"), ImageCurrentCount));
+							}
+
+							Pin = CustomCreatePin(EGPD_Output, Schema->PC_Image, FName(*PinName)); // Store the sufix in the category to be able to remove the sufix from the pin name later to retrieve the textures
+							Pin->bDefaultValueIsIgnored = true;
+							Pin->bHidden = true;
+							Pin->PinFriendlyName = FText::FromString(PrefixId + PinFriendlyImageName);
+							Pin->PinToolTip = PinName;
+							LODs[LODIndex].Materials[MaterialIndex].ImagePinsRef.Add(FEdGraphPinReference(Pin));
+						}
+					}
+				}
+			}
+		}
+
+		DefaultPin = FEdGraphPinReference();
+	}
+	else
+	{
+		UEdGraphPin* Pin = CustomCreatePin(EGPD_Output, Schema->PC_Mesh, FName(SKELETAL_MESH_PIN_NAME));
+		Pin->bDefaultValueIsIgnored = true;
+		Pin->PinFriendlyName = FText::FromString(SKELETAL_MESH_PIN_NAME);
+		
+		DefaultPin = FEdGraphPinReference(Pin);
+	}
+}
+
+
+UCustomizableObjectNodeRemapPinsByName* UCustomizableObjectNodeSkeletalMesh::CreateRemapPinsByName() const
+{
+	return NewObject<UCustomizableObjectNodeRemapPinsByNameDefaultPin>();
+}
+
+
+FText UCustomizableObjectNodeSkeletalMesh::GetNodeTitle(ENodeTitleType::Type TitleType) const
+{
+	if (SkeletalMesh)
+	{
+		FFormatNamedArguments Args;
+		Args.Add(TEXT("MeshName"), FText::FromString(SkeletalMesh->GetName()));
+
+		return FText::Format(LOCTEXT("SkeletalMesh_Title", "{MeshName}\nSkeletal Mesh"), Args);
+	}
+	else
+	{
+		return LOCTEXT("Skeletal_Mesh", "Skeletal Mesh");
+	}
+}
+
+
+FLinearColor UCustomizableObjectNodeSkeletalMesh::GetNodeTitleColor() const
+{
+	const UEdGraphSchema_CustomizableObject* Schema = GetDefault<UEdGraphSchema_CustomizableObject>();
+	return Schema->GetPinTypeColor(Schema->PC_Mesh);
+}
+
+
+UTexture2D* UCustomizableObjectNodeSkeletalMesh::FindTextureForPin(const UEdGraphPin* Pin) const
+{
+	if (!SkeletalMesh)
+	{
+		return nullptr;
+	}
+
+	if (const FSkeletalMeshModel* ImportedModel = Helper_GetImportedModel(SkeletalMesh))
+	{
+		for (int LODIndex = 0; LODIndex < LODs.Num(); ++LODIndex)
+		{
+			for (int MaterialIndex = 0; MaterialIndex < LODs[LODIndex].Materials.Num(); ++MaterialIndex)
+			{
+				if (UMaterialInterface* MaterialInterface = GetMaterialInterfaceFor(LODIndex, MaterialIndex, ImportedModel))
+				{
+					for (int ImageIndex = 0; ImageIndex < LODs[LODIndex].Materials[MaterialIndex].ImagePinsRef.Num(); ++ImageIndex)
+					{
+						const UEdGraphPin* OtherPin = LODs[LODIndex].Materials[MaterialIndex].ImagePinsRef[ImageIndex].Get();
+						if (OtherPin == Pin)
+						{
+							UTexture* Texture = nullptr;
+							FString TextureString = OtherPin->PinFriendlyName.IsEmptyOrWhitespace() ? Helper_GetPinName(OtherPin) : OtherPin->PinFriendlyName.ToString();
+							FName TextureName = FName(*TextureString);
+							MaterialInterface->GetTextureParameterValue(TextureName, Texture);
+							return Cast<UTexture2D>(Texture);
+						}
+					}
+				}
+			}
+		}
+	}
+
+	return nullptr;
+}
+
+
+void UCustomizableObjectNodeSkeletalMesh::GetUVChannelForPin(const UEdGraphPin* Pin, TArray<FVector2f>& OutSegments, int32 UVIndex) const
+{
+	check(Pin);
+	
+	if (!SkeletalMesh)
+	{
+		return;
+	}
+
+	int32 LODIndex;
+	int32 SectionIndex;
+	int32 LayoutIndex;
+	GetPinSection(*Pin, LODIndex, SectionIndex, LayoutIndex);
+
+	OutSegments = GetUV(*SkeletalMesh, LODIndex, SectionIndex, UVIndex);
+}
+
+
+TArray<UCustomizableObjectLayout*> UCustomizableObjectNodeSkeletalMesh::GetLayouts(const UEdGraphPin* OutPin) const
+{
+	TArray<UCustomizableObjectLayout*> Result;
+
+	for (int LODIndex = 0; LODIndex < LODs.Num(); ++LODIndex)
+	{
+		for (int MaterialIndex = 0; MaterialIndex < LODs[LODIndex].Materials.Num(); ++MaterialIndex)
+		{
+			if (LODs[LODIndex].Materials[MaterialIndex].MeshPinRef.Get() == OutPin)
+			{
+				for (int LayoutIndex = 0; LayoutIndex < LODs[LODIndex].Materials[MaterialIndex].LayoutPinsRef.Num(); ++LayoutIndex)
+				{
+					if (UEdGraphPin* LayoutInPin = LODs[LODIndex].Materials[MaterialIndex].LayoutPinsRef[LayoutIndex].Get())
+					{
+						if (const UEdGraphPin* ConnectedPin = FollowInputPin(*LayoutInPin))
+						{
+							const UCustomizableObjectNodeLayoutBlocks* LayoutNode = Cast<UCustomizableObjectNodeLayoutBlocks>(ConnectedPin->GetOwningNode());
+							if (LayoutNode && LayoutNode->Layout)
+							{
+								Result.Add(LayoutNode->Layout);
+							}
+						}
+					}
+				}
+			}
+		}
+	}
+
+	return Result;
+}
+
+
+UObject* UCustomizableObjectNodeSkeletalMesh::GetMesh() const
+{
+	return SkeletalMesh;
+}
+
+
+UEdGraphPin* UCustomizableObjectNodeSkeletalMesh::GetMeshPin(const int32 LODIndex, const int MaterialIndex) const
+{
+	if (LODIndex < LODs.Num())
+	{
+		if (const FCustomizableObjectNodeSkeletalMeshLOD& LOD = LODs[LODIndex];
+			MaterialIndex < LOD.Materials.Num())
+		{
+			return LOD.Materials[MaterialIndex].MeshPinRef.Get();
+		}
+	}
+
+	return nullptr;
+}
+
+void UCustomizableObjectNodeSkeletalMesh::GetPinSection(const UEdGraphPin& Pin, int32& OutLODIndex, int32& OutSectionIndex, int32& OutLayoutIndex) const
+{
+	if (LODs.IsEmpty())
+	{
+		OutLODIndex = -1;
+		OutSectionIndex = -1;
+		OutLayoutIndex = -1;
+
+		return;
+	}
+
+	for (OutLODIndex = 0; OutLODIndex < LODs.Num(); ++OutLODIndex)
+	{
+		const FCustomizableObjectNodeSkeletalMeshLOD& LOD = LODs[OutLODIndex];
+		
+		for (OutSectionIndex = 0; OutSectionIndex < LOD.Materials.Num(); ++OutSectionIndex)
+		{
+			const FCustomizableObjectNodeSkeletalMeshMaterial& Material = LOD.Materials[OutSectionIndex];
+			
+			for (OutLayoutIndex = 0; OutLayoutIndex < Material.LayoutPinsRef.Num(); ++OutLayoutIndex)
+			{
+				// Is a Layout pin?
+				if (Material.LayoutPinsRef[OutLayoutIndex].Get() == &Pin)
+				{
+					return;
+				}
+			
+				// Is a Mesh pin?
+				if (Material.MeshPinRef.Get() == &Pin)
+				{
+					return;
+				}
+
+				// Is an Image pin?
+				for (int ImageIndex = 0; ImageIndex < Material.ImagePinsRef.Num(); ++ImageIndex)
+				{
+					if (Material.ImagePinsRef[ImageIndex].Get() == &Pin)
+					{
+						return;
+					}
+				}
+			}
+		}
+	}
+
+	check(false); // Pin not found. Probably a new pin has been added which is not contemplated in this function.
+}
+
+UMaterialInterface* UCustomizableObjectNodeSkeletalMesh::GetMaterialFor(const UEdGraphPin* Pin) const
+{
+	if (SkeletalMesh)
+	{
+		if (FSkeletalMaterial* SkeletalMaterial = GetSkeletalMaterialFor(Pin))
+		{
+			return SkeletalMaterial->MaterialInterface;
+		}
+	}
+
+	return nullptr;
+}
+
+FSkeletalMaterial* UCustomizableObjectNodeSkeletalMesh::GetSkeletalMaterialFor(const UEdGraphPin* Pin) const
+{
+	check(Pin);
+	
+	int32 LODIndex;
+	int32 SectionIndex;
+	int32 LayoutIndex;
+	GetPinSection(*Pin, LODIndex, SectionIndex, LayoutIndex);
+
+	return GetSkeletalMaterialFor(LODIndex, SectionIndex);
+}
+
+bool UCustomizableObjectNodeSkeletalMesh::IsPinRelevant(const UEdGraphPin* Pin) const
+{
+	const UEdGraphSchema_CustomizableObject* Schema = GetDefault<UEdGraphSchema_CustomizableObject>();
+
+	if (Pin->Direction == EEdGraphPinDirection::EGPD_Output)
+	{
+		return Pin->PinType.PinCategory == Schema->PC_Layout;
+	}
+
+	if (Pin->Direction == EEdGraphPinDirection::EGPD_Input)
+	{	
+		return Pin->PinType.PinCategory == Schema->PC_Mesh;
+	}
+
+	return false;
+}
+
+bool UCustomizableObjectNodeSkeletalMesh::IsNodeOutDatedAndNeedsRefresh()
+{
+	bool Result = false;
+
+	if (SkeletalMesh)
+	{
+		int32 numLods = Helper_GetLODNum(SkeletalMesh);
+		if (LODs.Num() != numLods)
+		{
+			Result = true;
+		}
+
+		if (const FSkeletalMeshModel* ImportedModel = Helper_GetImportedModel(SkeletalMesh))
+		{
+			for (int LODIndex = 0; !Result && LODIndex < LODs.Num(); ++LODIndex)
+			{
+				int32 NumMaterials = ImportedModel->LODModels[LODIndex].Sections.Num();
+				if (LODs[LODIndex].Materials.Num() != NumMaterials)
+				{
+					Result = true;
+				}
+
+				for (int MaterialIndex = 0; !Result && MaterialIndex < NumMaterials; ++MaterialIndex)
+				{
+					if (UMaterialInterface* MaterialInterface = GetMaterialInterfaceFor(LODIndex, MaterialIndex, ImportedModel))
+					{
+						if (LODs[LODIndex].Materials[MaterialIndex].Name != MaterialInterface->GetName())
+						{
+							Result = true;
+							break;
+						}
+						UMaterial* Material = MaterialInterface->GetMaterial();
+
+						TArray<FMaterialParameterInfo> ImageInfos;
+						TArray<FGuid> ImageIds;
+						Material->GetAllTextureParameterInfo(ImageInfos, ImageIds);
+
+						if (LODs[LODIndex].Materials[MaterialIndex].ImagePinsRef.Num() != ImageInfos.Num())
+						{
+							Result = true;
+							break;
+						}
+
+						for (int ImageIndex = 0; ImageIndex < ImageInfos.Num(); ++ImageIndex)
+						{
+							const FString ImagePinName = LODs[LODIndex].Materials[MaterialIndex].ImagePinsRef[ImageIndex].Get()->GetName();
+							if (!ImagePinName.Contains(ImageInfos[ImageIndex].Name.ToString()))
+							{
+								Result = true;
+								break;
+							}
+						}
+					}
+				}
+			}
+		}
+	}
+
+	// Remove previous compilation warnings
+	if (!Result && bHasCompilerMessage)
+	{
+		RemoveWarnings();
+		GetGraph()->NotifyGraphChanged();
+	}
+
+    return Result;
+}
+
+FString UCustomizableObjectNodeSkeletalMesh::GetRefreshMessage() const
+{
+    return "Node data outdated. Please refresh node.";
+}
+
+
+FText UCustomizableObjectNodeSkeletalMesh::GetTooltipText() const
+{
+	return LOCTEXT("Skeletal_Mesh_Tooltip", "Get access to the sections (also known as material slots) of a skeletal mesh and to each of the sections texture parameters.");
+}
+
+
+void UCustomizableObjectNodeSkeletalMesh::BackwardsCompatibleFixup()
+{
+	Super::BackwardsCompatibleFixup();
+
+	const int32 CustomizableObjectCustomVersion = GetLinkerCustomVersion(FCustomizableObjectCustomVersion::GUID);
+
+	if (CustomizableObjectCustomVersion < FCustomizableObjectCustomVersion::NodeSkeletalMeshCorruptedPinRef)
+	{
+		ReconstructNode();
+	}
+}
+
+
+TSharedPtr<SGraphNode> UCustomizableObjectNodeSkeletalMesh::CreateVisualWidget()
+{
+	TSharedPtr<SGraphNodeSkeletalMesh> GraphNode;
+	SAssignNew(GraphNode, SGraphNodeSkeletalMesh,this);
+
+	GraphNodeSkeletalMesh = GraphNode;
+
+	return GraphNode;
+}
+
+
+bool UCustomizableObjectNodeSkeletalMesh::CheckIsValidLayout(const UEdGraphPin* Pin, int32& LayoutIndex, FString& MaterialName)
+{
+	if (const UEdGraphPin* ConnectedPin = FollowOutputPin(*Pin))
+	{
+		int32 LODIndex;
+		int32 SectionIndex;
+		GetPinSection(*ConnectedPin, LODIndex, SectionIndex, LayoutIndex);
+
+		if (LayoutIndex == 0)
+		{
+			return true;
+		}
+		else
+		{
+			TArray<FEdGraphPinReference> LayoutPins = LODs[LODIndex].Materials[SectionIndex].LayoutPinsRef;
+			MaterialName = LODs[LODIndex].Materials[SectionIndex].Name;
+
+			for (int32 i = 0; i < LayoutPins.Num(); ++i)
+			{
+				if (i == LayoutIndex)
+				{
+					break;
+				}
+
+				if (!FollowInputPin(*LayoutPins[i].Get()))
+				{
+					return false;
+				}
+			}
+
+		}
+	}
+
+	return true;
+}
+
+UMaterialInterface* UCustomizableObjectNodeSkeletalMesh::GetMaterialInterfaceFor(const int LODIndex, const int MaterialIndex, const FSkeletalMeshModel* ImportedModel) const
+{
+	if (FSkeletalMaterial* SkeletalMaterial = GetSkeletalMaterialFor(LODIndex, MaterialIndex, ImportedModel))
+	{
+		return SkeletalMaterial->MaterialInterface;
+	}
+
+	return nullptr;
+}
+
+FSkeletalMaterial* UCustomizableObjectNodeSkeletalMesh::GetSkeletalMaterialFor(const int LODIndex, const int MaterialIndex, const FSkeletalMeshModel* ImportedModel) const
+{
+	if (!SkeletalMesh)
+	{
+		return nullptr;
+	}
+
+	// We assume that LODIndex and MaterialIndex are valid for the imported model
+	int SkeletalMeshMaterialIndex;
+	if (ImportedModel)
+	{
+		SkeletalMeshMaterialIndex = ImportedModel->LODModels[LODIndex].Sections[MaterialIndex].MaterialIndex;
+	}
+	else
+	{
+		SkeletalMeshMaterialIndex = Helper_GetImportedModel(SkeletalMesh)->LODModels[LODIndex].Sections[MaterialIndex].MaterialIndex;
+	}
+	
+	// Check if we have lod info map to get the correct material index
+	if (const FSkeletalMeshLODInfo* LodInfo = SkeletalMesh->GetLODInfo(LODIndex))
+	{
+		if (LodInfo->LODMaterialMap.IsValidIndex(MaterialIndex))
+		{
+			SkeletalMeshMaterialIndex = LodInfo->LODMaterialMap[MaterialIndex];
+		}
+	}
+	
+	if (SkeletalMesh->GetMaterials().IsValidIndex(SkeletalMeshMaterialIndex))
+	{
+		return &SkeletalMesh->GetMaterials()[SkeletalMeshMaterialIndex];
+	}
+	
+	return nullptr;
+}
+
+//SGraphNode --------------------------------------------
+
+void SGraphNodeSkeletalMesh::Construct(const FArguments& InArgs, UEdGraphNode* InGraphNode)
+{
+	GraphNode = InGraphNode;
+	NodeSkeletalMesh = Cast< UCustomizableObjectNodeSkeletalMesh >(GraphNode);
+
+	WidgetSize = 128.0f;
+	ThumbnailSize = 128;
+
+	TSharedPtr<FCustomizableObjectEditor> Editor = StaticCastSharedPtr< FCustomizableObjectEditor >(NodeSkeletalMesh->GetGraphEditor());
+
+	// Thumbnail
+	AssetThumbnailPool = MakeShareable(new FAssetThumbnailPool(32));
+	AssetThumbnail = MakeShareable(new FAssetThumbnail(NodeSkeletalMesh->SkeletalMesh, ThumbnailSize, ThumbnailSize, AssetThumbnailPool));
+
+	// Selector
+	FPropertyEditorModule& PropPlugin = FModuleManager::LoadModuleChecked<FPropertyEditorModule>("PropertyEditor");
+	FSinglePropertyParams SingleDetails;
+	SingleDetails.NamePlacement = EPropertyNamePlacement::Hidden;
+	SingleDetails.NotifyHook = Editor.Get();
+	SingleDetails.bHideAssetThumbnail = true;
+
+	SkeletalMeshSelector = PropPlugin.CreateSingleProperty(NodeSkeletalMesh, "SkeletalMesh", SingleDetails);
+
+	UpdateGraphNode();
+}
+
+
+void SGraphNodeSkeletalMesh::UpdateGraphNode()
+{
+	SGraphNode::UpdateGraphNode();
+}
+
+
+void SGraphNodeSkeletalMesh::SetDefaultTitleAreaWidget(TSharedRef<SOverlay> DefaultTitleAreaWidget)
+{
+	DefaultTitleAreaWidget->AddSlot()
+		.HAlign(HAlign_Right)
+		.VAlign(VAlign_Center)
+		.Padding(FMargin(5))
+		[
+			SNew(SCheckBox)
+			.OnCheckStateChanged(this, &SGraphNodeSkeletalMesh::OnExpressionPreviewChanged)
+			.IsChecked(IsExpressionPreviewChecked())
+			.Cursor(EMouseCursor::Default)
+			.Style(FAppStyle::Get(), "Graph.Node.AdvancedView")
+			[
+				SNew(SHorizontalBox)
+				+ SHorizontalBox::Slot()
+				.VAlign(VAlign_Center)
+				.HAlign(HAlign_Center)
+				[
+					SNew(SImage)
+					.Image(GetExpressionPreviewArrow())
+				]
+			]
+	];
+}
+
+
+void SGraphNodeSkeletalMesh::CreateBelowPinControls(TSharedPtr<SVerticalBox> MainBox)
+{
+	LeftNodeBox->AddSlot()
+		.AutoHeight()
+		.MaxHeight(WidgetSize)
+		.Padding(10.0f,10.0f,0.0f,0.0f)
+		[
+			SNew(SHorizontalBox)
+			.Visibility(ExpressionPreviewVisibility())
+
+			+SHorizontalBox::Slot()
+			.MaxWidth(WidgetSize)
+			.Padding(5.0f,5.0f,5.0f,5.0f)
+			[
+				AssetThumbnail->MakeThumbnailWidget()
+			]
+		];
+
+	if (SkeletalMeshSelector.IsValid())
+	{
+		LeftNodeBox->AddSlot()
+		.AutoHeight()
+		.Padding(10.0f, 5.0f, 0.0f, 0.0f)
+		[
+			SNew(SHorizontalBox)
+			.Visibility(ExpressionPreviewVisibility())
+
+			+ SHorizontalBox::Slot()
+			.AutoWidth()
+			.Padding(1.0f,0.0f, 5.0f, 5.0f)
+			.VAlign(VAlign_Center)
+			.HAlign(HAlign_Left)
+			[
+				SkeletalMeshSelector.ToSharedRef()
+			]
+		];
+	}
+}
+
+
+void SGraphNodeSkeletalMesh::OnExpressionPreviewChanged(const ECheckBoxState NewCheckedState)
+{
+	NodeSkeletalMesh->bCollapsed = (NewCheckedState != ECheckBoxState::Checked);
+	UpdateGraphNode();
+}
+
+
+ECheckBoxState SGraphNodeSkeletalMesh::IsExpressionPreviewChecked() const
+{
+	return NodeSkeletalMesh->bCollapsed ? ECheckBoxState::Unchecked : ECheckBoxState::Checked;
+}
+
+
+const FSlateBrush* SGraphNodeSkeletalMesh::GetExpressionPreviewArrow() const
+{
+	return FCustomizableObjectEditorStyle::Get().GetBrush(NodeSkeletalMesh->bCollapsed ? TEXT("Nodes.ArrowDown") : TEXT("Nodes.ArrowUp"));
+}
+
+
+EVisibility SGraphNodeSkeletalMesh::ExpressionPreviewVisibility() const
+{
+	return NodeSkeletalMesh->bCollapsed ? EVisibility::Collapsed : EVisibility::Visible;
+}
+
+
+#undef LOCTEXT_NAMESPACE
diff --git a/Engine/Plugins/Experimental/Mutable/Source/CustomizableObjectEditor/Private/Nodes/CustomizableObjectNodeStaticMesh.cpp b/Engine/Plugins/Experimental/Mutable/Source/CustomizableObjectEditor/Private/Nodes/CustomizableObjectNodeStaticMesh.cpp
new file mode 100644
--- /dev/null
+++ b/Engine/Plugins/Experimental/Mutable/Source/CustomizableObjectEditor/Private/Nodes/CustomizableObjectNodeStaticMesh.cpp
@@ -0,0 +1,595 @@
+// Copyright Epic Games, Inc. All Rights Reserved.
+
+#include "Nodes/CustomizableObjectNodeStaticMesh.h"
+#include "CustomizableObjectEditorModule.h"
+#include "UnrealEditorPortabilityHelpers.h"
+
+#include "Engine/StaticMesh.h"
+#include "Materials/Material.h"
+#include "ISinglePropertyView.h"
+#include "CustomizableObjectEditorStyle.h"
+#include "CustomizableObjectEditor.h"
+#include "GraphTraversal.h"
+#include "MutableUtils.h"
+#include "RemapPins/CustomizableObjectNodeRemapPinsByNameDefaultPin.h"
+
+#define LOCTEXT_NAMESPACE "CustomizableObjectEditor"
+
+
+/** Default node pin configuration pin name (node does not have an static mesh). */
+static const TCHAR* STATIC_MESH_PIN_NAME = TEXT("Static Mesh");
+
+
+void UCustomizableObjectNodeStaticMesh::PostEditChangeProperty(FPropertyChangedEvent& PropertyChangedEvent)
+{
+	FProperty* PropertyThatChanged = PropertyChangedEvent.Property;
+	if (PropertyThatChanged && PropertyThatChanged->GetName() == TEXT("StaticMesh"))
+	{
+		ReconstructNode();
+	}
+
+	Super::PostEditChangeProperty(PropertyChangedEvent);
+}
+
+
+void UCustomizableObjectNodeStaticMesh::Serialize(FArchive& Ar)
+{
+	Super::Serialize(Ar);
+
+	Ar.UsingCustomVersion(FCustomizableObjectCustomVersion::GUID);
+
+	if (Ar.CustomVer(FCustomizableObjectCustomVersion::GUID) < FCustomizableObjectCustomVersion::PostLoadToCustomVersion)
+	{
+		for (FCustomizableObjectNodeStaticMeshLOD& LOD : LODs)
+		{
+			for (FCustomizableObjectNodeStaticMeshMaterial& Material : LOD.Materials)
+			{
+				if (Material.MeshPin && !Material.MeshPinRef.Get())
+				{
+					UEdGraphPin* AuxPin = UEdGraphPin::FindPinCreatedFromDeprecatedPin(Material.MeshPin);
+					Material.MeshPinRef.SetPin(AuxPin);
+				}
+
+				if (Material.LayoutPin && !Material.LayoutPinRef.Get())
+				{
+					UEdGraphPin* AuxPin = UEdGraphPin::FindPinCreatedFromDeprecatedPin(Material.LayoutPin);
+					Material.LayoutPinRef.SetPin(AuxPin);
+				}
+
+				if (Material.ImagePins.Num() && !Material.ImagePinsRef.Num())
+				{
+					for (UEdGraphPin_Deprecated* DeprecatedPin : Material.ImagePins)
+					{
+						UEdGraphPin* AuxPin = UEdGraphPin::FindPinCreatedFromDeprecatedPin(DeprecatedPin);
+						FEdGraphPinReference AuxEdGraphPinReference(AuxPin);
+						Material.ImagePinsRef.Add(AuxEdGraphPinReference);
+					}
+				}
+			}
+		}
+	}
+}
+
+
+void UCustomizableObjectNodeStaticMesh::AllocateDefaultPins(UCustomizableObjectNodeRemapPins* RemapPins)
+{
+	// Pass information to the remap pins action context
+	if (UCustomizableObjectNodeRemapPinsByNameDefaultPin* RemapPinsCustom = Cast<UCustomizableObjectNodeRemapPinsByNameDefaultPin>(RemapPins))
+	{
+		RemapPinsCustom->DefaultPin = DefaultPin.Get();
+	}
+	
+	const UEdGraphSchema_CustomizableObject* Schema = GetDefault<UEdGraphSchema_CustomizableObject>();
+
+	if (StaticMesh)
+	{
+		const int NumLODs = StaticMesh->GetRenderData()->LODResources.Num();
+
+		LODs.Reset(NumLODs);
+		LODs.SetNum(NumLODs);
+		for (int32 i = 0; i < NumLODs; ++i)
+		{
+			FString LODName = FString::Printf(TEXT("LOD %d, "), i);
+
+			const int NumMaterials = StaticMesh->GetRenderData()->LODResources[i].Sections.Num();
+
+			LODs[i].Materials.Reset(NumMaterials);
+			LODs[i].Materials.SetNum(NumMaterials);
+			for (int32 MaterialIndex = 0; MaterialIndex < NumMaterials; ++MaterialIndex)
+			{
+				LODs[i].Materials.Reset(NumMaterials);
+				LODs[i].Materials.SetNum(NumMaterials);
+				
+				UMaterialInterface* MaterialInterface = 0;
+				int StaticMeshMaterialIndex = StaticMesh->GetRenderData()->LODResources[i].Sections[MaterialIndex].MaterialIndex;
+				if (StaticMeshMaterialIndex >= 0 && StaticMeshMaterialIndex < StaticMesh->GetStaticMaterials().Num())
+				{
+					MaterialInterface = StaticMesh->GetStaticMaterials()[StaticMeshMaterialIndex].MaterialInterface;
+				}
+
+				FString MaterialName = TEXT("Unnamed Material");
+				if (MaterialInterface)
+				{
+					MaterialName = MaterialInterface->GetName();
+				}
+				LODs[i].Materials[MaterialIndex].Name = MaterialName;
+
+				// Mesh
+				const FString MeshName = MaterialName + TEXT(" Mesh");
+				UEdGraphPin* Pin = CustomCreatePin(EGPD_Output, Schema->PC_Mesh, FName(*MeshName));
+				Pin->bDefaultValueIsIgnored = true;
+				LODs[i].Materials[MaterialIndex].MeshPinRef.SetPin(Pin);
+
+				// Layout
+				// TODO: Multiple layout ui
+				const FString LayoutName = MaterialName + TEXT(" Layout");
+				Pin = CustomCreatePin(EGPD_Input, Schema->PC_Layout, FName(*LayoutName));
+				Pin->bDefaultValueIsIgnored = true;
+				LODs[i].Materials[MaterialIndex].LayoutPinRef.SetPin(Pin);
+
+				// Images
+				if (MaterialInterface)
+				{
+					UMaterial* Material = MaterialInterface->GetMaterial();
+
+					Helper_ForEachTextureParameterName(Material, [&](const FString& ImageName)
+						{
+							Pin = CustomCreatePin(EGPD_Output, Schema->PC_Image, *ImageName);
+							Pin->bDefaultValueIsIgnored = true;
+							LODs[i].Materials[MaterialIndex].ImagePinsRef.Add(Pin);
+						});
+				}
+			}
+		}
+
+		DefaultPin = FEdGraphPinReference();
+	}
+	else
+	{
+		UEdGraphPin* Pin = CustomCreatePin(EGPD_Output, Schema->PC_Mesh, FName(STATIC_MESH_PIN_NAME));
+		Pin->bDefaultValueIsIgnored = true;
+		Pin->PinFriendlyName = FText::FromString(STATIC_MESH_PIN_NAME);
+
+		DefaultPin = FEdGraphPinReference(Pin);
+	}
+}
+
+
+FText UCustomizableObjectNodeStaticMesh::GetNodeTitle(ENodeTitleType::Type TitleType) const
+{
+	if (StaticMesh)
+	{
+		FFormatNamedArguments Args;
+		Args.Add(TEXT("MeshName"), FText::FromString(StaticMesh->GetName()));
+
+		return FText::Format(LOCTEXT("StaticMesh_Title", "{MeshName}\nStatic Mesh"), Args);
+	}
+	else
+	{
+		return LOCTEXT("Static_Mesh", "Static Mesh");
+	}
+}
+
+
+FLinearColor UCustomizableObjectNodeStaticMesh::GetNodeTitleColor() const
+{
+	const UEdGraphSchema_CustomizableObject* Schema = GetDefault<UEdGraphSchema_CustomizableObject>();
+	return Schema->GetPinTypeColor(Schema->PC_Mesh);
+}
+
+
+UTexture2D* UCustomizableObjectNodeStaticMesh::FindTextureForPin(const UEdGraphPin* Pin) const
+{
+	if (!StaticMesh)
+	{
+		return 0;
+	}
+
+	for (int LODIndex = 0; LODIndex < LODs.Num(); ++LODIndex)
+	{
+		for (int MaterialIndex = 0; MaterialIndex<LODs[LODIndex].Materials.Num(); ++MaterialIndex)
+		{
+			UMaterialInterface* MaterialInterface = 0;
+			int StaticMeshMaterialIndex = -1;
+			if (StaticMesh->GetRenderData()->LODResources.Num() > LODIndex && StaticMesh->GetRenderData()->LODResources[LODIndex].Sections.Num()>MaterialIndex)
+			{
+				StaticMeshMaterialIndex = StaticMesh->GetRenderData()->LODResources[LODIndex].Sections[MaterialIndex].MaterialIndex;
+			}
+
+			if (StaticMeshMaterialIndex >= 0 && StaticMeshMaterialIndex < StaticMesh->GetStaticMaterials().Num() )
+			{
+				MaterialInterface = StaticMesh->GetStaticMaterials()[ StaticMeshMaterialIndex ].MaterialInterface;
+			}
+
+			if (MaterialInterface)
+			{
+				for (int ImageIndex = 0; ImageIndex < LODs[LODIndex].Materials[MaterialIndex].ImagePinsRef.Num(); ++ImageIndex)
+				{
+					if (LODs[LODIndex].Materials[MaterialIndex].ImagePinsRef[ImageIndex].Get() == Pin)
+					{
+						UTexture* Texture=nullptr;
+						MaterialInterface->GetTextureParameterValue(Helper_GetPinFName(Pin), Texture);
+						return Cast<UTexture2D>(Texture);
+					}
+				}
+			}
+		}
+	}
+
+	return 0;
+}
+
+
+void UCustomizableObjectNodeStaticMesh::GetUVChannelForPin(const UEdGraphPin* Pin, TArray<FVector2f>& OutSegments, int32 UVIndex) const
+{
+	check(Pin);
+	
+	if (!StaticMesh)
+	{
+		return;
+	}
+
+	int32 LODIndex;
+	int32 SectionIndex;
+	int32 LayoutIndex;
+	GetPinSection(*Pin, LODIndex, SectionIndex, LayoutIndex);
+	
+	OutSegments = GetUV(*StaticMesh, LODIndex, SectionIndex, UVIndex);
+}
+
+
+TArray<class UCustomizableObjectLayout*> UCustomizableObjectNodeStaticMesh::GetLayouts(const UEdGraphPin* OutPin) const
+{
+	TArray<class UCustomizableObjectLayout*> Result;
+
+	for (int LODIndex = 0; LODIndex<LODs.Num(); ++LODIndex)
+	{
+		for (int MaterialIndex = 0; MaterialIndex<LODs[LODIndex].Materials.Num(); ++MaterialIndex)
+		{
+			if (LODs[LODIndex].Materials[MaterialIndex].MeshPinRef.Get() == OutPin)
+			{
+				if (UEdGraphPin* LayoutInPin = LODs[LODIndex].Materials[MaterialIndex].LayoutPinRef.Get())
+				{
+					if (const UEdGraphPin* ConnectedPin = FollowInputPin(*LayoutInPin))
+					{
+						const UCustomizableObjectNodeLayoutBlocks* LayoutNode = Cast<UCustomizableObjectNodeLayoutBlocks>(ConnectedPin->GetOwningNode());
+						if (LayoutNode && LayoutNode->Layout)
+						{
+							Result.Add(LayoutNode->Layout);
+						}
+					}
+				}
+			}
+		}
+	}
+
+	return Result;
+}
+
+
+UObject* UCustomizableObjectNodeStaticMesh::GetMesh() const
+{
+	return StaticMesh;
+}
+
+
+UEdGraphPin* UCustomizableObjectNodeStaticMesh::GetMeshPin(const int32 LODIndex, const int MaterialIndex) const
+{
+	if (LODIndex < LODs.Num())
+	{
+		if (const FCustomizableObjectNodeStaticMeshLOD& LOD = LODs[LODIndex];
+			MaterialIndex < LOD.Materials.Num())
+		{
+			return LOD.Materials[MaterialIndex].MeshPinRef.Get();
+		}
+	}
+
+	return nullptr;
+}
+
+
+bool UCustomizableObjectNodeStaticMesh::IsPinRelevant(const UEdGraphPin* Pin) const
+{
+	const UEdGraphSchema_CustomizableObject* Schema = GetDefault<UEdGraphSchema_CustomizableObject>();
+
+	if (Pin->Direction == EEdGraphPinDirection::EGPD_Output)
+	{
+		return Pin->PinType.PinCategory == Schema->PC_Layout;
+	}
+
+	if (Pin->Direction == EEdGraphPinDirection::EGPD_Input)
+	{	
+		return Pin->PinType.PinCategory == Schema->PC_Mesh;
+	}
+
+	return false;
+}
+
+
+UCustomizableObjectNodeRemapPinsByName* UCustomizableObjectNodeStaticMesh::CreateRemapPinsByName() const
+{
+	return NewObject<UCustomizableObjectNodeRemapPinsByNameDefaultPin>();
+}
+
+
+bool UCustomizableObjectNodeStaticMesh::IsNodeOutDatedAndNeedsRefresh()
+{
+	bool Result = false;
+
+	if (StaticMesh)
+	{
+		int32 numLods = StaticMesh->GetRenderData()->LODResources.Num();
+		if (LODs.Num() != numLods)
+		{
+			Result = true;
+		}
+
+		for (int LODIndex = 0; !Result && LODIndex < LODs.Num(); ++LODIndex)
+		{
+			int32 NumMaterials = StaticMesh->GetRenderData()->LODResources[LODIndex].Sections.Num();
+			if (LODs[LODIndex].Materials.Num() != NumMaterials)
+			{
+				Result = true;
+			}
+
+			for (int MaterialIndex = 0; !Result && MaterialIndex < NumMaterials; ++MaterialIndex)
+			{
+				UMaterialInterface* MaterialInterface = nullptr;
+				int SkeletalMeshMaterialIndex = StaticMesh->GetRenderData()->LODResources[LODIndex].Sections[MaterialIndex].MaterialIndex;
+				if (SkeletalMeshMaterialIndex >= 0 && SkeletalMeshMaterialIndex < StaticMesh->GetStaticMaterials().Num())
+				{
+					MaterialInterface = StaticMesh->GetStaticMaterials()[SkeletalMeshMaterialIndex].MaterialInterface;
+				}
+
+				if (MaterialInterface)
+				{
+					/*if (LODs[LODIndex].Materials[MaterialIndex].Name != MaterialInterface->GetName())
+					{
+						Result = true;
+						break;
+					}*/
+					UMaterial* Material = MaterialInterface->GetMaterial();
+
+					TArray<FMaterialParameterInfo> ImageInfos;
+					TArray<FGuid> ImageIds;
+					Material->GetAllTextureParameterInfo(ImageInfos, ImageIds);
+
+					if (LODs[LODIndex].Materials[MaterialIndex].ImagePinsRef.Num() != ImageInfos.Num())
+					{
+						Result = true;
+						break;
+					}
+
+					for (int ImageIndex = 0; ImageIndex < ImageInfos.Num(); ++ImageIndex)
+					{
+						const FString ImagePinName = LODs[LODIndex].Materials[MaterialIndex].ImagePinsRef[ImageIndex].Get()->GetName();
+						if (!ImagePinName.Contains(ImageInfos[ImageIndex].Name.ToString()))
+						{
+							Result = true;
+							break;
+						}
+					}
+				}
+			}
+		}
+	}
+
+	// Remove previous compilation warnings
+	if (!Result && bHasCompilerMessage)
+	{
+		RemoveWarnings();
+		GetGraph()->NotifyGraphChanged();
+	}
+
+	return Result;
+}
+
+FString UCustomizableObjectNodeStaticMesh::GetRefreshMessage() const
+{
+	return "Node data outdated. Please refresh node.";
+}
+
+
+FText UCustomizableObjectNodeStaticMesh::GetTooltipText() const
+{
+	return LOCTEXT("Static_Mesh_Tooltip", "Static Mesh");
+}
+
+
+void UCustomizableObjectNodeStaticMesh::GetPinSection(const UEdGraphPin& Pin, int32& OutLODIndex, int32& OutSectionIndex, int32& OutLayoutIndex) const
+{
+	OutLayoutIndex = 0; // Currently layouts different from 0 not supported for static mesh.
+	
+	for (OutLODIndex = 0; OutLODIndex < LODs.Num(); ++OutLODIndex)
+	{
+		const FCustomizableObjectNodeStaticMeshLOD& LOD = LODs[OutLODIndex];
+
+		for (OutSectionIndex = 0; OutSectionIndex < LOD.Materials.Num(); ++OutSectionIndex)
+		{
+			const FCustomizableObjectNodeStaticMeshMaterial& Material = LOD.Materials[OutSectionIndex];
+
+			// Is a Layout pin?
+			if (Material.LayoutPinRef.Get() == &Pin)
+			{
+				return;
+			}
+
+			// Is a Mesh pin?
+			if (Material.MeshPinRef.Get() == &Pin)
+			{
+				return;
+			}
+			
+			// Is an Image pin?
+			for (int ImageIndex = 0; ImageIndex < Material.ImagePinsRef.Num(); ++ImageIndex)
+			{
+				if (Material.ImagePinsRef[ImageIndex].Get() == &Pin)
+				{
+					return;
+				}
+			}
+		}
+	}
+	
+	check(false); // Pin not found. Probably a new pin has been added which is not contemplated in this function.
+}
+
+
+UMaterialInterface* UCustomizableObjectNodeStaticMesh::GetMaterialFor(const UEdGraphPin* Pin) const
+{
+	if (StaticMesh)
+	{
+		check(Pin);
+		int32 LODIndex;
+		int32 SectionIndex;
+		int32 LayoutIndex;
+		GetPinSection(*Pin, LODIndex, SectionIndex, LayoutIndex);
+
+		const int32 MaterialIndex = StaticMesh->GetRenderData()->LODResources[LODIndex].Sections[SectionIndex].MaterialIndex;
+
+		if (const TArray<FStaticMaterial>& StaticMaterials = StaticMesh->GetStaticMaterials();
+			StaticMaterials.IsValidIndex(MaterialIndex))
+		{
+			return StaticMaterials[MaterialIndex].MaterialInterface;
+		}
+	}
+	
+	return nullptr;
+}
+
+
+TSharedPtr<SGraphNode> UCustomizableObjectNodeStaticMesh::CreateVisualWidget()
+{
+	TSharedPtr<SGraphNodeStaticMesh> GraphNode;
+	SAssignNew(GraphNode, SGraphNodeStaticMesh, this);
+
+	GraphNodeStaticMesh = GraphNode;
+
+	return GraphNode;
+}
+
+
+void SGraphNodeStaticMesh::Construct(const FArguments& InArgs, UEdGraphNode* InGraphNode)
+{
+	GraphNode = InGraphNode;
+	NodeStaticMesh = Cast< UCustomizableObjectNodeStaticMesh >(GraphNode);
+
+	WidgetSize = 128.0f;
+	ThumbnailSize = 128;
+
+	TSharedPtr<FCustomizableObjectEditor> Editor = StaticCastSharedPtr< FCustomizableObjectEditor >(NodeStaticMesh->GetGraphEditor());
+
+	// Thumbnail
+	AssetThumbnailPool = MakeShareable(new FAssetThumbnailPool(32));
+	AssetThumbnail = MakeShareable(new FAssetThumbnail(NodeStaticMesh->StaticMesh, ThumbnailSize, ThumbnailSize, AssetThumbnailPool));
+
+	FPropertyEditorModule& PropPlugin = FModuleManager::LoadModuleChecked<FPropertyEditorModule>("PropertyEditor");
+	FSinglePropertyParams SingleDetails;
+	SingleDetails.NamePlacement = EPropertyNamePlacement::Hidden;
+	SingleDetails.NotifyHook = Editor.Get();
+	SingleDetails.bHideAssetThumbnail = true;
+
+	StaticMeshSelector = PropPlugin.CreateSingleProperty(NodeStaticMesh, "StaticMesh", SingleDetails);	
+
+	UpdateGraphNode();
+}
+
+
+void SGraphNodeStaticMesh::UpdateGraphNode()
+{
+	SGraphNode::UpdateGraphNode();
+}
+
+
+void SGraphNodeStaticMesh::SetDefaultTitleAreaWidget(TSharedRef<SOverlay> DefaultTitleAreaWidget)
+{
+	DefaultTitleAreaWidget->AddSlot()
+		.HAlign(HAlign_Right)
+		.VAlign(VAlign_Center)
+		.Padding(FMargin(5))
+		[
+			SNew(SCheckBox)
+			.OnCheckStateChanged(this, &SGraphNodeStaticMesh::OnExpressionPreviewChanged)
+		.IsChecked(IsExpressionPreviewChecked())
+		.Cursor(EMouseCursor::Default)
+		.Style(FAppStyle::Get(), "Graph.Node.AdvancedView")
+		[
+			SNew(SHorizontalBox)
+			+ SHorizontalBox::Slot()
+		.VAlign(VAlign_Center)
+		.HAlign(HAlign_Center)
+		[
+			SNew(SImage)
+			.Image(GetExpressionPreviewArrow())
+		]
+		]
+		];
+}
+
+
+void SGraphNodeStaticMesh::CreateBelowPinControls(TSharedPtr<SVerticalBox> MainBox)
+{
+	LeftNodeBox->AddSlot()
+		.AutoHeight()
+		.MaxHeight(WidgetSize)
+		.Padding(10.0f, 10.0f, 0.0f, 0.0f)
+		[
+			SNew(SHorizontalBox)
+			.Visibility(ExpressionPreviewVisibility())
+
+		+ SHorizontalBox::Slot()
+		.MaxWidth(WidgetSize)
+		.Padding(5.0f, 5.0f, 5.0f, 5.0f)
+		[
+			AssetThumbnail->MakeThumbnailWidget()
+		]
+		];
+
+	if (StaticMeshSelector.IsValid())
+	{
+		LeftNodeBox->AddSlot()
+			.AutoHeight()
+			.Padding(10.0f, 5.0f, 0.0f, 0.0f)
+			[
+				SNew(SHorizontalBox)
+				.Visibility(ExpressionPreviewVisibility())
+
+			+ SHorizontalBox::Slot()
+			.AutoWidth()
+			.Padding(1.0f, 0.0f, 5.0f, 5.0f)
+			.VAlign(VAlign_Center)
+			.HAlign(HAlign_Left)
+			[
+				StaticMeshSelector.ToSharedRef()
+			]
+			];
+	}
+}
+
+
+void SGraphNodeStaticMesh::OnExpressionPreviewChanged(const ECheckBoxState NewCheckedState)
+{
+	NodeStaticMesh->bCollapsed = (NewCheckedState != ECheckBoxState::Checked);
+	UpdateGraphNode();
+}
+
+
+ECheckBoxState SGraphNodeStaticMesh::IsExpressionPreviewChecked() const
+{
+	return NodeStaticMesh->bCollapsed ? ECheckBoxState::Unchecked : ECheckBoxState::Checked;
+}
+
+
+const FSlateBrush* SGraphNodeStaticMesh::GetExpressionPreviewArrow() const
+{
+	return FCustomizableObjectEditorStyle::Get().GetBrush(NodeStaticMesh->bCollapsed ? TEXT("Nodes.ArrowDown") : TEXT("Nodes.ArrowUp"));
+}
+
+
+EVisibility SGraphNodeStaticMesh::ExpressionPreviewVisibility() const
+{
+	return NodeStaticMesh->bCollapsed ? EVisibility::Collapsed : EVisibility::Visible;
+}
+
+
+#undef LOCTEXT_NAMESPACE
diff --git a/Engine/Plugins/Experimental/Mutable/Source/CustomizableObjectEditor/Private/Nodes/CustomizableObjectNodeTable.cpp b/Engine/Plugins/Experimental/Mutable/Source/CustomizableObjectEditor/Private/Nodes/CustomizableObjectNodeTable.cpp
new file mode 100644
--- /dev/null
+++ b/Engine/Plugins/Experimental/Mutable/Source/CustomizableObjectEditor/Private/Nodes/CustomizableObjectNodeTable.cpp
@@ -0,0 +1,852 @@
+// Copyright Epic Games, Inc. All Rights Reserved.
+
+#include "Nodes/CustomizableObjectNodeTable.h"
+#include "CustomizableObjectEditorModule.h"
+#include "CustomizableObjectEditor.h"
+#include "SCustomizableObjectNodeLayoutBlocksEditor.h"
+#include "CustomizableObjectEditorLogger.h"
+#include "GameplayTagContainer.h"
+
+#include "DataTableUtils.h"
+#include "GraphTraversal.h"
+#include "MutableUtils.h"
+#include "Materials/MaterialInstance.h"
+
+#include "UnrealEditorPortabilityHelpers.h"
+
+#define LOCTEXT_NAMESPACE "CustomizableObjectEditor"
+
+
+void UCustomizableObjectNodeTable::PostBackwardsCompatibleFixup()
+{
+	Super::PostBackwardsCompatibleFixup();
+
+	if (Table)
+	{
+		OnTableChangedDelegateHandle = Table->OnDataTableChanged().AddUObject(this, &UCustomizableObjectNodeTable::OnTableChanged);
+	}
+}
+
+
+void UCustomizableObjectNodeTable::PreEditChange(FProperty* PropertyAboutToChange)
+{
+	Super::PreEditChange(PropertyAboutToChange);
+
+	if (PropertyAboutToChange && PropertyAboutToChange->GetFName() == GET_MEMBER_NAME_CHECKED(UCustomizableObjectNodeTable, Table) &&
+		Table)
+	{
+		Table->OnDataTableChanged().Remove(OnTableChangedDelegateHandle);
+	}
+}
+
+
+void UCustomizableObjectNodeTable::PostEditChangeChainProperty(FPropertyChangedChainEvent& PropertyChangedEvent)
+{
+	Super::PostEditChangeChainProperty(PropertyChangedEvent);
+
+	if (const FProperty* PropertyThatChanged = PropertyChangedEvent.Property)
+	{
+		if (PropertyThatChanged->GetFName() == GET_MEMBER_NAME_CHECKED(UCustomizableObjectNodeTable, Table))
+		{
+			if (Table)
+			{
+				OnTableChangedDelegateHandle = Table->OnDataTableChanged().AddUObject(this, &UCustomizableObjectNodeTable::OnTableChanged);
+			}
+			
+			OnTableChanged();
+		}
+	}
+}
+
+
+FText UCustomizableObjectNodeTable::GetNodeTitle(ENodeTitleType::Type TitleType) const
+{
+	if (Table)
+	{
+		FFormatNamedArguments Args;
+		Args.Add(TEXT("TableName"), FText::FromString(Table->GetName()));
+
+		return FText::Format(LOCTEXT("TableNode_Title", "{TableName}\nTable"), Args);
+	}
+	
+	return LOCTEXT("Mutable Table", "Table");
+}
+
+
+FLinearColor UCustomizableObjectNodeTable::GetNodeTitleColor() const
+{
+	const UEdGraphSchema_CustomizableObject* Schema = GetDefault<UEdGraphSchema_CustomizableObject>();
+	return Schema->GetPinTypeColor(Schema->PC_Object);
+}
+
+
+FText UCustomizableObjectNodeTable::GetTooltipText() const
+{
+	return LOCTEXT("Node_Table_Tooltip", "Represents all the columns of Data Table asset.");
+}
+
+
+void UCustomizableObjectNodeTable::PostEditChangeProperty(FPropertyChangedEvent& PropertyChangedEvent)
+{
+	Super::PostEditChangeProperty(PropertyChangedEvent);
+
+	FProperty* PropertyThatChanged = PropertyChangedEvent.Property;
+	if (PropertyThatChanged && PropertyThatChanged->GetName() == TEXT("Table"))
+	{
+		MarkForReconstruct();
+	}
+}
+
+
+void UCustomizableObjectNodeTable::AllocateDefaultPins(UCustomizableObjectNodeRemapPins* RemapPins)
+{
+	if (!Table)
+	{
+		return;
+	}
+
+	// Getting Struct Pointer
+	const UScriptStruct* TableStruct = Table->GetRowStruct();
+
+	if (!TableStruct)
+	{
+		return;
+	}
+
+	NumProperties = Table->GetColumnTitles().Num();
+
+	// Getting Default Struct Values
+	uint8* DefaultRowData = (uint8*)FMemory::Malloc(TableStruct->GetStructureSize());
+	
+	if (!DefaultRowData)
+	{
+		return;
+	}
+
+	TableStruct->InitializeStruct(DefaultRowData);
+	
+	for (TFieldIterator<FProperty> It(TableStruct); It; ++It)
+	{
+		FProperty* ColumnProperty = *It;
+
+		if (!ColumnProperty)
+		{
+			continue;
+		}
+		
+		const UEdGraphSchema_CustomizableObject* Schema = GetDefault<UEdGraphSchema_CustomizableObject>();
+
+		UEdGraphPin* OutPin = nullptr;
+
+		FString ColumnName = DataTableUtils::GetPropertyExportName(ColumnProperty);
+		FString PinName = ColumnName;
+
+		if (const FSoftObjectProperty* SoftObjectProperty = CastField<FSoftObjectProperty>(ColumnProperty))
+		{
+			UObject* Object = nullptr;
+
+			// Getting default UObject
+			uint8* CellData = SoftObjectProperty->ContainerPtrToValuePtr<uint8>(DefaultRowData);
+
+			if (CellData)
+			{
+				Object = SoftObjectProperty->GetPropertyValue(CellData).LoadSynchronous();
+			}
+
+			if (Object)
+			{
+				if (Object->IsA(USkeletalMesh::StaticClass()) || Object->IsA(UStaticMesh::StaticClass()))
+				{
+					GenerateMeshPins(Object, ColumnName);
+				}
+
+				else if (Object->IsA(UTexture2D::StaticClass()))
+				{
+					UCustomizableObjectNodeTableObjectPinData* PinData = NewObject<UCustomizableObjectNodeTableObjectPinData>(this);
+					PinData->ColumnName = ColumnName;
+
+					OutPin = CustomCreatePin(EGPD_Output, Schema->PC_Image, FName(*PinName), PinData);
+				}
+
+				else if (Object->IsA(UMaterialInstance::StaticClass()))
+				{
+					UCustomizableObjectNodeTableObjectPinData* PinData = NewObject<UCustomizableObjectNodeTableObjectPinData>(this);
+					PinData->ColumnName = ColumnName;
+
+					OutPin = CustomCreatePin(EGPD_Output, Schema->PC_MaterialAsset, FName(*PinName), PinData);
+				}
+			}
+			else
+			{
+				const FText Text = FText::FromString(FString::Printf(TEXT("Could not find a Default Value in Structure member [%s]"),*ColumnName));
+
+				FCustomizableObjectEditorLogger::CreateLog(Text)
+					.Category(ELoggerCategory::General)
+					.Severity(EMessageSeverity::Warning)
+					.Node(*this)
+					.Log();
+			}
+		}
+
+		else if (const FStructProperty* StructProperty = CastField<FStructProperty>(ColumnProperty))
+		{
+			if (StructProperty->Struct == TBaseStructure<FLinearColor>::Get())
+			{
+				OutPin = CustomCreatePin(EGPD_Output, Schema->PC_Color, FName(*PinName));
+			}
+		}
+
+		else if (const FNumericProperty* NumProperty = CastField<FFloatProperty>(ColumnProperty))
+		{
+			OutPin = CustomCreatePin(EGPD_Output, Schema->PC_Float, FName(*PinName));
+		}
+
+		else if (const FObjectProperty* ObjectProperty = CastField<FObjectProperty>(ColumnProperty))
+		{
+			const FText Text = FText::FromString(FString::Printf(TEXT("Asset format not supported in Structure member [%s]. All assets should be Soft References."), *ColumnName));
+
+			FCustomizableObjectEditorLogger::CreateLog(Text)
+				.Category(ELoggerCategory::General)
+				.Severity(EMessageSeverity::Warning)
+				.Node(*this)
+				.Log();
+		}
+
+		if (OutPin)
+		{
+			OutPin->PinFriendlyName = FText::FromString(PinName);
+			OutPin->PinToolTip = PinName;
+			OutPin->SafeSetHidden(false);
+		}
+	}
+
+	// Cleaning Default Structure Pointer
+	TableStruct->DestroyStruct(DefaultRowData);
+	FMemory::Free(DefaultRowData);
+}
+
+void UCustomizableObjectNodeTable::GenerateMeshPins(UObject* Mesh, FString Name)
+{
+	if (USkeletalMesh* SkeletalMesh = Cast<USkeletalMesh>(Mesh))
+	{
+		int NumLODs = Helper_GetLODInfoArray(SkeletalMesh).Num();
+		const UEdGraphSchema_CustomizableObject* Schema = GetDefault<UEdGraphSchema_CustomizableObject>();
+
+		for (int32 LODIndex = 0; LODIndex < NumLODs; ++LODIndex)
+		{
+			int32 NumMaterials = Helper_GetImportedModel(SkeletalMesh)->LODModels[LODIndex].Sections.Num();
+
+			for (int32 MatIndex = 0; MatIndex < NumMaterials; ++MatIndex)
+			{
+				FString TableMeshPinName = Name + FString::Printf(TEXT(" LOD_%d "), LODIndex) + FString::Printf(TEXT("Mat_%d"), MatIndex);
+
+				UCustomizableObjectNodeTableMeshPinData* PinData = NewObject<UCustomizableObjectNodeTableMeshPinData>(this);
+				PinData->ColumnName = Name;
+				PinData->MutableColumnName = TableMeshPinName;
+				PinData->LOD = LODIndex;
+				PinData->Material = MatIndex;
+
+				UEdGraphPin* MeshPin = CustomCreatePin(EGPD_Output, Schema->PC_Mesh, FName(*TableMeshPinName), PinData);
+				MeshPin->PinFriendlyName = FText::FromString(TableMeshPinName);
+				MeshPin->SafeSetHidden(false);
+
+				FSkeletalMeshModel* importModel = Helper_GetImportedModel(SkeletalMesh);
+				if (importModel && importModel->LODModels.Num() > LODIndex)
+				{
+					const uint32 NumberOfUVLayouts = importModel->LODModels[LODIndex].NumTexCoords;
+
+					for (uint32 LayoutIndex = 0; LayoutIndex < NumberOfUVLayouts; ++LayoutIndex)
+					{
+						UCustomizableObjectLayout* Layout = NewObject<UCustomizableObjectLayout>(this);
+						FString LayoutName = TableMeshPinName;
+
+						if (NumberOfUVLayouts > 1)
+						{
+							LayoutName += FString::Printf(TEXT(" UV_%d"), LayoutIndex);
+						}
+
+						if (Layout)
+						{
+							PinData->Layouts.Add(Layout);
+
+							Layout->SetLayout(SkeletalMesh, LODIndex, MatIndex, LayoutIndex);
+							Layout->SetLayoutName(LayoutName);
+						}
+					}
+				}
+			}
+		}
+	}
+	
+	else if (UStaticMesh* StaticMesh = Cast<UStaticMesh>(Mesh))
+	{
+		int NumLODs = StaticMesh->GetRenderData()->LODResources.Num();
+		const UEdGraphSchema_CustomizableObject* Schema = GetDefault<UEdGraphSchema_CustomizableObject>();
+
+		for (int32 LODIndex = 0; LODIndex < NumLODs; ++LODIndex)
+		{
+			int32 NumMaterials = StaticMesh->GetRenderData()->LODResources[LODIndex].Sections.Num();
+
+			for (int32 MatIndex = 0; MatIndex < NumMaterials; ++MatIndex)
+			{
+				FString TableMeshPinName = Name + FString::Printf(TEXT(" LOD_%d "), LODIndex) + FString::Printf(TEXT("Mat_%d"), MatIndex);
+
+				UCustomizableObjectNodeTableMeshPinData* PinData = NewObject<UCustomizableObjectNodeTableMeshPinData>(this);
+				PinData->ColumnName = Name;
+				PinData->MutableColumnName = TableMeshPinName;
+				PinData->LOD = LODIndex;
+				PinData->Material = MatIndex;
+
+				UEdGraphPin* MeshPin = CustomCreatePin(EGPD_Output, Schema->PC_Mesh, FName(*TableMeshPinName), PinData);
+				MeshPin->PinFriendlyName = FText::FromString(TableMeshPinName);
+				MeshPin->SafeSetHidden(false);
+
+				UCustomizableObjectLayout* Layout = NewObject<UCustomizableObjectLayout>(this);
+				FString LayoutName = TableMeshPinName;
+
+				if (Layout)
+				{
+					PinData->Layouts.Add(Layout);
+					Layout->SetLayout(StaticMesh, LODIndex, MatIndex, 0);
+					Layout->SetLayoutName(LayoutName);
+				}
+			}
+		}
+	}
+}
+
+
+bool UCustomizableObjectNodeTable::IsNodeOutDatedAndNeedsRefresh()
+{
+	if (!Table)
+	{
+		return Pins.Num() > 0;
+	}
+
+	// Getting Struct Pointer
+	const UScriptStruct* TableStruct = Table->GetRowStruct();
+
+	if (!TableStruct)
+	{
+		return Pins.Num() > 0;
+	}
+
+	if (NumProperties != Table->GetColumnTitles().Num())
+	{
+		return true;
+	}
+
+	// Getting Default Struct 
+	uint8* DefaultRowData = (uint8*)FMemory::Malloc(TableStruct->GetStructureSize());
+
+	if (!DefaultRowData)
+	{
+		return Pins.Num() > 0;
+	}
+
+	TableStruct->InitializeStruct(DefaultRowData);
+
+	int32  NumPins = 0;
+
+	bool bNeedsUpdate = false;
+
+	for (TFieldIterator<FProperty> It(TableStruct); It && !bNeedsUpdate; ++It)
+	{
+		FProperty* ColumnProperty = *It;
+
+		if (ColumnProperty)
+		{
+			const UEdGraphSchema_CustomizableObject* Schema = GetDefault<UEdGraphSchema_CustomizableObject>();
+
+			UEdGraphPin* OutPin = nullptr;
+
+			if (const FSoftObjectProperty* SoftObjectProperty = CastField<FSoftObjectProperty>(ColumnProperty))
+			{
+				UObject* Object = nullptr;
+
+				// Getting default UObject
+				uint8* CellData = SoftObjectProperty->ContainerPtrToValuePtr<uint8>(DefaultRowData);
+
+				if (CellData)
+				{
+					Object = SoftObjectProperty->GetPropertyValue(CellData).LoadSynchronous();
+				}
+
+				if (!Object)
+				{
+					continue;
+				}
+
+				if (Object->IsA(USkeletalMesh::StaticClass()) || Object->IsA(UStaticMesh::StaticClass()))
+				{
+					if (USkeletalMesh* SkeletalMesh = Cast<USkeletalMesh>(Object))
+					{
+						int NumLODs = Helper_GetLODInfoArray(SkeletalMesh).Num();
+
+						for (int32 LODIndex = 0; LODIndex < NumLODs; ++LODIndex)
+						{
+							int32 NumMaterials = Helper_GetImportedModel(SkeletalMesh)->LODModels[LODIndex].Sections.Num();
+
+							for (int32 MatIndex = 0; MatIndex < NumMaterials; ++MatIndex)
+							{
+								FString PinName = DataTableUtils::GetPropertyExportName(ColumnProperty);
+								PinName += FString::Printf(TEXT(" LOD_%d "), LODIndex) + FString::Printf(TEXT("Mat_%d"), MatIndex);
+
+								if (CheckPinUpdated(PinName, Schema->PC_Mesh))
+								{
+									bNeedsUpdate = true;
+								}
+
+								NumPins++;
+							}
+						}
+					}
+					else if (UStaticMesh* StaticMesh = Cast<UStaticMesh>(Object))
+					{
+						int NumLODs = StaticMesh->GetRenderData()->LODResources.Num();
+
+						for (int32 LODIndex = 0; LODIndex < NumLODs; ++LODIndex)
+						{
+							int32 NumMaterials = StaticMesh->GetRenderData()->LODResources[LODIndex].Sections.Num();
+
+							for (int32 MatIndex = 0; MatIndex < NumMaterials; ++MatIndex)
+							{
+								FString PinName = DataTableUtils::GetPropertyExportName(ColumnProperty);
+								PinName += FString::Printf(TEXT(" LOD_%d "), LODIndex) + FString::Printf(TEXT("Mat_%d"), MatIndex);
+
+								if (CheckPinUpdated(PinName, Schema->PC_Mesh))
+								{
+									bNeedsUpdate = true;
+								}
+
+								NumPins++;
+							}
+						}
+					}
+				}
+				else if (Object->IsA(UTexture2D::StaticClass()))
+				{
+					FString PinName = DataTableUtils::GetPropertyExportName(ColumnProperty);
+					
+					if (CheckPinUpdated(PinName, Schema->PC_Image))
+					{
+						bNeedsUpdate = true;
+					}
+
+					NumPins++;
+				}
+				else if (Object->IsA(UMaterialInstance::StaticClass()))
+				{
+					FString PinName = DataTableUtils::GetPropertyExportName(ColumnProperty);
+					
+					if (CheckPinUpdated(PinName, Schema->PC_MaterialAsset))
+					{
+						bNeedsUpdate = true;
+					}
+
+					NumPins++;
+				}
+				
+			}
+			else if (const FStructProperty* StructProperty = CastField<FStructProperty>(ColumnProperty))
+			{
+				if (StructProperty->Struct == TBaseStructure<FLinearColor>::Get())
+				{
+					FString PinName = DataTableUtils::GetPropertyExportName(ColumnProperty);
+
+					if (CheckPinUpdated(PinName, Schema->PC_Color))
+					{
+						bNeedsUpdate = true;
+					}
+
+					NumPins++;
+				}
+			}
+			else if (const FFloatProperty* NumProperty = CastField<FFloatProperty>(ColumnProperty))
+			{
+				FString PinName = DataTableUtils::GetPropertyExportName(ColumnProperty);
+
+				if (CheckPinUpdated(PinName, Schema->PC_Float))
+				{
+					bNeedsUpdate = true;
+				}
+
+				NumPins++;
+			}
+		}
+	}
+
+	// Cleaning Default Structure Pointer
+	TableStruct->DestroyStruct(DefaultRowData);
+	FMemory::Free(DefaultRowData);
+
+	if (Pins.Num() != NumPins)
+	{
+		bNeedsUpdate = true;
+	}
+
+	return bNeedsUpdate;
+}
+
+
+FString UCustomizableObjectNodeTable::GetRefreshMessage() const
+{
+	return "Node data outdated. Please refresh node.";
+}
+
+
+void UCustomizableObjectNodeTable::RemapPins(const TMap<UEdGraphPin*, UEdGraphPin*>& PinsToRemap)
+{
+	const UEdGraphSchema_CustomizableObject* Schema = GetDefault<UEdGraphSchema_CustomizableObject>();
+
+	for (const TTuple<UEdGraphPin*, UEdGraphPin*>& Pair : PinsToRemap)
+	{
+		if (Pair.Key->PinType.PinCategory == Schema->PC_Mesh)
+		{
+			UCustomizableObjectNodeTableMeshPinData* PinDataOldPin =Cast<UCustomizableObjectNodeTableMeshPinData>(GetPinData(*(Pair.Key)));
+			UCustomizableObjectNodeTableMeshPinData* PinDataNewPin =Cast<UCustomizableObjectNodeTableMeshPinData>(GetPinData(*(Pair.Value)));
+			
+			if (PinDataOldPin && PinDataNewPin)
+			{
+				PinDataOldPin->ColumnName = PinDataNewPin->ColumnName;
+				PinDataOldPin->MutableColumnName = PinDataNewPin->MutableColumnName;
+				PinDataOldPin->LOD = PinDataNewPin->LOD;
+				PinDataOldPin->Material = PinDataNewPin->Material;
+			}
+		}
+	}
+
+	Super::RemapPins(PinsToRemap);
+}
+
+
+bool UCustomizableObjectNodeTable::ProvidesCustomPinRelevancyTest() const
+{
+	return true;
+}
+
+
+bool UCustomizableObjectNodeTable::IsPinRelevant(const UEdGraphPin* Pin) const
+{
+	const UEdGraphSchema_CustomizableObject* Schema = GetDefault<UEdGraphSchema_CustomizableObject>();
+	
+	return Pin->Direction == EGPD_Input &&
+		(Pin->PinType.PinCategory == Schema->PC_MaterialAsset ||
+		Pin->PinType.PinCategory == Schema->PC_Image ||
+		Pin->PinType.PinCategory == Schema->PC_Color ||
+		Pin->PinType.PinCategory == Schema->PC_Float ||
+		Pin->PinType.PinCategory == Schema->PC_Mesh);
+}
+
+
+bool UCustomizableObjectNodeTable::ForceImageMutableMode(const UEdGraphPin* Pin, FGuid ParameterId) const
+{
+	UMaterialInstance* DefaultPinValue = GetColumnDefaultAssetByType<UMaterialInstance>(Pin);
+	
+	if (DefaultPinValue && DefaultPinValue->TextureParameterValues.Num())
+	{
+		TArray<FMaterialParameterInfo> TextureParameterInfo;
+		TArray<FGuid> TextureGuids;
+
+		DefaultPinValue->GetMaterial()->GetAllTextureParameterInfo(TextureParameterInfo, TextureGuids);
+
+		int32 TextureIndex = TextureGuids.Find(ParameterId);
+		
+		if (TextureIndex == INDEX_NONE)
+		{
+			return false;
+		}
+		
+		FName TextureName = TextureParameterInfo[TextureIndex].Name;
+		
+		for (const FTextureParameterValue Texture : DefaultPinValue->TextureParameterValues)
+		{
+			if (TextureName == Texture.ParameterInfo.Name)
+			{
+				return true;
+			}
+		}
+	}
+
+	return false;
+}
+
+
+void UCustomizableObjectNodeTable::SetLayoutInLayoutEditor(UCustomizableObjectLayout* CurrentLayout)
+{
+	TSharedPtr<ICustomizableObjectEditor> Editor = GetGraphEditor();
+
+	if (TSharedPtr<FCustomizableObjectEditor> GraphEditor = StaticCastSharedPtr<FCustomizableObjectEditor>(GetGraphEditor()))
+	{
+		if (GraphEditor->GetLayoutBlocksEditor().IsValid())
+		{
+			GraphEditor->GetLayoutBlocksEditor()->SetCurrentLayout(CurrentLayout);
+		}
+	}
+}
+
+
+UTexture2D* UCustomizableObjectNodeTable::FindReferenceTextureParameter(const UEdGraphPin* Pin, FString ParameterImageName) const
+{
+	UMaterialInterface* Material = GetColumnDefaultAssetByType<UMaterialInterface>(Pin);
+
+	if (Material)
+	{
+		UTexture* OutTexture;
+		bool bFound = Material->GetTextureParameterValue(FName(*ParameterImageName), OutTexture);
+
+		if (bFound)
+		{
+			if (UTexture2D* Texture = Cast<UTexture2D>(OutTexture))
+			{
+				return Texture;
+			}
+		}
+	}
+
+	return nullptr;
+}
+
+
+void UCustomizableObjectNodeTable::GetUVChannel(const UCustomizableObjectLayout* CurrentLayout, TArray<FVector2f>& OutSegments) const
+{
+	FString ColumnName;
+
+	for (const UEdGraphPin* Pin : Pins)
+	{
+		const UCustomizableObjectNodeTableMeshPinData* PinData = Cast<UCustomizableObjectNodeTableMeshPinData > (GetPinData(*Pin));
+		
+		if (!PinData)
+		{
+			continue;
+		}
+
+		bool bFound = false;
+
+		for (const UCustomizableObjectLayout* Layout : PinData->Layouts)
+		{
+			if (Layout == CurrentLayout)
+			{
+				ColumnName = PinData->ColumnName;
+				bFound = true;
+				break;
+			}
+		}
+
+		if (bFound)
+		{
+			break;
+		}
+	}
+
+	const int32 LODIndex = CurrentLayout->GetLOD();
+	const int32 MaterialIndex = CurrentLayout->GetMaterial();
+	const int32 LayoutIndex = CurrentLayout->GetUVChannel();
+
+	if (const USkeletalMesh* SkeletalMesh = GetColumnDefaultAssetByType<USkeletalMesh>(ColumnName))
+	{
+		OutSegments = GetUV(*SkeletalMesh, LODIndex, MaterialIndex, LayoutIndex);
+	}
+	else if (const UStaticMesh* StaticMesh = GetColumnDefaultAssetByType<UStaticMesh>(ColumnName))
+	{
+		OutSegments = GetUV(*StaticMesh, LODIndex, MaterialIndex, LayoutIndex);
+	}
+}
+
+
+void UCustomizableObjectNodeTable::GetUVChannelForPin(const UEdGraphPin* Pin, TArray<FVector2f>& OutSegments, int32 UVChannel) const
+{
+	const UCustomizableObjectNodeTableMeshPinData* PinData = Cast<UCustomizableObjectNodeTableMeshPinData >(GetPinData(*Pin));
+
+	for (int32 LayoutIndex = 0; LayoutIndex < PinData->Layouts.Num(); ++LayoutIndex)
+	{
+		if (PinData->Layouts[LayoutIndex]->GetUVChannel() == UVChannel)
+		{
+			// Getting LOD and Mat index from Layouts
+			const uint32 LODIndex = PinData->Layouts[LayoutIndex]->GetLOD();
+			const uint32 MaterialIndex = PinData->Layouts[LayoutIndex]->GetMaterial();
+
+			if (const USkeletalMesh* SkeletalMesh = GetColumnDefaultAssetByType<USkeletalMesh>(PinData->ColumnName))
+			{
+				OutSegments.Append(GetUV(*SkeletalMesh, LODIndex, MaterialIndex, LayoutIndex));
+			}
+			else if (const UStaticMesh* StaticMesh = GetColumnDefaultAssetByType<UStaticMesh>(PinData->ColumnName))
+			{
+				OutSegments.Append(GetUV(*StaticMesh, LODIndex, MaterialIndex, LayoutIndex));
+			}
+		}
+	}
+}
+
+
+TArray<UCustomizableObjectLayout*> UCustomizableObjectNodeTable::GetLayouts(const UEdGraphPin* Pin) const
+{
+	TArray<UCustomizableObjectLayout*> Result;
+
+	const UCustomizableObjectNodeTableMeshPinData* PinData = Cast<UCustomizableObjectNodeTableMeshPinData >(GetPinData(*Pin));
+
+	if (PinData)
+	{
+		for (int32 LayoutIndex = 0; LayoutIndex < PinData->Layouts.Num(); ++LayoutIndex)
+		{
+			Result.Add(PinData->Layouts[LayoutIndex]);
+		}
+	}
+	
+	return Result;
+}
+
+
+FString UCustomizableObjectNodeTable::GetColumnNameByPin(const UEdGraphPin* Pin) const
+{
+	const UCustomizableObjectNodeTableObjectPinData* PinData = Cast<UCustomizableObjectNodeTableObjectPinData >(GetPinData(*Pin));
+	
+	if (PinData)
+	{
+		return PinData->ColumnName;
+	}
+
+	return FString();
+}
+
+
+FString UCustomizableObjectNodeTable::GetMutableColumnName(const UEdGraphPin* Pin, const int32& LOD) const
+{
+	UCustomizableObjectNodeTableMeshPinData* PinData = Cast<UCustomizableObjectNodeTableMeshPinData >(GetPinData(*Pin));
+
+	if (PinData)
+	{
+		FString ColumnName = PinData->ColumnName;
+		int32 MaterialIndex = PinData->Material;
+
+		for (const UEdGraphPin* NodePin : Pins)
+		{
+			const UCustomizableObjectNodeTableMeshPinData* MeshPinData = Cast<UCustomizableObjectNodeTableMeshPinData >(GetPinData(*NodePin));
+
+			if (MeshPinData && MeshPinData->ColumnName == ColumnName && MeshPinData->LOD == LOD && MeshPinData->Material == MaterialIndex)
+			{
+				return MeshPinData->MutableColumnName;
+			}
+		}
+	}
+
+	return FString();
+}
+
+
+void UCustomizableObjectNodeTable::GetPinLODAndMaterial(const UEdGraphPin* Pin, int32& LOD, int32& Material) const
+{
+	if (const UCustomizableObjectNodeTableMeshPinData* PinData = Cast<UCustomizableObjectNodeTableMeshPinData >(GetPinData(*Pin)))
+	{
+		LOD = PinData->LOD;
+		Material = PinData->Material;
+	}
+	else
+	{
+		LOD = -1;
+		Material = -1;
+	}
+}
+
+
+void UCustomizableObjectNodeTable::GetAnimationColumns(const FString& ColumnName, FString& AnimBPColumnName, FString& AnimSlotColumnName, FString& AnimTagColumnName) const
+{
+	for (const UEdGraphPin* Pin : Pins)
+	{
+		const UCustomizableObjectNodeTableMeshPinData* PinData = Cast<UCustomizableObjectNodeTableMeshPinData >(GetPinData(*Pin));
+
+		if (PinData)
+		{
+			if (PinData->ColumnName == ColumnName)
+			{
+				AnimBPColumnName = PinData->AnimInstanceColumnName;
+				AnimSlotColumnName = PinData->AnimSlotColumnName;
+				AnimTagColumnName = PinData->AnimTagColumnName;
+			}
+		}
+	}
+}
+
+
+void UCustomizableObjectNodeTable::OnTableChanged()
+{	
+	for (const UEdGraphPin* Pin : GetAllNonOrphanPins())
+	{
+		if (Pin->Direction == EGPD_Output)
+		{
+			for (const UEdGraphPin* ConnectedPin : FollowOutputPinArray(*Pin))
+        	{
+        		if (UCustomizableObjectNodeMaterial* NodeMaterial = Cast<UCustomizableObjectNodeMaterial>(ConnectedPin->GetOwningNode()))
+        		{
+        			NodeMaterial->UpdateAllImagesPinMode();
+        		}
+        	}
+		}
+	}
+}
+
+
+bool UCustomizableObjectNodeTable::CheckPinUpdated(const FString& PinName, const FName& PinType) const
+{
+	if (UEdGraphPin* Pin = FindPin(PinName))
+	{
+		if (Pin->PinType.PinCategory != PinType)
+		{
+			return true;
+		}
+	}
+	else
+	{
+		return true;
+	}
+
+	return false;
+}
+
+
+USkeletalMesh* UCustomizableObjectNodeTable::GetSkeletalMeshAt(const UEdGraphPin* Pin, int32 RowIndex) const
+{
+	if (!Table || !Table->GetRowStruct() || !Pin || !Table->GetRowNames().IsValidIndex(RowIndex))
+	{
+		return nullptr;
+	}
+
+	// Getting Struct Pointer
+	const UScriptStruct* TableStruct = Table->GetRowStruct();
+
+	FString ColumnName = GetColumnNameByPin(Pin);
+
+	FName RowName = Table->GetRowNames()[RowIndex];
+
+	for (TFieldIterator<FProperty> It(TableStruct); It; ++It)
+	{
+		FProperty* ColumnProperty = *It;
+
+		if (!ColumnProperty || ColumnName != DataTableUtils::GetPropertyExportName(ColumnProperty))
+		{
+			continue;
+		}
+
+		if (const FSoftObjectProperty* SoftObjectProperty = CastField<FSoftObjectProperty>(ColumnProperty))
+		{
+			if (uint8* RowData = Table->FindRowUnchecked(RowName))
+			{
+				if (uint8* CellData = ColumnProperty->ContainerPtrToValuePtr<uint8>(RowData, 0))
+				{
+					if (UObject* Object = SoftObjectProperty->GetPropertyValue(CellData).LoadSynchronous())
+					{
+						if (USkeletalMesh* SkeletalMesh = Cast<USkeletalMesh>(Object))
+						{
+							return SkeletalMesh;
+						}
+					}
+				}
+			}
+		}
+	}
+
+	return nullptr;
+}
+
+
+#undef LOCTEXT_NAMESPACE
diff --git a/Engine/Plugins/Experimental/Mutable/Source/CustomizableObjectEditor/Private/Nodes/CustomizableObjectNodeTableDetails.cpp b/Engine/Plugins/Experimental/Mutable/Source/CustomizableObjectEditor/Private/Nodes/CustomizableObjectNodeTableDetails.cpp
new file mode 100644
--- /dev/null
+++ b/Engine/Plugins/Experimental/Mutable/Source/CustomizableObjectEditor/Private/Nodes/CustomizableObjectNodeTableDetails.cpp
@@ -0,0 +1,505 @@
+// Copyright Epic Games, Inc. All Rights Reserved.
+
+#include "CustomizableObjectNodeTableDetails.h"
+#include "Nodes/CustomizableObjectNode.h"
+#include "Nodes/CustomizableObjectNodeTable.h"
+
+#include "DetailLayoutBuilder.h"
+#include "DetailCategoryBuilder.h"
+#include "DetailWidgetRow.h"
+
+#include "Engine/DataTable.h"
+#include "Engine/StaticMesh.h"
+#include "Animation/AnimInstance.h"
+#include "GameplayTagContainer.h"
+
+#include "Styling/AppStyle.h"
+
+#define LOCTEXT_NAMESPACE "CustomizableObjectDetails"
+
+TSharedRef<IDetailCustomization> FCustomizableObjectNodeTableDetails::MakeInstance()
+{
+	return MakeShareable(new FCustomizableObjectNodeTableDetails);
+}
+
+
+void FCustomizableObjectNodeTableDetails::CustomizeDetails(IDetailLayoutBuilder& DetailBuilder)
+{
+	Node = 0;
+
+	const IDetailsView* DetailsView = DetailBuilder.GetDetailsView();
+	if (DetailsView->GetSelectedObjects().Num())
+	{
+		Node = Cast<UCustomizableObjectNodeTable>(DetailsView->GetSelectedObjects()[0].Get());
+	}
+
+	if (Node)
+	{
+		IDetailCategoryBuilder& CustomizableObjectCategory = DetailBuilder.EditCategory("CustomizableObject");
+		IDetailCategoryBuilder& UICategory = DetailBuilder.EditCategory("UI");
+		IDetailCategoryBuilder& LayoutCategory = DetailBuilder.EditCategory("LayoutEditor");
+
+		GenerateColumnComboBoxOptions();
+
+		LayoutCategory.AddCustomRow(LOCTEXT("TestName", "Test Name"))
+		[
+			SNew(SVerticalBox)
+			
+			// Mesh Column selection widget
+			+SVerticalBox::Slot()
+			.Padding(0.0f, 5.0f, 6.0f, 0.0f)
+			.AutoHeight()
+			[
+				SNew(SHorizontalBox)
+				+ SHorizontalBox::Slot()
+				.AutoWidth()
+				.Padding(0.0f, 5.0f, 6.0f, 0.0f)
+				[
+					SNew(STextBlock)
+					.Text(LOCTEXT("ColumnText", "Column: "))
+				]
+
+				+ SHorizontalBox::Slot()
+				[
+					SNew(SBorder)
+					.BorderImage(FAppStyle::GetBrush("NoBorder"))
+					.Padding(FMargin(0.0f, 0.0f, 10.0f, 0.0f))
+					[
+						SAssignNew(ColumnComboBox, STextComboBox)
+						.OptionsSource(&ColumnOptionNames)
+						.InitiallySelectedItem(nullptr)
+						.OnSelectionChanged(this, &FCustomizableObjectNodeTableDetails::OnColumnComboBoxSelectionChanged)
+					]
+				]
+			]
+			
+			// Layout selection widget
+			+SVerticalBox::Slot()
+			.Padding(0.0f, 5.0f, 6.0f, 0.0f)
+			.AutoHeight()
+			[
+				SNew(SHorizontalBox)
+				+ SHorizontalBox::Slot()
+				.AutoWidth()
+				.Padding(0.0f, 5.0f, 6.0f, 0.0f)
+				[
+					SNew(STextBlock)
+					.Text(LOCTEXT("LayoutText", "Layout: "))
+				]
+
+				+ SHorizontalBox::Slot()
+				[
+					SNew(SBorder)
+					.BorderImage(FAppStyle::GetBrush("NoBorder"))
+					.Padding(FMargin(0.0f, 0.0f, 10.0f, 0.0f))
+					[
+						SAssignNew(LayoutComboBox, STextComboBox)
+						.OptionsSource(&LayoutOptionNames)
+						.InitiallySelectedItem(nullptr)
+						.OnSelectionChanged(this, &FCustomizableObjectNodeTableDetails::OnLayoutComboBoxSelectionChanged)
+					]
+				]
+			]
+
+			// Animation Blueprint selection widget
+			+SVerticalBox::Slot()
+			.Padding(0.0f, 5.0f, 6.0f, 0.0f)
+			.AutoHeight()
+			[
+				SNew(SHorizontalBox)
+				+ SHorizontalBox::Slot()
+				.AutoWidth()
+				.Padding(0.0f, 5.0f, 6.0f, 0.0f)
+				[
+					SNew(STextBlock)
+					.Text(LOCTEXT("AnimBPText", "Animation Column: "))
+					.Visibility_Lambda([this]() -> EVisibility
+					{
+						if (!AnimComboBox.IsValid())
+						{
+							return EVisibility::Collapsed;
+						}
+
+						return AnimComboBox->GetVisibility();
+					})
+				]
+
+				+ SHorizontalBox::Slot()
+				[
+					SNew(SBorder)
+					.BorderImage(FAppStyle::GetBrush("NoBorder"))
+					.Padding(FMargin(0.0f, 0.0f, 10.0f, 0.0f))
+					[
+						SAssignNew(AnimComboBox, STextComboBox)
+						.Visibility(EVisibility::Collapsed)
+						.OptionsSource(&AnimOptionNames)
+						.OnSelectionChanged(this, &FCustomizableObjectNodeTableDetails::OnAnimInstanceComboBoxSelectionChanged)
+					]
+				]
+			]
+
+			// Animation Slot selection widget
+			+SVerticalBox::Slot()
+			.Padding(0.0f, 5.0f, 6.0f, 0.0f)
+			.AutoHeight()
+			[
+				SNew(SHorizontalBox)
+				+ SHorizontalBox::Slot()
+				.AutoWidth()
+				.Padding(0.0f, 5.0f, 6.0f, 0.0f)
+				[
+					SNew(STextBlock)
+					.Text(LOCTEXT("AnimSlotText", "Animation Slot Column: "))
+					.Visibility_Lambda([this]() -> EVisibility 
+						{
+							if (!AnimSlotComboBox.IsValid())
+							{
+								return EVisibility::Collapsed;
+							}
+
+							return AnimSlotComboBox->GetVisibility();
+						})
+				]
+
+				+ SHorizontalBox::Slot()
+				[
+					SNew(SBorder)
+					.BorderImage(FAppStyle::GetBrush("NoBorder"))
+					.Padding(FMargin(0.0f, 0.0f, 10.0f, 0.0f))
+					[
+						SAssignNew(AnimSlotComboBox, STextComboBox)
+						.Visibility(EVisibility::Collapsed)
+						.OptionsSource(&AnimSlotOptionNames)
+						.OnSelectionChanged(this, &FCustomizableObjectNodeTableDetails::OnAnimSlotComboBoxSelectionChanged)
+					]
+				]
+			]
+
+			// Animation Tags selection widget
+			+SVerticalBox::Slot()
+			.Padding(0.0f, 5.0f, 6.0f, 0.0f)
+			.AutoHeight()
+			[
+				SNew(SHorizontalBox)
+				+ SHorizontalBox::Slot()
+				.AutoWidth()
+				.Padding(0.0f, 5.0f, 6.0f, 0.0f)
+				[
+					SNew(STextBlock)
+					.Text(LOCTEXT("AnimTagsText", "Animation Tags Column: "))
+					.Visibility_Lambda([this]() -> EVisibility 
+						{
+							if (!AnimTagsComboBox.IsValid())
+							{
+								return EVisibility::Collapsed;
+							}
+
+							return AnimTagsComboBox->GetVisibility();
+						})
+				]
+
+				+ SHorizontalBox::Slot()
+				[
+					SNew(SBorder)
+					.BorderImage(FAppStyle::GetBrush("NoBorder"))
+					.Padding(FMargin(0.0f, 0.0f, 10.0f, 0.0f))
+					[
+						SAssignNew(AnimTagsComboBox, STextComboBox)
+						.Visibility(EVisibility::Collapsed)
+						.OptionsSource(&AnimTagsOptionNames)
+						.OnSelectionChanged(this, &FCustomizableObjectNodeTableDetails::OnAnimTagsComboBoxSelectionChanged)
+					]
+				]
+			]
+		];
+	}
+
+}
+
+
+void FCustomizableObjectNodeTableDetails::GenerateColumnComboBoxOptions()
+{
+	ColumnOptionNames.Empty();
+
+	if (Node->Table)
+	{
+		const UScriptStruct* TableStruct = Node->Table->GetRowStruct();
+
+		// we just need the mesh columns
+		for (TFieldIterator<FProperty> It(TableStruct); It; ++It)
+		{
+			FProperty* ColumnProperty = *It;
+
+			if (ColumnProperty)
+			{
+				if (const FSoftObjectProperty* SoftObjectProperty = CastField<FSoftObjectProperty>(ColumnProperty))
+				{
+					if (SoftObjectProperty->PropertyClass->IsChildOf(USkeletalMesh::StaticClass())
+						|| SoftObjectProperty->PropertyClass->IsChildOf(UStaticMesh::StaticClass()))
+					{
+						ColumnOptionNames.Add(MakeShareable(new FString(DataTableUtils::GetPropertyExportName(ColumnProperty))));
+					}
+				}
+			}
+		}
+	}
+}
+
+
+void FCustomizableObjectNodeTableDetails::GenerateLayoutComboBoxOptions()
+{
+	LayoutOptionNames.Empty();
+	
+	for (const UEdGraphPin* Pin : Node->Pins)
+	{
+		const UCustomizableObjectNodeTableMeshPinData* PinData = Cast<UCustomizableObjectNodeTableMeshPinData >(Node->GetPinData(*Pin));
+
+		if (PinData && PinData->ColumnName == *ColumnComboBox->GetSelectedItem().Get())
+		{
+			if (PinData->Layouts.Num() > 1)
+			{
+				for (int32 LayoutIndex = 0; LayoutIndex < PinData->Layouts.Num(); ++LayoutIndex)
+				{
+					FString LayoutName = Pin->PinFriendlyName.ToString() + FString::Printf(TEXT(" UV_%d"), LayoutIndex);
+					LayoutOptionNames.Add(MakeShareable(new FString(LayoutName)));
+				}
+			}
+			else
+			{
+				LayoutOptionNames.Add(MakeShareable(new FString(Pin->PinFriendlyName.ToString())));
+			}
+		}
+	}
+}
+
+
+void FCustomizableObjectNodeTableDetails::GenerateAnimInstanceComboBoxOptions()
+{
+	AnimOptionNames.Empty();
+
+	const UScriptStruct* TableStruct = Node->Table->GetRowStruct();
+
+	for (TFieldIterator<FProperty> It(TableStruct); It; ++It)
+	{
+		if (FProperty* ColumnProperty = *It)
+		{
+			if (const FSoftClassProperty* SoftClassProperty = CastField<FSoftClassProperty>(ColumnProperty))
+			{
+				if (SoftClassProperty->MetaClass->IsChildOf(UAnimInstance::StaticClass()))
+				{
+					TSharedPtr<FString> Option = MakeShareable(new FString(DataTableUtils::GetPropertyExportName(ColumnProperty)));
+					AnimOptionNames.Add(Option);
+
+					for (const UEdGraphPin* Pin : Node->Pins)
+					{
+						const UCustomizableObjectNodeTableMeshPinData* PinData = Cast<UCustomizableObjectNodeTableMeshPinData >(Node->GetPinData(*Pin));
+
+						if (PinData && PinData->ColumnName == *ColumnComboBox->GetSelectedItem().Get() && PinData->AnimInstanceColumnName == *Option)
+						{
+							AnimComboBox->SetSelectedItem(Option);
+							break;
+						}
+					}
+				}
+			}
+		}
+	}
+	
+}
+
+
+void FCustomizableObjectNodeTableDetails::GenerateAnimSlotComboBoxOptions()
+{
+	AnimSlotOptionNames.Empty();
+
+	if (Node->Table)
+	{
+		const UScriptStruct* TableStruct = Node->Table->GetRowStruct();
+
+		// we just need the mesh columns
+		for (TFieldIterator<FProperty> It(TableStruct); It; ++It)
+		{
+			FProperty* ColumnProperty = *It;
+
+			if (ColumnProperty)
+			{
+				if (const FIntProperty* NumProperty = CastField<FIntProperty>(ColumnProperty))
+				{
+					TSharedPtr<FString> Option = MakeShareable(new FString(DataTableUtils::GetPropertyExportName(ColumnProperty)));
+					AnimSlotOptionNames.Add(Option);
+
+					for (const UEdGraphPin* Pin : Node->Pins)
+					{
+						const UCustomizableObjectNodeTableMeshPinData* PinData = Cast<UCustomizableObjectNodeTableMeshPinData >(Node->GetPinData(*Pin));
+
+						if (PinData && PinData->ColumnName == *ColumnComboBox->GetSelectedItem().Get() && PinData->AnimSlotColumnName == *Option)
+						{
+							AnimSlotComboBox->SetSelectedItem(Option);
+							break;
+						}
+					}
+				}
+			}
+		}
+	}
+}
+
+
+void FCustomizableObjectNodeTableDetails::GenerateAnimTagsComboBoxOptions()
+{
+	AnimTagsOptionNames.Empty();
+
+	if (Node->Table)
+	{
+		const UScriptStruct* TableStruct = Node->Table->GetRowStruct();
+
+		// we just need the mesh columns
+		for (TFieldIterator<FProperty> It(TableStruct); It; ++It)
+		{
+			FProperty* ColumnProperty = *It;
+
+			if (ColumnProperty)
+			{
+				if (const FStructProperty* StructProperty = CastField<FStructProperty>(ColumnProperty))
+				{
+					if (StructProperty->Struct == TBaseStructure<FGameplayTagContainer>::Get())
+					{
+						TSharedPtr<FString> Option = MakeShareable(new FString(DataTableUtils::GetPropertyExportName(ColumnProperty)));
+						AnimTagsOptionNames.Add(Option);
+
+						for (const UEdGraphPin* Pin : Node->Pins)
+						{
+							const UCustomizableObjectNodeTableMeshPinData* PinData = Cast<UCustomizableObjectNodeTableMeshPinData >(Node->GetPinData(*Pin));
+
+							if (PinData && PinData->ColumnName == *ColumnComboBox->GetSelectedItem().Get() && PinData->AnimTagColumnName == *Option)
+							{
+								AnimTagsComboBox->SetSelectedItem(Option);
+								break;
+							}
+						}
+					}
+				}
+			}
+		}
+	}
+}
+
+
+void FCustomizableObjectNodeTableDetails::OnColumnComboBoxSelectionChanged(TSharedPtr<FString> Selection, ESelectInfo::Type SelectInfo)
+{
+	if (Selection.IsValid())
+	{
+		if (LayoutComboBox.IsValid())
+		{
+			LayoutComboBox->ClearSelection();
+			LayoutComboBox->RefreshOptions();
+
+			GenerateLayoutComboBoxOptions();
+		}
+
+		if (AnimComboBox.IsValid() && AnimSlotComboBox.IsValid())
+		{
+			AnimComboBox->SetVisibility(EVisibility::Visible);
+			AnimSlotComboBox->SetVisibility(EVisibility::Visible);
+			AnimTagsComboBox->SetVisibility(EVisibility::Visible);
+
+			GenerateAnimInstanceComboBoxOptions();
+			GenerateAnimSlotComboBoxOptions();
+			GenerateAnimTagsComboBoxOptions();
+		}
+
+		Node->SetLayoutInLayoutEditor(nullptr);
+	}
+}
+
+
+void FCustomizableObjectNodeTableDetails::OnLayoutComboBoxSelectionChanged(TSharedPtr<FString> Selection, ESelectInfo::Type SelectInfo)
+{
+	if (Selection.IsValid())
+	{
+		FString ColumnName = *ColumnComboBox->GetSelectedItem();
+		FString LayoutName = *Selection;
+
+		for (const UEdGraphPin* Pin : Node->Pins)
+		{
+			const UCustomizableObjectNodeTableMeshPinData* PinData = Cast<UCustomizableObjectNodeTableMeshPinData >(Node->GetPinData(*Pin));
+			
+			if (PinData && PinData->ColumnName == ColumnName)
+			{
+				for (int32 LayoutIndex = 0; LayoutIndex < PinData->Layouts.Num(); ++LayoutIndex)
+				{
+					if (PinData->Layouts[LayoutIndex]->GetLayoutName() == LayoutName)
+					{
+						Node->SetLayoutInLayoutEditor(PinData->Layouts[LayoutIndex]);
+					}
+				}
+			}
+		}
+	}
+}
+
+
+void FCustomizableObjectNodeTableDetails::OnAnimInstanceComboBoxSelectionChanged(TSharedPtr<FString> Selection, ESelectInfo::Type SelectInfo)
+{
+	if (Selection.IsValid() && ColumnComboBox->GetSelectedItem().IsValid())
+	{
+		FString ColumnName = *ColumnComboBox->GetSelectedItem();
+
+		for (const UEdGraphPin* Pin : Node->Pins)
+		{
+			UCustomizableObjectNodeTableMeshPinData* PinData = Cast< UCustomizableObjectNodeTableMeshPinData>(Node->GetPinData(*Pin));
+
+			if (PinData && PinData->ColumnName == ColumnName)
+			{
+				PinData->AnimInstanceColumnName = *Selection;
+			}
+		}
+
+		Node->MarkPackageDirty();
+	}
+}
+
+
+void FCustomizableObjectNodeTableDetails::OnAnimSlotComboBoxSelectionChanged(TSharedPtr<FString> Selection, ESelectInfo::Type SelectInfo)
+{
+	if (Selection.IsValid() && ColumnComboBox->GetSelectedItem().IsValid())
+	{
+		FString ColumnName = *ColumnComboBox->GetSelectedItem();
+
+		for (const UEdGraphPin* Pin : Node->Pins)
+		{
+			UCustomizableObjectNodeTableMeshPinData* PinData = Cast< UCustomizableObjectNodeTableMeshPinData>(Node->GetPinData(*Pin));
+
+			if (PinData && PinData->ColumnName == ColumnName)
+			{
+				PinData->AnimSlotColumnName = *Selection;
+			}
+		}
+
+		Node->MarkPackageDirty();
+	}
+}
+
+
+void FCustomizableObjectNodeTableDetails::OnAnimTagsComboBoxSelectionChanged(TSharedPtr<FString> Selection, ESelectInfo::Type SelectInfo)
+{
+	if (Selection.IsValid() && ColumnComboBox->GetSelectedItem().IsValid())
+	{
+		FString ColumnName = *ColumnComboBox->GetSelectedItem();
+
+		for (const UEdGraphPin* Pin : Node->Pins)
+		{
+			UCustomizableObjectNodeTableMeshPinData* PinData = Cast< UCustomizableObjectNodeTableMeshPinData>(Node->GetPinData(*Pin));
+
+			if (PinData && PinData->ColumnName == ColumnName)
+			{
+				PinData->AnimTagColumnName = *Selection;
+			}
+		}
+
+		Node->MarkPackageDirty();
+	}
+}
+
+
+#undef LOCTEXT_NAMESPACE
diff --git a/Engine/Plugins/Experimental/Mutable/Source/CustomizableObjectEditor/Private/Nodes/CustomizableObjectNodeTableDetails.h b/Engine/Plugins/Experimental/Mutable/Source/CustomizableObjectEditor/Private/Nodes/CustomizableObjectNodeTableDetails.h
new file mode 100644
--- /dev/null
+++ b/Engine/Plugins/Experimental/Mutable/Source/CustomizableObjectEditor/Private/Nodes/CustomizableObjectNodeTableDetails.h
@@ -0,0 +1,89 @@
+// Copyright Epic Games, Inc. All Rights Reserved.
+
+#pragma once
+
+#include "IDetailCustomization.h"
+#include "Templates/SharedPointer.h"
+
+#include "Types/SlateEnums.h"
+#include "Widgets/Text/STextBlock.h"
+#include "Widgets/Input/STextComboBox.h"
+
+/** Copy Material node details panel. Hides all properties from the inheret Material node. */
+class FCustomizableObjectNodeTableDetails : public IDetailCustomization
+{
+public:
+	
+	static TSharedRef<IDetailCustomization> MakeInstance();
+
+	/** Hides details copied from CustomizableObjectNodeMaterial. */
+	virtual void CustomizeDetails(IDetailLayoutBuilder& DetailBuilder) override;
+
+private:
+
+	// Generates columns combobox options
+	void GenerateColumnComboBoxOptions();
+	
+	// Generates layouts combobox options
+	void GenerateLayoutComboBoxOptions();
+	
+	// Generates Animation Instance combobox options
+	void GenerateAnimInstanceComboBoxOptions();
+
+	// Generates Animation Slot combobox options
+	void GenerateAnimSlotComboBoxOptions();
+	
+	// Generates Animation Tags combobox options
+	void GenerateAnimTagsComboBoxOptions();
+
+	// OnComboBoxSelectionChanged Callback for Column ComboBox
+	void OnColumnComboBoxSelectionChanged(TSharedPtr<FString> Selection, ESelectInfo::Type SelectInfo);
+	
+	// OnComboBoxSelectionChanged Callback for Layout ComboBox
+	void OnLayoutComboBoxSelectionChanged(TSharedPtr<FString> Selection, ESelectInfo::Type SelectInfo);
+	
+	// OnComboBoxSelectionChanged Callback for AnimInstance ComboBox
+	void OnAnimInstanceComboBoxSelectionChanged(TSharedPtr<FString> Selection, ESelectInfo::Type SelectInfo);
+	
+	// OnComboBoxSelectionChanged Callback for Anim Slot ComboBox
+	void OnAnimSlotComboBoxSelectionChanged(TSharedPtr<FString> Selection, ESelectInfo::Type SelectInfo);
+
+	// OnComboBoxSelectionChanged Callback for Anim Tags ComboBox
+	void OnAnimTagsComboBoxSelectionChanged(TSharedPtr<FString> Selection, ESelectInfo::Type SelectInfo);
+
+private:
+
+	// Pointer to the node represented in this details
+	class UCustomizableObjectNodeTable* Node;
+
+	// ComboBox widget to select a column from the NodeTable
+	TSharedPtr<STextComboBox> ColumnComboBox;
+
+	// Array with the name of the table columns as combobox options
+	TArray<TSharedPtr<FString>> ColumnOptionNames;	
+	
+	// ComboBox widget to select a column from the NodeTable
+	TSharedPtr<STextComboBox> LayoutComboBox;
+
+	// Array with the name of the table columns as combobox options
+	TArray<TSharedPtr<FString>> LayoutOptionNames;
+
+	// ComboBox widget to select an Animation Instance column from the NodeTable
+	TSharedPtr<STextComboBox> AnimComboBox;
+
+	// Array with the name of the Animation Instance columns as combobox options
+	TArray<TSharedPtr<FString>> AnimOptionNames;
+
+	// ComboBox widget to select an Animation Slot column from the NodeTable
+	TSharedPtr<STextComboBox> AnimSlotComboBox;
+
+	// Array with the name of the Animation Slot columns as combobox options
+	TArray<TSharedPtr<FString>> AnimSlotOptionNames;
+
+	// ComboBox widget to select an Animation Tags column from the NodeTable
+	TSharedPtr<STextComboBox> AnimTagsComboBox;
+
+	// Array with the name of the Animation Tags columns as combobox options
+	TArray<TSharedPtr<FString>> AnimTagsOptionNames;
+
+};
\ No newline at end of file
diff --git a/Engine/Plugins/Experimental/Mutable/Source/CustomizableObjectEditor/Private/Nodes/CustomizableObjectNodeTextureBinarise.cpp b/Engine/Plugins/Experimental/Mutable/Source/CustomizableObjectEditor/Private/Nodes/CustomizableObjectNodeTextureBinarise.cpp
new file mode 100644
--- /dev/null
+++ b/Engine/Plugins/Experimental/Mutable/Source/CustomizableObjectEditor/Private/Nodes/CustomizableObjectNodeTextureBinarise.cpp
@@ -0,0 +1,65 @@
+// Copyright Epic Games, Inc. All Rights Reserved.
+
+#include "Nodes/CustomizableObjectNodeTextureBinarise.h"
+#include "CustomizableObjectEditorModule.h"
+
+#define LOCTEXT_NAMESPACE "CustomizableObjectEditor"
+
+
+void UCustomizableObjectNodeTextureBinarise::AllocateDefaultPins(UCustomizableObjectNodeRemapPins* RemapPins)
+{
+	const UEdGraphSchema_CustomizableObject* Schema = GetDefault<UEdGraphSchema_CustomizableObject>();
+
+	FString PinName = TEXT("Texture");
+	UEdGraphPin* OutputPin = CustomCreatePin(EGPD_Output, Schema->PC_Image, FName(*PinName));
+	OutputPin->bDefaultValueIsIgnored = true;
+
+	PinName = TEXT("Base Texture");
+	UEdGraphPin* ImagePin = CustomCreatePin(EGPD_Input, Schema->PC_Image, FName(*PinName));
+
+	BaseImagePinReference = FEdGraphPinReference(ImagePin);
+	
+	PinName = TEXT("Threshold");
+	UEdGraphPin* ThresholdPin = CustomCreatePin(EGPD_Input, Schema->PC_Float, FName(*PinName));
+}
+
+
+UEdGraphPin* UCustomizableObjectNodeTextureBinarise::GetBaseImagePin() const
+{
+	return BaseImagePinReference.Get();
+}
+
+
+FText UCustomizableObjectNodeTextureBinarise::GetNodeTitle(ENodeTitleType::Type TitleType) const
+{
+	return LOCTEXT("Texture_Binarise", "Texture Binarise");
+}
+
+
+FLinearColor UCustomizableObjectNodeTextureBinarise::GetNodeTitleColor() const
+{
+	const UEdGraphSchema_CustomizableObject* Schema = GetDefault<UEdGraphSchema_CustomizableObject>();
+	return Schema->GetPinTypeColor(Schema->PC_Image);
+}
+
+
+FText UCustomizableObjectNodeTextureBinarise::GetTooltipText() const
+{
+	return LOCTEXT("Texture_Binarise_Tooltip", "Turns a base texture into black and white using a threshold.");
+}
+
+
+void UCustomizableObjectNodeTextureBinarise::Serialize(FArchive& Ar)
+{
+	Super::Serialize(Ar);
+
+	Ar.UsingCustomVersion(FCustomizableObjectCustomVersion::GUID);
+
+	if (Ar.CustomVer(FCustomizableObjectCustomVersion::GUID) < FCustomizableObjectCustomVersion::PinsNamesImageToTexture)
+	{
+		BaseImagePinReference = FEdGraphPinReference(FindPin(TEXT("Base Image")));
+	}
+}
+
+
+#undef LOCTEXT_NAMESPACE
diff --git a/Engine/Plugins/Experimental/Mutable/Source/CustomizableObjectEditor/Private/Nodes/CustomizableObjectNodeTextureColourMap.cpp b/Engine/Plugins/Experimental/Mutable/Source/CustomizableObjectEditor/Private/Nodes/CustomizableObjectNodeTextureColourMap.cpp
new file mode 100644
--- /dev/null
+++ b/Engine/Plugins/Experimental/Mutable/Source/CustomizableObjectEditor/Private/Nodes/CustomizableObjectNodeTextureColourMap.cpp
@@ -0,0 +1,51 @@
+// Copyright Epic Games, Inc. All Rights Reserved.
+
+#include "Nodes/CustomizableObjectNodeTextureColourMap.h"
+#include "CustomizableObjectEditorModule.h"
+
+#define LOCTEXT_NAMESPACE "CustomizableObjectEditor"
+
+
+void UCustomizableObjectNodeTextureColourMap::AllocateDefaultPins(UCustomizableObjectNodeRemapPins* RemapPins)
+{
+	const UEdGraphSchema_CustomizableObject* Schema = GetDefault<UEdGraphSchema_CustomizableObject>();
+
+	FString PinName = TEXT("Texture");
+	UEdGraphPin* OutputPin = CustomCreatePin(EGPD_Output, Schema->PC_Image, FName(*PinName));
+	OutputPin->bDefaultValueIsIgnored = true;
+
+	PinName = TEXT("Base");
+	UEdGraphPin* SourcePin = CustomCreatePin(EGPD_Input, Schema->PC_Image, FName(*PinName));
+	SourcePin->bDefaultValueIsIgnored = true;
+
+	PinName = TEXT("Mask");
+	UEdGraphPin* MaskPin = CustomCreatePin(EGPD_Input, Schema->PC_Image, FName(*PinName));
+	MaskPin->bDefaultValueIsIgnored = true;
+
+	PinName = TEXT("Map");
+	UEdGraphPin* GradientPin = CustomCreatePin(EGPD_Input, Schema->PC_Image, FName(*PinName));
+	GradientPin->bDefaultValueIsIgnored = true;
+
+}
+
+
+FText UCustomizableObjectNodeTextureColourMap::GetNodeTitle(ENodeTitleType::Type TitleType) const
+{
+	return LOCTEXT("Texture_Map", "Texture Colour Map");
+}
+
+
+FLinearColor UCustomizableObjectNodeTextureColourMap::GetNodeTitleColor() const
+{
+	const UEdGraphSchema_CustomizableObject* Schema = GetDefault<UEdGraphSchema_CustomizableObject>();
+	return Schema->GetPinTypeColor(Schema->PC_Image);
+}
+
+
+FText UCustomizableObjectNodeTextureColourMap::GetTooltipText() const
+{
+	return LOCTEXT("Texture_Gradient_Sample_Tooltip", "Map colours of map using values form image.");
+}
+
+#undef LOCTEXT_NAMESPACE
+
diff --git a/Engine/Plugins/Experimental/Mutable/Source/CustomizableObjectEditor/Private/Nodes/CustomizableObjectNodeTextureFromChannels.cpp b/Engine/Plugins/Experimental/Mutable/Source/CustomizableObjectEditor/Private/Nodes/CustomizableObjectNodeTextureFromChannels.cpp
new file mode 100644
--- /dev/null
+++ b/Engine/Plugins/Experimental/Mutable/Source/CustomizableObjectEditor/Private/Nodes/CustomizableObjectNodeTextureFromChannels.cpp
@@ -0,0 +1,54 @@
+// Copyright Epic Games, Inc. All Rights Reserved.
+
+#include "Nodes/CustomizableObjectNodeTextureFromChannels.h"
+#include "CustomizableObjectEditorModule.h"
+
+#define LOCTEXT_NAMESPACE "CustomizableObjectEditor"
+
+
+void UCustomizableObjectNodeTextureFromChannels::AllocateDefaultPins(UCustomizableObjectNodeRemapPins* RemapPins)
+{
+	const UEdGraphSchema_CustomizableObject* Schema = GetDefault<UEdGraphSchema_CustomizableObject>();
+
+	FString PinName = TEXT("Texture");
+	UEdGraphPin* OutputPin = CustomCreatePin(EGPD_Output, Schema->PC_Image, FName(*PinName));
+	OutputPin->bDefaultValueIsIgnored = true;
+
+	PinName = TEXT("R");
+	UEdGraphPin* RPin = CustomCreatePin(EGPD_Input, Schema->PC_Image, FName(*PinName));
+	RPin->bDefaultValueIsIgnored = true;
+
+	PinName = TEXT("G");
+	UEdGraphPin* GPin = CustomCreatePin(EGPD_Input, Schema->PC_Image, FName(*PinName));
+	GPin->bDefaultValueIsIgnored = true;
+
+	PinName = TEXT("B");
+	UEdGraphPin* BPin = CustomCreatePin(EGPD_Input, Schema->PC_Image, FName(*PinName));
+	BPin->bDefaultValueIsIgnored = true;
+
+	PinName = TEXT("A");
+	UEdGraphPin* APin = CustomCreatePin(EGPD_Input, Schema->PC_Image, FName(*PinName));
+	APin->bDefaultValueIsIgnored = true;
+}
+
+
+FText UCustomizableObjectNodeTextureFromChannels::GetNodeTitle(ENodeTitleType::Type TitleType) const
+{
+	return LOCTEXT("Texture_From_Channels", "Make Texture");
+}
+
+
+FLinearColor UCustomizableObjectNodeTextureFromChannels::GetNodeTitleColor() const
+{
+	const UEdGraphSchema_CustomizableObject* Schema = GetDefault<UEdGraphSchema_CustomizableObject>();
+	return Schema->GetPinTypeColor(Schema->PC_Image);
+}
+
+
+FText UCustomizableObjectNodeTextureFromChannels::GetTooltipText() const
+{
+	return LOCTEXT("Texture_From_Channels_Tooltip", "Make a colored texture with transparency from four grayscale textures that represent the values of each RGBA channel.");
+}
+
+#undef LOCTEXT_NAMESPACE
+
diff --git a/Engine/Plugins/Experimental/Mutable/Source/CustomizableObjectEditor/Private/Nodes/CustomizableObjectNodeTextureFromColor.cpp b/Engine/Plugins/Experimental/Mutable/Source/CustomizableObjectEditor/Private/Nodes/CustomizableObjectNodeTextureFromColor.cpp
new file mode 100644
--- /dev/null
+++ b/Engine/Plugins/Experimental/Mutable/Source/CustomizableObjectEditor/Private/Nodes/CustomizableObjectNodeTextureFromColor.cpp
@@ -0,0 +1,41 @@
+// Copyright Epic Games, Inc. All Rights Reserved.
+
+#include "Nodes/CustomizableObjectNodeTextureFromColor.h"
+#include "CustomizableObjectEditorModule.h"
+
+#define LOCTEXT_NAMESPACE "CustomizableObjectEditor"
+
+
+void UCustomizableObjectNodeTextureFromColor::AllocateDefaultPins(UCustomizableObjectNodeRemapPins* RemapPins)
+{
+	const UEdGraphSchema_CustomizableObject* Schema = GetDefault<UEdGraphSchema_CustomizableObject>();
+
+	FString PinName = TEXT("Texture");
+	UEdGraphPin* OutputPin = CustomCreatePin(EGPD_Output, Schema->PC_Image, FName(*PinName));
+	OutputPin->bDefaultValueIsIgnored = true;
+
+	PinName = TEXT("Color");
+	UEdGraphPin* RPin = CustomCreatePin(EGPD_Input, Schema->PC_Color, FName(*PinName));
+	RPin->bDefaultValueIsIgnored = true;
+}
+
+
+FText UCustomizableObjectNodeTextureFromColor::GetNodeTitle(ENodeTitleType::Type TitleType) const
+{
+	return LOCTEXT("Texture_From_Color", "Texture From Color");
+}
+
+
+FLinearColor UCustomizableObjectNodeTextureFromColor::GetNodeTitleColor() const
+{
+	const UEdGraphSchema_CustomizableObject* Schema = GetDefault<UEdGraphSchema_CustomizableObject>();
+	return Schema->GetPinTypeColor(Schema->PC_Image);
+}
+
+
+FText UCustomizableObjectNodeTextureFromColor::GetTooltipText() const
+{
+	return LOCTEXT("Texture_From_Color_Tooltip", "Creates a flat color texture from the color provided.");
+}
+
+#undef LOCTEXT_NAMESPACE
diff --git a/Engine/Plugins/Experimental/Mutable/Source/CustomizableObjectEditor/Private/Nodes/CustomizableObjectNodeTextureInterpolate.cpp b/Engine/Plugins/Experimental/Mutable/Source/CustomizableObjectEditor/Private/Nodes/CustomizableObjectNodeTextureInterpolate.cpp
new file mode 100644
--- /dev/null
+++ b/Engine/Plugins/Experimental/Mutable/Source/CustomizableObjectEditor/Private/Nodes/CustomizableObjectNodeTextureInterpolate.cpp
@@ -0,0 +1,68 @@
+// Copyright Epic Games, Inc. All Rights Reserved.
+
+#include "Nodes/CustomizableObjectNodeTextureInterpolate.h"
+#include "CustomizableObjectEditorModule.h"
+
+#define LOCTEXT_NAMESPACE "CustomizableObjectEditor"
+
+
+UCustomizableObjectNodeTextureInterpolate::UCustomizableObjectNodeTextureInterpolate()
+	: Super()
+{
+	NumTargets = 2;
+}
+
+
+void UCustomizableObjectNodeTextureInterpolate::PostEditChangeProperty(FPropertyChangedEvent& PropertyChangedEvent)
+{
+	FProperty* PropertyThatChanged = PropertyChangedEvent.Property;
+	if ( PropertyThatChanged && PropertyThatChanged->GetName() == TEXT("NumTargets") )
+	{
+		ReconstructNode();
+	}
+
+	Super::PostEditChangeProperty(PropertyChangedEvent);
+}
+
+
+void UCustomizableObjectNodeTextureInterpolate::AllocateDefaultPins(UCustomizableObjectNodeRemapPins* RemapPins)
+{
+	const UEdGraphSchema_CustomizableObject* Schema = GetDefault<UEdGraphSchema_CustomizableObject>();
+
+	FString PinName = TEXT("Texture");
+	UEdGraphPin* OutputPin = CustomCreatePin(EGPD_Output, Schema->PC_Image, FName(*PinName));
+	OutputPin->bDefaultValueIsIgnored = true;
+	
+	PinName = TEXT("Factor");
+	UEdGraphPin* FactorPin = CustomCreatePin(EGPD_Input, Schema->PC_Float, FName(*PinName));
+	FactorPin->bDefaultValueIsIgnored = true;
+
+	for ( int LayerIndex = 0; LayerIndex < NumTargets; ++LayerIndex )
+	{
+		PinName = FString::Printf(TEXT("Target %d"), LayerIndex);
+		UEdGraphPin* LayerPin = CustomCreatePin(EGPD_Input, Schema->PC_Image, FName(*PinName));
+		LayerPin->bDefaultValueIsIgnored = true;
+	}
+}
+
+
+FText UCustomizableObjectNodeTextureInterpolate::GetNodeTitle(ENodeTitleType::Type TitleType) const
+{
+	return LOCTEXT("Texture_Interpolate", "Texture Interpolate");
+}
+
+
+FLinearColor UCustomizableObjectNodeTextureInterpolate::GetNodeTitleColor() const
+{
+	const UEdGraphSchema_CustomizableObject* Schema = GetDefault<UEdGraphSchema_CustomizableObject>();
+	return Schema->GetPinTypeColor(Schema->PC_Image);
+}
+
+
+FText UCustomizableObjectNodeTextureInterpolate::GetTooltipText() const
+{
+	return LOCTEXT("Texture_Interpolate_Tooltip", "Changes between textures gradually, allowing to fully apply one texture, the next one, or any proportion between them.");
+}
+
+#undef LOCTEXT_NAMESPACE
+
