[mutable] Move Mutable plugin from the restricted folder to the experimental plugins folder.

diff --git a/Engine/Plugins/Experimental/Mutable/Source/CustomizableObjectEditor/Private/Nodes/CustomizableObjectNodeTextureInvert.cpp b/Engine/Plugins/Experimental/Mutable/Source/CustomizableObjectEditor/Private/Nodes/CustomizableObjectNodeTextureInvert.cpp
new file mode 100644
--- /dev/null
+++ b/Engine/Plugins/Experimental/Mutable/Source/CustomizableObjectEditor/Private/Nodes/CustomizableObjectNodeTextureInvert.cpp
@@ -0,0 +1,62 @@
+// Copyright Epic Games, Inc. All Rights Reserved.
+
+#include "Nodes/CustomizableObjectNodeTextureInvert.h"
+#include "CustomizableObjectEditorModule.h"
+
+#define LOCTEXT_NAMESPACE "CustomizableObjectEditor"
+
+
+void UCustomizableObjectNodeTextureInvert::AllocateDefaultPins(UCustomizableObjectNodeRemapPins* RemapPins)
+{
+	const UEdGraphSchema_CustomizableObject* Schema = GetDefault<UEdGraphSchema_CustomizableObject>();
+
+	FString PinName = TEXT("Texture");
+	UEdGraphPin* OutputPin = CustomCreatePin(EGPD_Output, Schema->PC_Image, FName(*PinName));
+	OutputPin->bDefaultValueIsIgnored = true;
+
+	PinName = TEXT("Base Texture");
+	UEdGraphPin* ImagePin = CustomCreatePin(EGPD_Input, Schema->PC_Image, FName(*PinName));
+
+	BaseImagePinReference = FEdGraphPinReference(ImagePin);
+}
+
+
+UEdGraphPin* UCustomizableObjectNodeTextureInvert::GetBaseImagePin() const
+{
+	return BaseImagePinReference.Get();
+}
+
+
+FText UCustomizableObjectNodeTextureInvert::GetNodeTitle(ENodeTitleType::Type TitleType)const
+{
+	return LOCTEXT("Texture_Invert", "Texture Invert");
+}
+
+
+FLinearColor UCustomizableObjectNodeTextureInvert::GetNodeTitleColor() const
+{
+	const UEdGraphSchema_CustomizableObject* Schema = GetDefault<UEdGraphSchema_CustomizableObject>();
+	return Schema->GetPinTypeColor(Schema->PC_Image);
+}
+
+
+FText UCustomizableObjectNodeTextureInvert::GetTooltipText() const
+{
+	return LOCTEXT("Texture_Invert_Tooltip", "Inverts the colors of a base texture.");
+}
+
+
+void UCustomizableObjectNodeTextureInvert::Serialize(FArchive& Ar)
+{
+	Super::Serialize(Ar);
+
+	Ar.UsingCustomVersion(FCustomizableObjectCustomVersion::GUID);
+
+	if (Ar.CustomVer(FCustomizableObjectCustomVersion::GUID) < FCustomizableObjectCustomVersion::PinsNamesImageToTexture)
+	{
+		BaseImagePinReference = FEdGraphPinReference(FindPin(TEXT("Base Image")));
+	}
+}
+
+
+#undef LOCTEXT_NAMESPACE
diff --git a/Engine/Plugins/Experimental/Mutable/Source/CustomizableObjectEditor/Private/Nodes/CustomizableObjectNodeTextureLayer.cpp b/Engine/Plugins/Experimental/Mutable/Source/CustomizableObjectEditor/Private/Nodes/CustomizableObjectNodeTextureLayer.cpp
new file mode 100644
--- /dev/null
+++ b/Engine/Plugins/Experimental/Mutable/Source/CustomizableObjectEditor/Private/Nodes/CustomizableObjectNodeTextureLayer.cpp
@@ -0,0 +1,160 @@
+// Copyright Epic Games, Inc. All Rights Reserved.
+
+#include "Nodes/CustomizableObjectNodeTextureLayer.h"
+#include "CustomizableObjectEditorModule.h"
+#include "UnrealEditorPortabilityHelpers.h"
+
+#define LOCTEXT_NAMESPACE "CustomizableObjectEditor"
+
+TArray<UCustomizableObjectNodeTextureLayer::TextureLayerLayersData> UCustomizableObjectNodeTextureLayer::ArrayDestroyedNodes;
+
+UCustomizableObjectNodeTextureLayer::UCustomizableObjectNodeTextureLayer()
+	: Super()
+{
+
+}
+
+
+void UCustomizableObjectNodeTextureLayer::PostEditChangeProperty(FPropertyChangedEvent& PropertyChangedEvent)
+{
+	FProperty* PropertyThatChanged = PropertyChangedEvent.Property;
+	if ( PropertyThatChanged && PropertyThatChanged->GetName() == TEXT("Layers") )
+	{
+		ReconstructNode();
+	}
+
+	Super::PostEditChangeProperty(PropertyChangedEvent);
+}
+
+
+void UCustomizableObjectNodeTextureLayer::Serialize(FArchive& Ar)
+{
+	Super::Serialize(Ar);
+
+	Ar.UsingCustomVersion(FCustomizableObjectCustomVersion::GUID);
+
+	if (Ar.CustomVer(FCustomizableObjectCustomVersion::GUID) < FCustomizableObjectCustomVersion::PinsNamesImageToTexture)
+	{
+		OutputPinReference = FEdGraphPinReference(FindPin(TEXT("Image")));
+	}
+}
+
+
+void UCustomizableObjectNodeTextureLayer::AllocateDefaultPins(UCustomizableObjectNodeRemapPins* RemapPins)
+{
+	const UEdGraphSchema_CustomizableObject* Schema = GetDefault<UEdGraphSchema_CustomizableObject>();
+
+	UEdGraphPin* OutputPin = CustomCreatePin(EGPD_Output, Schema->PC_Image, FName(TEXT("Texture")));
+	OutputPin->bDefaultValueIsIgnored = true;
+
+	OutputPinReference = FEdGraphPinReference(OutputPin);
+
+	for (int LayerIndex = Layers.Num() - 1; LayerIndex >=0 ; --LayerIndex )
+	{
+		FString PinName = FString::Printf( TEXT("Layer %d "), LayerIndex );
+		UEdGraphPin* LayerPin = CustomCreatePin(EGPD_Input, Schema->PC_Image, *PinName);
+		LayerPin->bDefaultValueIsIgnored = true;
+
+		PinName = FString::Printf( TEXT("Mask %d "), LayerIndex );
+		UEdGraphPin* MaskPin = CustomCreatePin(EGPD_Input, Schema->PC_Image, *PinName);
+		MaskPin->bDefaultValueIsIgnored = true;
+	}
+
+	UEdGraphPin* BasePin = CustomCreatePin(EGPD_Input, Schema->PC_Image, FName(TEXT("Base")));
+	BasePin->bDefaultValueIsIgnored = true;
+}
+
+
+FText UCustomizableObjectNodeTextureLayer::GetNodeTitle(ENodeTitleType::Type TitleType) const
+{
+	return LOCTEXT("Texture_Layer", "Texture Layer");
+}
+
+
+FLinearColor UCustomizableObjectNodeTextureLayer::GetNodeTitleColor() const
+{
+	const UEdGraphSchema_CustomizableObject* Schema = GetDefault<UEdGraphSchema_CustomizableObject>();
+	return Schema->GetPinTypeColor(Schema->PC_Image);
+}
+
+
+void UCustomizableObjectNodeTextureLayer::PrepareForCopying()
+{
+	// Destroy previous copy with possible out of date information just in case
+	int32 IndexToRemove = GetIndexInArrayDestroyedNodes(this);
+
+	if (IndexToRemove != -1)
+	{
+		ArrayDestroyedNodes.RemoveAt(IndexToRemove);
+	}
+
+	TextureLayerLayersData Data;
+	Data.Name = GetName();
+	Data.Graph = GetGraph();
+	Data.Layers = Layers;
+	ArrayDestroyedNodes.Add(Data);
+}
+
+void UCustomizableObjectNodeTextureLayer::PostPasteNode()
+{
+	int32 IndexToRemove = GetIndexInArrayDestroyedNodes(this);
+
+	if (IndexToRemove != -1)
+	{
+		Layers = ArrayDestroyedNodes[IndexToRemove].Layers;
+		ArrayDestroyedNodes.RemoveAt(IndexToRemove);
+	}
+}
+
+
+bool UCustomizableObjectNodeTextureLayer::CanConnect(const UEdGraphPin* InOwnedInputPin, const UEdGraphPin* InOutputPin, bool& bOutIsOtherNodeBlocklisted, bool& bOutArePinsCompatible) const
+{
+	const UEdGraphSchema_CustomizableObject* Schema = GetDefault<UEdGraphSchema_CustomizableObject>();
+
+	// In layer pins we accept both colors and images
+	for ( int LayerIndex = 0; LayerIndex < Layers.Num(); ++LayerIndex )
+	{
+		if (LayerPin(LayerIndex) == InOwnedInputPin )
+		{
+			return (InOutputPin->PinType.PinCategory==Helper_GetPinCategory(Schema->PC_Image))
+				|| (InOutputPin->PinType.PinCategory==Helper_GetPinCategory(Schema->PC_Color));
+		}
+	}
+
+	return Super::CanConnect( InOwnedInputPin, InOutputPin, bOutIsOtherNodeBlocklisted, bOutArePinsCompatible);
+}
+
+
+UEdGraphPin* UCustomizableObjectNodeTextureLayer::OutputPin() const
+{
+	return OutputPinReference.Get();
+}
+
+
+int32 UCustomizableObjectNodeTextureLayer::GetIndexInArrayDestroyedNodes(UCustomizableObjectNodeTextureLayer* Node)
+{
+	const int32 MaxIndex = ArrayDestroyedNodes.Num();
+	FString NodeName = Node->GetName();
+	UEdGraph* NodeGraph = Node->GetGraph();
+	for (int32 i = 0; i < MaxIndex; ++i)
+	{
+		TextureLayerLayersData& DataTemp = ArrayDestroyedNodes[i];
+
+		//if ((DataTemp.Name == NodeName) && (DataTemp.Graph == NodeGraph))
+		if (DataTemp.Name == NodeName) // Copying between graphs would yeld incorredct array data, are the names of the elements unique?
+		{
+			return i;
+		}
+	}
+
+	return -1;
+}
+
+
+FText UCustomizableObjectNodeTextureLayer::GetTooltipText() const
+{
+	return LOCTEXT("Texture_Layer_Tooltip", "Combines multiple textures into one.");
+}
+
+#undef LOCTEXT_NAMESPACE
+
diff --git a/Engine/Plugins/Experimental/Mutable/Source/CustomizableObjectEditor/Private/Nodes/CustomizableObjectNodeTextureParameter.cpp b/Engine/Plugins/Experimental/Mutable/Source/CustomizableObjectEditor/Private/Nodes/CustomizableObjectNodeTextureParameter.cpp
new file mode 100644
--- /dev/null
+++ b/Engine/Plugins/Experimental/Mutable/Source/CustomizableObjectEditor/Private/Nodes/CustomizableObjectNodeTextureParameter.cpp
@@ -0,0 +1,38 @@
+// Copyright Epic Games, Inc. All Rights Reserved.
+
+#include "Nodes/CustomizableObjectNodeTextureParameter.h"
+#include "CustomizableObjectEditorModule.h"
+
+#define LOCTEXT_NAMESPACE "CustomizableObjectEditor"
+
+
+void UCustomizableObjectNodeTextureParameter::AllocateDefaultPins(UCustomizableObjectNodeRemapPins* RemapPins)
+{
+	const UEdGraphSchema_CustomizableObject* Schema = GetDefault<UEdGraphSchema_CustomizableObject>();
+
+	FString PinName = TEXT("Value");
+	UEdGraphPin* ValuePin = CustomCreatePin(EGPD_Output, Schema->PC_Image, FName(*PinName));
+	ValuePin->bDefaultValueIsIgnored = true;
+}
+
+
+FText UCustomizableObjectNodeTextureParameter::GetNodeTitle(ENodeTitleType::Type TitleType) const
+{
+	return LOCTEXT("Texture Parameter", "Texture Parameter");
+}
+
+
+FLinearColor UCustomizableObjectNodeTextureParameter::GetNodeTitleColor() const
+{
+	const UEdGraphSchema_CustomizableObject* Schema = GetDefault<UEdGraphSchema_CustomizableObject>();
+	return Schema->GetPinTypeColor(Schema->PC_Image);
+}
+
+
+FText UCustomizableObjectNodeTextureParameter::GetTooltipText() const
+{
+	return LOCTEXT("Texture_Parameter_Tooltip", "Expose a runtime modifiable texture parameter from the Customizable Object.");
+}
+
+#undef LOCTEXT_NAMESPACE
+
diff --git a/Engine/Plugins/Experimental/Mutable/Source/CustomizableObjectEditor/Private/Nodes/CustomizableObjectNodeTextureProject.cpp b/Engine/Plugins/Experimental/Mutable/Source/CustomizableObjectEditor/Private/Nodes/CustomizableObjectNodeTextureProject.cpp
new file mode 100644
--- /dev/null
+++ b/Engine/Plugins/Experimental/Mutable/Source/CustomizableObjectEditor/Private/Nodes/CustomizableObjectNodeTextureProject.cpp
@@ -0,0 +1,160 @@
+// Copyright Epic Games, Inc. All Rights Reserved.
+
+#include "Nodes/CustomizableObjectNodeTextureProject.h"
+#include "CustomizableObjectEditorModule.h"
+
+#define LOCTEXT_NAMESPACE "CustomizableObjectEditor"
+
+
+UCustomizableObjectNodeTextureProject::UCustomizableObjectNodeTextureProject()
+	: Super()
+{
+	Textures = 1;
+}
+
+
+void UCustomizableObjectNodeTextureProject::PostEditChangeProperty(FPropertyChangedEvent& PropertyChangedEvent)
+{
+	FProperty* PropertyThatChanged = PropertyChangedEvent.Property;
+	if ( PropertyThatChanged && PropertyThatChanged->GetName() == TEXT("Textures") )
+	{
+		ReconstructNode();
+	}
+
+	Super::PostEditChangeProperty(PropertyChangedEvent);
+}
+
+
+void UCustomizableObjectNodeTextureProject::Serialize(FArchive& Ar)
+{
+	Super::Serialize(Ar);
+
+	Ar.UsingCustomVersion(FCustomizableObjectCustomVersion::GUID);
+
+	if (Ar.CustomVer(FCustomizableObjectCustomVersion::GUID) < FCustomizableObjectCustomVersion::PinsNamesImageToTexture)
+	{
+		{
+			uint32 TexturePinsCount = 0;
+			uint32 OutputPinsCount = 0;
+
+			TArray<UEdGraphPin*> NonOrphanPins = GetAllNonOrphanPins();
+			
+			for (const UEdGraphPin* Pin : NonOrphanPins)
+			{
+				if (Pin->GetName().StartsWith(TEXT("Image ")) && Pin->Direction == EGPD_Input)
+				{
+					TexturePinsCount++;
+				}
+				else if (Pin->GetName().StartsWith(TEXT("Image ")) && Pin->Direction == EGPD_Output)
+				{
+					OutputPinsCount++;
+				}
+			}
+
+			TexturePinsReferences.SetNum(OutputPinsCount);
+			for (uint32 Index = 0; Index < TexturePinsCount; ++Index)
+			{
+				TexturePinsReferences.Add(FEdGraphPinReference(FindPin(FString::Printf(TEXT("Image %d "), Index), EGPD_Input)));
+			}
+			
+			OutputPinsReferences.SetNum(TexturePinsCount);
+			for (uint32 Index = 0; Index < TexturePinsCount; ++Index)
+			{
+				OutputPinsReferences.Add(FEdGraphPinReference(FindPin(FString::Printf(TEXT("Image %d "), Index), EGPD_Output)));
+			}
+		}
+	}
+}
+
+
+void UCustomizableObjectNodeTextureProject::AllocateDefaultPins(UCustomizableObjectNodeRemapPins* RemapPins)
+{
+	const UEdGraphSchema_CustomizableObject* Schema = GetDefault<UEdGraphSchema_CustomizableObject>();
+
+	TexturePinsReferences.Reset(0);
+	OutputPinsReferences.Reset(0);
+	
+	for ( int LayerIndex = 0; LayerIndex < Textures; ++LayerIndex )
+	{
+		FString PinName = FString::Printf(TEXT("Texture %d"), LayerIndex);
+		UEdGraphPin* TexturePin = CustomCreatePin(EGPD_Input, Schema->PC_Image, FName(*PinName));
+		TexturePin->bDefaultValueIsIgnored = true;
+
+		TexturePinsReferences.Add(FEdGraphPinReference(TexturePin));
+		
+		PinName = FString::Printf(TEXT("Texture %d"), LayerIndex);
+		UEdGraphPin* OutputPin = CustomCreatePin(EGPD_Output, Schema->PC_Image, FName(*PinName));
+		OutputPin->bDefaultValueIsIgnored = true;
+
+		OutputPinsReferences.Add(FEdGraphPinReference(OutputPin));
+	}
+
+	FString PinName = TEXT("Mesh");
+	UEdGraphPin* MeshPin = CustomCreatePin(EGPD_Input, Schema->PC_Mesh, FName(*PinName));
+	MeshPin->bDefaultValueIsIgnored = true;
+
+	PinName = TEXT("Mesh Mask");
+	UEdGraphPin* MeshMaskPin = CustomCreatePin(EGPD_Input, Schema->PC_Image, FName(*PinName));
+	MeshMaskPin->bDefaultValueIsIgnored = true;
+
+	PinName = TEXT("Fade Start Angle");
+	UEdGraphPin* AngleFadeStartPin = CustomCreatePin(EGPD_Input, Schema->PC_Float, FName(*PinName));
+	AngleFadeStartPin->bDefaultValueIsIgnored = false;
+	AngleFadeStartPin->DefaultValue = "90";
+
+	PinName = TEXT("Fade End Angle");
+	UEdGraphPin* AngleFadeEndPin = CustomCreatePin(EGPD_Input, Schema->PC_Float, FName(*PinName));
+	AngleFadeEndPin->bDefaultValueIsIgnored = false;
+	AngleFadeEndPin->DefaultValue = "95";
+
+	PinName = TEXT("Projector");
+	UEdGraphPin* ProjectorPin = CustomCreatePin(EGPD_Input, Schema->PC_Projector, FName(*PinName));
+	ProjectorPin->bDefaultValueIsIgnored = true;
+}
+
+
+UEdGraphPin* UCustomizableObjectNodeTextureProject::TexturePins(int32 Index) const
+{
+	return TexturePinsReferences[Index].Get();	
+}
+
+
+UEdGraphPin* UCustomizableObjectNodeTextureProject::OutputPins(int32 Index) const
+{
+	return OutputPinsReferences[Index].Get();
+}
+
+
+int32 UCustomizableObjectNodeTextureProject::GetNumTextures() const
+{
+	return TexturePinsReferences.Num();
+}
+
+int32 UCustomizableObjectNodeTextureProject::GetNumOutputs() const
+{
+	return OutputPinsReferences.Num();
+}
+
+
+FText UCustomizableObjectNodeTextureProject::GetNodeTitle(ENodeTitleType::Type TitleType) const
+{
+	return LOCTEXT("Texture_Project", "Texture Project");
+}
+
+
+FLinearColor UCustomizableObjectNodeTextureProject::GetNodeTitleColor() const
+{
+	const UEdGraphSchema_CustomizableObject* Schema = GetDefault<UEdGraphSchema_CustomizableObject>();
+	return Schema->GetPinTypeColor(Schema->PC_Image);
+}
+
+
+FText UCustomizableObjectNodeTextureProject::GetTooltipText() const
+{
+	return LOCTEXT("Texture_Project_Tooltip",
+	"Projects one or more textures on a mesh. Transforms a flat image to a layed out version usable as texture for the input mesh. This node does not modify the mesh section or material.");
+}
+
+
+#undef LOCTEXT_NAMESPACE
+
diff --git a/Engine/Plugins/Experimental/Mutable/Source/CustomizableObjectEditor/Private/Nodes/CustomizableObjectNodeTextureSample.cpp b/Engine/Plugins/Experimental/Mutable/Source/CustomizableObjectEditor/Private/Nodes/CustomizableObjectNodeTextureSample.cpp
new file mode 100644
--- /dev/null
+++ b/Engine/Plugins/Experimental/Mutable/Source/CustomizableObjectEditor/Private/Nodes/CustomizableObjectNodeTextureSample.cpp
@@ -0,0 +1,68 @@
+// Copyright Epic Games, Inc. All Rights Reserved.
+
+#include "Nodes/CustomizableObjectNodeTextureSample.h"
+#include "CustomizableObjectEditorModule.h"
+
+#define LOCTEXT_NAMESPACE "CustomizableObjectEditor"
+
+
+UCustomizableObjectNodeTextureSample::UCustomizableObjectNodeTextureSample()
+	: Super()
+{
+
+}
+
+
+void UCustomizableObjectNodeTextureSample::PostEditChangeProperty(FPropertyChangedEvent& PropertyChangedEvent)
+{
+	//UProperty* PropertyThatChanged = PropertyChangedEvent.Property;
+	//if ( PropertyThatChanged && PropertyThatChanged->GetName() == TEXT("NumLODs") )
+	{
+	}
+
+	Super::PostEditChangeProperty(PropertyChangedEvent);
+}
+
+
+void UCustomizableObjectNodeTextureSample::AllocateDefaultPins(UCustomizableObjectNodeRemapPins* RemapPins)
+{
+	const UEdGraphSchema_CustomizableObject* Schema = GetDefault<UEdGraphSchema_CustomizableObject>();
+
+	FString PinName = TEXT("Color");
+	UEdGraphPin* ColorPin = CustomCreatePin(EGPD_Output, Schema->PC_Color, FName(*PinName));
+	ColorPin->bDefaultValueIsIgnored = true;
+
+	PinName = TEXT("Texture");
+	UEdGraphPin* TexturePin = CustomCreatePin(EGPD_Input, Schema->PC_Image, FName(*PinName));
+	TexturePin->bDefaultValueIsIgnored = true;
+
+	PinName = TEXT("X");
+	UEdGraphPin* XPin = CustomCreatePin(EGPD_Input, Schema->PC_Float, FName(*PinName));
+	XPin->bDefaultValueIsIgnored = true;
+
+	PinName = TEXT("Y");
+	UEdGraphPin* YPin = CustomCreatePin(EGPD_Input, Schema->PC_Float, FName(*PinName));
+	YPin->bDefaultValueIsIgnored = true;
+}
+
+
+FText UCustomizableObjectNodeTextureSample::GetNodeTitle(ENodeTitleType::Type TitleType) const
+{
+	return LOCTEXT("Sample_Texture", "Sample Texture");
+}
+
+
+FLinearColor UCustomizableObjectNodeTextureSample::GetNodeTitleColor() const
+{
+	const UEdGraphSchema_CustomizableObject* Schema = GetDefault<UEdGraphSchema_CustomizableObject>();
+	return Schema->GetPinTypeColor(Schema->PC_Color);
+}
+
+
+FText UCustomizableObjectNodeTextureSample::GetTooltipText() const
+{
+	return LOCTEXT("Texture_Sample_Tooltip","Texture Sample");
+}
+
+#undef LOCTEXT_NAMESPACE
+
diff --git a/Engine/Plugins/Experimental/Mutable/Source/CustomizableObjectEditor/Private/Nodes/CustomizableObjectNodeTextureSwitch.cpp b/Engine/Plugins/Experimental/Mutable/Source/CustomizableObjectEditor/Private/Nodes/CustomizableObjectNodeTextureSwitch.cpp
new file mode 100644
--- /dev/null
+++ b/Engine/Plugins/Experimental/Mutable/Source/CustomizableObjectEditor/Private/Nodes/CustomizableObjectNodeTextureSwitch.cpp
@@ -0,0 +1,42 @@
+// Copyright Epic Games, Inc. All Rights Reserved.
+
+#include "Nodes/CustomizableObjectNodeTextureSwitch.h"
+
+#include "EdGraphSchema_CustomizableObject.h"
+
+#define LOCTEXT_NAMESPACE "CustomizableObjectEditor"
+
+
+void UCustomizableObjectNodeTextureSwitch::Serialize(FArchive& Ar)
+{
+	Super::Serialize(Ar);
+	
+	Ar.UsingCustomVersion(FCustomizableObjectCustomVersion::GUID);
+
+	if (Ar.CustomVer(FCustomizableObjectCustomVersion::GUID) < FCustomizableObjectCustomVersion::PinsNamesImageToTexture)
+	{
+		OutputPinReference = FindPin(TEXT("Image"));	
+	}
+}
+
+
+FString UCustomizableObjectNodeTextureSwitch::GetOutputPinName() const
+{
+	return TEXT("Texture");
+}
+
+
+FName UCustomizableObjectNodeTextureSwitch::GetCategory() const
+{
+	return UEdGraphSchema_CustomizableObject::PC_Image;
+}
+
+
+FString UCustomizableObjectNodeTextureSwitch::GetPinPrefix() const
+{
+	return TEXT("Texture ");
+}
+
+
+#undef LOCTEXT_NAMESPACE
+
diff --git a/Engine/Plugins/Experimental/Mutable/Source/CustomizableObjectEditor/Private/Nodes/CustomizableObjectNodeTextureToChannels.cpp b/Engine/Plugins/Experimental/Mutable/Source/CustomizableObjectEditor/Private/Nodes/CustomizableObjectNodeTextureToChannels.cpp
new file mode 100644
--- /dev/null
+++ b/Engine/Plugins/Experimental/Mutable/Source/CustomizableObjectEditor/Private/Nodes/CustomizableObjectNodeTextureToChannels.cpp
@@ -0,0 +1,92 @@
+// Copyright Epic Games, Inc. All Rights Reserved.
+
+#include "Nodes/CustomizableObjectNodeTextureToChannels.h"
+#include "CustomizableObjectEditorModule.h"
+
+#define LOCTEXT_NAMESPACE "CustomizableObjectEditor"
+
+
+void UCustomizableObjectNodeTextureToChannels::PostEditChangeProperty(FPropertyChangedEvent& PropertyChangedEvent)
+{
+	Super::PostEditChangeProperty(PropertyChangedEvent);
+}
+
+
+void UCustomizableObjectNodeTextureToChannels::Serialize(FArchive& Ar)
+{
+	Super::Serialize(Ar);
+
+	Ar.UsingCustomVersion(FCustomizableObjectCustomVersion::GUID);
+
+	if (Ar.CustomVer(FCustomizableObjectCustomVersion::GUID) < FCustomizableObjectCustomVersion::PinsNamesImageToTexture)
+	{
+		InputPinReference = FEdGraphPinReference(FindPin(TEXT("Image")));
+	}
+}
+
+
+void UCustomizableObjectNodeTextureToChannels::AllocateDefaultPins(UCustomizableObjectNodeRemapPins* RemapPins)
+{
+	const UEdGraphSchema_CustomizableObject* Schema = GetDefault<UEdGraphSchema_CustomizableObject>();
+
+	{
+		FString PinName = TEXT("Texture");
+		UEdGraphPin* InputPin = CustomCreatePin(EGPD_Input, Schema->PC_Image, FName(*PinName));
+		InputPin->bDefaultValueIsIgnored = true;
+
+		InputPinReference = FEdGraphPinReference(InputPin);
+	}
+
+	{
+		FString PinName = TEXT("R");
+		UEdGraphPin* RPin = CustomCreatePin(EGPD_Output, Schema->PC_Image, FName(*PinName));
+		RPin->bDefaultValueIsIgnored = true;
+	}
+
+	{
+		FString PinName = TEXT("G");
+		UEdGraphPin* GPin = CustomCreatePin(EGPD_Output, Schema->PC_Image, FName(*PinName));
+		GPin->bDefaultValueIsIgnored = true;
+	}
+
+	{
+		FString PinName = TEXT("B");
+		UEdGraphPin* BPin = CustomCreatePin(EGPD_Output, Schema->PC_Image, FName(*PinName));
+		BPin->bDefaultValueIsIgnored = true;
+	}
+
+	{
+		FString PinName = TEXT("A");
+		UEdGraphPin* APin = CustomCreatePin(EGPD_Output, Schema->PC_Image, FName(*PinName));
+		APin->bDefaultValueIsIgnored = true;
+	}
+}
+
+
+UEdGraphPin* UCustomizableObjectNodeTextureToChannels::InputPin() const
+{
+	return InputPinReference.Get();
+}
+
+
+FText UCustomizableObjectNodeTextureToChannels::GetNodeTitle(ENodeTitleType::Type TitleType) const
+{
+	return LOCTEXT("Texture_To_Channels", "Break Texture");
+}
+
+
+FLinearColor UCustomizableObjectNodeTextureToChannels::GetNodeTitleColor() const
+{
+	const UEdGraphSchema_CustomizableObject* Schema = GetDefault<UEdGraphSchema_CustomizableObject>();
+	return Schema->GetPinTypeColor(Schema->PC_Image);
+}
+
+
+FText UCustomizableObjectNodeTextureToChannels::GetTooltipText() const
+{
+	return LOCTEXT("Texture_To_Channels_Tooltip", "Get the red, green, blue and alpha channels of a texture as four separate grayscale textures.");
+}
+
+
+#undef LOCTEXT_NAMESPACE
+
diff --git a/Engine/Plugins/Experimental/Mutable/Source/CustomizableObjectEditor/Private/Nodes/CustomizableObjectNodeTextureTransform.cpp b/Engine/Plugins/Experimental/Mutable/Source/CustomizableObjectEditor/Private/Nodes/CustomizableObjectNodeTextureTransform.cpp
new file mode 100644
--- /dev/null
+++ b/Engine/Plugins/Experimental/Mutable/Source/CustomizableObjectEditor/Private/Nodes/CustomizableObjectNodeTextureTransform.cpp
@@ -0,0 +1,71 @@
+// Copyright Epic Games, Inc. All Rights Reserved.
+
+#include "Nodes/CustomizableObjectNodeTextureTransform.h"
+#include "CustomizableObjectEditorModule.h"
+
+#define LOCTEXT_NAMESPACE "CustomizableObjectEditor"
+
+void UCustomizableObjectNodeTextureTransform::AllocateDefaultPins(UCustomizableObjectNodeRemapPins* RemapPins)
+{
+	const UEdGraphSchema_CustomizableObject* Schema = GetDefault<UEdGraphSchema_CustomizableObject>();
+
+	FString PinName = TEXT("Texture");
+	UEdGraphPin* OutputPin = CustomCreatePin(EGPD_Output, Schema->PC_Image, FName(*PinName));
+	OutputPin->bDefaultValueIsIgnored = true;
+
+	PinName = TEXT("Base Texture");
+	UEdGraphPin* ImagePin = CustomCreatePin(EGPD_Input, Schema->PC_Image, FName(*PinName));
+
+	BaseImagePinReference = FEdGraphPinReference(ImagePin);
+	
+	PinName = TEXT("Offset X");
+	UEdGraphPin* OffsetXPin = CustomCreatePin(EGPD_Input, Schema->PC_Float, FName(*PinName));
+	OffsetXPin->bDefaultValueIsIgnored = true;
+	
+	PinName = TEXT("Offset Y");
+	UEdGraphPin* OffsetYPin = CustomCreatePin(EGPD_Input, Schema->PC_Float, FName(*PinName));
+	OffsetYPin->bDefaultValueIsIgnored = true;
+
+	PinName = TEXT("Scale X");
+	UEdGraphPin* ScaleXPin = CustomCreatePin(EGPD_Input, Schema->PC_Float, FName(*PinName));
+	ScaleXPin->bDefaultValueIsIgnored = true;
+
+	PinName = TEXT("Scale Y");
+	UEdGraphPin* ScaleYPin = CustomCreatePin(EGPD_Input, Schema->PC_Float, FName(*PinName));
+	ScaleYPin->bDefaultValueIsIgnored = true;
+	
+	PinName = TEXT("Rotation");
+	UEdGraphPin* RotationPin = CustomCreatePin(EGPD_Input, Schema->PC_Float, FName(*PinName));
+	RotationPin->bDefaultValueIsIgnored = true;
+}
+
+UEdGraphPin* UCustomizableObjectNodeTextureTransform::GetBaseImagePin() const
+{
+	return BaseImagePinReference.Get();
+}
+
+FText UCustomizableObjectNodeTextureTransform::GetNodeTitle(ENodeTitleType::Type TitleType) const
+{
+	return LOCTEXT("Texture_Transform", "Texture Transform");
+}
+
+FLinearColor UCustomizableObjectNodeTextureTransform::GetNodeTitleColor() const
+{
+	const UEdGraphSchema_CustomizableObject* Schema = GetDefault<UEdGraphSchema_CustomizableObject>();
+	return Schema->GetPinTypeColor(Schema->PC_Image);
+}
+
+FText UCustomizableObjectNodeTextureTransform::GetTooltipText() const
+{
+	return LOCTEXT("Texture_Transform_Tooltip", "Applies linear transform, Offset followed by Scale followed by Rotation, to the content of Base Texture. Samples outside the image get tiled.");
+}
+
+void UCustomizableObjectNodeTextureTransform::Serialize(FArchive& Ar)
+{
+	Super::Serialize(Ar);
+
+	Ar.UsingCustomVersion(FCustomizableObjectCustomVersion::GUID);
+}
+
+
+#undef LOCTEXT_NAMESPACE
diff --git a/Engine/Plugins/Experimental/Mutable/Source/CustomizableObjectEditor/Private/Nodes/CustomizableObjectNodeTextureVariation.cpp b/Engine/Plugins/Experimental/Mutable/Source/CustomizableObjectEditor/Private/Nodes/CustomizableObjectNodeTextureVariation.cpp
new file mode 100644
--- /dev/null
+++ b/Engine/Plugins/Experimental/Mutable/Source/CustomizableObjectEditor/Private/Nodes/CustomizableObjectNodeTextureVariation.cpp
@@ -0,0 +1,71 @@
+// Copyright Epic Games, Inc. All Rights Reserved.
+
+#include "Nodes/CustomizableObjectNodeTextureVariation.h"
+#include "CustomizableObjectEditorModule.h"
+
+#include "UnrealEditorPortabilityHelpers.h"
+
+#define LOCTEXT_NAMESPACE "CustomizableObjectEditor"
+
+
+UCustomizableObjectNodeTextureVariation::UCustomizableObjectNodeTextureVariation()
+	: Super()
+{
+
+}
+
+
+void UCustomizableObjectNodeTextureVariation::PostEditChangeProperty(FPropertyChangedEvent& PropertyChangedEvent)
+{
+	FProperty* PropertyThatChanged = PropertyChangedEvent.Property;
+	if ( PropertyThatChanged  )
+	{
+		ReconstructNode();
+	}
+
+	Super::PostEditChangeProperty(PropertyChangedEvent);
+}
+
+
+void UCustomizableObjectNodeTextureVariation::AllocateDefaultPins(UCustomizableObjectNodeRemapPins* RemapPins)
+{
+	const UEdGraphSchema_CustomizableObject* Schema = GetDefault<UEdGraphSchema_CustomizableObject>();
+
+	UEdGraphPin* OutputPin = CustomCreatePin(EGPD_Output, Schema->PC_Image, FName(TEXT("Texture")));
+	OutputPin->bDefaultValueIsIgnored = true;
+
+	for ( int VariationIndex = Variations.Num() - 1; VariationIndex >= 0; --VariationIndex)
+	{
+		FString PinName = FString::Printf( TEXT("Variation %d"), VariationIndex);
+		UEdGraphPin* VariationPin = CustomCreatePin(EGPD_Input, Schema->PC_Image, FName(*PinName), false);
+		VariationPin->bDefaultValueIsIgnored = true;
+
+		FString FriendlyName = FString::Printf(TEXT("Variation %d [%s]"), VariationIndex, *Variations[VariationIndex].Tag);
+		VariationPin->PinFriendlyName = FText::FromString(FriendlyName);
+	}
+
+	UEdGraphPin* DefaultVariation = CustomCreatePin(EGPD_Input, Schema->PC_Image, FName(TEXT("Default")), false);
+	DefaultVariation->bDefaultValueIsIgnored = true;
+}
+
+
+FText UCustomizableObjectNodeTextureVariation::GetNodeTitle(ENodeTitleType::Type TitleType) const
+{
+	return LOCTEXT("Texture_Variation", "Texture Variation");
+}
+
+
+FLinearColor UCustomizableObjectNodeTextureVariation::GetNodeTitleColor() const
+{
+	const UEdGraphSchema_CustomizableObject* Schema = GetDefault<UEdGraphSchema_CustomizableObject>();
+	return Schema->GetPinTypeColor(Schema->PC_Image);
+}
+
+
+FText UCustomizableObjectNodeTextureVariation::GetTooltipText() const
+{
+	return LOCTEXT("Texture_Variation_Tooltip", "Select a texture depending on what tags are active.");
+}
+
+#undef LOCTEXT_NAMESPACE
+
diff --git a/Engine/Plugins/Experimental/Mutable/Source/CustomizableObjectEditor/Private/Nodes/CustomizableObjectNodeUseMaterial.cpp b/Engine/Plugins/Experimental/Mutable/Source/CustomizableObjectEditor/Private/Nodes/CustomizableObjectNodeUseMaterial.cpp
new file mode 100644
--- /dev/null
+++ b/Engine/Plugins/Experimental/Mutable/Source/CustomizableObjectEditor/Private/Nodes/CustomizableObjectNodeUseMaterial.cpp
@@ -0,0 +1,164 @@
+// Copyright Epic Games, Inc. All Rights Reserved.
+
+#include "Nodes/CustomizableObjectNodeUseMaterial.h"
+
+#include "CustomizableObjectPin.h"
+#include "GraphTraversal.h"
+#include "ICustomizableObjectEditor.h"
+
+#include "Nodes/CustomizableObjectNodeMaterial.h"
+
+
+bool FCustomizableObjectNodeUseMaterial::IsNodeOutDatedAndNeedsRefreshWork()
+{
+	UCustomizableObjectNode& Node = GetNode();
+	const FCustomizableObjectNodeParentedMaterial& NodeParentedMaterial = GetNodeParentedMaterial();
+
+	const TMap<FGuid, FEdGraphPinReference>& PinsParameter = GetPinsParameter();
+	
+	const bool bOutdated = [&]()
+	{
+		if (const UCustomizableObjectNodeMaterial* ParentMaterialNode = NodeParentedMaterial.GetParentMaterialNode())
+		{
+			if (ParentMaterialNode->RealMaterialDataHasChanged())
+			{
+				return true;
+			}
+
+			for (const TTuple<FGuid, FEdGraphPinReference> Pair : PinsParameter)
+			{
+				const UEdGraphPin& Pin = *Pair.Value.Get();
+				if (!IsPinOrphan(Pin))
+				{
+					if (UsesImage(Pair.Key) && !ParentMaterialNode->HasParameter(Pair.Key))
+					{
+						return true;
+					}
+				}
+			}
+		}
+
+		return false;
+	}();
+	
+	// Remove previous compilation warnings
+	if (!bOutdated && Node.bHasCompilerMessage)
+	{
+		Node.RemoveWarnings();
+		Node.GetGraph()->NotifyGraphChanged();
+	}
+
+	return bOutdated;
+}
+
+
+void FCustomizableObjectNodeUseMaterial::PreSetParentNodeWork(UCustomizableObject* Object, const FGuid NodeId)
+{
+	const FCustomizableObjectNodeParentedMaterial& NodeParentedMaterial = GetNodeParentedMaterial();
+
+	if (UCustomizableObjectNodeMaterial* ParentMaterialNode = NodeParentedMaterial.GetParentMaterialNode())
+	{
+		ParentMaterialNode->PostReconstructNodeDelegate.Remove(PostReconstructNodeDelegateHandler);
+		ParentMaterialNode->PostImagePinModeChangedDelegate.Remove(PostTextureParameterModeChangedDelegateHandle);
+	}
+}
+
+
+void FCustomizableObjectNodeUseMaterial::PostSetParentNodeWork(UCustomizableObject* Object, const FGuid NodeId)
+{
+	UCustomizableObjectNode& Node = GetNode();
+	const FCustomizableObjectNodeParentedMaterial& NodeParentedMaterial = GetNodeParentedMaterial();
+
+	if (UCustomizableObjectNodeMaterial* ParentMaterialNode = NodeParentedMaterial.GetParentMaterialNode())
+	{
+		PostReconstructNodeDelegateHandler = ParentMaterialNode->PostReconstructNodeDelegate.AddUObject(&Node, &UCustomizableObjectNode::ReconstructNode);
+		PostTextureParameterModeChangedDelegateHandle = ParentMaterialNode->PostImagePinModeChangedDelegate.AddUObject(&Node, &UCustomizableObjectNode::ReconstructNode);
+	}
+	
+	GetNode().ReconstructNode();
+}
+
+
+void FCustomizableObjectNodeUseMaterial::PinConnectionListChangedWork(UEdGraphPin* Pin)
+{
+	if (Pin == OutputPin())
+	{
+		if (const TSharedPtr<ICustomizableObjectEditor> Editor = GetNode().GetGraphEditor())
+		{
+			Editor->UpdateGraphNodeProperties();
+		}
+
+		GetNode().ReconstructNode();
+	}
+}
+
+
+void FCustomizableObjectNodeUseMaterial::CustomRemovePinWork(UEdGraphPin& Pin)
+{
+	TMap<FGuid, FEdGraphPinReference>& PinsParameter = GetPinsParameter();
+	
+	for(TMap<FGuid, FEdGraphPinReference>::TIterator It = PinsParameter.CreateIterator(); It; ++It)
+	{
+		if (It.Value().Get() == &Pin) // We could improve performance if FEdGraphPinReference exposed the pin id.
+		{
+			It.RemoveCurrent();
+			break;
+		}
+	}
+}
+
+
+const UCustomizableObjectNode& FCustomizableObjectNodeUseMaterial::GetNode() const
+{
+	return const_cast<FCustomizableObjectNodeUseMaterial*>(this)->GetNode();
+}
+
+
+const TMap<FGuid, FEdGraphPinReference>& FCustomizableObjectNodeUseMaterial::GetPinsParameter() const
+{
+	return const_cast<FCustomizableObjectNodeUseMaterial*>(this)->GetPinsParameter();
+}
+
+
+bool FCustomizableObjectNodeUseMaterial::UsesImage(const FGuid& ImageId) const
+{
+	if (const UEdGraphPin* Pin = GetUsedImagePin(ImageId))
+	{
+		return FollowInputPin(*Pin) != nullptr;
+	}
+	else
+	{
+		return false;
+	}
+}
+
+
+const UEdGraphPin* FCustomizableObjectNodeUseMaterial::GetUsedImagePin(const FGuid& ImageId) const
+{
+	const TMap<FGuid, FEdGraphPinReference>& PinsParameter = GetPinsParameter();
+
+	if (const FEdGraphPinReference* PinReference = PinsParameter.Find(ImageId))
+	{
+		if (const UEdGraphPin& Pin = *PinReference->Get(); !IsPinOrphan(Pin)) // We always have a valid pin reference. If it is nullptr, it means that something went wrong.
+		{
+			return &Pin;
+		}
+	}
+	
+	return nullptr;
+}
+
+
+void FCustomizableObjectNodeUseMaterial::PostBackwardsCompatibleFixupWork()
+{
+	UCustomizableObjectNode& Node = GetNode();
+	const FCustomizableObjectNodeParentedMaterial& NodeParentedMaterial = GetNodeParentedMaterial();
+
+	if (UCustomizableObjectNodeMaterial* ParentMaterialNode = NodeParentedMaterial.GetParentMaterialNode())
+	{
+		PostReconstructNodeDelegateHandler = ParentMaterialNode->PostReconstructNodeDelegate.AddUObject(&Node, &UCustomizableObjectNode::ReconstructNode);
+		PostTextureParameterModeChangedDelegateHandle = ParentMaterialNode->PostImagePinModeChangedDelegate.AddUObject(&Node, &UCustomizableObjectNode::ReconstructNode);
+	}
+
+	Node.ReconstructNode(); // Reconstruct the node since the parent could have changed while not loaded. 
+}
diff --git a/Engine/Plugins/Experimental/Mutable/Source/CustomizableObjectEditor/Private/Nodes/SCustomizableObjectNodeMaterial.cpp b/Engine/Plugins/Experimental/Mutable/Source/CustomizableObjectEditor/Private/Nodes/SCustomizableObjectNodeMaterial.cpp
new file mode 100644
--- /dev/null
+++ b/Engine/Plugins/Experimental/Mutable/Source/CustomizableObjectEditor/Private/Nodes/SCustomizableObjectNodeMaterial.cpp
@@ -0,0 +1,24 @@
+// Copyright Epic Games, Inc. All Rights Reserved.
+
+#include "Nodes/SCustomizableObjectNodeMaterial.h"
+#include "EdGraphSchema_CustomizableObject.h"
+#include "Nodes/SCustomizableObjectNodeMaterialPinImage.h"
+
+
+void SCustomizableObjectNodeMaterial::Construct(const FArguments& InArgs, UEdGraphNode* InGraphNode)
+{
+	GraphNode = InGraphNode;
+	UpdateGraphNode();
+}
+
+
+TSharedPtr<SGraphPin> SCustomizableObjectNodeMaterial::CreatePinWidget(UEdGraphPin* Pin) const
+{
+	if (Pin->PinType.PinCategory == UEdGraphSchema_CustomizableObject::PC_Image &&
+		Pin->Direction == EGPD_Input)
+	{
+		return SNew(SCustomizableObjectNodeMaterialPinImage, Pin);
+	}
+
+	return SGraphNode::CreatePinWidget(Pin);
+}
diff --git a/Engine/Plugins/Experimental/Mutable/Source/CustomizableObjectEditor/Private/Nodes/SCustomizableObjectNodeMaterial.h b/Engine/Plugins/Experimental/Mutable/Source/CustomizableObjectEditor/Private/Nodes/SCustomizableObjectNodeMaterial.h
new file mode 100644
--- /dev/null
+++ b/Engine/Plugins/Experimental/Mutable/Source/CustomizableObjectEditor/Private/Nodes/SCustomizableObjectNodeMaterial.h
@@ -0,0 +1,18 @@
+// Copyright Epic Games, Inc. All Rights Reserved.
+
+#pragma once
+
+#include "SGraphNode.h"
+
+
+/** Custom widget for the Material node. */
+class SCustomizableObjectNodeMaterial : public SGraphNode
+{
+public:
+	SLATE_BEGIN_ARGS(SCustomizableObjectNodeMaterial) {}
+	SLATE_END_ARGS();
+
+	void Construct(const FArguments& InArgs, UEdGraphNode* InGraphNode);
+
+	TSharedPtr<SGraphPin> CreatePinWidget(UEdGraphPin* Pin) const override;
+};
\ No newline at end of file
diff --git a/Engine/Plugins/Experimental/Mutable/Source/CustomizableObjectEditor/Private/Nodes/SCustomizableObjectNodeMaterialPinImage.cpp b/Engine/Plugins/Experimental/Mutable/Source/CustomizableObjectEditor/Private/Nodes/SCustomizableObjectNodeMaterialPinImage.cpp
new file mode 100644
--- /dev/null
+++ b/Engine/Plugins/Experimental/Mutable/Source/CustomizableObjectEditor/Private/Nodes/SCustomizableObjectNodeMaterialPinImage.cpp
@@ -0,0 +1,97 @@
+// Copyright Epic Games, Inc. All Rights Reserved.
+
+#include "Nodes/SCustomizableObjectNodeMaterialPinImage.h"
+
+#include "GraphTraversal.h"
+#include "Nodes/CustomizableObjectNode.h"
+#include "Nodes/CustomizableObjectNodeMaterial.h"
+#include "Widgets/Input/SEditableTextBox.h"
+#include "Widgets/Layout/SWrapBox.h"
+#include "Widgets/SToolTip.h"
+
+#define LOCTEXT_NAMESPACE "CustomizableObjectEditor"
+
+
+void SCustomizableObjectNodeMaterialPinImage::Construct(const FArguments& InArgs, UEdGraphPin* InGraphPinObj)
+{
+	SGraphPin::Construct(SGraphPin::FArguments(), InGraphPinObj);
+
+	// Override previously defined tool tip.
+	TSharedPtr<IToolTip> TooltipWidget = SNew(SToolTip)
+		.Text(this, &SCustomizableObjectNodeMaterialPinImage::GetTooltipText);
+
+	SetToolTip(TooltipWidget);
+
+	// Cache pin icons.
+	CachedPinMutableConnected = FAppStyle::GetBrush(TEXT("Graph.Pin.Connected"));
+	CachedPinMutableDisconnected = FAppStyle::GetBrush(TEXT("Graph.Pin.Disconnected"));
+	CachedPinPassthroughDisconnected = FAppStyle::GetBrush(TEXT("Graph.ExecPin.Disconnected"));
+}
+
+
+TSharedRef<SWidget>	SCustomizableObjectNodeMaterialPinImage::GetDefaultValueWidget()
+{
+	LabelAndValue->SetWrapSize(TNumericLimits<float>::Max()); // Remove warping.
+
+	return SNew(SEditableTextBox)
+		.Style(FAppStyle::Get(), "Graph.EditableTextBox")
+		.Text(this, &SCustomizableObjectNodeMaterialPinImage::GetDefaultValueText)
+		.SelectAllTextWhenFocused(false)
+		.Visibility(this, &SCustomizableObjectNodeMaterialPinImage::GetDefaultValueVisibility)
+		.IsReadOnly(true)
+		.ForegroundColor(FSlateColor::UseForeground());
+}
+
+
+const FSlateBrush* SCustomizableObjectNodeMaterialPinImage::GetPinIcon() const
+{
+	if (CastChecked<UCustomizableObjectNodeMaterial>(GraphPinObj->GetOwningNode())->IsImageMutableMode(*GraphPinObj))
+	{
+		if (FollowInputPin(*GraphPinObj))
+		{
+			return CachedPinMutableConnected;
+		}
+		else
+		{
+			return CachedPinMutableDisconnected;
+		}
+	}
+	else
+	{
+		return CachedPinPassthroughDisconnected;
+	}
+}
+
+
+FText SCustomizableObjectNodeMaterialPinImage::GetTooltipText() const
+{
+	if (CastChecked<UCustomizableObjectNodeMaterial>(GraphPinObj->GetOwningNode())->IsImageMutableMode(*GraphPinObj))
+	{
+		return  LOCTEXT("PinModeMutableTooltip", "Texture Parameter goes through Mutable.");
+	}
+	else
+	{
+		return LOCTEXT("PinModePassthroughTooltip", "Texture Parameter is ignored by Mutable.");
+	}
+}
+
+
+FText SCustomizableObjectNodeMaterialPinImage::GetDefaultValueText() const
+{
+	if (CastChecked<UCustomizableObjectNodeMaterial>(GraphPinObj->GetOwningNode())->IsImageMutableMode(*GraphPinObj))
+	{
+		return LOCTEXT("PinModeMutable", "mutable");
+	}
+	else 
+	{
+		return LOCTEXT("PinModePassthrough", "passthrough");
+	}
+}
+
+
+EVisibility SCustomizableObjectNodeMaterialPinImage::GetDefaultValueVisibility() const
+{
+	return EVisibility::Visible;
+}
+
+#undef LOCTEXT_NAMESPACE
\ No newline at end of file
diff --git a/Engine/Plugins/Experimental/Mutable/Source/CustomizableObjectEditor/Private/Nodes/SCustomizableObjectNodeMaterialPinImage.h b/Engine/Plugins/Experimental/Mutable/Source/CustomizableObjectEditor/Private/Nodes/SCustomizableObjectNodeMaterialPinImage.h
new file mode 100644
--- /dev/null
+++ b/Engine/Plugins/Experimental/Mutable/Source/CustomizableObjectEditor/Private/Nodes/SCustomizableObjectNodeMaterialPinImage.h
@@ -0,0 +1,39 @@
+// Copyright Epic Games, Inc. All Rights Reserved.
+
+#pragma once
+
+#include "SGraphPin.h"
+
+
+/** Material node custom Image pin. Allows to define a custom style to Image pins.
+ *  
+ * Currently these different styles allow the user to differentiate the different states an Image pin can have:
+ * - passthrough: The texture does NOT go trough the mutable pipeline.
+ * - mutable: The texture goes trough the mutable pipeline.
+ */
+class SCustomizableObjectNodeMaterialPinImage : public SGraphPin
+{
+public:
+
+	SLATE_BEGIN_ARGS(SCustomizableObjectNodeMaterialPinImage) {}
+	SLATE_END_ARGS()
+
+	void Construct(const FArguments& InArgs, UEdGraphPin* InGraphPinObj);
+
+	TSharedRef<SWidget>	GetDefaultValueWidget() override;
+	const FSlateBrush* GetPinIcon() const override;
+
+private:
+	/** Return pin state text. */
+	FText GetDefaultValueText() const;
+
+	/** Return true if the pin state is visible. */
+	EVisibility GetDefaultValueVisibility() const;
+
+	/** Return pin tool tip. */
+	FText GetTooltipText() const;
+
+	const FSlateBrush* CachedPinMutableConnected;
+	const FSlateBrush* CachedPinMutableDisconnected;
+	const FSlateBrush* CachedPinPassthroughDisconnected;
+};
diff --git a/Engine/Plugins/Experimental/Mutable/Source/CustomizableObjectEditor/Private/Nodes/SPinViewerNodeMaterialPinImageDetails.cpp b/Engine/Plugins/Experimental/Mutable/Source/CustomizableObjectEditor/Private/Nodes/SPinViewerNodeMaterialPinImageDetails.cpp
new file mode 100644
--- /dev/null
+++ b/Engine/Plugins/Experimental/Mutable/Source/CustomizableObjectEditor/Private/Nodes/SPinViewerNodeMaterialPinImageDetails.cpp
@@ -0,0 +1,207 @@
+// Copyright Epic Games, Inc. All Rights Reserved.
+
+#include "SPinViewerNodeMaterialPinImageDetails.h"
+
+#include "ISinglePropertyView.h"
+#include "PerPlatformPropertyCustomization.h"
+#include "Widgets/Input/SNumericEntryBox.h"
+#include "SMutableTextComboBox.h"
+#include "Modules/ModuleManager.h"
+#include "ScopedTransaction.h"
+
+#define LOCTEXT_NAMESPACE "CustomizableObjectEditor"
+
+
+FText UVLayoutModeToText(const EUVLayoutMode Mode)
+{
+	switch(Mode)
+	{
+	case EUVLayoutMode::Default:
+		return LOCTEXT("EUVLayoutMode::Default","From Material");
+
+	case EUVLayoutMode::Ignore:
+		return LOCTEXT("EUVLayoutMode::Ignore","Ignore");
+
+	case EUVLayoutMode::Index:
+		return LOCTEXT("EUVLayoutMode::Index","Index");
+
+	default:
+		check(false); // Case not contemplated.
+		return FText();	
+	}
+}
+
+
+EUVLayoutMode UVLayoutToUVLayoutMode(const int32 UVLayout)
+{
+	if (UVLayout >= 0)
+	{
+		return EUVLayoutMode::Index;
+	}
+	else
+	{
+		switch (UVLayout)
+		{
+		case UCustomizableObjectNodeMaterialPinDataImage::UV_LAYOUT_DEFAULT:
+			return EUVLayoutMode::Default;
+
+		case UCustomizableObjectNodeMaterialPinDataImage::UV_LAYOUT_IGNORE:
+			return EUVLayoutMode::Ignore;
+
+		default:
+			check(false); // Case not contemplated.
+			return EUVLayoutMode::Index; // Fake return.
+		}
+	}
+}
+
+
+int32 UVLayoutModeToUVLayout(const EUVLayoutMode Mode, const int32 UVLayout)
+{
+	switch (Mode)
+	{
+	case EUVLayoutMode::Default:
+		return UCustomizableObjectNodeMaterialPinDataImage::UV_LAYOUT_DEFAULT;
+
+	case EUVLayoutMode::Ignore:
+		return UCustomizableObjectNodeMaterialPinDataImage::UV_LAYOUT_IGNORE;
+
+	case EUVLayoutMode::Index:
+		return UVLayout;
+
+		default:
+			check(false); // Case not contemplated.
+		return 0;
+	}
+}
+
+
+void SPinViewerNodeMaterialPinImageDetails::Construct(const FArguments& InArgs)
+{
+	SPinViewerPinDetails::Construct(SPinViewerPinDetails::FArguments());
+
+	PinData = InArgs._PinData;
+	
+	{
+		PinModeOptions.Add(MakeShared<EPinMode>(EPinMode::Default));
+		PinModeOptions.Add(MakeShared<EPinMode>(EPinMode::Mutable));
+		PinModeOptions.Add(MakeShared<EPinMode>(EPinMode::Passthrough));
+
+		const TSharedPtr<EPinMode> SelectedItem = *PinModeOptions.FindByPredicate([this](const TSharedPtr<EPinMode>& Element)
+		{
+			return *Element == PinData->GetPinMode();
+		});
+		
+		TSharedPtr<SWidget> Widget;
+		SAssignNew(Widget, SMutableTextComboBox<TSharedPtr<EPinMode>>)
+			.Options(&PinModeOptions)
+			.InitiallySelectedItem(SelectedItem)
+			.OnSelectionChanged(this, &SPinViewerNodeMaterialPinImageDetails::PinModeOnSelectionChanged)
+			.Translate_Lambda([this](const TSharedPtr<EPinMode> Option)
+			{
+				return EPinModeToText(*Option);
+			});
+
+		const FText Tooltip = LOCTEXT("PinModeTooltip", "Texture Parameter mode.");
+		AddRow(LOCTEXT("PinMode", "Mode"), Widget.ToSharedRef(), &Tooltip);
+	}
+
+	const EUVLayoutMode UVLayoutMode = UVLayoutToUVLayoutMode(PinData->UVLayout);
+	
+	{
+		UVLayoutModeOptions.Add(MakeShared<EUVLayoutMode>(EUVLayoutMode::Default));
+		UVLayoutModeOptions.Add(MakeShared<EUVLayoutMode>(EUVLayoutMode::Ignore));
+		UVLayoutModeOptions.Add(MakeShared<EUVLayoutMode>(EUVLayoutMode::Index));
+
+		const TSharedPtr<EUVLayoutMode> SelectedItem = *UVLayoutModeOptions.FindByPredicate([&](const TSharedPtr<EUVLayoutMode>& Element)
+		{
+			return *Element == UVLayoutMode;
+		});
+		
+		TSharedPtr<SWidget> Widget;
+		SAssignNew(Widget, SMutableTextComboBox<TSharedPtr<EUVLayoutMode>>)
+			.Options(&UVLayoutModeOptions)
+			.InitiallySelectedItem(SelectedItem)
+			.OnSelectionChanged(this, &SPinViewerNodeMaterialPinImageDetails::UVLayoutModeOnSelectionChanged)
+			.Translate_Lambda([this](const TSharedPtr<EUVLayoutMode> Option)
+			{
+				return UVLayoutModeToText(*Option);
+			});
+
+		const FText Tooltip = LOCTEXT("UVLayoutModeTooltip", "Specify how UV Layuout is defined.");
+		AddRow(LOCTEXT("UVLayoutMode", "UV Layout Mode"), Widget.ToSharedRef(), &Tooltip);
+	}
+	
+	{
+		SAssignNew(UVLayoutSSpinBox, SSpinBox<int32>)
+			.MinValue(0)
+			.OnValueChanged(this, &SPinViewerNodeMaterialPinImageDetails::UVLayoutOnValueChanged)
+			.Value_Lambda([this]
+			{
+				return FMath::Max(0, PinData->UVLayout);
+			});
+
+		UVLayout = AddRow(LOCTEXT("UVLayout", "UV Layout Index"), UVLayoutSSpinBox.ToSharedRef());
+		UVLayoutVisibility(UVLayoutMode, UVLayout);
+	}
+
+	{
+		FPropertyEditorModule& PropPlugin = FModuleManager::LoadModuleChecked<FPropertyEditorModule>("PropertyEditor");
+		FSinglePropertyParams SingleDetails;
+		SingleDetails.NamePlacement = EPropertyNamePlacement::Hidden;
+
+		const TSharedPtr<ISinglePropertyView> Widget = PropPlugin.CreateSingleProperty(PinData,
+			GET_MEMBER_NAME_CHECKED(UCustomizableObjectNodeMaterialPinDataImage, ReferenceTexture), SingleDetails);
+
+		AddRow(LOCTEXT("ReferenceTexture", "Reference Texture"), Widget.ToSharedRef());
+	}
+}
+
+
+void SPinViewerNodeMaterialPinImageDetails::PinModeOnSelectionChanged(const TSharedPtr<EPinMode> PinMode, ESelectInfo::Type Arg)
+{
+	FScopedTransaction Transaction(LOCTEXT("ChangedPinModeTransaction", "Changed Pin Mode"));
+	PinData->Modify();
+	
+	PinData->SetPinMode(*PinMode);
+}
+
+
+void SPinViewerNodeMaterialPinImageDetails::UVLayoutModeOnSelectionChanged(const TSharedPtr<EUVLayoutMode> LayoutMode, ESelectInfo::Type Arg)
+{
+	FScopedTransaction Transaction(LOCTEXT("ChangedUVLayoutModeTransaction", "Changed UV Layout Mode"));
+	PinData->Modify();
+
+	PinData->UVLayout = UVLayoutModeToUVLayout(*LayoutMode, UVLayoutSSpinBox->GetValue());
+	UVLayoutVisibility(*LayoutMode, UVLayout);
+}
+
+
+void SPinViewerNodeMaterialPinImageDetails::UVLayoutVisibility(const EUVLayoutMode LayoutMode, const TSharedPtr<SWidget> Widget)
+{
+	switch (LayoutMode)
+	{
+	case EUVLayoutMode::Default:
+	case EUVLayoutMode::Ignore:
+		Widget->SetVisibility(EVisibility::Collapsed);
+		break;
+	
+	case EUVLayoutMode::Index:
+		Widget->SetVisibility(EVisibility::Visible);
+		break;
+
+	default:
+		check(false); // Case not contemplated.
+	}
+}
+
+
+void SPinViewerNodeMaterialPinImageDetails::UVLayoutOnValueChanged(int32 Value)
+{
+	FScopedTransaction Transaction(LOCTEXT("ChangedUVLayoutTransaction", "Changed UV Layout Index"));
+	PinData->Modify();
+	
+	PinData->UVLayout = Value;
+}
+
+#undef LOCTEXT_NAMESPACE
\ No newline at end of file
diff --git a/Engine/Plugins/Experimental/Mutable/Source/CustomizableObjectEditor/Private/Nodes/SPinViewerNodeMaterialPinImageDetails.h b/Engine/Plugins/Experimental/Mutable/Source/CustomizableObjectEditor/Private/Nodes/SPinViewerNodeMaterialPinImageDetails.h
new file mode 100644
--- /dev/null
+++ b/Engine/Plugins/Experimental/Mutable/Source/CustomizableObjectEditor/Private/Nodes/SPinViewerNodeMaterialPinImageDetails.h
@@ -0,0 +1,58 @@
+// Copyright Epic Games, Inc. All Rights Reserved.
+
+#pragma once
+
+#include "Nodes/CustomizableObjectNodeMaterial.h"
+#include "PinViewer/SPinViewerPinDetails.h"
+#include "Widgets/Input/SSpinBox.h"
+
+/** UV Layout Mode UI Enum. */
+UENUM()
+enum class EUVLayoutMode
+{
+	/** Equal to UV_LAYOUT_DEFAULT */
+	Default,
+	/** Equal to UV_LAYOUT_IGNORE */
+	Ignore,
+	/** Any valid index. */
+	Index,
+};
+
+/** Codify the UVLayout into a single integer. */
+int32 UVLayoutModeToUVLayout(const EUVLayoutMode Mode, const int32 UVLayout);
+
+/** Given the codified UV Layout, returns the UV Layout mode. */
+EUVLayoutMode UVLayoutToUVLayoutMode(int32 UVLayout);
+
+/** Node Material Pin Image custom details. Allows Image Pin Data to be modified by the user. */
+class SPinViewerNodeMaterialPinImageDetails : public SPinViewerPinDetails
+{
+public:
+	SLATE_BEGIN_ARGS(SPinViewerNodeMaterialPinImageDetails) {}
+		SLATE_ARGUMENT(UCustomizableObjectNodeMaterialPinDataImage*, PinData)
+	SLATE_END_ARGS()
+
+	/** Slate constructor. */
+	void Construct(const FArguments& InArgs);
+
+private:
+	TArray<TSharedPtr<EPinMode>> PinModeOptions;
+	TArray<TSharedPtr<EUVLayoutMode>> UVLayoutModeOptions;
+	TSharedPtr<SSpinBox<int32>> UVLayoutSSpinBox;
+	TSharedPtr<SWidget> UVLayout;
+	TMap<TSharedPtr<FString>, EPinMode> PinModeComboBoxOptionValues;
+	
+	/** Pin Data to show and modify. */
+	UCustomizableObjectNodeMaterialPinDataImage* PinData = nullptr;
+	
+	/** Callback for Pin Mode changed. */
+	void PinModeOnSelectionChanged(TSharedPtr<EPinMode> PinMode, ESelectInfo::Type Arg);
+
+	/** Callback for UV Layout Mode changed. */
+	void UVLayoutModeOnSelectionChanged(TSharedPtr<EUVLayoutMode> LayoutMode, ESelectInfo::Type Arg);
+
+	static void UVLayoutVisibility(EUVLayoutMode LayoutMode, TSharedPtr<SWidget> Widget);
+	
+	/** Callback for UV Layout Index changed. */
+	void UVLayoutOnValueChanged(int32 Value);
+};
diff --git a/Engine/Plugins/Experimental/Mutable/Source/CustomizableObjectEditor/Private/RemapPins/CustomizableObjectNodeRemapPinsByName.cpp b/Engine/Plugins/Experimental/Mutable/Source/CustomizableObjectEditor/Private/RemapPins/CustomizableObjectNodeRemapPinsByName.cpp
new file mode 100644
--- /dev/null
+++ b/Engine/Plugins/Experimental/Mutable/Source/CustomizableObjectEditor/Private/RemapPins/CustomizableObjectNodeRemapPinsByName.cpp
@@ -0,0 +1,37 @@
+// Copyright Epic Games, Inc. All Rights Reserved.
+
+#include "RemapPins/CustomizableObjectNodeRemapPinsByName.h"
+#include "UnrealEditorPortabilityHelpers.h"
+
+
+bool UCustomizableObjectNodeRemapPinsByName::Equal(const UEdGraphPin& OldPin, const UEdGraphPin& NewPin) const
+{
+	return Helper_GetPinName(&OldPin) == Helper_GetPinName(&NewPin) &&
+				OldPin.PinType == NewPin.PinType &&
+				OldPin.Direction == NewPin.Direction;
+}
+
+
+void UCustomizableObjectNodeRemapPinsByName::RemapPins(const TArray<UEdGraphPin*>& OldPins, const TArray<UEdGraphPin*>& NewPins, TMap<UEdGraphPin*, UEdGraphPin*>& PinsToRemap, TArray<UEdGraphPin*>& PinsToOrphan)
+{
+	for (UEdGraphPin* OldPin : OldPins)
+	{
+		bool bFound = false;
+
+		for (UEdGraphPin* NewPin : NewPins)
+		{
+			if (Equal(*OldPin, *NewPin))
+			{
+				bFound = true;
+
+				PinsToRemap.Add(OldPin, NewPin);
+				break;
+			}
+		}
+
+		if (!bFound && OldPin->LinkedTo.Num())
+		{
+			PinsToOrphan.Add(OldPin);
+		}
+	}
+}
diff --git a/Engine/Plugins/Experimental/Mutable/Source/CustomizableObjectEditor/Private/RemapPins/CustomizableObjectNodeRemapPinsByNameDefaultPin.cpp b/Engine/Plugins/Experimental/Mutable/Source/CustomizableObjectEditor/Private/RemapPins/CustomizableObjectNodeRemapPinsByNameDefaultPin.cpp
new file mode 100644
--- /dev/null
+++ b/Engine/Plugins/Experimental/Mutable/Source/CustomizableObjectEditor/Private/RemapPins/CustomizableObjectNodeRemapPinsByNameDefaultPin.cpp
@@ -0,0 +1,43 @@
+// Copyright Epic Games, Inc. All Rights Reserved.
+
+#include "RemapPins/CustomizableObjectNodeRemapPinsByNameDefaultPin.h"
+
+
+void UCustomizableObjectNodeRemapPinsByNameDefaultPin::RemapPins(const TArray<UEdGraphPin*>& OldPins, const TArray<UEdGraphPin*>& NewPins, TMap<UEdGraphPin*, UEdGraphPin*>& PinsToRemap, TArray<UEdGraphPin*>& PinsToOrphan)
+{
+	if (DefaultPin)
+	{
+		check(DefaultPin->Direction == EGPD_Output);
+
+		// Remap the default pin to the first mesh pin
+		// Orphan all pins except the default
+		for (UEdGraphPin* Pin : OldPins)
+		{
+			if (Pin->LinkedTo.Num() && Pin != DefaultPin)
+			{
+				PinsToOrphan.Add(Pin);
+			}
+		}
+
+		// Find the first mesh pin
+		UEdGraphPin* FirstMeshPin = nullptr;
+		for (UEdGraphPin* Pin : NewPins)
+		{
+			if (Pin->Direction == EGPD_Output && Pin->PinType == DefaultPin->PinType)
+			{
+				FirstMeshPin = Pin;
+				break;
+			}
+		}
+
+		if (FirstMeshPin)
+		{
+			PinsToRemap.Add(DefaultPin, FirstMeshPin);
+		}
+	}
+	else
+	{
+		Super::RemapPins(OldPins, NewPins, PinsToRemap, PinsToOrphan);
+	}
+}
+
