[mutable] Move Mutable plugin from the restricted folder to the experimental plugins folder.

diff --git a/Engine/Plugins/Experimental/Mutable/Source/CustomizableObjectEditor/Private/SCustomizableInstanceProperties.cpp b/Engine/Plugins/Experimental/Mutable/Source/CustomizableObjectEditor/Private/SCustomizableInstanceProperties.cpp
new file mode 100644
--- /dev/null
+++ b/Engine/Plugins/Experimental/Mutable/Source/CustomizableObjectEditor/Private/SCustomizableInstanceProperties.cpp
@@ -0,0 +1,2437 @@
+// Copyright Epic Games, Inc. All Rights Reserved.
+
+#include "SCustomizableInstanceProperties.h"
+
+#include "Runtime/Core/Public/UObject/PropertyPortFlags.h"
+
+#include "CustomizableObjectInstance.h"
+#include "MultilayerProjector.h"
+#include "CustomizableObjectSystem.h"
+#include "CustomizableObjectEditorUtilities.h"
+
+#include "Widgets/Input/STextComboBox.h"
+#include "SMutableTextSearchBox.h"
+#include "Widgets/Colors/SColorPicker.h"
+#include "Widgets/Colors/SColorBlock.h"
+#include "Widgets/Input/SCheckBox.h"
+#include "Widgets/Images/SImage.h"
+#include "Widgets/Input/SButton.h"
+#include "Widgets/Layout/SExpandableArea.h"
+#include "Slate/DeferredCleanupSlateBrush.h"
+
+#include "Serialization/BufferArchive.h"
+#include "Widgets/Layout/SUniformGridPanel.h"
+
+#include "ContentBrowserModule.h"
+#include "IContentBrowserSingleton.h"
+#include "AssetRegistry/AssetRegistryModule.h"
+#include "Editor.h"
+#include "ISkeletonTree.h"
+#include "EnvironmentQuery/EnvQueryTypes.h"
+
+#include "UObject/UObjectGlobals.h"
+
+#include "HAL/PlatformApplicationMisc.h"
+
+#define LOCTEXT_NAMESPACE "CustomizableObjectEditor"
+
+
+void SCustomizableInstanceProperties::Construct(const FArguments& InArgs)
+{
+	CustomInstance = InArgs._CustomInstance;
+	check(CustomInstance.IsValid());
+	
+	InstanceDetails = InArgs._InstanceDetails;
+	
+	NoInstanceMessage = LOCTEXT("Model not compiled", "Model not compiled");
+	SimplifiedProjectorInterface = false;
+	
+	if (CustomInstance->IsSelectedParameterProfileDirty())
+	{
+		CustomInstance->SaveParametersToProfile(CustomInstance->SelectedProfileIndex);
+	}
+	
+	CustomInstance->UpdatedNativeDelegate.AddSP(this, &SCustomizableInstanceProperties::InstanceUpdated);
+	
+	StateNames.Reset();
+	TextureParameterValueNames.Reset();
+	TextureParameterValues.Reset();	
+
+	if (UCustomizableObject* CustomizableObject = CustomInstance->GetCustomizableObject())
+	{
+		// Store the state data required for the ui.
+		const int numStates = CustomizableObject->GetStateCount();
+		const int currentState = CustomInstance->GetState();
+		TSharedPtr<FString> CurrentStateName = nullptr;
+
+		for (int i = 0; i < numStates; ++i)
+		{
+			if (i == currentState)
+			{
+				CurrentStateName = MakeShareable(new FString(CustomizableObject->GetStateName(i)));
+				StateNames.Add(CurrentStateName);
+			}
+			else
+			{
+				StateNames.Add(MakeShareable(new FString(CustomizableObject->GetStateName(i))));
+			}
+		}
+		if (!StateNames.Num())
+		{
+			StateNames.Add(MakeShareable(new FString("default")));
+		}
+
+		StateNames.Sort(&CompareNames);
+
+		// Store the texture parameters data required for the ui.
+		TextureParameterValueNames.Add(MakeShareable(new FString("None")));
+		TextureParameterValues.Add(-1);
+		TArray<FCustomizableObjectExternalTexture> Textures = UCustomizableObjectSystem::GetInstance()->GetTextureParameterValues();
+		for (int i = 0; i < Textures.Num(); ++i)
+		{
+			TextureParameterValueNames.Add(MakeShareable(new FString(Textures[i].Name)));
+			TextureParameterValues.Add(Textures[i].Value);
+		}
+
+		SetParameterProfileNamesOnEditor();
+		
+		int32 ProfileStringIndex = 0;
+
+		const int32 ProfileIdx = CustomInstance->SelectedProfileIndex;
+
+		if (ProfileIdx != INDEX_NONE)
+		{
+			FProfileParameterDat& Profile = CustomizableObject->InstancePropertiesProfiles[ProfileIdx];
+			for (int i = 1; i < ParameterProfileNames.Num(); ++i)
+			{
+				if (ParameterProfileNames[i]->Equals(Profile.ProfileName))
+				{
+					ProfileStringIndex = i; 
+					break; 
+				}
+			}
+		}
+
+		FText SearchBoxText = FText::GetEmpty();
+
+		if (SB_Properties.IsValid())
+		{
+			SearchBoxText = SB_Properties->GetText();
+		}
+
+		this->ChildSlot
+		[
+			SNew(SVerticalBox)
+
+			+ SVerticalBox::Slot()
+			.AutoHeight()
+			.Padding(2.0f)
+			[
+				SNew(SHorizontalBox)
+
+				+ SHorizontalBox::Slot()
+				.AutoWidth()
+				.VAlign(VAlign_Center)
+				.Padding(2.0f)
+				[
+					SNew(STextBlock)
+					.Text(LOCTEXT("State", "State"))
+				]
+
+				+ SHorizontalBox::Slot()
+				.FillWidth(10.0f)
+				.VAlign(VAlign_Center)
+				.HAlign(HAlign_Fill)
+				.Padding(2.0f)
+				[
+					SNew(STextComboBox)
+					.OptionsSource(&StateNames)
+					.InitiallySelectedItem((CustomInstance->GetState() != -1) ? CurrentStateName : StateNames[0])
+					.OnSelectionChanged(this, &SCustomizableInstanceProperties::OnStateComboBoxSelectionChanged)
+				]
+			]
+
+			+ SVerticalBox::Slot()
+			.AutoHeight()
+			.Padding(2.0f)
+			[
+				SNew(SHorizontalBox)
+				+ SHorizontalBox::Slot()
+				.AutoWidth()
+				.VAlign(VAlign_Center)
+				[
+					SNew(SCheckBox)
+					.IsChecked(CustomInstance->bShowOnlyRuntimeParameters ? ECheckBoxState::Checked : ECheckBoxState::Unchecked)
+					.OnCheckStateChanged(this, &SCustomizableInstanceProperties::OnShowOnlyRuntimeSelectionChanged)
+				]
+
+				+ SHorizontalBox::Slot()
+				.AutoWidth()
+				.VAlign(VAlign_Center)
+				[
+					SNew(STextBlock)
+					.Text(LOCTEXT("only runtime", "only runtime"))
+				]
+			]
+
+			+ SVerticalBox::Slot()
+			.AutoHeight()
+			.Padding(2.0f)
+			[
+				SNew(SHorizontalBox)
+				+ SHorizontalBox::Slot()
+				.AutoWidth()
+				.VAlign(VAlign_Center)
+				[
+					SNew(SCheckBox)
+					.IsChecked(CustomInstance->bShowOnlyRelevantParameters ? ECheckBoxState::Checked : ECheckBoxState::Unchecked)
+					.OnCheckStateChanged(this, &SCustomizableInstanceProperties::OnShowOnlyRelevantSelectionChanged)
+				]
+			
+				+ SHorizontalBox::Slot()
+				.AutoWidth()
+				.VAlign(VAlign_Center)
+				[
+					SNew(STextBlock)
+					.Text(LOCTEXT("only relevant", "only relevant"))
+				]
+			]
+
+			//Properties profile functionality
+			+SVerticalBox::Slot()
+			.AutoHeight()
+			.Padding(6.0f)
+			.HAlign(HAlign_Fill)
+			[
+				SNew(SVerticalBox)
+				+ SVerticalBox::Slot()
+				.AutoHeight()
+				.HAlign(HAlign_Fill)
+				.Padding(-2.f, 0.f)
+				[
+					SNew(STextBlock)
+					.Text(LOCTEXT("Preview_Instance_Parameter_Profiles", "Preview Instance Parameter Profiles"))
+				]
+			]
+			
+			+ SVerticalBox::Slot()
+			.AutoHeight()
+			.Padding(2.0f)
+			.HAlign(HAlign_Fill)
+			.MaxHeight(22.f)
+			.Padding(2.f, 0, 0, 6.f)
+			[
+				SNew(SVerticalBox)
+				+ SVerticalBox::Slot()
+				[
+					SNew(SHorizontalBox)
+					+ SHorizontalBox::Slot()
+					.FillWidth(1.0f)
+					[
+						SNew(STextComboBox)
+						.OptionsSource(&ParameterProfileNames)
+						.OnSelectionChanged(this, &SCustomizableInstanceProperties::OnProfileSelectedChanged)
+						.ToolTipText(FText::FromString(ParameterProfileNames.Num() > 1 ? FString("Select an existing profile") : FString("No profiles are available")))
+						.InitiallySelectedItem(ParameterProfileNames[ProfileStringIndex])
+						.IsEnabled(ParameterProfileNames.Num() > 1)
+					]
+					
+					+ SHorizontalBox::Slot()
+					.AutoWidth()
+					[
+						SNew(SButton)
+						.Text(LOCTEXT("AddButtonLabel", " + "))
+						.ToolTipText(FText::FromString(HasAnyParameters() ? FString("Add new profile") : FString("Create a profile functionality is not available, no parameters were found.")))
+						.IsEnabled(HasAnyParameters())
+						.IsFocusable(false)
+						.OnClicked(this, &SCustomizableInstanceProperties::CreateParameterProfileWindow)
+					]	
+					
+					+ SHorizontalBox::Slot()
+					.AutoWidth()
+					[
+						SNew(SButton)
+						.Text(LOCTEXT("RemoveButtonLabel", " - "))
+						.ToolTipText(FText::FromString(HasAnyParameters() ? FString("Delete selected profile") : FString("Delete selected profile functionality is not available, no profile is selected.")))
+						.IsEnabled(CustomInstance->SelectedProfileIndex != INDEX_NONE)
+						.IsFocusable(false)
+						.OnClicked(this, &SCustomizableInstanceProperties::RemoveParameterProfile)
+					]
+				]
+			]
+
+			+ SVerticalBox::Slot()
+			.AutoHeight()
+			.Padding(2.0f)
+			[
+				SNew(SHorizontalBox)
+				+ SHorizontalBox::Slot()
+				.FillWidth(2.0f)
+				.Padding(.5f)
+				[
+					SNew(SButton)
+					.HAlign(HAlign_Center)
+					.Text(LOCTEXT("Copy_Parameters", "Copy Parameters"))
+					.OnClicked(this, &SCustomizableInstanceProperties::OnCopyAllParameters)
+					.IsEnabled(HasAnyParameters())
+					.ToolTipText(FText::FromString(FString("Copy the preview Instance parameters")))
+				]
+				+ SHorizontalBox::Slot()
+				.FillWidth(2.0f)
+				.Padding(.5f)
+				[
+					SNew(SButton)
+					.HAlign(HAlign_Center)
+					.Text(LOCTEXT("Paste_Parameters", "Paste Parameters"))
+					.OnClicked(this, &SCustomizableInstanceProperties::OnPasteAllParameters)
+					.IsEnabled(HasAnyParameters())
+					.ToolTipText(FText::FromString(FString("Paste the preview Instance parameters")))
+				]
+				+ SHorizontalBox::Slot()
+				.FillWidth(2.0f)
+				.Padding(.5f)
+				[
+					SNew(SButton)
+					.HAlign(HAlign_Center)
+					.Text(LOCTEXT("Reset_Integer_Paramaters", "Reset parameters"))
+					.OnClicked(this, &SCustomizableInstanceProperties::OnResetAllParameters)
+					.IsEnabled(HasAnyParameters())
+					.ToolTipText(FText::FromString(FString("Clear the preview Instance parameters")))
+				]
+			]
+
+			//Space between last section and the SearchBox
+			+SVerticalBox::Slot()
+			.AutoHeight()
+			.Padding(0, 6.f)
+			.HAlign(HAlign_Left)
+
+			/*SearchBox declaration*/
+
+			+SVerticalBox::Slot()
+			.AutoHeight()
+			.Padding(0, 6.f)
+			[
+				SNew(SHorizontalBox)
+				+ SHorizontalBox::Slot()
+				.FillWidth(10.0f)
+				.VAlign(VAlign_Center)
+				.HAlign(HAlign_Fill)
+
+				.Padding(2.0f)
+				[
+					SAssignNew(SB_Properties, SSearchBox)
+					.HintText(LOCTEXT("SearchHint", "Search Properties"))
+					.InitialText(SearchBoxText)
+					.OnTextChanged(this, &SCustomizableInstanceProperties::OnFilterTextChanged)
+				]
+			]
+
+			//Space between the SearchBox and the parameters
+			+SVerticalBox::Slot()
+			.AutoHeight()
+			.Padding(0, 4.0f)
+			.HAlign(HAlign_Fill)
+			[
+				SAssignNew(ParamBox, SVerticalBox)
+			]
+		];
+	}
+	else
+	{
+		this->ChildSlot
+			[
+				SNew(STextBlock)
+				.Text(NoInstanceMessage)
+			];
+	}
+
+	ResetParamBox();
+}
+
+
+void SCustomizableInstanceProperties::SetNoInstanceMessage(const FText& InNoInstanceMessage)
+{
+	NoInstanceMessage = InNoInstanceMessage;
+}
+
+
+void SCustomizableInstanceProperties::SetSimplifiedProjectorInterface(bool Enabled)
+{
+	SimplifiedProjectorInterface = Enabled;
+}
+
+
+const UCustomizableObjectInstance* SCustomizableInstanceProperties::GetInstance() const
+{
+	return CustomInstance.Get();
+}
+
+
+struct FCustomizableInstancePropertyArrayElem
+{
+	int32 PropertyCOIndex{-1};
+	int32 UIOrder{0};
+	FString ParameterName{""};
+
+	friend bool operator <(const FCustomizableInstancePropertyArrayElem& A, const FCustomizableInstancePropertyArrayElem& B)
+	{
+		if (A.UIOrder != B.UIOrder)
+		{
+			return A.UIOrder < B.UIOrder;
+		}
+		else if (A.ParameterName.IsEmpty())
+		{
+			if (B.ParameterName.IsEmpty())
+			{
+				return (A.ParameterName.Compare(B.ParameterName, ESearchCase::IgnoreCase) < 0);
+			}
+			else
+			{
+				return false;
+			}
+		}
+		else
+		{
+			return true;
+		}
+	}
+};
+
+
+void SCustomizableInstanceProperties::ResetParamBox()
+{
+	if (ParamBox)
+	{
+		ParamBox->ClearChildren();
+	}
+	//IntOptionNames.Reset();
+	FloatSliders.Reset();
+	DynamicBrushes.Reset();
+
+	ParamChildren.Empty();
+	ParamHasParent.Empty();
+	ParamNameToExpandableAreaMap.Empty();
+
+	MapProjectorButton.Empty();
+	MapProjectorButtonTextBlock.Empty();
+
+	if (UCustomizableObject* CustomizableObject = CustomInstance->GetCustomizableObject())
+	{
+		if (CustomInstance->bShowOnlyRuntimeParameters)
+		{
+			const int32 NumStateParameters = CustomizableObject->GetStateParameterCount(CustomInstance->GetState());
+			TArray<FCustomizableInstancePropertyArrayElem> ParamIndexesInState;
+			for (int32 ParamIndexInState = 0; ParamIndexInState < NumStateParameters; ++ParamIndexInState)
+			{
+				FCustomizableInstancePropertyArrayElem ParameterSortInfo;
+				ParameterSortInfo.PropertyCOIndex = CustomizableObject->GetStateParameterIndex(CustomInstance->GetState(), ParamIndexInState);
+				if (CustomInstance->IsParameterRelevant(ParameterSortInfo.PropertyCOIndex))
+				{
+					ParameterSortInfo.UIOrder = CustomizableObject->GetParameterUIMetadataFromIndex(ParameterSortInfo.PropertyCOIndex).ParamUIMetadata.UIOrder;
+					ParameterSortInfo.ParameterName = CustomizableObject->GetParameterName(ParameterSortInfo.PropertyCOIndex);
+					ParamIndexesInState.Add(ParameterSortInfo);
+				}
+			}
+
+			ParamIndexesInState.Sort();
+
+			for (int32 ParamArrayNum = 0; ParamArrayNum < ParamIndexesInState.Num(); ++ParamArrayNum)
+			{
+				AddParameter(ParamIndexesInState[ParamArrayNum].PropertyCOIndex);
+			}
+		}
+		else
+		{
+			const int32 NumObjectParameter = CustomizableObject->GetParameterCount();
+
+			//TODO: get all parameters and sort, then make the next "for" use that sorted list as source of indexes
+			TArray<FCustomizableInstancePropertyArrayElem> ParamIndexesInObject;
+			for (int32 ParamIndexInObject = 0; ParamIndexInObject < NumObjectParameter; ++ParamIndexInObject)
+			{
+				FCustomizableInstancePropertyArrayElem ParameterSortInfo;
+				ParameterSortInfo.PropertyCOIndex = ParamIndexInObject;
+				if (CustomInstance->IsParameterRelevant(ParameterSortInfo.PropertyCOIndex))
+				{
+					ParameterSortInfo.UIOrder = CustomizableObject->GetParameterUIMetadataFromIndex(ParameterSortInfo.PropertyCOIndex).ParamUIMetadata.UIOrder;
+					ParameterSortInfo.ParameterName = CustomizableObject->GetParameterName(ParameterSortInfo.PropertyCOIndex);
+					ParamIndexesInObject.Add(ParameterSortInfo);
+				}
+			}
+
+			ParamIndexesInObject.Sort();
+
+			for (int32 ParamIndexInObject = 0; ParamIndexInObject < ParamIndexesInObject.Num(); ++ParamIndexInObject)
+			{
+				if (CustomInstance->IsParameterRelevant(ParamIndexesInObject[ParamIndexInObject].PropertyCOIndex))
+				{
+					FillChildrenMap(ParamIndexesInObject[ParamIndexInObject].PropertyCOIndex);
+				}
+			}
+
+			for (int32 ParamIndexInObject = 0; ParamIndexInObject < ParamIndexesInObject.Num(); ++ParamIndexInObject)
+			{
+				if (CustomInstance->IsParameterRelevant(ParamIndexesInObject[ParamIndexInObject].PropertyCOIndex))
+				{
+					if (!ParamHasParent.Find(ParamIndexesInObject[ParamIndexInObject].PropertyCOIndex))
+					{
+						RecursivelyAddParamAndChildren(ParamIndexesInObject[ParamIndexInObject].PropertyCOIndex);
+					}
+				}
+			}
+		}
+	}
+}
+
+
+void SCustomizableInstanceProperties::RecursivelyAddParamAndChildren(int32 ParamIndexInObject)
+{
+	const UCustomizableObject* CustomizableObject = CustomInstance->GetCustomizableObject();
+	const FString ParamName = CustomizableObject->GetParameterName(ParamIndexInObject);
+
+	AddParameter(ParamIndexInObject);
+
+	TArray<int32> Children;
+
+	ParamChildren.MultiFind(ParamName, Children, true);
+
+	for (const int32 ChildIndexInObject : Children)
+	{
+		RecursivelyAddParamAndChildren(ChildIndexInObject);
+	}
+}
+
+
+void SCustomizableInstanceProperties::FillChildrenMap(int32 ParamIndexInObject)
+{
+	const UCustomizableObject* CustomizableObject = CustomInstance->GetCustomizableObject();
+	FString ParamName = CustomizableObject->GetParameterName(ParamIndexInObject);
+	FParameterUIData UIData = CustomizableObject->GetParameterUIMetadataFromIndex(ParamIndexInObject);
+
+	const FString* ParentName = UIData.ParamUIMetadata.ExtraInformation.Find(FString("__ParentParamName"));
+
+	if (ParentName)
+	{
+		ParamChildren.Add(*ParentName, ParamIndexInObject);
+		ParamHasParent.Add(ParamIndexInObject, true);
+	}
+}
+
+
+void SCustomizableInstanceProperties::AddParameter(int32 ParamIndexInObject)
+{
+	UCustomizableObject* CustomizableObject = CustomInstance->GetCustomizableObject();
+	
+	FString ParamName = CustomizableObject->GetParameterName(ParamIndexInObject);
+
+	TSharedPtr<SVerticalBox> ActualParamBox = ParamBox;
+	TSharedPtr<SHorizontalBox> ParameterBox;
+
+	FParameterUIData UIData = CustomizableObject->GetParameterUIMetadataFromIndex(ParamIndexInObject);
+	FString* ParentName = UIData.ParamUIMetadata.ExtraInformation.Find(FString("__ParentParamName"));
+	bool bHasCollapsibleChildren = UIData.ParamUIMetadata.ExtraInformation.Find(FString("__HasCollapsibleChildren")) != nullptr;
+
+	bool bHideParam = ParamName.EndsWith(FMultilayerProjector::NUM_LAYERS_PARAMETER_POSTFIX)
+		|| (ParamName.EndsWith(FMultilayerProjector::IMAGE_PARAMETER_POSTFIX) && CustomInstance->IsParamMultidimensional(ParamIndexInObject))
+		|| (ParamName.EndsWith(FMultilayerProjector::OPACITY_PARAMETER_POSTFIX) && CustomInstance->IsParamMultidimensional(ParamIndexInObject))
+		|| (ParamName.EndsWith(FMultilayerProjector::POSE_PARAMETER_POSTFIX));
+	
+	FString Name = SB_Properties->GetText().ToString();
+
+	//If SearchBox text is not empty
+	if (!Name.IsEmpty())
+	{
+		//Check if text entered has any similarity with any parameter. 
+		if (!ParamName.Contains(Name))
+		{
+			return;
+		}
+	}
+
+	if (bHideParam)
+	{
+		return;
+	}
+
+	if (ParentName)
+	{
+		if (UIData.ParamUIMetadata.ExtraInformation.Find(FString("CollapseUnderParent")))
+		{
+			const FMutableParamExpandableArea* ParentExpandableAreaPtr = ParamNameToExpandableAreaMap.Find(*ParentName);
+
+			if (ParentExpandableAreaPtr)
+			{
+				ActualParamBox = ParentExpandableAreaPtr->VerticalBox;
+			}
+		}
+
+		if (CustomInstance->bShowOnlyRelevantParameters)
+		{
+			FString *Value = UIData.ParamUIMetadata.ExtraInformation.Find(FString("__DisplayWhenParentValueEquals"));
+			if (Value && CustomInstance->GetIntParameterSelectedOption(*ParentName) != *Value)
+			{
+				return;
+			}
+		}
+	}
+
+	if (bHasCollapsibleChildren)
+	{
+		TSharedPtr<SExpandableArea> ChildParamExpandableArea;
+		TSharedPtr<SVerticalBox> VerticalBox;
+
+		bool* bIsExpanded = CustomInstance->ParamNameToExpandedMap.Find(ParamName);
+
+		ActualParamBox->AddSlot()
+			.AutoHeight()
+			.Padding(1.0f, 4.0f)
+			.HAlign(HAlign_Fill)
+			[
+				SAssignNew(ChildParamExpandableArea, SExpandableArea)
+				.AreaTitle(LOCTEXT("ChildParamExpandableArea_Title", "Child params"))
+			.HeaderContent()
+			[
+				SAssignNew(ParameterBox, SHorizontalBox)
+				+ SHorizontalBox::Slot()
+			.FillWidth(2.0f)
+			.VAlign(EVerticalAlignment::VAlign_Center)
+			[
+				SNew(STextBlock)
+				.Text(FText::FromString(ParamName))
+			]
+			]
+		.InitiallyCollapsed(bIsExpanded == nullptr || !*bIsExpanded)
+			.OnAreaExpansionChanged(this, &SCustomizableInstanceProperties::OnAreaExpansionChanged, ParamName)
+			.Padding(1.0)
+			.Visibility(EVisibility::Visible)
+			.BodyContent()
+			[
+				SAssignNew(VerticalBox, SVerticalBox)
+				.Visibility(EVisibility::Visible)
+			]
+			];
+
+		ParamNameToExpandableAreaMap.Add(ParamName, FMutableParamExpandableArea(ChildParamExpandableArea, VerticalBox));
+	}
+	else
+	{
+		/*In this horizontal box slot, we add the CO parameter name*/
+		ActualParamBox->AddSlot()
+			.Padding(1.0f, 6.0f)
+			.HAlign(HAlign_Fill)
+			.VAlign(VAlign_Fill)
+			.AutoHeight()
+			[
+				SAssignNew(ParameterBox, SHorizontalBox)
+
+				+ SHorizontalBox::Slot()
+			.FillWidth(1.8f)
+			.VAlign(VAlign_Center)
+			[
+				SNew(STextBlock)
+				.Text(FText::FromString(ParamName))
+			.Justification(ETextJustify::Left)
+			.AutoWrapText(true)
+			]
+			];
+	}
+
+	switch (CustomizableObject->GetParameterType(ParamIndexInObject))
+	{
+	case EMutableParameterType::Bool:
+	{
+		ParameterBox->AddSlot()
+		.Padding(1.0f)
+		.FillWidth(10.f)
+		[
+			SNew(SHorizontalBox)
+			+ SHorizontalBox::Slot()
+			.HAlign(HAlign_Right)
+			[
+				SNew(SCheckBox)
+				.HAlign(HAlign_Right)
+				.IsChecked(this, &SCustomizableInstanceProperties::GetBoolParameterValue, ParamName)
+				.OnCheckStateChanged(this, &SCustomizableInstanceProperties::OnBoolParameterChanged, ParamName)
+			]
+		];
+		
+		break;
+	}
+
+	case EMutableParameterType::Float:
+	{
+		TSharedPtr<SVerticalBox> SliderBox;
+
+		ParameterBox->AddSlot()
+		.FillWidth(10.0f)
+		.HAlign(HAlign_Fill)
+		.Padding(1.0f)
+		[
+			SAssignNew(SliderBox, SVerticalBox)
+		];
+
+		// Color bar?
+		if (CustomizableObject->GetParameterDescriptionCount(ParamIndexInObject) > 0)
+		{
+			if (UTexture2D* Bar = CustomInstance->GetParameterDescription(ParamIndexInObject, 0))
+			{
+				TSharedRef<FDeferredCleanupSlateBrush> Brush = FDeferredCleanupSlateBrush::CreateBrush(Bar, FVector2D(Bar->GetSizeX(), 12), FLinearColor(1.0f, 1.0f, 1.0f, 1.0f), ESlateBrushTileType::NoTile, ESlateBrushImageType::Linear);
+				DynamicBrushes.Add(Brush);
+
+				SliderBox->AddSlot()
+				.AutoHeight()
+				.FillHeight(1.0f)
+				.HAlign(HAlign_Fill)
+				.Padding(0.0f, 0.0f, 0.0f, 2.0f)
+				[
+					SNew(SImage)
+					.Image(Brush->GetSlateBrush())
+				];
+			}
+		}
+
+		TSharedPtr<SSpinBox<float>> Slider;
+
+		int SliderIndex = FloatSliders.Num();
+		SliderBox->AddSlot()
+		.AutoHeight()
+		.HAlign(HAlign_Fill)
+		[
+			SAssignNew(Slider, SSpinBox<float>)
+			.Value(this, &SCustomizableInstanceProperties::GetSliderValue, SliderIndex)
+			.MinValue(CustomizableObject->GetParameterUIMetadataFromIndex(ParamIndexInObject).ParamUIMetadata.MinimumValue)
+			.MaxValue(CustomizableObject->GetParameterUIMetadataFromIndex(ParamIndexInObject).ParamUIMetadata.MaximumValue)
+			.OnValueChanged(this, &SCustomizableInstanceProperties::OnFloatParameterChanged, SliderIndex)
+			.OnBeginSliderMovement(this, &SCustomizableInstanceProperties::OnFloatParameterSliderBegin)
+			.OnEndSliderMovement(this, &SCustomizableInstanceProperties::OnFloatParameterSliderEnd)
+
+		];
+
+		FloatSliders.Add(FSliderData(Slider, ParamName, -1, GetFloatParameterValue(ParamName, -1)));
+		break;
+	}
+
+	case EMutableParameterType::Int:
+	{
+		int numValues = CustomizableObject->GetIntParameterNumOptions(ParamIndexInObject);
+		bool bIsParamMultidimensional = CustomInstance->IsParamMultidimensional(ParamIndexInObject);
+
+		if (!bIsParamMultidimensional && numValues)
+		{
+			FString ToolTipText = FString("None");
+			//TSharedPtr< TArray< TSharedPtr<FString> > > OptionNames = MakeShareable(new TArray< TSharedPtr<FString> >());
+			//IntOptionNames.Add(OptionNames);
+			
+			TArray<FString> OptionNamesAttribute;
+			FString Value = GetIntParameterValue(ParamName);
+			int ValueIndex = 0;
+			
+			for (int i = 0; i < numValues; ++i)
+			{
+				FString PossibleValue = CustomizableObject->GetIntParameterAvailableOption(ParamIndexInObject, i);
+				
+				if (PossibleValue == Value)
+				{
+					ValueIndex = i;
+
+					const FString* Identifier = CustomizableObject->GroupNodeMap.FindKey(FCustomizableObjectIdPair(ParamName, PossibleValue));
+					if (Identifier)
+					{
+						if (FString* CustomizableObjectPath = CustomizableObject->CustomizableObjectPathMap.Find(*Identifier))
+						{
+							ToolTipText = *CustomizableObjectPath;
+						}
+					}
+				}
+
+				//OptionNames->Add(MakeShareable(new FString(CustomizableObject->GetIntParameterAvailableOption(ParamIndexInObject, i))));
+				OptionNamesAttribute.Add(CustomizableObject->GetIntParameterAvailableOption(ParamIndexInObject, i));
+			}
+
+			//OptionNames->Sort(::CompareNames);
+
+			ParameterBox->AddSlot()
+			.FillWidth(8.0f)
+			.Padding(0.2f, 2.0f)
+			.HAlign(HAlign_Fill)
+			.VAlign(VAlign_Fill)
+			[
+				SNew(SMutableTextSearchBox)
+				.ToolTipText(FText::FromString(ToolTipText))
+				.PossibleSuggestions(OptionNamesAttribute)
+				.InitialText(FText::FromString(OptionNamesAttribute[ValueIndex]))
+				.MustMatchPossibleSuggestions(TAttribute<bool>(true))
+				.SuggestionListPlacement(EMenuPlacement::MenuPlacement_ComboBox)
+				.OnTextCommitted(this, &SCustomizableInstanceProperties::OnIntParameterComboBoxChanged, ParamName)
+			];
+		}
+		else
+		{
+			// We only support named integers now.
+			//ParameterBox->AddSlot()
+			//	.AutoWidth()
+			//	.FillWidth(1.0f)
+			//	.Padding(2.0f)
+			//	[
+			//		SNew( SSpinBox<int32> )
+			//		.Value( MutableParameters->GetIntValue( ParamIndex ) )
+			//		.OnValueChanged( this, &SCustomizableInstanceProperties::OnIntParameterChanged, ParamName)
+			//	];
+		}
+		break;
+	}
+
+	case EMutableParameterType::Color:
+	{
+		ParameterBox->AddSlot()
+		.AutoWidth()
+		.FillWidth(10.f)
+		.Padding(2.0f)
+		[
+			SNew(SColorBlock)
+			.Color(this, &SCustomizableInstanceProperties::GetColorParameterValue, ParamName)
+			.ShowBackgroundForAlpha(false)
+			.AlphaDisplayMode(EColorBlockAlphaDisplayMode::Ignore)
+			.UseSRGB(false)
+			.OnMouseButtonDown(this, &SCustomizableInstanceProperties::OnColorBlockMouseButtonDown, ParamName)
+			.Size(FVector2D(10.0f, 10.0f))
+		];
+
+		break;
+	}
+
+	case EMutableParameterType::Texture:
+	{
+		TArray<FCustomizableObjectTextureParameterValue>& TextureParameters = CustomInstance->GetTextureParameters();
+		TSharedPtr<FString> InitiallySelected = TextureParameterValueNames[0];
+		for (int i = 0; i < TextureParameters.Num(); ++i)
+		{
+			FString LocalParamName = TextureParameters[i].ParameterName;
+			uint64 ParamValue = TextureParameters[i].ParameterValue;
+			if (TextureParameters[i].ParameterName == LocalParamName)
+			{
+				// Look for the value index
+				for (int ValueIndex = 0; ValueIndex < TextureParameterValueNames.Num(); ++ValueIndex)
+				{
+					if (TextureParameterValues[ValueIndex] == ParamValue)
+					{
+						InitiallySelected = TextureParameterValueNames[ValueIndex];
+
+						/*if (CustomInstance->ProfileParameterDat.)
+						{
+							CustomInstance->ProfileParameterDat->TextureParameters[i].ParameterValue = ParamValue;
+						}*/	
+					}
+				}
+
+				break;
+			}
+		}
+
+		ParameterBox->AddSlot()
+		.AutoWidth()
+		.FillWidth(1.0f)
+		.Padding(2.0f)
+		[
+			SNew(STextComboBox)
+			.OptionsSource(&TextureParameterValueNames)
+			.InitiallySelectedItem(InitiallySelected)
+			.OnSelectionChanged(this, &SCustomizableInstanceProperties::OnTextureParameterComboBoxSelectionChanged, ParamName)
+		];
+		
+		break;
+	}
+
+	case EMutableParameterType::Projector:
+	{
+		if (SimplifiedProjectorInterface)
+		{
+			ParameterBox->AddSlot()
+			.AutoWidth()
+			.FillWidth(1.0f)
+			.Padding(2.0f)
+			[
+				SNew(SHorizontalBox)
+				+ SHorizontalBox::Slot()
+				.AutoWidth()
+				[
+					SNew(SCheckBox)
+					.Type(ESlateCheckBoxType::ToggleButton)
+					.ToolTipText(LOCTEXT("Hide_Projector", "Hide projector"))
+					.IsChecked(this, &SCustomizableInstanceProperties::IsProjectorHidden, ParamName)
+					[
+						SNew(SImage)
+						// todo: editor style not available in game
+						//.Image(FAppStyle::GetBrush(TEXT("LevelEditor.ViewOptions.Small")))
+					]
+				]
+			
+				+ SHorizontalBox::Slot()
+				.AutoWidth()
+				[
+					SNew(SCheckBox)
+					.Type(ESlateCheckBoxType::ToggleButton)
+					.ToolTipText(LOCTEXT("Show_Projector", "Show projector widget"))
+					.IsChecked(this, &SCustomizableInstanceProperties::IsProjectorTranslate, ParamName)
+					.OnCheckStateChanged(this, &SCustomizableInstanceProperties::OnProjectorTranslateStateChanged, ParamName)
+					.Padding(2.f)
+					[
+						SNew(SImage)
+						// todo: editor style not available in game
+						// .Image(FAppStyle::GetBrush(TEXT("LevelEditor.TranslateMode.Small")))
+					]
+				]
+			];
+		}
+		else
+		{
+			bool bIsParamMultidimensional = CustomInstance->IsParamMultidimensional(ParamIndexInObject);
+
+			if (!bIsParamMultidimensional)
+			{
+				FString ParamNameWithIndex = ParamName + FString::Printf(TEXT("__%d"), -1);
+
+				TSharedPtr<SButton> Button;
+				TSharedPtr<STextBlock> TextBlock;
+
+				ParameterBox->AddSlot()
+				.AutoWidth()
+				[
+					SNew(SHorizontalBox)
+					+ SHorizontalBox::Slot()
+					.AutoWidth()
+					.VAlign(VAlign_Center)
+					[
+						SAssignNew(Button, SButton)
+						.OnClicked(this, &SCustomizableInstanceProperties::OnProjectorSelectStateChanged, ParamNameWithIndex)
+						.HAlign(HAlign_Center)
+						.Content()
+						[
+							SAssignNew(TextBlock, STextBlock)
+							.Text(LOCTEXT("Select Projector", "Select Projector"))
+							.ToolTipText(LOCTEXT("Select Projector", "Select Projector"))
+						]
+					]
+				
+					+SHorizontalBox::Slot()
+					.VAlign(VAlign_Center)
+					.AutoWidth()
+					[
+						SNew(SButton)
+						.HAlign(HAlign_Center)
+						.ToolTipText(LOCTEXT("Copy Transform", "Copy Transform"))
+						.Text(LOCTEXT("Copy Transform", "Copy Transform"))
+						.OnClicked(this, &SCustomizableInstanceProperties::OnProjectorCopyTransform, ParamNameWithIndex)
+					]
+					
+					+ SHorizontalBox::Slot()
+					.VAlign(VAlign_Center)
+					.AutoWidth()
+					[
+						SNew(SButton)
+						.HAlign(HAlign_Center)
+						.ToolTipText(LOCTEXT("Paste Transform", "Paste Transform"))
+						.Text(LOCTEXT("Paste Transform", "Paste Transform"))
+						.OnClicked(this, &SCustomizableInstanceProperties::OnProjectorPasteTransform, ParamNameWithIndex)
+					]
+					
+					+ SHorizontalBox::Slot()
+					.VAlign(VAlign_Center)
+					.AutoWidth()
+					[
+						SNew(SButton)
+						.HAlign(HAlign_Center)
+						.ToolTipText(LOCTEXT("Reset Transform", "Reset Transform"))
+						.Text(LOCTEXT("Reset Transform", "Reset Transform"))
+						.OnClicked(this, &SCustomizableInstanceProperties::OnProjectorResetTransform, ParamNameWithIndex)
+					]
+				];
+				
+				ArrayProjectorParameterNameWithIndex.Add(ParamNameWithIndex);
+
+				MapProjectorButton.Add(ParamNameWithIndex, Button);
+				MapProjectorButtonTextBlock.Add(ParamNameWithIndex, TextBlock);
+
+				if (!CustomInstance->LastSelectedProjectorParameter.IsEmpty() && (CustomInstance->LastSelectedProjectorParameter == ParamName))
+				{
+					if (Button.Get() != nullptr)
+					{
+						Button->SetBorderBackgroundColor(FLinearColor::Green);
+					}
+
+					if (TextBlock.Get() != nullptr)
+					{
+						TextBlock->SetText(LOCTEXT("Unselect Projector", "Unselect Projector"));
+					}
+				}
+			}
+			else
+			{
+				TArray<FCustomizableObjectProjectorParameterValue>& ProjectorParameters = CustomInstance->GetProjectorParameters();
+				int32 ProjectorParamIndex = CustomInstance->FindProjectorParameterNameIndex(ParamName);
+				check(ProjectorParamIndex < ProjectorParameters.Num());
+
+				FString PoseSwitchEnumParamName = ParamName + FMultilayerProjector::POSE_PARAMETER_POSTFIX;
+				FString PoseSwitchEnumParamNameWithRange = PoseSwitchEnumParamName + FString::Printf(TEXT("__%d"), -1);
+				int32 PoseSwitchEnumParamIndexInObject = CustomizableObject->FindParameter(PoseSwitchEnumParamName);
+				check(PoseSwitchEnumParamIndexInObject >= 0);
+
+				int32 NumPoseValues = CustomizableObject->GetIntParameterNumOptions(PoseSwitchEnumParamIndexInObject);
+
+				//TSharedPtr< TArray< TSharedPtr<FString> > > PoseOptionNames = MakeShareable(new TArray< TSharedPtr<FString> >());
+				//IntOptionNames.Add(PoseOptionNames);
+				
+				TArray<FString> PoseOptionNamesAttribute;
+				FString PoseValue = GetIntParameterValue(PoseSwitchEnumParamName, -1);
+				int32 PoseValueIndex = 0;
+				
+				for (int32 j = 0; j < NumPoseValues; ++j)
+				{
+					FString PossibleValue = CustomizableObject->GetIntParameterAvailableOption(PoseSwitchEnumParamIndexInObject, j);
+					if (PossibleValue == PoseValue)
+					{
+						PoseValueIndex = j;
+					}
+					//PoseOptionNames->Add(MakeShareable(new FString(CustomizableObject->GetIntParameterAvailableOption(PoseSwitchEnumParamIndexInObject, j))));
+					PoseOptionNamesAttribute.Add(CustomizableObject->GetIntParameterAvailableOption(PoseSwitchEnumParamIndexInObject, j));
+				}
+
+				TSharedPtr<SHorizontalBox> Box;
+
+				ParameterBox->AddSlot()
+				.HAlign(HAlign_Right)
+				.VAlign(VAlign_Fill)
+				.FillWidth(10.f)
+				[
+					SAssignNew(Box, SHorizontalBox)
+					+ SHorizontalBox::Slot()
+					.HAlign(HAlign_Fill)
+					.VAlign(VAlign_Center)
+					.FillWidth(0.45f)
+					[
+						SNew(SMutableTextSearchBox)
+						.PossibleSuggestions(PoseOptionNamesAttribute)
+						.InitialText(FText::FromString(PoseOptionNamesAttribute[PoseValueIndex]))
+						.MustMatchPossibleSuggestions(TAttribute<bool>(true))
+						.SuggestionListPlacement(EMenuPlacement::MenuPlacement_ComboBox)
+						.OnTextCommitted(this, &SCustomizableInstanceProperties::OnProjectorTextureParameterComboBoxChanged, PoseSwitchEnumParamNameWithRange)
+					]
+					
+					+ SHorizontalBox::Slot()
+					.HAlign(HAlign_Fill)
+					.VAlign(VAlign_Center)
+					.FillWidth(0.3f)
+					[
+						SNew(SButton)
+						.ToolTipText(LOCTEXT("Add Layer", "Add Layer"))
+						.Text(LOCTEXT("Add Layer", "Add Layer"))
+						.OnClicked(this, &SCustomizableInstanceProperties::OnProjectorLayerAdded, ParamName)
+						.HAlign(HAlign_Fill)
+					]
+				];
+
+				for (int32 i = 0; i < ProjectorParameters[ProjectorParamIndex].RangeValues.Num(); ++i)
+				{
+					FString ParamNameWithIndex = ParamName + FString::Printf(TEXT("__%d"), i);
+
+					FString TextureSwitchEnumParamName = ParamName + FMultilayerProjector::IMAGE_PARAMETER_POSTFIX;
+					FString TextureSwitchEnumParamNameWithRange = TextureSwitchEnumParamName + FString::Printf(TEXT("__%d"), i);
+					int32 TextureSwitchEnumParamIndexInObject = CustomizableObject->FindParameter(TextureSwitchEnumParamName);
+					check(TextureSwitchEnumParamIndexInObject >= 0);
+					int32 numValues = CustomizableObject->GetIntParameterNumOptions(TextureSwitchEnumParamIndexInObject);
+
+					FString OpacitySliderParamName = ParamName + FMultilayerProjector::OPACITY_PARAMETER_POSTFIX;
+					FString OpacitySliderParamNameWithRange = OpacitySliderParamName + FString::Printf(TEXT("__%d"), i);
+
+					//if (numValues)
+					//{
+					//TSharedPtr< TArray< TSharedPtr<FString> > > OptionNames = MakeShareable(new TArray< TSharedPtr<FString> >());
+					//IntOptionNames.Add(OptionNames);
+
+					TArray<FString> OptionNamesAttribute;
+					FString Value = GetIntParameterValue(TextureSwitchEnumParamName, i);
+					int32 ValueIndex = 0;
+
+					for (int32 CandidateIndex = 0; CandidateIndex < numValues; ++CandidateIndex)
+					{
+						FString PossibleValue = CustomizableObject->GetIntParameterAvailableOption(TextureSwitchEnumParamIndexInObject, CandidateIndex);
+						if (PossibleValue == Value)
+						{
+							ValueIndex = CandidateIndex;
+						}
+						//OptionNames->Add(MakeShareable(new FString(CustomizableObject->GetIntParameterAvailableOption(TextureSwitchEnumParamIndexInObject, CandidateIndex))));
+						OptionNamesAttribute.Add(CustomizableObject->GetIntParameterAvailableOption(TextureSwitchEnumParamIndexInObject, CandidateIndex));
+					}
+					//}
+
+					TSharedPtr<SHorizontalBox> SliderBox;
+					TSharedPtr<SSpinBox<float>> Slider;
+					TSharedPtr<SHorizontalBox> PropertiesHB;
+					int32 SliderIndex = FloatSliders.Num();
+
+					TSharedPtr<SButton> Button;
+					TSharedPtr<STextBlock> TextBlock;
+
+					//Horizontal box that owns all the hidden layer properties
+					ActualParamBox->AddSlot()
+					.HAlign(HAlign_Fill)
+					.Padding(0, 10.f)
+					[
+						SAssignNew(PropertiesHB, SHorizontalBox)
+						
+						//Select Projector slot
+						+ SHorizontalBox::Slot()
+						.Padding(1, 0)
+						.AutoWidth()
+						[
+							SNew(SHorizontalBox)
+							+ SHorizontalBox::Slot()
+							.Padding(1, 0)
+							[
+								SNew(SBox)
+								.MinDesiredWidth(115.f)
+								.MaxDesiredWidth(115.f)
+								[
+									SAssignNew(Button, SButton)
+									.OnClicked(this, &SCustomizableInstanceProperties::OnProjectorSelectStateChanged, ParamNameWithIndex)
+									.VAlign(VAlign_Center)
+									.HAlign(HAlign_Center)
+									.Content()
+									[
+										SAssignNew(TextBlock, STextBlock)
+										.Text(LOCTEXT("Select Projector", "Select Projector"))
+										.ToolTipText(LOCTEXT("Select Projector", "Select Projector"))
+										.Justification(ETextJustify::Center)
+										//.AutoWrapText(true)
+									]
+								]
+							]
+
+							+ SHorizontalBox::Slot()
+							.Padding(1, 0)
+							[
+								SNew(SBox)
+								.MinDesiredWidth(120.f)
+								[
+									SNew(SMutableTextSearchBox)
+									.PossibleSuggestions(OptionNamesAttribute)
+									.InitialText(FText::FromString(OptionNamesAttribute[ValueIndex]))
+									.MustMatchPossibleSuggestions(TAttribute<bool>(true))
+									.SuggestionListPlacement(EMenuPlacement::MenuPlacement_ComboBox)
+									.OnTextCommitted(this, &SCustomizableInstanceProperties::OnProjectorTextureParameterComboBoxChanged, TextureSwitchEnumParamNameWithRange)
+								]
+							]
+						]
+					];
+						
+					MapProjectorButton.Add(ParamNameWithIndex, Button);
+					MapProjectorButtonTextBlock.Add(ParamNameWithIndex, TextBlock);
+
+					PropertiesHB->AddSlot()
+					.HAlign(HAlign_Fill)
+					[
+						SNew(SHorizontalBox)
+						+ SHorizontalBox::Slot()
+						.Padding(1, 0)
+						.HAlign(HAlign_Fill)
+						[
+							SNew(SHorizontalBox)
+							+ SHorizontalBox::Slot()
+							[
+								SAssignNew(Slider, SSpinBox<float>)
+								.MinValue(0.0f)
+								.MaxValue(1.0f)
+								.Value(this, &SCustomizableInstanceProperties::GetSliderValue, SliderIndex)
+								.OnValueChanged(this, &SCustomizableInstanceProperties::OnProjectorFloatParameterChanged, SliderIndex, ParamNameWithIndex)
+								.OnBeginSliderMovement(this, &SCustomizableInstanceProperties::OnProjectorFloatParameterSliderBegin)
+								.OnEndSliderMovement(this, &SCustomizableInstanceProperties::OnProjectorFloatParameterSliderEnd)
+							]
+
+						]
+					];
+
+					PropertiesHB->AddSlot()
+					.HAlign(HAlign_Right)
+					.AutoWidth()
+					[
+						SNew(SHorizontalBox)
+						+ SHorizontalBox::Slot()
+						.Padding(1, 0)
+						[
+							SNew(SBox)
+							.MinDesiredWidth(80.f)
+							[
+								SNew(SButton)
+								.OnClicked(this, &SCustomizableInstanceProperties::OnProjectorCopyTransform, ParamNameWithIndex)
+								.VAlign(VAlign_Center)
+								.HAlign(HAlign_Center)
+								.Content()
+								[
+									SNew(STextBlock)
+									.ToolTipText(LOCTEXT("Copy Transform", "Copy Transform"))
+									.Text(LOCTEXT("Copy Transform", "Copy Transform"))
+									.AutoWrapText(true)
+									.Justification(ETextJustify::Center)
+								]
+							]
+						]
+
+						+ SHorizontalBox::Slot()
+						.Padding(1, 0)
+						[
+							SNew(SBox)
+							.MinDesiredWidth(80.f)
+							[
+								SNew(SButton)
+								.OnClicked(this, &SCustomizableInstanceProperties::OnProjectorPasteTransform, ParamNameWithIndex)
+								.VAlign(VAlign_Center)
+								.HAlign(HAlign_Center)
+								.Content()
+								[
+									SNew(STextBlock)
+									.ToolTipText(LOCTEXT("Paste Transform", "Paste Transform"))
+									.Text(LOCTEXT("Paste Transform", "Paste Transform"))
+									.Justification(ETextJustify::Center)
+									.AutoWrapText(true)
+								]
+							]
+						]
+
+						+ SHorizontalBox::Slot()
+						.Padding(1, 0)
+						[
+							SNew(SBox)
+							.MinDesiredWidth(80.f)
+							[
+								SNew(SButton)
+								.OnClicked(this, &SCustomizableInstanceProperties::OnProjectorResetTransform, ParamNameWithIndex)
+								.VAlign(VAlign_Center)
+								.HAlign(HAlign_Center)
+								.Content()
+								[
+									SNew(STextBlock)
+									.ToolTipText(LOCTEXT("Reset Transform", "Reset Transform"))
+									.Text(LOCTEXT("Reset Transform", "Reset Transform"))
+									.Justification(ETextJustify::Center)
+									.AutoWrapText(true)
+								]
+							]
+						]
+
+						+ SHorizontalBox::Slot()
+						.Padding(1, 0)
+						[
+							SNew(SBox)
+							.MinDesiredWidth(80.f)
+							[
+								SNew(SButton)
+								.OnClicked(this, &SCustomizableInstanceProperties::OnProjectorLayerRemoved, ParamNameWithIndex)
+								.VAlign(VAlign_Center)
+								.HAlign(HAlign_Center)
+								.Content()
+								[
+									SNew(STextBlock)
+									.ToolTipText(LOCTEXT("Remove Layer", "Remove Layer"))
+									.Text(LOCTEXT("Remove Layer", "Remove Layer"))
+									.Justification(ETextJustify::Center)
+									.AutoWrapText(true)
+								]
+							]
+						]
+					];
+
+					ArrayProjectorParameterNameWithIndex.Add(ParamNameWithIndex);
+
+					if (!CustomInstance->LastSelectedProjectorParameterWithIndex.IsEmpty() && (CustomInstance->LastSelectedProjectorParameterWithIndex == ParamNameWithIndex))
+					{
+						if (Button.Get() != nullptr)
+						{
+							Button->SetBorderBackgroundColor(FLinearColor::Green);
+						}
+
+						if (TextBlock.Get() != nullptr)
+						{
+							TextBlock->SetText(LOCTEXT("Unselect Projector", "Unselect Projector"));
+						}
+					}
+
+					FloatSliders.Add(FSliderData(Slider, OpacitySliderParamName, i, GetFloatParameterValue(OpacitySliderParamName, i)));
+				}
+			}
+		}
+		break;
+	}
+
+	default:
+		break;
+	}
+}
+
+void SCustomizableInstanceProperties::Tick(const FGeometry& AllottedGeometry, const double InCurrentTime, const float InDeltaTime)
+{
+	SCompoundWidget::Tick(AllottedGeometry, InCurrentTime, InDeltaTime);
+
+	// We set this every frame, in case other widgets are also modifying this flag
+	if (!CustomInstance->bBuildParameterDecorations)
+	{
+		CustomInstance->bBuildParameterDecorations = true;
+		CustomInstance->UpdateSkeletalMeshAsync(true);
+	}
+
+	// Check the slider values
+	// We do this because the slider "changed" event seems to be called many times per-frame
+	TArray<FCustomizableObjectFloatParameterValue>& FloatParameters = CustomInstance->GetFloatParameters();
+	for (int SliderIndex = 0; SliderIndex < FloatSliders.Num() && bSlidersChanged; ++SliderIndex)
+	{
+		for (int i = 0; i < FloatParameters.Num(); ++i)
+		{
+			if (FloatParameters[i].ParameterName == FloatSliders[SliderIndex].ParameterName)
+			{
+				const int32 RangeIndex = FloatSliders[SliderIndex].RangeIndex;
+
+				const float NewValue = FloatSliders[SliderIndex].LastValueSet;
+				float CurrentValue = NewValue;
+
+				if (FloatSliders[SliderIndex].RangeIndex < 0)
+				{
+					CurrentValue = FloatParameters[i].ParameterValue;
+				}
+				else if (FloatParameters[i].ParameterRangeValues.IsValidIndex(FloatSliders[SliderIndex].RangeIndex))
+				{
+					CurrentValue = FloatParameters[i].ParameterRangeValues[FloatSliders[SliderIndex].RangeIndex];
+				}
+
+				if (CurrentValue != NewValue)
+				{
+					CustomInstance->PreEditChange(nullptr);
+
+					if (RangeIndex < 0)
+					{
+						FloatParameters[i].ParameterValue = NewValue;
+					}
+					else
+					{
+						check(FloatParameters[i].ParameterRangeValues.IsValidIndex(RangeIndex));
+						FloatParameters[i].ParameterRangeValues[RangeIndex] = NewValue;
+					}
+
+					CustomInstance->SetSelectedParameterProfileDirty();
+
+					CustomInstance->PostEditChange();
+					CustomInstance->UpdateSkeletalMeshAsync(true, true);
+				}
+			}
+		}
+	}
+	
+	bSlidersChanged = false;
+}
+
+
+void SCustomizableInstanceProperties::OnShowOnlyRuntimeSelectionChanged(ECheckBoxState InCheckboxState)
+{
+	CustomInstance->bShowOnlyRuntimeParameters = InCheckboxState == ECheckBoxState::Checked;
+	ResetParamBox();
+}
+
+
+void SCustomizableInstanceProperties::OnShowOnlyRelevantSelectionChanged(ECheckBoxState InCheckboxState)
+{
+	CustomInstance->bShowOnlyRelevantParameters = InCheckboxState == ECheckBoxState::Checked;
+	ResetParamBox();
+}
+
+
+FLinearColor SCustomizableInstanceProperties::GetColorParameterValue(FString ParamName) const
+{
+	return CustomInstance->GetColorParameterSelectedOption(ParamName);
+}
+
+
+void SCustomizableInstanceProperties::OnStateComboBoxSelectionChanged(TSharedPtr<FString> Selection, ESelectInfo::Type SelectInfo)
+{
+	CustomInstance->PreEditChange(nullptr);
+
+	CustomInstance->SetCurrentState(*Selection);
+
+	CustomInstance->UpdateSkeletalMeshAsync(true, true);
+
+	CustomInstance->PostEditChange();
+
+	// Non-continuous change: collect garbage.
+	GEngine->ForceGarbageCollection();
+
+	CustomInstance->InstanceUpdated = false;
+}
+
+
+FReply SCustomizableInstanceProperties::OnCopyAllParameters()
+{
+	ToBinary = new FBufferArchive(false, TEXT("BufferArchive"));
+
+	UE_LOG(LogMutable, Warning, TEXT("Parameters have been copied."));
+
+	CustomInstance->SaveDescriptor(*ToBinary);
+
+	// FString StringifyedData = FString::FromHexBlob(ToBinary->GetData(), ToBinary->Num());
+	// FPlatformApplicationMisc::ClipboardCopy(StringifyedData.GetCharArray().GetData());
+	ParametersHaveBeenRead = true;
+	
+	return FReply::Handled();
+}
+
+
+FReply SCustomizableInstanceProperties::OnPasteAllParameters()
+{
+	if (ParametersHaveBeenRead)
+	{
+		//FString StringifiedData;
+		//FPlatformApplicationMisc::ClipboardPaste(StringifiedData);
+		//
+		//TArray<uint8> DescriptorData;
+		//// Each char represents a hex number i.e. half byte.
+		//DescriptorData.SetNum(StringifiedData.Len()*2);
+		//FString::ToHexBlob(StringifiedData, DescriptorData.GetData(), DescriptorData.Num());
+		//FMemoryReader FromBinary = FMemoryReader(DescriptorData, true); //true, free data after done;
+
+		FMemoryReader FromBinary = FMemoryReader(*ToBinary, true); //true, free data after done;
+		FromBinary.Seek(0);
+
+		CustomInstance->LoadDescriptor(FromBinary);
+
+		CustomInstance->SetSelectedParameterProfileDirty();
+
+		CustomInstance->UpdateSkeletalMeshAsync(true, true);
+
+		CustomInstance->PostEditChange();
+
+		// Non-continuous change: collect garbage.
+		GEngine->ForceGarbageCollection();
+	}
+
+	return FReply::Handled();
+}
+
+
+FReply SCustomizableInstanceProperties::OnResetAllParameters()
+{
+	CustomInstance->PreEditChange(nullptr);
+
+	TArray<FCustomizableObjectIntParameterValue>& IntParameters = CustomInstance->GetIntParameters();
+	
+	const UCustomizableObject* CustomObject = CustomInstance->GetCustomizableObject();
+	for (int32 i = 0; i < IntParameters.Num(); ++i)
+	{
+		const int32 ParamIndexInObject = CustomObject->FindParameter(IntParameters[i].ParameterName);
+		if (ParamIndexInObject != INDEX_NONE)
+		{
+			const int32 NumOptions = CustomObject->GetIntParameterNumOptions(ParamIndexInObject);
+			for (int32 v = 0; v < NumOptions; ++v)
+			{
+				FString ValueName = CustomObject->GetIntParameterAvailableOption(ParamIndexInObject, v);
+				if (ValueName == "None")
+				{
+					IntParameters[i].ParameterValueName = "None";
+				}
+				break;
+			}
+		}
+		else
+		{
+			UE_LOG(LogMutable, Error, TEXT("Failed to find parameter."));
+		}
+	}
+
+	CustomInstance->UpdateSkeletalMeshAsync(true, true);
+
+	CustomInstance->PostEditChange();
+
+	// Non-continuous change: collect garbage.
+	GEngine->ForceGarbageCollection();
+
+	CustomInstance->SelectedProfileIndex = INDEX_NONE;
+
+	InstanceDetails.Pin()->Refresh();
+
+	return FReply::Handled();
+}
+
+
+ECheckBoxState SCustomizableInstanceProperties::GetBoolParameterValue(FString ParamName) const
+{
+	const bool bResult = CustomInstance->GetBoolParameterSelectedOption(ParamName);
+	return bResult ? ECheckBoxState::Checked : ECheckBoxState::Unchecked;
+}
+
+
+void SCustomizableInstanceProperties::OnAreaExpansionChanged(bool bExpanded, FString ParamName) const
+{
+	if (bExpanded)
+	{
+		CustomInstance->ParamNameToExpandedMap.Add(ParamName, true);
+	}
+	else
+	{
+		CustomInstance->ParamNameToExpandedMap.Remove(ParamName);
+	}
+}
+
+
+void SCustomizableInstanceProperties::OnBoolParameterChanged(ECheckBoxState InCheckboxState, FString ParamName)
+{
+	CustomInstance->PreEditChange(nullptr);
+
+	CustomInstance->SetBoolParameterSelectedOption(ParamName, InCheckboxState == ECheckBoxState::Checked);
+
+	CustomInstance->SetSelectedParameterProfileDirty();
+
+	CustomInstance->UpdateSkeletalMeshAsync(true, true);
+
+	CustomInstance->PostEditChange();
+
+	// Non-continuous change: collect garbage.
+	GEngine->ForceGarbageCollection();
+}
+
+
+float SCustomizableInstanceProperties::GetSliderValue(int Slider) const
+{
+	if (Slider >= 0 && Slider < FloatSliders.Num())
+	{
+		return FloatSliders[Slider].LastValueSet;
+	}
+
+	return 0.0f;
+}
+
+
+float SCustomizableInstanceProperties::GetFloatParameterValue(FString ParamName, int32 RangeIndex) const
+{
+	return CustomInstance->GetFloatParameterSelectedOption(ParamName, RangeIndex);
+}
+
+
+void SCustomizableInstanceProperties::OnFloatParameterChanged(float Value, int SliderIndex)
+{
+	// Actually updated in tick, since this seems to be called many times per frame.
+	if (SliderIndex >= 0 && SliderIndex < FloatSliders.Num())
+	{
+		FloatSliders[SliderIndex].LastValueSet = Value;
+
+		// Update the float parameters on the tick
+		bSlidersChanged = true;
+	}
+}
+
+
+void SCustomizableInstanceProperties::OnFloatParameterSliderBegin()
+{
+	bShouldResetParamBox = false;
+}
+
+
+void SCustomizableInstanceProperties::OnFloatParameterSliderEnd(float Value)
+{
+	bShouldResetParamBox = true;
+
+	// Non-continuous change: collect garbage.
+	GEngine->ForceGarbageCollection();
+}
+
+
+FString SCustomizableInstanceProperties::GetIntParameterValue(FString ParamName, int RangeIndex) const
+{
+	return CustomInstance->GetIntParameterSelectedOption(ParamName, RangeIndex);
+}
+
+
+void SCustomizableInstanceProperties::OnIntParameterChanged(int32 Value, FString ParamName)
+{
+	//		CustomInstance->PreEditChange(nullptr);
+	//
+	//		mu::ParametersPtr MutableParameters = CustomInstance->ReloadParametersFromObject();
+	//
+	//		for (int i=0; i<CustomInstance->IntParameters.Num(); ++i )
+	//		{
+	//			if ( CustomInstance->IntParameters[i].MutableIndex==ParamIndex )
+	//			{
+	//				int possibleValue = 0;
+	//				if (Value >= 0 &&
+	//					Value < MutableParameters->GetIntPossibleValueCount(ParamIndex))
+	//				{
+	//					possibleValue = MutableParameters->GetIntPossibleValue(ParamIndex, Value);
+	//				}
+	//				CustomInstance->IntParameters[i].ParameterValue = possibleValue;
+	//			}
+	//		}
+	//
+	//		if (OnChangeCallback)
+	//		{
+	//			OnChangeCallback(CustomInstance.Get());
+	//		}
+	//
+	//		CustomInstance->PostEditChange();
+	//	
+	//		if (CustomInstance->bShowOnlyRelevantParameters)
+	//		{
+	//			ResetParamBox();
+	//		}
+}
+
+
+void SCustomizableInstanceProperties::OnIntParameterComboBoxChanged(TSharedPtr<FString> Selection, ESelectInfo::Type SelectInfo, FString ParamName)
+{
+	TArray<FCustomizableObjectIntParameterValue>& IntParameters = CustomInstance->GetIntParameters();
+	
+	const UCustomizableObject* CustomObject = CustomInstance->GetCustomizableObject();
+	for (int32 i = 0; i < IntParameters.Num(); ++i)
+	{
+		if (IntParameters[i].ParameterName == ParamName)
+		{
+			const int32 ParamIndexInObject = CustomObject->FindParameter(ParamName);
+			if (ParamIndexInObject != INDEX_NONE)
+			{
+				for (int32 v = 0; v < CustomObject->GetIntParameterNumOptions(ParamIndexInObject); ++v)
+				{
+					FString ValueName = CustomObject->GetIntParameterAvailableOption(ParamIndexInObject, v);
+
+					if (ValueName == *Selection)
+					{
+						if (ValueName == "None" || CustomInstance->FindProjectorParameterNameIndex(ParamName) == INDEX_NONE)
+						{
+							// Hiding the projector when the selection doesn't have one
+							CustomInstance->ResetProjectorStates();
+							CustomInstance->UnselectProjector = true;
+						}
+						else
+						{
+							// Variable used in a non multilayer projector to not lose the focus of the gizmo
+							CustomInstance->ProjectorLayerChange = true;
+						}
+						
+
+						IntParameters[i].ParameterValueName = ValueName;
+						break;
+					}
+				}
+			}
+			else
+			{
+				UE_LOG(LogMutable, Error, TEXT("Failed to find parameter."));
+			}
+		}
+	}
+
+	CustomInstance->SetSelectedParameterProfileDirty();
+
+	CustomInstance->PreEditChange(nullptr);
+
+	CustomInstance->UpdateSkeletalMeshAsync(true, true);
+
+	CustomInstance->PostEditChange();
+
+	// Non-continuous change: collect garbage.
+	GEngine->ForceGarbageCollection();
+}
+
+
+void SCustomizableInstanceProperties::OnTextureParameterComboBoxSelectionChanged(TSharedPtr<FString> Selection, ESelectInfo::Type SelectInfo, FString ParamName)
+{
+	CustomInstance->PreEditChange(nullptr);
+
+	auto& TextureParameters = CustomInstance->GetTextureParameters();
+
+	for (int i = 0; i < TextureParameters.Num(); ++i)
+	{
+		if (TextureParameters[i].ParameterName == ParamName)
+		{
+			for (int o = 0; o < TextureParameterValueNames.Num(); ++o)
+			{
+				if (Selection.IsValid() && *TextureParameterValueNames[o] == *Selection)
+				{
+					TextureParameters[i].ParameterValue = TextureParameterValues[o];
+				}
+			}
+		}
+	}
+
+	CustomInstance->SetSelectedParameterProfileDirty();
+
+	CustomInstance->UpdateSkeletalMeshAsync(true, true);
+
+	CustomInstance->PostEditChange();
+
+	// Non-continuous change: collect garbage.
+	GEngine->ForceGarbageCollection();
+}
+
+
+void SCustomizableInstanceProperties::OnFilterTextChanged(const FText & InFilterText)
+{
+	ResetParamBox();
+}
+
+
+FReply SCustomizableInstanceProperties::OnColorBlockMouseButtonDown(const FGeometry& MyGeometry, const FPointerEvent& MouseEvent, FString ParamName)
+{
+	if (MouseEvent.GetEffectingButton() != EKeys::LeftMouseButton)
+	{
+		return FReply::Unhandled();
+	}
+
+	FLinearColor col = GetColorParameterValue(ParamName);
+
+	TArray<FLinearColor*> LinearColorArray;
+	LinearColorArray.Add(&col);
+
+	FColorPickerArgs args;
+	args.bIsModal = true;
+	args.bUseAlpha = false;
+	args.bOnlyRefreshOnMouseUp = false;
+	args.InitialColorOverride = col;
+	args.OnColorCommitted = FOnLinearColorValueChanged::CreateSP(this, &SCustomizableInstanceProperties::OnSetColorFromColorPicker, ParamName);
+	OpenColorPicker(args);
+
+	return FReply::Handled();
+}
+
+
+void SCustomizableInstanceProperties::OnSetColorFromColorPicker(FLinearColor NewColor, FString PickerParamName) const
+{
+	CustomInstance->PreEditChange(nullptr);
+
+	CustomInstance->SetColorParameterSelectedOption(PickerParamName, NewColor);
+	
+	CustomInstance->SetSelectedParameterProfileDirty();
+
+	PickerParamName = FString();
+
+	CustomInstance->UpdateSkeletalMeshAsync(true, true);
+
+	CustomInstance->PostEditChange();
+}
+
+
+FString ParseNameAndIndex(const FString& ParamNameWithIndex, int& OutIndex)
+{
+	FString ParamName, Index;
+	const bool bSplit = ParamNameWithIndex.Split(FString("__"), &ParamName, &Index);
+	check(bSplit && Index.IsNumeric());
+
+	OutIndex = FCString::Atoi(*Index);
+
+	return ParamName;
+}
+
+
+ECheckBoxState SCustomizableInstanceProperties::IsProjectorHidden(FString ParamNameWithIndex) const
+{
+	int32 RangeIndex;
+	const  FString ParamName = ParseNameAndIndex(ParamNameWithIndex, RangeIndex);
+
+	if (CustomInstance->GetProjectorState(ParamName, RangeIndex) == EProjectorState::Hidden)
+	{
+		return ECheckBoxState::Checked;
+	}
+	return ECheckBoxState::Unchecked;
+}
+
+
+ECheckBoxState SCustomizableInstanceProperties::IsProjectorTranslate(FString ParamNameWithIndex) const
+{
+	int32 RangeIndex;
+	const FString ParamName = ParseNameAndIndex(ParamNameWithIndex, RangeIndex);
+
+	if (CustomInstance->GetProjectorState(ParamName, RangeIndex) == EProjectorState::Translate)
+	{
+		return ECheckBoxState::Checked;
+	}
+	return ECheckBoxState::Unchecked;
+}
+
+
+ECheckBoxState SCustomizableInstanceProperties::IsProjectorRotate(FString ParamNameWithIndex) const
+{
+	int32 RangeIndex;
+	const FString ParamName = ParseNameAndIndex(ParamNameWithIndex, RangeIndex);
+
+	if (CustomInstance->GetProjectorState(ParamName, RangeIndex) == EProjectorState::Rotate)
+	{
+		return ECheckBoxState::Checked;
+	}
+	return ECheckBoxState::Unchecked;
+}
+
+
+ECheckBoxState SCustomizableInstanceProperties::IsProjectorScale(FString ParamNameWithIndex) const
+{
+	int32 RangeIndex;
+	const FString ParamName = ParseNameAndIndex(ParamNameWithIndex, RangeIndex);
+
+	if (CustomInstance->GetProjectorState(ParamName, RangeIndex) == EProjectorState::Scale)
+	{
+		return ECheckBoxState::Checked;
+	}
+	return ECheckBoxState::Unchecked;
+}
+
+
+FReply SCustomizableInstanceProperties::OnProjectorLayerAdded(FString ParamName) const
+{
+	const int32 NumLayers = CustomInstance->AddValueToProjectorRange(ParamName) + 1;
+	if (NumLayers == 0)
+	{
+		return FReply::Handled();
+	}
+
+	const FString TextureSwitchEnumParamName = ParamName + FMultilayerProjector::IMAGE_PARAMETER_POSTFIX;
+	CustomInstance->AddValueToIntRange(TextureSwitchEnumParamName);
+	check(CustomInstance->FindIntParameterNameIndex(TextureSwitchEnumParamName) != INDEX_NONE);
+	check(NumLayers == CustomInstance->GetIntParameters()[CustomInstance->FindIntParameterNameIndex(TextureSwitchEnumParamName)].ParameterRangeValueNames.Num());
+
+	TArray<FCustomizableObjectFloatParameterValue>& FloatParameters = CustomInstance->GetFloatParameters();
+	
+	const FString NumLayersParamName = ParamName + FMultilayerProjector::NUM_LAYERS_PARAMETER_POSTFIX;
+	check(CustomInstance->FindFloatParameterNameIndex(NumLayersParamName) != INDEX_NONE);
+	FloatParameters[CustomInstance->FindFloatParameterNameIndex(NumLayersParamName)].ParameterValue = NumLayers;
+
+	const FString OpacitySliderParamName = ParamName + FMultilayerProjector::OPACITY_PARAMETER_POSTFIX;
+	CustomInstance->AddValueToFloatRange(OpacitySliderParamName);
+	check(CustomInstance->FindFloatParameterNameIndex(OpacitySliderParamName) != INDEX_NONE);
+	check(NumLayers == FloatParameters[CustomInstance->FindFloatParameterNameIndex(OpacitySliderParamName)].ParameterRangeValues.Num());
+
+	CustomInstance->SetSelectedParameterProfileDirty();
+
+	CustomInstance->PreEditChange(nullptr);
+
+	CustomInstance->UpdateSkeletalMeshAsync(true, true);
+
+	CustomInstance->PostEditChange();
+	FCoreUObjectDelegates::BroadcastOnObjectModified(CustomInstance.Get());
+
+	return FReply::Handled();
+}
+
+
+FReply SCustomizableInstanceProperties::OnProjectorLayerRemoved(FString ParamNameWithIndex) const
+{
+	int32 RangeIndex;
+	const FString ParamName = ParseNameAndIndex(ParamNameWithIndex, RangeIndex);
+
+	const int32 projectorParameterIndex = CustomInstance->FindProjectorParameterNameIndex(ParamName);
+	if (projectorParameterIndex == INDEX_NONE
+		|| CustomInstance->GetProjectorParameters()[projectorParameterIndex].RangeValues.Num() <= 0)
+	{
+		return FReply::Handled();
+	}
+
+	CustomInstance->SetProjectorState(ParamName, RangeIndex, EProjectorState::Hidden);  // If the Projector was selected, unselect it
+
+	const int32 NumLayers = CustomInstance->RemoveValueFromProjectorRange(ParamName, RangeIndex) + 1;
+	CustomInstance->RemovedProjectorParameterNameWithIndex = ParamNameWithIndex; // Since the gizmo updates the CustomObjectInstance when unselected, a check is required to avoid updating the removed projector parameter range index
+
+	const FString TextureSwitchEnumParamName = ParamName + FMultilayerProjector::IMAGE_PARAMETER_POSTFIX;
+	CustomInstance->RemoveValueFromIntRange(TextureSwitchEnumParamName, RangeIndex);
+	check(CustomInstance->FindIntParameterNameIndex(TextureSwitchEnumParamName) != INDEX_NONE);
+	check(NumLayers == CustomInstance->GetIntParameters()[CustomInstance->FindIntParameterNameIndex(TextureSwitchEnumParamName)].ParameterRangeValueNames.Num());
+
+	TArray<FCustomizableObjectFloatParameterValue>& FloatParameters = CustomInstance->GetFloatParameters();
+
+	const FString NumLayersParamName = ParamName + FMultilayerProjector::NUM_LAYERS_PARAMETER_POSTFIX;
+	check(CustomInstance->FindFloatParameterNameIndex(NumLayersParamName) != INDEX_NONE);
+	FloatParameters[CustomInstance->FindFloatParameterNameIndex(NumLayersParamName)].ParameterValue = NumLayers;
+
+	const FString OpacitySliderParamName = ParamName + FMultilayerProjector::OPACITY_PARAMETER_POSTFIX;
+	CustomInstance->RemoveValueFromFloatRange(OpacitySliderParamName, RangeIndex);
+	check(CustomInstance->FindFloatParameterNameIndex(OpacitySliderParamName) != INDEX_NONE);
+	check(NumLayers == FloatParameters[CustomInstance->FindFloatParameterNameIndex(OpacitySliderParamName)].ParameterRangeValues.Num());
+
+	CustomInstance->SetSelectedParameterProfileDirty();
+	
+	CustomInstance->PreEditChange(nullptr);
+
+	CustomInstance->UpdateSkeletalMeshAsync(true, true);
+
+	CustomInstance->PostEditChange();
+	FCoreUObjectDelegates::BroadcastOnObjectModified(CustomInstance.Get());
+
+	return FReply::Handled();
+}
+
+
+FReply SCustomizableInstanceProperties::OnProjectorSelectStateChanged(FString ParamNameWithIndex)
+{
+	const int32 MaxIndex = ArrayProjectorParameterNameWithIndex.Num();
+
+	const TSharedPtr<SButton> Button = MapProjectorButton[ParamNameWithIndex];
+	const TSharedPtr<STextBlock> TextBlock = MapProjectorButtonTextBlock[ParamNameWithIndex];
+
+	CustomInstance->ProjectorLayerChange = false;
+
+	for (TPair<FString, TSharedPtr<SButton>>& Element : MapProjectorButton)
+	{
+		if (/*Element.Value && */Element.Value != Button)
+		{
+			Element.Value->SetBorderBackgroundColor(FLinearColor(1.0f, 1.0f, 1.0f, 1.0f));
+		}
+	}
+
+	for (TPair<FString, TSharedPtr<STextBlock>>& Element : MapProjectorButtonTextBlock)
+	{
+		if (/*Element.Value &&*/ Element.Value != TextBlock)
+		{
+			Element.Value->SetText(LOCTEXT("Select Projector", "Select Projector"));
+		}
+	}
+
+	bool SelectProjector = true;
+
+	if ((TextBlock.Get() != nullptr) && (TextBlock.Get()->GetText().ToString() == FString("Unselect Projector")))
+	{
+		SelectProjector = false;
+	}
+
+	if (Button.Get() != nullptr)
+	{
+		Button->SetBorderBackgroundColor(SelectProjector ? FLinearColor::Green : FLinearColor::White);
+	}
+
+	if (TextBlock.Get() != nullptr)
+	{
+		TextBlock->SetText(SelectProjector ? LOCTEXT("Unselect Projector", "Unselect Projector") : LOCTEXT("Select Projector", "Select Projector"));
+	}
+
+	for (int32 i = 0; i < MaxIndex; ++i)
+	{
+		int32 ArrayRangeIndex;
+		FString ArrayParamName = ParseNameAndIndex(ArrayProjectorParameterNameWithIndex[i], ArrayRangeIndex);
+
+		if ((ArrayProjectorParameterNameWithIndex[i] != ParamNameWithIndex) &&
+			(CustomInstance->GetProjectorState(ArrayParamName, ArrayRangeIndex) != EProjectorState::Hidden))
+		{
+			CustomInstance->SetProjectorState(ArrayParamName, ArrayRangeIndex, EProjectorState::Hidden);
+		}
+
+		if (ArrayProjectorParameterNameWithIndex[i] == ParamNameWithIndex)
+		{
+			CustomInstance->SetProjectorState(ArrayParamName, ArrayRangeIndex, SelectProjector ? EProjectorState::Selected : EProjectorState::Hidden);
+
+			if (SelectProjector)
+			{
+				CustomInstance->LastSelectedProjectorParameterWithIndex = ParamNameWithIndex;
+			}
+		}
+	}
+
+	FCoreUObjectDelegates::BroadcastOnObjectModified(CustomInstance.Get());
+
+	return FReply::Handled();
+}
+
+
+void SCustomizableInstanceProperties::OnProjectorTranslateStateChanged(ECheckBoxState InCheckboxState, FString ParamNameWithIndex) const
+{
+	int32 RangeIndex;
+	const FString ParamName = ParseNameAndIndex(ParamNameWithIndex, RangeIndex);
+
+	CustomInstance->SetProjectorState(ParamName, RangeIndex, EProjectorState::Translate);
+}
+
+
+void SCustomizableInstanceProperties::OnProjectorRotateStateChanged(ECheckBoxState InCheckboxState, FString ParamNameWithIndex) const
+{
+	int32 RangeIndex;
+	const FString ParamName = ParseNameAndIndex(ParamNameWithIndex, RangeIndex);
+
+	CustomInstance->SetProjectorState(ParamName, RangeIndex, EProjectorState::Rotate);
+}
+
+
+void SCustomizableInstanceProperties::OnProjectorScaleStateChanged(ECheckBoxState InCheckboxState, FString ParamNameWithIndex) const
+{
+	int32 RangeIndex;
+	const FString ParamName = ParseNameAndIndex(ParamNameWithIndex, RangeIndex);
+
+	CustomInstance->SetProjectorState(ParamName, RangeIndex, EProjectorState::Scale);
+}
+
+
+void SCustomizableInstanceProperties::OnProjectorTextureParameterComboBoxChanged(TSharedPtr<FString> Selection, ESelectInfo::Type SelectInfo, FString ParamNameWithIndex) const
+{
+	TArray<FCustomizableObjectIntParameterValue>& IntParameters = CustomInstance->GetIntParameters();
+
+	int32 RangeIndex;
+	const FString ParamName = ParseNameAndIndex(ParamNameWithIndex, RangeIndex);
+
+	const UCustomizableObject* CustomObject = CustomInstance->GetCustomizableObject();
+	for (int32 i = 0; i < IntParameters.Num(); ++i)
+	{
+		if (IntParameters[i].ParameterName == ParamName)
+		{
+			const int32 ParamIndexInObject = CustomObject->FindParameter(ParamName);
+			if (ParamIndexInObject != INDEX_NONE)
+			{
+				for (int32 v = 0; v < CustomObject->GetIntParameterNumOptions(ParamIndexInObject); ++v)
+				{
+					FString ValueName = CustomObject->GetIntParameterAvailableOption(ParamIndexInObject, v);
+					if (ValueName == *Selection)
+					{
+						if (RangeIndex >= 0)
+						{
+							check(IntParameters[i].ParameterRangeValueNames.IsValidIndex(RangeIndex));
+							IntParameters[i].ParameterRangeValueNames[RangeIndex] = ValueName;
+						}
+						else
+						{
+							IntParameters[i].ParameterValueName = ValueName;
+						}
+
+						break;
+					}
+				}
+			}
+			else
+			{
+				UE_LOG(LogMutable, Error, TEXT("Failed to find parameter."));
+			}
+		}
+	}
+
+	CustomInstance->SetSelectedParameterProfileDirty();
+
+	CustomInstance->ProjectorLayerChange = true;
+
+	CustomInstance->UpdateSkeletalMeshAsync(true, true);
+
+	CustomInstance->PostEditChange();
+
+	// Non-continuous change: collect garbage.
+	GEngine->ForceGarbageCollection();
+}
+
+
+void SCustomizableInstanceProperties::OnProjectorFloatParameterChanged(float Value, int SliderIndex, FString ParamNameWithIndex)
+{
+	OnFloatParameterChanged(Value, SliderIndex);
+
+	// Code commented since it was causing problems with the alpha slider. Selecting the projector removed the focus on the slider.
+	//int32 RangeIndex;
+	//FString ParamName = ParseNameAndIndex(ParamNameWithIndex, RangeIndex);
+
+	//EProjectorState::Type ProjectorState = CustomInstance->GetProjectorState(ParamName, RangeIndex);
+	//if (ProjectorState != EProjectorState::Selected)
+	//{
+	//	CustomInstance->ResetProjectorStates();
+	//	CustomInstance->SetProjectorState(ParamName, RangeIndex, EProjectorState::Selected);
+
+	//	CustomInstance->LastSelectedProjectorParameterWithIndex = ParamNameWithIndex;
+	//	FCoreUObjectDelegates::BroadcastOnObjectModified(CustomInstance.Get());
+	//}
+}
+
+
+void SCustomizableInstanceProperties::OnProjectorFloatParameterSliderBegin()
+{
+	OnFloatParameterSliderBegin();
+
+	CustomInstance->ProjectorAlphaChange = true;
+}
+
+
+void SCustomizableInstanceProperties::OnProjectorFloatParameterSliderEnd(float Value)
+{
+	OnFloatParameterSliderEnd(Value);
+
+	CustomInstance->ProjectorAlphaChange = false;
+}
+
+
+FReply SCustomizableInstanceProperties::OnProjectorCopyTransform(FString ParamNameWithIndex) const
+{
+	int32 RangeIndex;
+	const FString ParamName = ParseNameAndIndex(ParamNameWithIndex, RangeIndex);
+
+	const int32 ParameterIndexInObject = CustomInstance->GetCustomizableObject()->FindParameter(ParamName);
+	const int32 ProjectorParamIndex = CustomInstance->FindProjectorParameterNameIndex(ParamName);
+
+	if ((ParameterIndexInObject >= 0) && (ProjectorParamIndex >= 0))
+	{
+		TArray<FCustomizableObjectProjectorParameterValue>& ProjectorParameters = CustomInstance->GetProjectorParameters();
+		FCustomizableObjectProjector Value;
+
+		if (RangeIndex == -1)
+		{
+			Value = ProjectorParameters[ProjectorParamIndex].Value;
+		}
+		else if (ProjectorParameters[ProjectorParamIndex].RangeValues.IsValidIndex(RangeIndex))
+		{
+			Value = ProjectorParameters[ProjectorParamIndex].RangeValues[RangeIndex];
+		}
+		else
+		{
+			check(false);
+		}
+
+		const UScriptStruct* Struct = Value.StaticStruct();
+		FString Output = TEXT("");
+		Struct->ExportText(Output, &Value, nullptr, nullptr, (PPF_ExportsNotFullyQualified | PPF_Copy | PPF_Delimited | PPF_IncludeTransient), nullptr);
+
+		FPlatformApplicationMisc::ClipboardCopy(*Output);
+	}
+
+	return FReply::Handled();
+}
+
+
+FReply SCustomizableInstanceProperties::OnProjectorPasteTransform(FString ParamNameWithIndex) const
+{
+	int32 RangeIndex;
+	FString ParamName = ParseNameAndIndex(ParamNameWithIndex, RangeIndex);
+
+	FString ClipboardText;
+	FPlatformApplicationMisc::ClipboardPaste(ClipboardText);
+	FCustomizableObjectProjector DefaultValue;
+	UScriptStruct* Struct = DefaultValue.StaticStruct();
+	Struct->ImportText(*ClipboardText, &DefaultValue, nullptr, 0, GLog, GetPathNameSafe(Struct));
+
+	CustomInstance->SetProjectorValue(ParamName,
+		static_cast<FVector>(DefaultValue.Position),
+		static_cast<FVector>(DefaultValue.Direction),
+		static_cast<FVector>(DefaultValue.Up),
+		static_cast<FVector>(DefaultValue.Scale),
+		DefaultValue.Angle,
+		RangeIndex);
+
+	//Instance->UpdateSkeletalMeshAsync(true);
+	CustomInstance->TempUpdateGizmoInViewport = true;
+	CustomInstance->TempProjectorParameterName = ParamName;
+	CustomInstance->TempProjectorParameterRangeIndex = RangeIndex;
+
+	CustomInstance->ProjectorLayerChange = false;
+
+	return FReply::Handled();
+}
+
+
+FReply SCustomizableInstanceProperties::OnProjectorResetTransform(FString ParamNameWithIndex) const
+{
+	int32 RangeIndex;
+
+	FString ParamName = ParseNameAndIndex(ParamNameWithIndex, RangeIndex);
+
+	const FCustomizableObjectProjector DefaultValue = CustomInstance->GetProjectorDefaultValue(CustomInstance->GetCustomizableObject()->FindParameter(ParamName));
+
+	CustomInstance->SetProjectorValue(ParamName,
+		static_cast<FVector>(DefaultValue.Position),
+		static_cast<FVector>(DefaultValue.Direction),
+		static_cast<FVector>(DefaultValue.Up),
+		static_cast<FVector>(DefaultValue.Scale),
+		DefaultValue.Angle,
+		RangeIndex);
+
+	CustomInstance->TempUpdateGizmoInViewport = true;
+	CustomInstance->TempProjectorParameterName = ParamName;
+	CustomInstance->TempProjectorParameterRangeIndex = RangeIndex;
+
+	CustomInstance->ProjectorLayerChange = false;
+
+	return FReply::Handled();
+}
+
+
+void SCustomizableInstanceProperties::InstanceUpdated(UCustomizableObjectInstance* Instance) const
+{
+	if (CustomInstance->bShowOnlyRelevantParameters &&
+		!CustomInstance->ProjectorUpdatedInViewport &&
+		bShouldResetParamBox)
+	{
+		if (const TSharedPtr<FCustomizableInstanceDetails> Details = InstanceDetails.Pin())
+		{
+			Details->Refresh();
+		}
+	}
+}
+
+
+class SCreateProfileParameters : public SWindow
+{
+public:
+	SLATE_BEGIN_ARGS(SCreateProfileParameters) {}
+	SLATE_ARGUMENT(FText, DefaultAssetPath)
+	SLATE_ARGUMENT(FText, DefaultFileName)
+	SLATE_END_ARGS()
+
+	SCreateProfileParameters()
+		: UserResponse(EAppReturnType::Cancel)
+	{
+	}
+
+	void Construct(const FArguments& InArgs);
+
+public:
+	/** Displays the dialog in a blocking fashion */
+	EAppReturnType::Type ShowModal();
+
+	/** FileName getter */
+	FString GetFileName() const;
+
+protected:
+	FReply OnButtonClick(EAppReturnType::Type ButtonID);
+	void OnNameChange(const FText& NewName, ETextCommit::Type CommitInfo);
+	EAppReturnType::Type UserResponse;
+	FText AssetPath;
+	FText FileName;
+
+public:
+	TWeakObjectPtr<UCustomizableObjectInstance> CustomInstance;
+
+	SCustomizableInstanceProperties* CIProperties = nullptr;
+};
+
+
+FReply SCustomizableInstanceProperties::CreateParameterProfileWindow()
+{
+	const TSharedRef<SCreateProfileParameters> FolderDlg =
+		SNew(SCreateProfileParameters)
+		.DefaultAssetPath(LOCTEXT("DefaultAssethPath", "/Game"))
+		.DefaultFileName(FText::FromString("ProfileParameterData"));
+
+	FolderDlg->CustomInstance = CustomInstance;
+	FolderDlg->CIProperties = this;
+
+	FolderDlg->Construct(SCreateProfileParameters::FArguments());
+
+	FolderDlg->ShowModal();
+	
+	return FReply::Handled();
+}
+
+
+FReply SCustomizableInstanceProperties::RemoveParameterProfile()
+{
+	UCustomizableObject* CustomizableObject = CustomInstance->GetCustomizableObject(); 
+
+	const int32 ProfileIdx = CustomInstance->SelectedProfileIndex;
+	if (ProfileIdx == INDEX_NONE)
+	{
+		return FReply::Handled();
+	}
+
+	TArray<FProfileParameterDat>& Profiles = CustomizableObject->InstancePropertiesProfiles;
+	
+	Profiles.RemoveAt(ProfileIdx);
+	CustomInstance->SelectedProfileIndex = INDEX_NONE;
+	CustomizableObject->Modify();
+	if (const TSharedPtr<FCustomizableInstanceDetails> Details = InstanceDetails.Pin())
+	{
+	    Details->Refresh();
+	}
+
+	CustomInstance->PreEditChange(nullptr);
+
+	CustomInstance->UpdateSkeletalMeshAsync(true, true);
+
+	CustomInstance->PostEditChange();
+
+	// Non-continuous change: collect garbage.
+	GEngine->ForceGarbageCollection();
+
+	return FReply::Handled();
+}
+
+
+void SCustomizableInstanceProperties::OnProfileSelectedChanged(TSharedPtr<FString> Selection, ESelectInfo::Type SelectInfo)
+{
+	const int32 ProfileIdx = CustomInstance->SelectedProfileIndex;
+	if (CustomInstance->IsSelectedParameterProfileDirty())
+	{
+		CustomInstance->SaveParametersToProfile(ProfileIdx);
+	}
+
+	if (*Selection == "None")
+	{
+		CustomInstance->SelectedProfileIndex = INDEX_NONE;
+		return;
+	}	
+
+	//Set selected profile
+	TArray<FProfileParameterDat>& Profiles = CustomInstance->GetCustomizableObject()->InstancePropertiesProfiles;
+	for (int32 Idx = 0; Idx < Profiles.Num(); ++Idx)
+	{
+		if (Profiles[Idx].ProfileName == *Selection)
+		{
+			CustomInstance->SelectedProfileIndex = Idx;
+			break;
+		}
+	}
+	
+	CustomInstance->LoadParametersFromProfile(CustomInstance->SelectedProfileIndex);
+
+	CustomInstance->PreEditChange(nullptr);
+
+	CustomInstance->UpdateSkeletalMeshAsync(true, true);
+
+	CustomInstance->PostEditChange();
+
+	// Non-continuous change: collect garbage.
+	GEngine->ForceGarbageCollection();
+
+	if (const TSharedPtr<FCustomizableInstanceDetails> Details = InstanceDetails.Pin())
+	{
+		Details->Refresh();
+	}
+}
+
+
+void SCustomizableInstanceProperties::SetParameterProfileNamesOnEditor()
+{
+	ParameterProfileNames.Empty();
+
+	for (FProfileParameterDat& Profile : CustomInstance->GetCustomizableObject()->InstancePropertiesProfiles)
+	{
+		ParameterProfileNames.Emplace(MakeShared<FString>(Profile.ProfileName));
+	}
+	ParameterProfileNames.Sort(&CompareNames);
+
+	ParameterProfileNames.EmplaceAt(0, MakeShared<FString>("None"));
+}
+
+
+bool SCustomizableInstanceProperties::HasAnyParameters() const
+{
+	return CustomInstance->HasAnyParameters();
+}
+
+
+void SCreateProfileParameters::Construct(const FArguments& InArgs)
+{
+	if (AssetPath.IsEmpty())
+	{
+		AssetPath = FText::FromString(TEXT("/Game/"));
+	}
+	
+	FModuleManager::LoadModuleChecked<FContentBrowserModule>("ContentBrowser");
+
+	SWindow::Construct(SWindow::FArguments()
+		.Title(LOCTEXT("SSelectFolderDlg_Title", "Add a name to the new profile"))
+		.SupportsMinimize(false)
+		.SupportsMaximize(false)
+		//.SizingRule( ESizingRule::Autosized )
+		.ClientSize(FVector2D(450, 85))
+		[
+			SNew(SVerticalBox)
+
+			+ SVerticalBox::Slot() // Add user input block
+		.Padding(2)
+		[
+			SNew(SBorder)
+			.BorderImage(FAppStyle::GetBrush("ToolPanel.GroupBorder"))
+		[
+			SNew(SVerticalBox)
+	+ SVerticalBox::Slot()
+		.AutoHeight()
+		[
+			SNew(STextBlock)
+			.Text(LOCTEXT("FileName", "Customizable Profile Name"))
+		.Font(FSlateFontInfo(FPaths::EngineContentDir() / TEXT("Slate/Fonts/Roboto-Regular.ttf"), 14))
+		]
+
+	+ SVerticalBox::Slot()
+		.AutoHeight()
+		[
+			SNew(SEditableTextBox)
+			.Text(InArgs._DefaultFileName)
+		.OnTextCommitted(this, &SCreateProfileParameters::OnNameChange)
+		.MinDesiredWidth(250)
+		]
+
+		]
+		]
+
+	+ SVerticalBox::Slot()
+		.AutoHeight()
+		.HAlign(HAlign_Right)
+		.Padding(5)
+		[
+			SNew(SUniformGridPanel)
+			.SlotPadding(FAppStyle::GetMargin("StandardDialog.SlotPadding"))
+			.MinDesiredSlotWidth(FAppStyle::GetFloat("StandardDialog.MinDesiredSlotWidth"))
+			.MinDesiredSlotHeight(FAppStyle::GetFloat("StandardDialog.MinDesiredSlotHeight"))
+		+ SUniformGridPanel::Slot(0, 0)
+		[
+			SNew(SButton)
+			.HAlign(HAlign_Center)
+			.ContentPadding(FAppStyle::GetMargin("StandardDialog.ContentPadding"))
+			.Text(LOCTEXT("OK", "OK"))
+			.OnClicked(this, &SCreateProfileParameters::OnButtonClick, EAppReturnType::Ok)
+		]
+	+ SUniformGridPanel::Slot(1, 0)
+		[
+			SNew(SButton)
+			.HAlign(HAlign_Center)
+			.ContentPadding(FAppStyle::GetMargin("StandardDialog.ContentPadding"))
+			.Text(LOCTEXT("Cancel", "Cancel"))
+			.OnClicked(this, &SCreateProfileParameters::OnButtonClick, EAppReturnType::Cancel)
+		]
+		]
+		]);
+}
+
+
+EAppReturnType::Type SCreateProfileParameters::ShowModal()
+{
+	GEditor->EditorAddModalWindow(SharedThis(this));
+	
+	return UserResponse;
+}
+
+
+FReply SCreateProfileParameters::OnButtonClick(EAppReturnType::Type ButtonID)
+{
+	if (ButtonID == EAppReturnType::Ok)
+	{
+		UserResponse = ButtonID;
+
+		RequestDestroyWindow();
+
+		UCustomizableObject* CustomizableObject = CustomInstance->GetCustomizableObject(); 
+		CustomizableObject->AddNewParameterProfile(GetFileName(), *CustomInstance.Get());
+
+		if (CustomInstance->bSelectedProfileDirty && CustomInstance->SelectedProfileIndex != INDEX_NONE)
+		{
+			CustomInstance->SaveParametersToProfile(CustomInstance->SelectedProfileIndex);
+		}
+		CustomInstance->SelectedProfileIndex = CustomizableObject->InstancePropertiesProfiles.Num() - 1;
+
+		if (CIProperties)
+		{
+			if (const TSharedPtr<FCustomizableInstanceDetails> Details = CIProperties->InstanceDetails.Pin())
+			{
+				Details->Refresh();
+			}
+		}
+	}
+	else if (ButtonID == EAppReturnType::Cancel)
+	{
+		UserResponse = ButtonID;
+
+		RequestDestroyWindow();
+	}
+	return FReply::Handled();
+}
+
+
+void SCreateProfileParameters::OnNameChange(const FText& NewName, ETextCommit::Type CommitInfo)
+{
+	FileName = NewName;
+}
+
+
+FString SCreateProfileParameters::GetFileName() const
+{
+	return FileName.ToString();
+}
+
+
+#undef LOCTEXT_NAMESPACE
+
