[mutable] Move Mutable plugin from the restricted folder to the experimental plugins folder.

diff --git a/Engine/Plugins/Experimental/Mutable/Source/CustomizableObjectEditor/Private/SCustomizableObjectEditorPerformanceReport.cpp b/Engine/Plugins/Experimental/Mutable/Source/CustomizableObjectEditor/Private/SCustomizableObjectEditorPerformanceReport.cpp
new file mode 100644
--- /dev/null
+++ b/Engine/Plugins/Experimental/Mutable/Source/CustomizableObjectEditor/Private/SCustomizableObjectEditorPerformanceReport.cpp
@@ -0,0 +1,1454 @@
+// Copyright Epic Games, Inc. All Rights Reserved.
+
+
+#include "SCustomizableObjectEditorPerformanceReport.h"
+#include "CustomizableObjectInstance.h"
+#include "StatePerformanceTesting.h"
+#include "PropertyCustomizationHelpers.h"
+#include "IDetailsView.h"
+#include "Widgets/Layout/SScrollBox.h"
+#include "Widgets/Layout/SSplitter.h"
+#include "CustomizableObjectEditorStyle.h"
+#include "Modules/ModuleManager.h"
+#include "PropertyCustomizationHelpers.h"
+#include "AssetThumbnail.h"
+
+
+#define LOCTEXT_NAMESPACE "CustomizableObjectEditor"
+
+static const FName ColumnID_PR_LOD_Label("LOD");
+static const FName ColumnID_PR_InstanceName_Label("InstanceName");
+static const FName ColumnID_PR_TexturesTotalSize_Label("TexturesTotalSize");
+static const FName ColumnID_PR_NumberOfTriangles_Label("NumberOfTriangles");
+static const FName ColumnID_PR_Name_Label("ParameterName");
+static const FName ColumnID_PR_UpdateTime_Label("UpdateTimeState");
+static const FName ColumnID_PR_Delta_Label("WorstCaseDelta");
+static const FName ColumnID_PR_Average_Label("Average");
+
+void UWorstCasePerformanceReportInstance::SetInstance(UCustomizableObjectInstance* UpdatedInstance)
+{
+	if (UpdatedInstance)
+	{
+		if (!WorstCaseInstance)
+		{
+			/* Uncommment next line if we need the report to have the valid time entries for state parameters expanded by default */
+			//SubVisibility = EVisibility::Visible;
+		}
+		WorstCaseInstance = UpdatedInstance->Clone();
+	}
+}
+
+void SCustomizableObjecEditorPerformanceReport::Construct(const FArguments& InArgs)
+{
+	ReportSubject = InArgs._CustomizableObject;
+	StopTests();
+	ResetSavedData();
+	ResetPerformanceReportOptions();
+
+	SAssignNew(PerformanceReportTextureAnalyzer, SCustomizableObjecEditorTextureAnalyzer)
+		.CustomizableObjectEditor(nullptr)
+		.CustomizableObjectInstanceEditor(nullptr);
+
+	FPropertyEditorModule& PropPlugin = FModuleManager::LoadModuleChecked<FPropertyEditorModule>("PropertyEditor");
+	FDetailsViewArgs DetailsViewArgs;
+	DetailsViewArgs.NotifyHook = this;
+	DetailsViewArgs.NameAreaSettings = FDetailsViewArgs::ENameAreaSettings::HideNameArea;
+	DetailsViewArgs.bAllowSearch = false;
+	DetailsViewArgs.bShowObjectLabel = false;
+	DetailsViewArgs.bAllowMultipleTopLevelObjects = true;
+
+	float WidgetSize = 128.0f;
+	int ThumbnailSize = 128;
+
+	ChildSlot
+	[
+		SNew(SSplitter)
+
+		// Performance report status and dynamic results (left section)
+		+ SSplitter::Slot()
+		[
+			SNew(SBox)
+			.Padding(5.0f)
+			.MinDesiredWidth(300)
+			[
+				SNew(SVerticalBox)
+
+				// Button with settings dropdown
+				+ SVerticalBox::Slot()
+				.HAlign(HAlign_Center)
+				.AutoHeight()
+				.Padding(5.0f)
+				[
+					SNew(SComboButton)
+					.OnGetMenuContent(this, &SCustomizableObjecEditorPerformanceReport::GenerateOptionsMenuContent)
+					.IsEnabled(ReportSubject != nullptr)
+					.ButtonContent()
+					[
+						SNew(SButton)
+						.Text(LOCTEXT("GeneratePerformanceReport", "Generate Performance Report"))
+						.OnClicked(this, &SCustomizableObjecEditorPerformanceReport::RunPerformanceReport)
+						.ButtonStyle(FAppStyle::Get(), "FlatButton.Success")
+						.ForegroundColor(FLinearColor::White)
+					]
+				]
+
+				+ SVerticalBox::Slot()
+				.HAlign(HAlign_Center)
+				.AutoHeight()
+				.Padding(5.0f)
+				[
+					SNew(SButton)
+					.Text(LOCTEXT("GetWorstTimeInstance", "Select the currently slowest instance"))
+					.OnClicked(this, &SCustomizableObjecEditorPerformanceReport::SelectSlowestInstance)
+					.ButtonStyle(FAppStyle::Get(), "FlatButton.Info")
+					.IsEnabled(this, &SCustomizableObjecEditorPerformanceReport::IsGetWorstInstanceButtonEnabled)
+					// This button target changes fast, I think it's better if it triggers on down instead of on up.
+					.ClickMethod(EButtonClickMethod::MouseDown)
+					.TouchMethod(EButtonTouchMethod::Down)
+					.PressMethod(EButtonPressMethod::ButtonPress)
+				]
+
+				+ SVerticalBox::Slot()
+				.HAlign(HAlign_Center)
+				.AutoHeight()
+				.Padding(5.0f)
+				[
+					SAssignNew(StatusText, STextBlock)
+				]
+
+				// Info on Texture Size worst cases
+				+ SVerticalBox::Slot()
+				.HAlign(HAlign_Left)
+				.AutoHeight()
+				.Padding(5.0f, 10.0f, 5.0f, 5.0f)
+				[
+					SNew(STextBlock)
+					.Text(LOCTEXT("MutableLargestTextures", "Instances with largest mutable textures:"))
+				]
+				+ SVerticalBox::Slot()
+				.AutoHeight()
+				[
+					SNew(SBorder)
+					.Padding(5.0f)
+					[
+						SAssignNew(TextureWorstCasesView, SVerticalBox)
+						.Visibility_Lambda([this]() { return WorstTextureInstances.Num() > 0 && WorstTextureInstances[0] && WorstTextureInstances[0]->WorstCaseInstance ? EVisibility::Visible : EVisibility::Collapsed; })
+					]
+				]
+
+				// Info on Geometry Amount worst cases
+				+ SVerticalBox::Slot()
+				.HAlign(HAlign_Left)
+				.AutoHeight()
+				.Padding(5.0f, 20.0f, 5.0f, 5.0f)
+				[
+					SNew(STextBlock)
+					.Text(LOCTEXT("MutableMostGeometry", "Instances with most triangles:"))
+				]
+				+ SVerticalBox::Slot()
+				.AutoHeight()
+				[
+					SNew(SBorder)
+					.Padding(5.0f)
+					.Visibility_Lambda([this]() { return WorstGeometryInstances.Num() > 0 && WorstGeometryInstances[0] && WorstGeometryInstances[0]->WorstCaseInstance ? EVisibility::Visible : EVisibility::Collapsed; })
+					[
+						SAssignNew(GeometryWorstCasesView, SVerticalBox)
+					]
+				]
+
+				// Info on Update Time worst cases detailed by state and parameter
+				+ SVerticalBox::Slot()
+				.HAlign(HAlign_Left)
+				.AutoHeight()
+				.Padding(5.0f, 20.0f, 5.0f, 5.0f)
+				[
+					SNew(STextBlock)
+					.Text(LOCTEXT("MutableLongestUpdateTimesDetailed", "Instances with longest updates by state and parameter:"))
+				]
+				+ SVerticalBox::Slot()
+				.FillHeight(1.0)
+				[
+					SNew(SBorder)
+					.Padding(5.0f)
+					.Visibility_Lambda([this]() { return WorstTimeStateInstances.Num() > 0
+						&& WorstTimeStateInstances[0].Num() > 0
+						&& WorstTimeStateInstances[0][0]
+						&& WorstTimeStateInstances[0][0]->WorstCaseInstance
+						? EVisibility::Visible : EVisibility::Collapsed; })
+					[
+						SAssignNew(TimeWorstCasesDetailedView, SScrollBox)
+					]
+				]
+			]
+		]
+
+		// Currently selected worst instance and its texture analyzer (right section)
+		+ SSplitter::Slot()
+		[
+			BuildAnalyzer()
+		]
+	];
+
+	BuildTextureWorstCasesView();
+	BuildGeometryWorstCasesView();
+	BuildTimeWorstCasesDetailedView();
+	StatusText->SetText(LOCTEXT("PerformanceReportStatusNone","Generate Performance Report to collect test data."));
+
+	HelperCallback = NewObject<UPerformanceReportHelper>(GetTransientPackage());
+	HelperCallback->PerformanceReport = this;
+}
+
+SCustomizableObjecEditorPerformanceReport::~SCustomizableObjecEditorPerformanceReport()
+{
+	if (CurrentReportWorstCase && CurrentReportWorstCase->WorstCaseInstance)
+	{
+		CurrentReportWorstCase->WorstCaseInstance->UpdatedDelegate.RemoveDynamic(HelperCallback, &UPerformanceReportHelper::DelegatedCallback);
+	}
+	HelperCallback = nullptr;
+	WorstTimeInstance = nullptr;
+	CurrentReportWorstCase = nullptr;
+	CurrentReportInstance = nullptr;
+}
+
+TSharedRef<SWidget> SCustomizableObjecEditorPerformanceReport::GenerateOptionsMenuContent()
+{
+	const bool bShouldCloseWindowAfterMenuSelection = false;
+	FMenuBuilder MenuBuilder(false, nullptr);
+
+	// settings
+	MenuBuilder.BeginSection("ResetOptions", LOCTEXT("MutableResetPerformanceReportOptions", "Reset Options"));
+	{
+		SAssignNew(PerformanceReportResetOptionsEntry, SButton)
+			.Text(LOCTEXT("MutableGeneratePerformanceReport", "Reset Performance Report Options"))
+			.OnClicked(this, &SCustomizableObjecEditorPerformanceReport::ResetPerformanceReportOptions);
+		MenuBuilder.AddWidget(PerformanceReportResetOptionsEntry.ToSharedRef(), FText());
+	}
+	MenuBuilder.EndSection();
+
+	MenuBuilder.BeginSection("StressTestParameters", LOCTEXT("MutablePerformanceReportStressTestParameters", "Stress Test Parameters"));
+	{
+		SAssignNew(StressTestOptionInstanceCountEntry, SNumericEntryBox<int32>)
+			.AllowSpin(true)
+			.MinValue(0)
+			.MaxValue(1000000)
+			.MinSliderValue(1)
+			.MaxSliderValue(1000)
+			.Value(this, &SCustomizableObjecEditorPerformanceReport::GetStressTestOptionInstanceCount)
+			.OnValueChanged(this, &SCustomizableObjecEditorPerformanceReport::OnStressTestOptionInstanceCountChanged);
+		MenuBuilder.AddWidget(StressTestOptionInstanceCountEntry.ToSharedRef(), LOCTEXT("MutableStressTestInstanceCount", "Instance Count"));
+
+		SAssignNew(StressTestOptionCreateInstanceTimeEntry, SNumericEntryBox<int32>)
+			.AllowSpin(true)
+			.MinValue(1)
+			.MaxValue(1000000)
+			.MinSliderValue(1)
+			.MaxSliderValue(1000)
+			.Value(this, &SCustomizableObjecEditorPerformanceReport::GetStressTestOptionCreateInstanceTime)
+			.OnValueChanged(this, &SCustomizableObjecEditorPerformanceReport::OnStressTestOptionCreateInstanceTimeChanged);
+		MenuBuilder.AddWidget(StressTestOptionCreateInstanceTimeEntry.ToSharedRef(), LOCTEXT("MutableStressTestCreateInstanceTime", "Instance Create Time (ms)"));
+
+		SAssignNew(StressTestOptionCreateInstanceTimeVarEntry, SNumericEntryBox<int32>)
+			.AllowSpin(true)
+			.MinValue(1)
+			.MaxValue(1000000)
+			.MinSliderValue(1)
+			.MaxSliderValue(10000)
+			.Value(this, &SCustomizableObjecEditorPerformanceReport::GetStressTestOptionCreateInstanceTimeVar)
+			.OnValueChanged(this, &SCustomizableObjecEditorPerformanceReport::OnStressTestOptionCreateInstanceTimeVarChanged);
+		MenuBuilder.AddWidget(StressTestOptionCreateInstanceTimeVarEntry.ToSharedRef(), LOCTEXT("MutableStressTestCreateInstanceTimeVar", "Instance Create Time variation (ms)"));
+
+		SAssignNew(StressTestOptionInstanceUpdateTimeEntry, SNumericEntryBox<int32>)
+			.AllowSpin(true)
+			.MinValue(0)
+			.MaxValue(1000000)
+			.MinSliderValue(1)
+			.MaxSliderValue(10000)
+			.Value(this, &SCustomizableObjecEditorPerformanceReport::GetStressTestOptionInstanceUpdateTime)
+			.OnValueChanged(this, &SCustomizableObjecEditorPerformanceReport::OnStressTestOptionInstanceUpdateTimeChanged);
+		MenuBuilder.AddWidget(StressTestOptionInstanceUpdateTimeEntry.ToSharedRef(), LOCTEXT("MutableStressTestInstanceUpdateTime", "Instance Update Time (ms)"));
+
+		SAssignNew(StressTestOptionInstanceUpdateTimeVarEntry, SNumericEntryBox<int32>)
+			.AllowSpin(true)
+			.MinValue(1)
+			.MaxValue(1000000)
+			.MinSliderValue(1)
+			.MaxSliderValue(10000)
+			.Value(this, &SCustomizableObjecEditorPerformanceReport::GetStressTestOptionInstanceUpdateTimeVar)
+			.OnValueChanged(this, &SCustomizableObjecEditorPerformanceReport::OnStressTestOptionInstanceUpdateTimeVarChanged);
+		MenuBuilder.AddWidget(StressTestOptionInstanceUpdateTimeVarEntry.ToSharedRef(), LOCTEXT("MutableStressTestInstanceUpdateTimeVar", "Instance Update Time variation (ms)"));
+
+		SAssignNew(StressTestOptionInstanceLifeTimeEntry, SNumericEntryBox<int32>)
+			.AllowSpin(true)
+			.MinValue(0)
+			.MaxValue(1000000)
+			.MinSliderValue(1)
+			.MaxSliderValue(10000)
+			.Value(this, &SCustomizableObjecEditorPerformanceReport::GetStressTestOptionInstanceLifeTime)
+			.OnValueChanged(this, &SCustomizableObjecEditorPerformanceReport::OnStressTestOptionInstanceLifeTimeChanged);
+		MenuBuilder.AddWidget(StressTestOptionInstanceLifeTimeEntry.ToSharedRef(), LOCTEXT("MutableStressTestInstanceLifeTime", "Instance Life Time (ms)"));
+
+		SAssignNew(StressTestOptionInstanceLifeTimeVarEntry, SNumericEntryBox<int32>)
+			.AllowSpin(true)
+			.MinValue(1)
+			.MaxValue(1000000)
+			.MinSliderValue(1)
+			.MaxSliderValue(10000)
+			.Value(this, &SCustomizableObjecEditorPerformanceReport::GetStressTestOptionInstanceLifeTimeVar)
+			.OnValueChanged(this, &SCustomizableObjecEditorPerformanceReport::OnStressTestOptionInstanceLifeTimeVarChanged);
+		MenuBuilder.AddWidget(StressTestOptionInstanceLifeTimeVarEntry.ToSharedRef(), LOCTEXT("MutableStressTestInstanceLifeTimeMs", "Instance Life Time variation (ms)"));
+	}
+	MenuBuilder.EndSection();
+
+	return MenuBuilder.MakeWidget();
+}
+
+void SCustomizableObjecEditorPerformanceReport::AddReferencedObjects(FReferenceCollector& Collector)
+{
+	Collector.AddReferencedObject(ReportSubject);
+	Collector.AddReferencedObject(HelperCallback);
+	Collector.AddReferencedObject(WorstTimeInstance);
+	Collector.AddReferencedObject(CurrentReportWorstCase);
+	Collector.AddReferencedObject(CurrentReportInstance);
+	Collector.AddReferencedObjects(WorstTextureInstances);
+	Collector.AddReferencedObjects(WorstGeometryInstances);
+	for (TArray<UWorstCasePerformanceReportInstance*>& Arr : WorstTimeStateInstances)
+	{
+		Collector.AddReferencedObjects(Arr);
+	}
+	for (TArray<TArray<UWorstCasePerformanceReportInstance*>>& ArrArr : WorstTimeParameterInstances)
+	{
+		for (TArray<UWorstCasePerformanceReportInstance*>& Arr : ArrArr)
+		{
+			Collector.AddReferencedObjects(Arr);
+		}
+	}
+}
+
+void SCustomizableObjecEditorPerformanceReport::Tick(const FGeometry& AllottedGeometry, const double InCurrentTime, const float InDeltaTime)
+{
+	SCompoundWidget::Tick(AllottedGeometry, InCurrentTime, InDeltaTime);
+}
+
+void SCustomizableObjecEditorPerformanceReport::StopTests()
+{
+	if (StressTest.IsValid() && StressTest->GetTestInCourse())
+	{
+		StressTest->WorstFacesDelegate.Unbind();
+		StressTest->WorstTextureDelegate.Unbind();
+		StressTest->StressTestEnded.Unbind();
+		StressTest->FinishTest();
+		StressTest.Reset();
+	}
+
+	if (StateTest.IsValid())
+	{
+		StateTest->WorstTimeDelegate.Unbind();
+		StateTest->StateTestEnded.Unbind();
+		StateTest->FinishTest();
+		StateTest.Reset();
+	}
+}
+
+FReply SCustomizableObjecEditorPerformanceReport::RunPerformanceReport()
+{
+	return RunPerformanceReport(nullptr);
+}
+
+FReply SCustomizableObjecEditorPerformanceReport::RunPerformanceReport(UCustomizableObject* ObjectToTest)
+{
+	if (ObjectToTest)
+	{
+		ReportSubject = ObjectToTest;
+	}
+
+	if (ReportSubject)
+	{
+		StopTests();
+
+		if (CurrentReportInstance || CurrentReportWorstCase)
+		{
+			CurrentReportInstance = nullptr;
+			CurrentReportWorstCase = nullptr;
+			UpdateSelectedInstanceData();
+		}
+
+		ResetSavedData();
+
+		StressTest = MakeShareable(new FRunningStressTest);
+		StressTest->SetCustomizableObject(ReportSubject);
+		StressTest->SetVerifyInstancesThreshold(5.0f);
+		StressTest->SetCreateInstanceTimeMs(StressTestOptionCreateInstanceTime);
+		StressTest->SetCreateInstanceTimeMsVar(StressTestOptionCreateInstanceTimeVar);
+		StressTest->SetInstanceLifeTimeMs(StressTestOptionInstanceLifeTime);
+		StressTest->SetInstanceLifeTimeMsVar(StressTestOptionInstanceLifeTimeVar);
+		StressTest->SetInstanceUpdateTimeMs(StressTestOptionInstanceUpdateTime);
+		StressTest->SetInstanceUpdateTimeMsVar(StressTestOptionInstanceUpdateTimeVar);
+		StressTest->SetPendingInstanceCount(StressTestOptionInstanceCount);
+		StressTest->SetNextInstanceTimeMs(0);
+		StressTest->WorstFacesDelegate.BindRaw(this, &SCustomizableObjecEditorPerformanceReport::WorstFacesFound);
+		StressTest->WorstTextureDelegate.BindRaw(this, &SCustomizableObjecEditorPerformanceReport::WorstTextureFound);
+		StressTest->StressTestEnded.BindRaw(this, &SCustomizableObjecEditorPerformanceReport::StressTestEnded);
+		StressTest->InitMeasureTest();
+
+		StatusText->SetText(LOCTEXT("PerformanceReportStatusStressRunning", "Running Stress test."));
+	}
+	else
+	{
+		StatusText->SetText(LOCTEXT("PerformanceReportStatusNoObject", "Generate Performance Report found no valid Customizable Object."));
+	}
+
+	return FReply::Handled();
+}
+
+FReply SCustomizableObjecEditorPerformanceReport::SelectSlowestInstance()
+{
+	check(WorstTimeInstance);
+	UWorstCasePerformanceReportInstance* SlowestInstanceInStructure = nullptr;
+	if (WorstTimeInstance)
+	{
+		switch (WorstTimeInstance->WorstCaseType)
+		{
+		case UWorstCasePerformanceReportInstance::EPerformanceReportWorstCaseType::SLOWEST_UPDATE_ENTER_STATE:
+		{
+			SlowestInstanceInStructure = GetWorstCaseStructureFromIndexes(WorstTimeInstance->WorstCaseType, WorstTimeInstance->LOD, WorstTimeInstance->LongestUpdateTimeStateIndex);
+			break;
+		}
+		case UWorstCasePerformanceReportInstance::EPerformanceReportWorstCaseType::SLOWEST_UPDATE_PARAM_IN_RELEVANT_STATE:
+		{
+			if (UWorstCasePerformanceReportInstance* SlowestInstanceInStructureParent = GetWorstCaseStructureFromIndexes(
+				UWorstCasePerformanceReportInstance::EPerformanceReportWorstCaseType::SLOWEST_UPDATE_ENTER_STATE,
+				WorstTimeInstance->LOD,
+				WorstTimeInstance->LongestUpdateTimeStateIndex))
+			{
+				SlowestInstanceInStructureParent->SubVisibility = EVisibility::Visible;
+			}
+			SlowestInstanceInStructure = GetWorstCaseStructureFromIndexes(WorstTimeInstance->WorstCaseType, WorstTimeInstance->LOD, WorstTimeInstance->LongestUpdateTimeStateIndex, WorstTimeInstance->LongestUpdateTimeParameterIndexInCO);
+			break;
+		}
+		default:
+		{
+			break;
+		}
+		}
+	}
+	if (SlowestInstanceInStructure)
+	{
+		SelectInstance(SlowestInstanceInStructure);
+	}
+	return FReply::Handled();
+}
+
+bool SCustomizableObjecEditorPerformanceReport::IsGetWorstInstanceButtonEnabled() const
+{
+	return WorstTimeInstance && WorstTimeInstance != CurrentReportWorstCase;
+}
+
+FReply SCustomizableObjecEditorPerformanceReport::ResetPerformanceReportOptions()
+{
+	StressTestOptionInstanceCount = 30;
+	StressTestOptionCreateInstanceTime = 200;
+	StressTestOptionCreateInstanceTimeVar = 200;
+	StressTestOptionInstanceLifeTime = 5000;
+	StressTestOptionInstanceLifeTimeVar = 5000;
+	StressTestOptionInstanceUpdateTime = 500;
+	StressTestOptionInstanceUpdateTimeVar = 500;
+	return FReply::Handled();
+}
+
+TOptional<int32> SCustomizableObjecEditorPerformanceReport::GetStressTestOptionInstanceCount() const
+{
+	return StressTestOptionInstanceCount;
+}
+
+void SCustomizableObjecEditorPerformanceReport::OnStressTestOptionInstanceCountChanged(int Value)
+{
+	StressTestOptionInstanceCount = Value;
+}
+
+TOptional<int32> SCustomizableObjecEditorPerformanceReport::GetStressTestOptionCreateInstanceTime() const
+{
+	return StressTestOptionCreateInstanceTime;
+}
+
+void SCustomizableObjecEditorPerformanceReport::OnStressTestOptionCreateInstanceTimeChanged(int Value)
+{
+	StressTestOptionCreateInstanceTime = Value;
+}
+
+TOptional<int32> SCustomizableObjecEditorPerformanceReport::GetStressTestOptionCreateInstanceTimeVar() const
+{
+	return StressTestOptionCreateInstanceTimeVar;
+}
+
+void SCustomizableObjecEditorPerformanceReport::OnStressTestOptionCreateInstanceTimeVarChanged(int Value)
+{
+	StressTestOptionCreateInstanceTimeVar = Value;
+}
+
+TOptional<int32> SCustomizableObjecEditorPerformanceReport::GetStressTestOptionInstanceUpdateTime() const
+{
+	return StressTestOptionInstanceUpdateTime;
+}
+
+void SCustomizableObjecEditorPerformanceReport::OnStressTestOptionInstanceUpdateTimeChanged(int Value)
+{
+	StressTestOptionInstanceUpdateTime = Value;
+}
+
+TOptional<int32> SCustomizableObjecEditorPerformanceReport::GetStressTestOptionInstanceUpdateTimeVar() const
+{
+	return StressTestOptionInstanceUpdateTimeVar;
+}
+
+void SCustomizableObjecEditorPerformanceReport::OnStressTestOptionInstanceUpdateTimeVarChanged(int Value)
+{
+	StressTestOptionInstanceUpdateTimeVar = Value;
+}
+
+TOptional<int32> SCustomizableObjecEditorPerformanceReport::GetStressTestOptionInstanceLifeTime() const
+{
+	return StressTestOptionInstanceLifeTime;
+}
+
+void SCustomizableObjecEditorPerformanceReport::OnStressTestOptionInstanceLifeTimeChanged(int Value)
+{
+	StressTestOptionInstanceLifeTime = Value;
+}
+
+TOptional<int32> SCustomizableObjecEditorPerformanceReport::GetStressTestOptionInstanceLifeTimeVar() const
+{
+	return StressTestOptionInstanceLifeTimeVar;
+}
+
+void SCustomizableObjecEditorPerformanceReport::OnStressTestOptionInstanceLifeTimeVarChanged(int Value)
+{
+	StressTestOptionInstanceLifeTimeVar = Value;
+}
+
+void SCustomizableObjecEditorPerformanceReport::ResetSavedData()
+{
+	if (ReportSubject)
+	{
+		NumLOD = ReportSubject->GetNumLODs();
+		NumStates = ReportSubject->GetStateCount();
+
+
+		WorstTextureInstances.Empty(NumLOD);
+		WorstTextureInstances.SetNum(NumLOD);
+		for (uint32 i = 0; i < NumLOD; ++i)
+		{
+			WorstTextureInstances[i] = NewObject<UWorstCasePerformanceReportInstance>(GetTransientPackage(), NAME_None, RF_Transient);
+			WorstTextureInstances[i]->LOD = i;
+			WorstTextureInstances[i]->WorstCaseType = UWorstCasePerformanceReportInstance::EPerformanceReportWorstCaseType::LARGEST_TEXTURE_SIZE;
+		}
+
+		WorstGeometryInstances.Empty(NumLOD);
+		WorstGeometryInstances.SetNum(NumLOD);
+		for (uint32 i = 0; i < NumLOD; ++i)
+		{
+			WorstGeometryInstances[i] = NewObject<UWorstCasePerformanceReportInstance>(GetTransientPackage(), NAME_None, RF_Transient);
+			WorstGeometryInstances[i]->LOD = i;
+			WorstGeometryInstances[i]->WorstCaseType = UWorstCasePerformanceReportInstance::EPerformanceReportWorstCaseType::MOST_TRIANGLES;
+		}
+
+		WorstTimeInstance = nullptr;
+
+		// TArray<TArray<UWorstCasePerformanceReportInstance*>> WorstTimeStateInstances;	// By LOD then by state index
+		WorstTimeStateInstances.Empty(NumLOD);
+		WorstTimeStateInstances.SetNum(NumLOD);
+		for (uint32 i = 0; i < NumLOD; ++i)
+		{
+			WorstTimeStateInstances[i].Empty(NumStates);
+			WorstTimeStateInstances[i].SetNum(NumStates);
+			for (uint32 j = 0; j < NumStates; ++j)
+			{
+				WorstTimeStateInstances[i][j] = NewObject<UWorstCasePerformanceReportInstance>(GetTransientPackage(), NAME_None, RF_Transient);
+				WorstTimeStateInstances[i][j]->LOD = i;
+				WorstTimeStateInstances[i][j]->WorstCaseType = UWorstCasePerformanceReportInstance::EPerformanceReportWorstCaseType::SLOWEST_UPDATE_ENTER_STATE;
+			}
+		}
+
+		// TArray<TArray<TArray<UWorstCasePerformanceReportInstance*>>> WorstTimeParameterInstances;	// By LOD then by state index then by parameter index
+		WorstTimeParameterInstances.Empty(NumLOD);
+		WorstTimeParameterInstances.SetNum(NumLOD);
+		for (uint32 i = 0; i < NumLOD; ++i)
+		{
+			WorstTimeParameterInstances[i].Empty(NumStates);
+			WorstTimeParameterInstances[i].SetNum(NumStates);
+			for (uint32 j = 0; j < NumStates; ++j)
+			{
+				WorstTimeParameterInstances[i][j].Reset();
+			}
+		}
+	}
+}
+
+void SCustomizableObjecEditorPerformanceReport::WorstFacesFound(uint32 CurrentLODNumFaces, uint32 LOD, UCustomizableObjectInstance* Instance, MeasuredData* data)
+{
+	// If any of there checks fails, we are receiving updates from a different CO than this was initialized with.
+	check(WorstGeometryInstances.Num() > (int32)LOD);
+	check(WorstGeometryInstances[LOD]);
+	check(WorstGeometryInstances[LOD]->LOD == LOD);
+	check(WorstGeometryInstances[LOD]->WorstCaseType == UWorstCasePerformanceReportInstance::EPerformanceReportWorstCaseType::MOST_TRIANGLES);
+
+	bool bSomethingUpdated = false;
+	UWorstCasePerformanceReportInstance* WorstCase = WorstGeometryInstances[LOD];
+	if ((int32)CurrentLODNumFaces > WorstGeometryInstances[LOD]->NumFaces)
+	{
+		WorstCase->SetInstance(Instance);
+		WorstCase->NumFaces = CurrentLODNumFaces;
+		bSomethingUpdated = true;
+	}
+
+	if (data)
+	{
+		WorstCase->MeasuredDataForThisWorstCase = *data;
+		bSomethingUpdated = true;
+	}
+
+	if (bSomethingUpdated)
+	{
+		BuildGeometryWorstCasesView();
+	}
+}
+
+void SCustomizableObjecEditorPerformanceReport::WorstTextureFound(MaterialInfoMap MapMaterialParam, uint32 CurrentLODTextureSize, uint32 LOD, UCustomizableObjectInstance* Instance, MeasuredData* data)
+{
+	// If any of there checks fails, we are receiving updates from a different CO than this was initialized with.
+	check(WorstTextureInstances.Num() > (int32)LOD);
+	check(WorstTextureInstances[LOD]);
+	check(WorstTextureInstances[LOD]->LOD == LOD);
+	check(WorstTextureInstances[LOD]->WorstCaseType == UWorstCasePerformanceReportInstance::EPerformanceReportWorstCaseType::LARGEST_TEXTURE_SIZE);
+
+	bool bSomethingUpdated = false;
+	UWorstCasePerformanceReportInstance* WorstCase = WorstTextureInstances[LOD];
+	if ((int32)CurrentLODTextureSize > WorstTextureInstances[LOD]->TextureBytes)
+	{
+		WorstCase->SetInstance(Instance);
+		WorstCase->TextureBytes = CurrentLODTextureSize;
+		bSomethingUpdated = true;
+	}
+
+	if (data)
+	{
+		WorstCase->MeasuredDataForThisWorstCase = *data;
+		bSomethingUpdated = true;
+	}
+
+	if (bSomethingUpdated)
+	{
+		BuildTextureWorstCasesView();
+	}
+}
+
+void SCustomizableObjecEditorPerformanceReport::WorstTimeFound(float CurrentInstanceGenerationTime, int32 LongesTimeStateIndex, int32 LongesTimeParameterIndexInCO, uint32 LOD, class UCustomizableObjectInstance* Instance, MeasuredData* data)
+{
+	UWorstCasePerformanceReportInstance* WorstCase = nullptr;
+	bool bIsWorstCaseState = LongesTimeParameterIndexInCO == INDEX_NONE;
+	bool IsNewPArameter = false;
+
+	if (bIsWorstCaseState)
+	{
+		// If any of there checks fails, we are receiving updates from a different CO than this was initialized with.
+		check(WorstTimeStateInstances.Num() > (int32)LOD);
+		check(WorstTimeStateInstances[LOD].Num() > LongesTimeStateIndex);
+		check(WorstTimeStateInstances[LOD][LongesTimeStateIndex]);
+		check(WorstTimeStateInstances[LOD][LongesTimeStateIndex]->LOD == LOD);
+		check(WorstTimeStateInstances[LOD][LongesTimeStateIndex]->WorstCaseType == UWorstCasePerformanceReportInstance::EPerformanceReportWorstCaseType::SLOWEST_UPDATE_ENTER_STATE);
+
+		if (CurrentInstanceGenerationTime > WorstTimeStateInstances[LOD][LongesTimeStateIndex]->UpdateTime)
+		{
+			WorstCase = WorstTimeStateInstances[LOD][LongesTimeStateIndex];
+		}
+	}
+	else
+	{
+		// If any of there checks fails, we are receiving updates from a different CO than this was initialized with.
+		check(WorstTimeParameterInstances.Num() > (int32)LOD);
+		check(WorstTimeParameterInstances[LOD].Num() > LongesTimeStateIndex);
+		UWorstCasePerformanceReportInstance* WorstCaseParameterTime = nullptr;
+
+		for (UWorstCasePerformanceReportInstance*& WCP : WorstTimeParameterInstances[LOD][LongesTimeStateIndex])
+		{
+			if (LongesTimeParameterIndexInCO == WCP->LongestUpdateTimeParameterIndexInCO)
+			{
+				WorstCaseParameterTime = WCP;
+				// If any of there checks fails, we are receiving updates from a different CO than this was initialized with.
+				check(WorstCaseParameterTime);
+				check(WorstCaseParameterTime->LOD == LOD);
+				check(WorstCaseParameterTime->WorstCaseType == UWorstCasePerformanceReportInstance::EPerformanceReportWorstCaseType::SLOWEST_UPDATE_PARAM_IN_RELEVANT_STATE);
+				break;
+			}
+		}
+
+		if (!WorstCaseParameterTime)
+		{
+			WorstCaseParameterTime = NewObject<UWorstCasePerformanceReportInstance>(GetTransientPackage(), NAME_None, RF_Transient);
+			WorstCaseParameterTime->LOD = LOD;
+			WorstCaseParameterTime->WorstCaseType = UWorstCasePerformanceReportInstance::EPerformanceReportWorstCaseType::SLOWEST_UPDATE_PARAM_IN_RELEVANT_STATE;
+			WorstTimeParameterInstances[LOD][LongesTimeStateIndex].Push(WorstCaseParameterTime);
+			IsNewPArameter = true;
+		}
+
+		if (WorstCaseParameterTime && CurrentInstanceGenerationTime > WorstCaseParameterTime->UpdateTime)
+		{
+			WorstCase = WorstCaseParameterTime;
+		}
+	}
+
+	if (WorstCase)
+	{
+		WorstCase->SetInstance(Instance);
+		if (WorstCase->WorstCaseInstance)
+		{
+			WorstCase->WorstCaseInstance->SetState(LongesTimeStateIndex);
+		}
+		WorstCase->LongestUpdateTimeStateIndex = LongesTimeStateIndex;
+		WorstCase->LongestUpdateTimeParameterIndexInCO = LongesTimeParameterIndexInCO;
+		WorstCase->UpdateTime = CurrentInstanceGenerationTime;
+		if (data)
+		{
+			WorstCase->MeasuredDataForThisWorstCase = *data;
+		}
+		UpdateWorstTimeInstance(WorstCase);
+		BuildTimeWorstCasesDetailedView();
+	}
+}
+
+void SCustomizableObjecEditorPerformanceReport::UpdateWorstTimeInstance(const UWorstCasePerformanceReportInstance* EvenWorseTime)
+{
+	if (!WorstTimeInstance || (EvenWorseTime && (EvenWorseTime->UpdateTime > WorstTimeInstance->UpdateTime)))
+	{
+		WorstTimeInstance = EvenWorseTime;
+	}
+}
+
+void SCustomizableObjecEditorPerformanceReport::StressTestEnded()
+{
+	if (ReportSubject)
+	{
+		StateTest = MakeShareable(new FRuntimeTest);
+		StateTest->WorstTimeDelegate.BindRaw(this, &SCustomizableObjecEditorPerformanceReport::WorstTimeFound);
+		StateTest->StateTestEnded.BindRaw(this, &SCustomizableObjecEditorPerformanceReport::StateTestEnded);
+		StateTest->StartTest(ReportSubject);
+		StatusText->SetText(LOCTEXT("PerformanceReportStatusStateRunning", "Running State test."));
+	}
+}
+
+void SCustomizableObjecEditorPerformanceReport::StateTestEnded()
+{
+	StatusText->SetText(LOCTEXT("PerformanceReportStatusFinished", "Tests completed."));
+}
+
+UWorstCasePerformanceReportInstance* SCustomizableObjecEditorPerformanceReport::GetWorstCaseStructureFromIndexes(UWorstCasePerformanceReportInstance::EPerformanceReportWorstCaseType WorstCaseType, uint32 LOD, int32 StateID, int32 ParamID) const
+{
+	if (LOD >= 0 && LOD < NumLOD)
+	{
+		switch (WorstCaseType)
+		{
+		case UWorstCasePerformanceReportInstance::EPerformanceReportWorstCaseType::LARGEST_TEXTURE_SIZE:
+		{
+			check((uint32)WorstTextureInstances.Num() > LOD);
+			return WorstTextureInstances[LOD];
+		}
+		case UWorstCasePerformanceReportInstance::EPerformanceReportWorstCaseType::MOST_TRIANGLES:
+		{
+			check((uint32)WorstGeometryInstances.Num() > LOD);
+			return WorstGeometryInstances[LOD];
+		}
+		case UWorstCasePerformanceReportInstance::EPerformanceReportWorstCaseType::SLOWEST_UPDATE_ENTER_STATE:
+		{
+			if (StateID != INDEX_NONE)
+			{
+				check((uint32)WorstTimeStateInstances.Num() > LOD);
+				check(StateID >= 0);
+				check(WorstTimeStateInstances[LOD].Num() > StateID);
+				return WorstTimeStateInstances[LOD][StateID];
+			}
+		}
+		case UWorstCasePerformanceReportInstance::EPerformanceReportWorstCaseType::SLOWEST_UPDATE_PARAM_IN_RELEVANT_STATE:
+		{
+			if (StateID != INDEX_NONE && ParamID != INDEX_NONE)
+			{
+				check((uint32)WorstTimeParameterInstances.Num() > LOD);
+				check(StateID >= 0);
+				check(WorstTimeParameterInstances[LOD].Num() > StateID);
+				const int32 NumParams = WorstTimeParameterInstances[LOD][StateID].Num();
+				for (int i = 0; i < NumParams; ++i)
+				{
+					if (WorstTimeParameterInstances[LOD][StateID][i] && WorstTimeParameterInstances[LOD][StateID][i]->LongestUpdateTimeParameterIndexInCO == ParamID)
+					{
+						return WorstTimeParameterInstances[LOD][StateID][i];
+					}
+				}
+			}
+		}
+		}
+	}
+	return nullptr;
+}
+
+TSharedRef<SWidget> SCustomizableObjecEditorPerformanceReport::BuildAnalyzer()
+{
+	AssetThumbnailPool = MakeShareable(new FAssetThumbnailPool(32));
+
+	TSharedRef<SVerticalBox> ResultVerticalBox = SNew(SVerticalBox)
+		.Visibility_Lambda([this]() { return CurrentReportInstance ? EVisibility::Visible : EVisibility::Collapsed; });
+
+	ResultVerticalBox->AddSlot()
+	.HAlign(HAlign_Left)
+	.AutoHeight()
+	[
+		SNew(SBorder)
+		.Padding(5.0f)
+		[
+			SNew(SVerticalBox)
+
+			+ SVerticalBox::Slot()
+			.HAlign(HAlign_Left)
+			.AutoHeight()
+			.Padding(5.0f)
+			[
+				SNew(SHorizontalBox)
+
+				+ SHorizontalBox::Slot()
+				.AutoWidth()
+				[
+					SNew(STextBlock)
+					.Text(LOCTEXT("SelectTextureAnalyzerInstance", "Selected Instance: "))
+				]
+
+				+ SHorizontalBox::Slot()
+				.AutoWidth()
+				[
+					SAssignNew(SelectedTextureAnalyzerText, STextBlock)
+				]
+			]
+
+			+ SVerticalBox::Slot()
+			.HAlign(HAlign_Left)
+			.AutoHeight()
+			[
+				SNew(SHorizontalBox)
+
+				+ SHorizontalBox::Slot()
+				.AutoWidth()
+				.Padding(5.0f)
+				[
+					SAssignNew(InstanceThumbnail, SObjectPropertyEntryBox)
+					.AllowedClass(UCustomizableObjectInstance::StaticClass())
+					.ObjectPath(this, &SCustomizableObjecEditorPerformanceReport::OnGetInstancePath)
+					.AllowClear(false)
+					.DisplayUseSelected(false)
+					.DisplayBrowse(false)
+					.EnableContentPicker(false)
+					.DisplayThumbnail(true)
+					.ThumbnailPool(AssetThumbnailPool)
+				]
+
+				+ SHorizontalBox::Slot()
+				.AutoWidth()
+				.VAlign(VAlign_Center)
+				.Padding(5.0f)
+				[
+					SAssignNew(PerformanceReportRefreshSelectedInstance, SButton)
+					.ButtonStyle(FAppStyle::Get(), "FlatButton.Danger")
+					.IsEnabled(this, &SCustomizableObjecEditorPerformanceReport::IsSelectedInstanceOutdated)
+					.OnClicked(this, &SCustomizableObjecEditorPerformanceReport::RefreshSelectedInstance)
+					[
+						SNew(STextBlock)
+						.Text(LOCTEXT("RefreshSelectedInstanceButton", "Refresh"))
+					]
+				]
+			]
+		]
+	];
+
+	ResultVerticalBox->AddSlot()
+	.HAlign(HAlign_Fill)
+	[
+		SNew(SBorder)
+		.Padding(5.0f)
+		[
+			SNew(SScrollBox)
+			+ SScrollBox::Slot()
+			[
+				PerformanceReportTextureAnalyzer.ToSharedRef()
+			]
+		]
+	];
+
+	UpdateSelectedInstanceData();
+
+	return ResultVerticalBox;
+}
+
+void SCustomizableObjecEditorPerformanceReport::UpdateSelectedInstanceData()
+{
+	// Refresh Identification text
+	if (!SelectedTextureAnalyzerText)
+	{
+		return;
+	}
+
+	FText Result = LOCTEXT("NoInstanceCurrentlySelected", "Nonee currently selected. Click on one of the rows on the left to view it's details.");
+	const FNumberFormattingOptions BytesFormat = FNumberFormattingOptions().SetMaximumFractionalDigits(2).SetMinimumFractionalDigits(2);
+	const FNumberFormattingOptions TimeFormat = FNumberFormattingOptions().SetMaximumFractionalDigits(2).SetMinimumFractionalDigits(2);
+	const FNumberFormattingOptions PolygonsFormat = FNumberFormattingOptions().SetMaximumFractionalDigits(0);
+	const FNumberFormattingOptions PercentageFormat = FNumberFormattingOptions().SetMaximumFractionalDigits(0).SetAlwaysSign(true);
+
+	if (CurrentReportInstance && CurrentReportWorstCase)
+	{
+		const FText TextLOD = FText::AsNumber(CurrentReportWorstCase->LOD, &FNumberFormattingOptions::DefaultNoGrouping());
+		FText TextReason;
+		FText TextAmount;
+		switch (CurrentReportWorstCase->WorstCaseType)
+		{
+		case UWorstCasePerformanceReportInstance::EPerformanceReportWorstCaseType::LARGEST_TEXTURE_SIZE:
+		{
+			TextReason = FText::FromString("Largest mutable textures size");
+			const float Bytes = CurrentReportWorstCase->TextureBytes;
+			const float Value = (Bytes > 1048576) ? float(Bytes) / 1048576.0f : float(Bytes) / 1024.0f;
+
+			TextAmount = FText::Format(LOCTEXT("MutablePerformanceTextureSize", "{0} {1}"), FText::AsNumber(Value, &BytesFormat), FText::FromString((Bytes > 1048576) ? "MB" : "KB"));
+			break;
+		}
+		case UWorstCasePerformanceReportInstance::EPerformanceReportWorstCaseType::MOST_TRIANGLES:
+		{
+			TextReason = FText::FromString("Most triangles");
+			TextAmount = FText::AsNumber(CurrentReportWorstCase->NumFaces, &PolygonsFormat);
+			break;
+		}
+		case UWorstCasePerformanceReportInstance::EPerformanceReportWorstCaseType::SLOWEST_UPDATE_ENTER_STATE:
+		{
+			TextReason = FText::FromString("Slowest update when entering the state");
+			TextAmount = FText::AsNumber(CurrentReportWorstCase->UpdateTime, &TimeFormat);
+			break;
+		}
+		case UWorstCasePerformanceReportInstance::EPerformanceReportWorstCaseType::SLOWEST_UPDATE_PARAM_IN_RELEVANT_STATE:
+		{
+			TextReason = FText::FromString("Slowest update when changing parameter");
+			TextAmount = FText::AsNumber(CurrentReportWorstCase->UpdateTime, &TimeFormat);
+			break;
+		}
+		default:
+		{
+			break;
+		}
+		}
+
+		Result = FText::Format(LOCTEXT("PerformanceReportReason", "{0} for LOD {1} ({2})"), TextReason, TextLOD, TextAmount);
+	}
+
+	SelectedTextureAnalyzerText->SetText(Result);
+
+	// Refresh Texture Analyzer
+	if (CurrentReportInstance && CurrentReportInstance->IsValidLowLevelFast())
+	{
+		// If we have at least one skeletal mesh in the instance, we know it's generated (at least once)
+		if (CurrentReportInstance->SkeletalMeshes.Num() > 0 && CurrentReportInstance->SkeletalMeshes[0])
+		{
+			PerformanceReportTextureAnalyzer->RefreshTextureAnalyzerTable(CurrentReportInstance);
+		}
+		else
+		{
+			CurrentReportInstance->UpdatedDelegate.AddUniqueDynamic(HelperCallback, &UPerformanceReportHelper::DelegatedCallback);
+			CurrentReportInstance->UpdateSkeletalMeshAsync(true, true);
+		}
+	}
+}
+
+void SCustomizableObjecEditorPerformanceReport::BuildTextureWorstCasesView()
+{
+	const FNumberFormattingOptions BytesFormat = FNumberFormattingOptions().SetMaximumFractionalDigits(2).SetMinimumFractionalDigits(2);
+
+	TextureWorstCasesView->ClearChildren();
+
+	TextureWorstCasesView->AddSlot()
+	.HAlign(HAlign_Fill)
+	.AutoHeight()
+	[
+		SNew(SListView<UWorstCasePerformanceReportInstance*>)
+		.ListItemsSource(&WorstTextureInstances)
+		.OnGenerateRow(this, &SCustomizableObjecEditorPerformanceReport::GenerateRow)
+		.OnMouseButtonClick(this, &SCustomizableObjecEditorPerformanceReport::SelectInstance)
+		.SelectionMode(ESelectionMode::Single)
+		.IsFocusable(true)
+		.HeaderRow
+		(
+			SNew(SHeaderRow)
+			+ SHeaderRow::Column(ColumnID_PR_LOD_Label)
+			.DefaultLabel(LOCTEXT("LOD", "LOD"))
+			.FillWidth(0.2)
+
+			+SHeaderRow::Column(ColumnID_PR_TexturesTotalSize_Label)
+			.DefaultLabel(LOCTEXT("TexturesTotalSize", "Textures total size"))
+			.FillWidth(0.2)
+
+			+ SHeaderRow::Column(ColumnID_PR_Delta_Label)
+			.DefaultLabel(LOCTEXT("Delta", "Difference"))
+			.FillWidth(0.1)
+
+			+ SHeaderRow::Column(ColumnID_PR_Average_Label)
+			.DefaultLabel(LOCTEXT("Average", "Average"))
+			.FillWidth(0.1)
+		)
+	];
+}
+
+void SCustomizableObjecEditorPerformanceReport::BuildGeometryWorstCasesView()
+{
+	GeometryWorstCasesView->ClearChildren();
+
+	GeometryWorstCasesView->AddSlot()
+	.HAlign(HAlign_Fill)
+	.AutoHeight()
+	[
+		SNew(SListView<UWorstCasePerformanceReportInstance*>)
+		.ListItemsSource(&WorstGeometryInstances)
+		.OnGenerateRow(this, &SCustomizableObjecEditorPerformanceReport::GenerateRow)
+		.OnMouseButtonClick(this, &SCustomizableObjecEditorPerformanceReport::SelectInstance)
+		.SelectionMode(ESelectionMode::Single)
+		.IsFocusable(true)
+		.HeaderRow
+		(
+			SNew(SHeaderRow)
+			+ SHeaderRow::Column(ColumnID_PR_LOD_Label)
+			.DefaultLabel(LOCTEXT("LOD", "LOD"))
+			.FillWidth(0.2)
+
+			+ SHeaderRow::Column(ColumnID_PR_NumberOfTriangles_Label)
+			.DefaultLabel(LOCTEXT("NumberOfTriangles", "Number of triangles"))
+			.FillWidth(0.2)
+
+			+ SHeaderRow::Column(ColumnID_PR_Delta_Label)
+			.DefaultLabel(LOCTEXT("Delta", "Difference"))
+			.FillWidth(0.1)
+
+			+ SHeaderRow::Column(ColumnID_PR_Average_Label)
+			.DefaultLabel(LOCTEXT("Average", "Average"))
+			.FillWidth(0.1)
+		)
+	];
+}
+
+void SCustomizableObjecEditorPerformanceReport::BuildTimeWorstCasesDetailedView()
+{
+	const FNumberFormattingOptions TimeFormat = FNumberFormattingOptions().SetMaximumFractionalDigits(2).SetMinimumFractionalDigits(2);
+
+	TimeWorstCasesDetailedView->ClearChildren();
+
+	for (int32 LOD = 0; LOD < WorstTimeStateInstances.Num(); ++LOD)
+	{
+		// Currently, states test only returns LOD 0 data
+		if (LOD != 0)
+		{
+			continue;
+		}
+
+		const FText TextLOD = FText::AsNumber(LOD, &FNumberFormattingOptions::DefaultNoGrouping());
+
+		TimeWorstCasesDetailedView->AddSlot()
+		.HAlign(HAlign_Left)
+		.Padding(5.0f)
+		[
+			SNew(STextBlock)
+			.Text(FText::Format(LOCTEXT("MutableLODNTitle", "LOD {0}:"), TextLOD))
+		];
+
+		TimeWorstCasesDetailedView->AddSlot()
+		.HAlign(HAlign_Fill)
+		[
+			SNew(SListView<UWorstCasePerformanceReportInstance*>)
+			.ListItemsSource(&WorstTimeStateInstances[LOD])
+			.OnGenerateRow(this, &SCustomizableObjecEditorPerformanceReport::GenerateRow)
+			.OnMouseButtonClick(this, &SCustomizableObjecEditorPerformanceReport::SelectInstance)
+			.SelectionMode(ESelectionMode::Single)
+			.IsFocusable(true)
+			.HeaderRow
+			(
+				SAssignNew(TimeWorstCasesViewHeader, SHeaderRow)
+				+ SHeaderRow::Column(ColumnID_PR_Name_Label)
+				.DefaultLabel(LOCTEXT("Name", "Name"))
+				.FillWidth(0.2)
+
+				+ SHeaderRow::Column(ColumnID_PR_UpdateTime_Label)
+				.DefaultLabel(LOCTEXT("StateEnterUpdateTime", "Update time (ms)"))
+				.FillWidth(0.2)
+
+				+ SHeaderRow::Column(ColumnID_PR_Delta_Label)
+				.DefaultLabel(LOCTEXT("Delta", "Difference"))
+				.FillWidth(0.1)
+
+				+ SHeaderRow::Column(ColumnID_PR_Average_Label)
+				.DefaultLabel(LOCTEXT("Average", "Average"))
+				.FillWidth(0.1)
+			)
+		];
+	}
+
+	TimeWorstCasesDetailedView->Invalidate(EInvalidateWidget::Layout);
+}
+
+void SCustomizableObjecEditorPerformanceReport::OnOpenAssetEditor(UObject* ObjectToEdit)
+{
+	if (ObjectToEdit)
+	{
+		GEditor->EditObject(ObjectToEdit);
+	}
+}
+
+void SCustomizableObjecEditorPerformanceReport::ChangeFocusedReportInstance(UWorstCasePerformanceReportInstance* InReportInstance)
+{
+	CurrentReportWorstCase = InReportInstance;
+}
+
+FString SCustomizableObjecEditorPerformanceReport::OnGetInstancePath() const
+{
+	if (CurrentReportInstance && CurrentReportInstance->IsValidLowLevelFast(false))
+	{
+		return CurrentReportInstance->GetPathName();
+	}
+
+	return FString();
+}
+
+bool SCustomizableObjecEditorPerformanceReport::IsSelectedInstanceOutdated() const
+{
+	if (CurrentReportWorstCase)
+	{
+		return CurrentReportWorstCase->WorstCaseInstance != CurrentReportInstance;
+		// // Alternative method to checking if a selected instance is outdated, if in the future, CurrentReportInstnce and CurrentReportWorstCase->WorstCaseInstance are kept synchronized
+		//const UWorstCasePerformanceReportInstance* InstanceInData = GetWorstCaseStructureFromIndexes(CurrentReportWorstCase->WorstCaseType, CurrentReportWorstCase->LOD, CurrentReportWorstCase->LongestUpdateTimeStateIndex, CurrentReportWorstCase->LongestUpdateTimeParameterIndexInCO);
+		//if (InstanceInData)
+		//{
+		//	switch (CurrentReportWorstCase->WorstCaseType)
+		//	{
+		//	case UWorstCasePerformanceReportInstance::EPerformanceReportWorstCaseType::LARGEST_TEXTURE_SIZE:
+		//	{
+		//		if (InstanceInData->TextureBytes > CurrentReportWorstCase->TextureBytes)
+		//		{
+		//			return true;
+		//		}
+		//		break;
+		//	}
+		//	case UWorstCasePerformanceReportInstance::EPerformanceReportWorstCaseType::MOST_TRIANGLES:
+		//	{
+		//		if (InstanceInData->NumFaces > CurrentReportWorstCase->NumFaces)
+		//		{
+		//			return true;
+		//		}
+		//		break;
+		//	}
+		//	case UWorstCasePerformanceReportInstance::EPerformanceReportWorstCaseType::SLOWEST_UPDATE_ENTER_STATE:
+		//	case UWorstCasePerformanceReportInstance::EPerformanceReportWorstCaseType::SLOWEST_UPDATE_PARAM_IN_RELEVANT_STATE:
+		//	{
+		//		if (InstanceInData->UpdateTime > CurrentReportWorstCase->UpdateTime)
+		//		{
+		//			return true;
+		//		}
+		//		break;
+		//	}
+		//	}
+		//}
+	}
+	return false;
+}
+
+FReply SCustomizableObjecEditorPerformanceReport::RefreshSelectedInstance()
+{
+	CurrentReportInstance = CurrentReportWorstCase->WorstCaseInstance;
+	UpdateSelectedInstanceData();
+
+	return FReply::Handled();
+}
+
+TSharedRef<ITableRow> SCustomizableObjecEditorPerformanceReport::GenerateRow(UWorstCasePerformanceReportInstance* ReportedInstance, const TSharedRef<STableViewBase>& OwnerTable)
+{
+	return SNew(SPerformanceReportWorstCaseRow, OwnerTable).ReportInstance(ReportedInstance).PerformanceReport(this).Style(&FCustomizableObjectEditorStyle::Get().GetWidgetStyle<FTableRowStyle>("CustomizableObjectPerformanceReport.Row"));
+}
+
+const TArray<UWorstCasePerformanceReportInstance*>* SCustomizableObjecEditorPerformanceReport::GetWorstTimeParameters(uint32 LOD, uint32 StateIdx) const
+{
+	check((uint32)WorstTimeParameterInstances.Num() > LOD);
+	//check((uint32)WorstTimeParameterInstances[LOD].Num() > StateIdx);
+
+	if ((uint32)WorstTimeParameterInstances.Num() > LOD && (uint32)WorstTimeParameterInstances[LOD].Num() > StateIdx)
+	{
+		return &WorstTimeParameterInstances[LOD][StateIdx];
+	}
+
+	return nullptr;
+}
+
+void SCustomizableObjecEditorPerformanceReport::SelectInstance(UWorstCasePerformanceReportInstance* ClickedInstance)
+{
+	if (CurrentReportInstance)
+	{
+		CurrentReportInstance->UpdatedDelegate.RemoveDynamic(HelperCallback, &UPerformanceReportHelper::DelegatedCallback);
+	}
+	CurrentReportWorstCase = ClickedInstance;
+	CurrentReportInstance = ClickedInstance->WorstCaseInstance;
+	UpdateSelectedInstanceData();
+
+	BuildTextureWorstCasesView();
+	BuildGeometryWorstCasesView();
+	BuildTimeWorstCasesDetailedView();
+}
+
+void SCustomizableObjecEditorPerformanceReport::SelectInstanceFinishedUpdate(UCustomizableObjectInstance* UpdatedInstance)
+{
+	if (CurrentReportInstance && UpdatedInstance == CurrentReportInstance)
+	{
+		PerformanceReportTextureAnalyzer->RefreshTextureAnalyzerTable(CurrentReportInstance);
+	}
+}
+
+void SPerformanceReportWorstCaseRow::Construct(const FArguments& InArgs, const TSharedRef<STableViewBase>& InOwnerTableView)
+{
+	ReportInstance = InArgs._ReportInstance;
+	PerformanceReport = InArgs._PerformanceReport;
+	SMutableExpandableTableRow<UWorstCasePerformanceReportInstance*>::Construct(FSuperRowType::FArguments().Style(InArgs._Style), InOwnerTableView);
+	Style = InArgs._Style;
+}
+
+void SPerformanceReportWorstCaseRow::AddReferencedObjects(FReferenceCollector& Collector)
+{
+	Collector.AddReferencedObject(ReportInstance);
+	Collector.AddReferencedObjects(SubWidgetData);
+}
+
+TSharedRef<SWidget> SPerformanceReportWorstCaseRow::GenerateWidgetForColumn(const FName& ColumnName)
+{
+	check(ReportInstance);
+
+	const FNumberFormattingOptions BytesFormat = FNumberFormattingOptions().SetMaximumFractionalDigits(2).SetMinimumFractionalDigits(2);
+	const FNumberFormattingOptions TimeFormat = FNumberFormattingOptions().SetMaximumFractionalDigits(2).SetMinimumFractionalDigits(2);
+	const FNumberFormattingOptions PolygonsFormat = FNumberFormattingOptions().SetMaximumFractionalDigits(0);
+	const FNumberFormattingOptions PercentageFormat = FNumberFormattingOptions().SetMaximumFractionalDigits(0).SetAlwaysSign(true);
+	const bool bIsSelected = PerformanceReport && ReportInstance == PerformanceReport->CurrentReportWorstCase;
+	TSharedRef<SOverlay> ResultWidget = SNew(SOverlay);
+
+	if (bIsSelected)
+	{
+		ResultWidget->AddSlot()
+		[
+			SNew(SBorder)
+			.BorderImage(this, &SPerformanceReportWorstCaseRow::GetSelectedRowHighlightColor)
+		];
+	}
+
+	if (ReportInstance->WorstCaseInstance)
+	{
+		if (ColumnName == ColumnID_PR_Delta_Label)
+		{
+			const float Mean = ReportInstance->MeasuredDataForThisWorstCase.MeanTime;
+			if (Mean)
+			{
+				float CurrentMaximum = 0.0f;
+
+				switch (ReportInstance->WorstCaseType)
+				{
+				case UWorstCasePerformanceReportInstance::EPerformanceReportWorstCaseType::LARGEST_TEXTURE_SIZE:
+				{
+					CurrentMaximum = (float)ReportInstance->TextureBytes;
+					break;
+				}
+				case UWorstCasePerformanceReportInstance::EPerformanceReportWorstCaseType::MOST_TRIANGLES:
+				{
+					CurrentMaximum = (float)ReportInstance->NumFaces;
+					break;
+				}
+				case UWorstCasePerformanceReportInstance::EPerformanceReportWorstCaseType::SLOWEST_UPDATE_ENTER_STATE:
+				case UWorstCasePerformanceReportInstance::EPerformanceReportWorstCaseType::SLOWEST_UPDATE_PARAM_IN_RELEVANT_STATE:
+				{
+					CurrentMaximum = ReportInstance->UpdateTime;
+					break;
+				}
+				default:
+				{
+					break;
+				}
+				}
+
+				FText Percentage = FText::AsNumber((CurrentMaximum * 100.0f / Mean) - 100.0f, &PercentageFormat);
+
+				ResultWidget->AddSlot()
+				[
+					SNew(STextBlock)
+					.Text(FText::Format(LOCTEXT("MutablePerformanceDeltaFromAverage", "{0}%"), Percentage))
+				];
+			}
+		}
+		else if (ColumnName == ColumnID_PR_Average_Label)
+		{
+			const float Mean = ReportInstance->MeasuredDataForThisWorstCase.MeanTime;
+			if (Mean)
+			{
+				switch (ReportInstance->WorstCaseType)
+				{
+				case UWorstCasePerformanceReportInstance::EPerformanceReportWorstCaseType::LARGEST_TEXTURE_SIZE:
+				{
+					const float Value = (Mean > 1048576) ? float(Mean) / 1048576.0f : float(Mean) / 1024.0f;
+
+					ResultWidget->AddSlot()
+					[
+						SNew(STextBlock).Text(FText::Format(LOCTEXT("MutablePerformanceTextureSize", "{0} {1}"), FText::AsNumber(Value, &BytesFormat), FText::FromString((Mean > 1048576) ? "MB" : "KB")))
+					];
+					break;
+				}
+				case UWorstCasePerformanceReportInstance::EPerformanceReportWorstCaseType::MOST_TRIANGLES:
+				{
+					ResultWidget->AddSlot()
+					[
+						SNew(STextBlock).Text(FText::AsNumber(Mean, &PolygonsFormat))
+					];
+					break;
+				}
+				case UWorstCasePerformanceReportInstance::EPerformanceReportWorstCaseType::SLOWEST_UPDATE_ENTER_STATE:
+				case UWorstCasePerformanceReportInstance::EPerformanceReportWorstCaseType::SLOWEST_UPDATE_PARAM_IN_RELEVANT_STATE:
+				{
+					ResultWidget->AddSlot()
+					[
+						SNew(STextBlock).Text(FText::AsNumber(Mean, &TimeFormat))
+					];
+					break;
+				}
+				default:
+				{
+					break;
+				}
+				}
+			}
+		}
+		else if (ColumnName == ColumnID_PR_Name_Label)
+		{
+			if (const UCustomizableObject* CustomizableObject = ReportInstance->WorstCaseInstance->GetCustomizableObject();
+				CustomizableObject &&
+				!CustomizableObject->IsLocked())
+			{
+				const bool bIsDetailParam = ReportInstance->WorstCaseType == UWorstCasePerformanceReportInstance::EPerformanceReportWorstCaseType::SLOWEST_UPDATE_PARAM_IN_RELEVANT_STATE;
+				const FText TextState = ReportInstance->WorstCaseType == UWorstCasePerformanceReportInstance::EPerformanceReportWorstCaseType::SLOWEST_UPDATE_ENTER_STATE ?
+					FText::FromString("State: " + CustomizableObject->GetStateName(ReportInstance->LongestUpdateTimeStateIndex))
+					: FText::FromString((bIsDetailParam ? "		" : "") + CustomizableObject->GetParameterName(ReportInstance->LongestUpdateTimeParameterIndexInCO));
+
+				ResultWidget->AddSlot()
+				[
+					SNew(STextBlock).Text(TextState)
+				];
+			}
+		}
+		else if (ColumnName == ColumnID_PR_UpdateTime_Label)
+		{
+			ResultWidget->AddSlot()
+			[
+				SNew(STextBlock).Text(FText::AsNumber(ReportInstance->UpdateTime, &TimeFormat))
+			];
+		}
+		else if (ColumnName == ColumnID_PR_LOD_Label)
+		{
+			const FText TextLOD = FText::AsNumber(ReportInstance->LOD, &FNumberFormattingOptions::DefaultNoGrouping());
+
+			ResultWidget->AddSlot()
+			[
+				SNew(STextBlock).Text(FText::Format(LOCTEXT("MutablePerformanceLOD", "{0}"), TextLOD))
+			];
+		}
+		else if (ColumnName == ColumnID_PR_TexturesTotalSize_Label)
+		{
+			const float Bytes = ReportInstance->TextureBytes;
+			const float Value = (Bytes > 1048576) ? float(Bytes) / 1048576.0f : float(Bytes) / 1024.0f;
+
+			ResultWidget->AddSlot()
+			[
+				SNew(STextBlock).Text(FText::Format(LOCTEXT("MutablePerformanceTextureSize", "{0} {1}"), FText::AsNumber(Value, &BytesFormat), FText::FromString((Bytes > 1048576) ? "MB" : "KB")))
+			];
+		}
+		else if (ColumnName == ColumnID_PR_NumberOfTriangles_Label)
+		{
+			ResultWidget->AddSlot()
+			[
+				SNew(STextBlock).Text(FText::AsNumber(ReportInstance->NumFaces, &PolygonsFormat))
+			];
+		}
+		else if (ColumnName == ColumnID_PR_InstanceName_Label)
+		{
+			const FText InstanceName = FText::FromString(ReportInstance->GetName());
+
+			ResultWidget->AddSlot()
+			[
+				SNew(STextBlock).Text(InstanceName)
+			];
+		}
+	}
+
+	return ResultWidget;
+}
+
+TSharedPtr<SWidget> SPerformanceReportWorstCaseRow::GenerateAdditionalWidgetForRow()
+{
+	if (ReportInstance && ReportInstance->WorstCaseType == UWorstCasePerformanceReportInstance::EPerformanceReportWorstCaseType::SLOWEST_UPDATE_ENTER_STATE)
+	{
+		const TArray<UWorstCasePerformanceReportInstance*>* WorstTimeParameters = PerformanceReport->GetWorstTimeParameters(ReportInstance->LOD, ReportInstance->LongestUpdateTimeStateIndex);
+
+		if (WorstTimeParameters)
+		{
+			SubWidgetData.Empty(WorstTimeParameters->Num());
+			for (UWorstCasePerformanceReportInstance* elem : *WorstTimeParameters)
+			{
+				SubWidgetData.Push(elem);
+			}
+			return SNew(SListView<UWorstCasePerformanceReportInstance*>)
+				.ListItemsSource(&SubWidgetData)
+				.OnGenerateRow(this, &SPerformanceReportWorstCaseRow::GenerateStyledRow)
+				.OnMouseButtonClick(PerformanceReport, &SCustomizableObjecEditorPerformanceReport::SelectInstance)
+				.SelectionMode(ESelectionMode::Single)
+				.IsFocusable(true)
+				.HeaderRow(PerformanceReport->TimeWorstCasesViewHeader);
+		}
+	}
+
+	SubWidgetData.Empty();
+	return nullptr;
+}
+
+EVisibility SPerformanceReportWorstCaseRow::GetAdditionalWidgetDefaultVisibility() const
+{
+	return (ReportInstance) ? ReportInstance->SubVisibility : EVisibility::Collapsed;
+}
+
+void SPerformanceReportWorstCaseRow::SetAdditionalWidgetVisibility(const EVisibility InVisibility)
+{
+	if (ReportInstance)
+	{
+		ReportInstance->SubVisibility = InVisibility;
+	}
+	SMutableExpandableTableRow::SetAdditionalWidgetVisibility(InVisibility);
+}
+
+const FSlateBrush* SPerformanceReportWorstCaseRow::GetSelectedRowHighlightColor() const
+{
+	const bool bIsSelected = PerformanceReport && ReportInstance == PerformanceReport->CurrentReportWorstCase;
+	const bool bIsUpToDate = bIsSelected && ReportInstance && ReportInstance->WorstCaseInstance == PerformanceReport->CurrentReportInstance;
+	return bIsUpToDate? FAppStyle::GetBrush("Brushes.Highlight") : FAppStyle::GetBrush("Brushes.AccentGray");
+}
+
+TSharedRef<ITableRow> SPerformanceReportWorstCaseRow::GenerateStyledRow(UWorstCasePerformanceReportInstance* ReportedInstance, const TSharedRef<STableViewBase>& OwnerTable)
+{
+	const bool bEvenEntryIndex = (GetIndexInList() % 2 == 0);
+	return SNew(SPerformanceReportWorstCaseRow, OwnerTable).ReportInstance(ReportedInstance).PerformanceReport(PerformanceReport).Style(
+		bEvenEntryIndex ?
+		&FCustomizableObjectEditorStyle::Get().GetWidgetStyle<FTableRowStyle>("CustomizableObjectPerformanceReportAllEven.Row")
+		: &FCustomizableObjectEditorStyle::Get().GetWidgetStyle<FTableRowStyle>("CustomizableObjectPerformanceReportAllOdd.Row")
+	);
+}
+
+void UPerformanceReportHelper::DelegatedCallback(UCustomizableObjectInstance* Instance)
+{
+	if (PerformanceReport)
+	{
+		PerformanceReport->SelectInstanceFinishedUpdate(Instance);
+	}
+}
+
+#undef LOCTEXT_NAMESPACE
diff --git a/Engine/Plugins/Experimental/Mutable/Source/CustomizableObjectEditor/Private/SCustomizableObjectEditorPerformanceReport.h b/Engine/Plugins/Experimental/Mutable/Source/CustomizableObjectEditor/Private/SCustomizableObjectEditorPerformanceReport.h
new file mode 100644
--- /dev/null
+++ b/Engine/Plugins/Experimental/Mutable/Source/CustomizableObjectEditor/Private/SCustomizableObjectEditorPerformanceReport.h
@@ -0,0 +1,272 @@
+// Copyright Epic Games, Inc. All Rights Reserved.
+
+#pragma once
+
+#include "Widgets/SCompoundWidget.h"
+#include "Widgets/Views/SListView.h"
+#include "Misc/NotifyHook.h"
+#include "CustomizableObject.h"
+#include "StressTest.h"
+#include "Widgets/SMutableExpandableTableRow.h"
+#include "SCustomizableObjectEditorTextureAnalyzer.h"
+#include "SCustomizableObjectEditorPerformanceReport.generated.h"
+
+
+// Instances performing worst in some aspect
+UCLASS()
+class UWorstCasePerformanceReportInstance : public UObject
+{
+public:
+	GENERATED_BODY()
+
+	enum EPerformanceReportWorstCaseType {
+		LARGEST_TEXTURE_SIZE,
+		MOST_TRIANGLES,
+		SLOWEST_UPDATE_ENTER_STATE,
+		SLOWEST_UPDATE_PARAM_IN_RELEVANT_STATE
+	};
+
+	// The instance that presented this worst case.
+	UPROPERTY(VisibleAnywhere, AssetRegistrySearchable, Transient, NoClear, Category = "Stats|PerformanceReport", meta = (DisplayName = "Instance", ColumnWidth = "90")/*"PerformanceReport|BiggestTextureSize"*/)
+	TObjectPtr<class UCustomizableObjectInstance> WorstCaseInstance = nullptr;
+
+	// What worst case does this represent
+	EPerformanceReportWorstCaseType WorstCaseType = LARGEST_TEXTURE_SIZE;
+
+	// LOD number where this instance generated the worst case  (currently only relevant when LARGEST_TEXTURE_SIZE or MOST_TRIANGLES)
+	UPROPERTY(VisibleAnywhere, AssetRegistrySearchable, Transient, NoClear, Category = "Stats|PerformanceReport", meta = (DisplayName = "LOD", ColumnWidth = "10", DisplayRight = "true")/*"PerformanceReport|BiggestTextureSize"*/)
+	int32 LOD = INDEX_NONE;
+
+	// Total texture memory size that this instance had when the worst case was recorded (currently only relevant when LARGEST_TEXTURE_SIZE)
+	UPROPERTY(VisibleAnywhere, AssetRegistrySearchable, Transient, NoClear, Category = "Stats|PerformanceReport|Texture", meta = (DisplayName = "Textures Total Size", ColumnWidth = "60", DisplayRight = "true"))
+	int32 TextureBytes = -1;
+
+	// Amount of triangles this instance had when the worst case was recorded (currently only relevant when MOST_TRIANGLES)
+	UPROPERTY(VisibleAnywhere, AssetRegistrySearchable, Transient, NoClear, Category = "Stats|PerformanceReport|Geometry", meta = (DisplayName = "Number of Triangles", ColumnWidth = "60", DisplayRight = "true"))
+	int32 NumFaces = -1;
+
+	// How long it took this instance to update when the worst case was recorded (currently only relevant when SLOWEST_UPDATE_PARAM_IN_RELEVANT_STATE or SLOWEST_UPDATE_PARAM_IN_RELEVANT_STATE)
+	UPROPERTY(VisibleAnywhere, AssetRegistrySearchable, Transient, NoClear, Category = "Stats|PerformanceReport|Time", meta = (DisplayName = "Update Time", ColumnWidth = "60", DisplayRight = "true"))
+	float UpdateTime = -1.0;
+
+	// State CO index where this instance took most cycles to update (currently only relevant when SLOWEST_UPDATE_ENTER_STATE or SLOWEST_UPDATE_PARAM_IN_RELEVANT_STATE)
+	UPROPERTY(VisibleAnywhere, AssetRegistrySearchable, Transient, NoClear, Category = "Stats|PerformanceReport|Time", meta = (DisplayName = "State ID", ColumnWidth = "60", DisplayRight = "true"))
+	int32 LongestUpdateTimeStateIndex = INDEX_NONE;
+
+	// Parameter CO index where this instance took most cycles to update (currently only relevant when SLOWEST_UPDATE_PARAM_IN_RELEVANT_STATE)
+	UPROPERTY(VisibleAnywhere, AssetRegistrySearchable, Transient, NoClear, Category = "Stats|PerformanceReport|Time", meta = (DisplayName = "Parameter ID", ColumnWidth = "60", DisplayRight = "true"))
+	int32 LongestUpdateTimeParameterIndexInCO = INDEX_NONE;
+
+	// Measured data for the current worst case, the type depends on which is the current WorstCaseType
+	MeasuredData MeasuredDataForThisWorstCase;
+
+	EVisibility SubVisibility = EVisibility::Collapsed;
+
+	// TODO: Obtain all data in all cases
+	//MaterialInfoMap MapMaterialParam;
+
+	void SetInstance(class UCustomizableObjectInstance* UpdatedInstance);
+};
+
+
+UCLASS()
+class UPerformanceReportHelper : public UObject
+{
+public:
+	GENERATED_BODY()
+	// Standard constructor
+	UPerformanceReportHelper() {}
+
+	// Method to assign for the callback
+	UFUNCTION()
+	void DelegatedCallback(UCustomizableObjectInstance* Instance);
+
+	// Properties to be updated once the instance has been updated
+	class SCustomizableObjecEditorPerformanceReport* PerformanceReport = nullptr;
+};
+
+
+// Editor GUI that executes all relevant testing of a CO and presents and prints results
+class SCustomizableObjecEditorPerformanceReport : public SCompoundWidget, public FNotifyHook, public FGCObject
+{
+public:
+
+	SLATE_BEGIN_ARGS(SCustomizableObjecEditorPerformanceReport){}
+		SLATE_ARGUMENT(UCustomizableObject *,CustomizableObject)
+	SLATE_END_ARGS()
+
+	void Construct(const FArguments& InArgs);
+	virtual ~SCustomizableObjecEditorPerformanceReport();
+
+	// FSerializableObject interface
+	void AddReferencedObjects(FReferenceCollector& Collector) override;
+	virtual FString GetReferencerName() const override
+	{
+		return TEXT("SCustomizableObjecEditorPerformanceReport");
+	}
+
+	void Tick(const FGeometry& AllottedGeometry, const double InCurrentTime, const float InDeltaTime) override;
+
+	void StopTests();
+
+	// Force the re-start of the performance report
+	FReply RunPerformanceReport();
+	// Force the re-start of the performance report and provides the instance to be tested
+	FReply RunPerformanceReport(UCustomizableObject* ObjectToTest);
+	// Select and view the details of the slowest instance found up to this moment
+	FReply SelectSlowestInstance();
+
+	bool IsGetWorstInstanceButtonEnabled() const;
+
+	FReply ResetPerformanceReportOptions();
+
+private:
+
+	// Customizable object whose performance is being analyzed
+	UCustomizableObject* ReportSubject;
+
+	// Number of LODs that the Customizable Object has
+	uint32 NumLOD;
+
+	// Number of States that the Customizable Object has
+	uint32 NumStates;
+
+	// Currently focused report case and instance
+	UWorstCasePerformanceReportInstance* CurrentReportWorstCase = nullptr;
+	class UCustomizableObjectInstance* CurrentReportInstance = nullptr;
+
+	// Stress tester ================================================================================================================================================================
+	TSharedPtr<FRunningStressTest> StressTest;
+
+	int StressTestOptionInstanceCount;
+	TSharedPtr<class SNumericEntryBox<int32>> StressTestOptionInstanceCountEntry;
+	TOptional<int32> GetStressTestOptionInstanceCount() const;
+	void OnStressTestOptionInstanceCountChanged(int32 Value);
+
+	int StressTestOptionCreateInstanceTime;
+	TSharedPtr<class SNumericEntryBox<int32>> StressTestOptionCreateInstanceTimeEntry;
+	TOptional<int32> GetStressTestOptionCreateInstanceTime() const;
+	void OnStressTestOptionCreateInstanceTimeChanged(int32 Value);
+
+	int StressTestOptionCreateInstanceTimeVar;
+	TSharedPtr<class SNumericEntryBox<int32>> StressTestOptionCreateInstanceTimeVarEntry;
+	TOptional<int32> GetStressTestOptionCreateInstanceTimeVar() const;
+	void OnStressTestOptionCreateInstanceTimeVarChanged(int32 Value);
+
+	int StressTestOptionInstanceUpdateTime;
+	TSharedPtr<class SNumericEntryBox<int32>> StressTestOptionInstanceUpdateTimeEntry;
+	TOptional<int32> GetStressTestOptionInstanceUpdateTime() const;
+	void OnStressTestOptionInstanceUpdateTimeChanged(int32 Value);
+
+	int StressTestOptionInstanceUpdateTimeVar;
+	TSharedPtr<class SNumericEntryBox<int32>> StressTestOptionInstanceUpdateTimeVarEntry;
+	TOptional<int32> GetStressTestOptionInstanceUpdateTimeVar() const;
+	void OnStressTestOptionInstanceUpdateTimeVarChanged(int32 Value);
+
+	int StressTestOptionInstanceLifeTime;
+	TSharedPtr<class SNumericEntryBox<int32>> StressTestOptionInstanceLifeTimeEntry;
+	TOptional<int32> GetStressTestOptionInstanceLifeTime() const;
+	void OnStressTestOptionInstanceLifeTimeChanged(int32 Value);
+
+	int StressTestOptionInstanceLifeTimeVar;
+	TSharedPtr<class SNumericEntryBox<int32>> StressTestOptionInstanceLifeTimeVarEntry;
+	TOptional<int32> GetStressTestOptionInstanceLifeTimeVar() const;
+	void OnStressTestOptionInstanceLifeTimeVarChanged(int32 Value);
+
+	// State tester =================================================================================================================================================================
+	TSharedPtr<class FTestingCustomizableObject> StateTest;
+
+	// Texture Analyzer =============================================================================================================================================================
+	TSharedPtr<class SCustomizableObjecEditorTextureAnalyzer> PerformanceReportTextureAnalyzer;
+
+	// Worst cases ==================================================================================================================================================================
+	TArray<UWorstCasePerformanceReportInstance*> WorstTextureInstances;	// By LOD
+	TArray<UWorstCasePerformanceReportInstance*> WorstGeometryInstances;	// By LOD
+	const UWorstCasePerformanceReportInstance* WorstTimeInstance;	// Slowest instance in any test that checks for time for this report
+	TArray<TArray<UWorstCasePerformanceReportInstance*>> WorstTimeStateInstances;	// By LOD then by state index
+	TArray<TArray<TArray<UWorstCasePerformanceReportInstance*>>> WorstTimeParameterInstances;	// By LOD then by state index then all parameters relevant for that state (not ordered)
+	void ResetSavedData();
+	void WorstFacesFound(uint32 CurrentLODNumFaces, uint32 LOD, class UCustomizableObjectInstance* Instance, MeasuredData* data);
+	void WorstTextureFound(MaterialInfoMap MapMaterialParam, uint32 CurrentLODTextureSize, uint32 LOD, class UCustomizableObjectInstance* Instance, MeasuredData* data);
+	void WorstTimeFound(float CurrentInstanceGenerationTime, int32 LongesTimeStateIndex, int32 LongesTimeParameterIndexInCO, uint32 LOD, class UCustomizableObjectInstance* Instance, MeasuredData* data);
+	void UpdateWorstTimeInstance(const UWorstCasePerformanceReportInstance* EvenWorseTime);
+	void StressTestEnded();
+	void StateTestEnded();
+	UWorstCasePerformanceReportInstance* GetWorstCaseStructureFromIndexes(UWorstCasePerformanceReportInstance::EPerformanceReportWorstCaseType WorstCaseType, uint32 LOD, int32 StateID = INDEX_NONE, int32 ParamID = INDEX_NONE) const;
+
+	// UI ==========================================================================================================================================================================
+	TSharedRef<SWidget> GenerateOptionsMenuContent();
+	TSharedRef<SWidget> BuildAnalyzer();
+	void UpdateSelectedInstanceData();
+	void BuildTextureWorstCasesView();
+	void BuildGeometryWorstCasesView();
+	void BuildTimeWorstCasesDetailedView();
+	void OnOpenAssetEditor(UObject* ObjectToEdit);
+	void ChangeFocusedReportInstance(UWorstCasePerformanceReportInstance* InReportInstance);
+	FString OnGetInstancePath() const;
+	bool IsSelectedInstanceOutdated() const;
+	FReply RefreshSelectedInstance();
+	TSharedRef<ITableRow> GenerateRow(UWorstCasePerformanceReportInstance* ReportedInstance, const TSharedRef<STableViewBase>& OwnerTable);
+	// This array pointer is only valid until the array of arrays WorstTimeParameterInstances is modified. Use this only to access the actual instance pointers, but don't keep the array pointer for later use.
+	const TArray<UWorstCasePerformanceReportInstance*>* GetWorstTimeParameters(uint32 LOD, uint32 StateIdx) const;
+	void SelectInstance(UWorstCasePerformanceReportInstance* ClickedInstance);
+	void SelectInstanceFinishedUpdate(class UCustomizableObjectInstance* ClickedInstance);
+
+	TSharedPtr<class SHeaderRow> TimeWorstCasesViewHeader;
+	TSharedPtr<class SVerticalBox> TextureWorstCasesView;
+	TSharedPtr<class SVerticalBox> GeometryWorstCasesView;
+	TSharedPtr<class SVerticalBox> TimeWorstCasesView;
+	TSharedPtr<class SScrollBox> TimeWorstCasesDetailedView;
+	TSharedPtr<class STextBlock> StatusText;
+	TSharedPtr<class STextBlock> SelectedTextureAnalyzerText;
+	TSharedPtr<class SButton> PerformanceReportResetOptionsEntry;
+	TSharedPtr<class SObjectPropertyEntryBox> InstanceThumbnail;
+	TSharedPtr<class FAssetThumbnailPool> AssetThumbnailPool;
+	TSharedPtr<class FAssetThumbnail> AssetThumbnail;
+	TSharedPtr<class SButton> PerformanceReportRefreshSelectedInstance;
+
+	/** UObject class to be able to use the UCustomizableObjectInstance::UpdatedDelegate update callback */
+	UPROPERTY()
+	UPerformanceReportHelper* HelperCallback;
+
+	friend class UPerformanceReportHelper;
+	friend class SPerformanceReportWorstCaseRow;
+};
+
+
+class SPerformanceReportWorstCaseRow : public SMutableExpandableTableRow<UWorstCasePerformanceReportInstance*>, public FGCObject
+{
+public:
+	SLATE_BEGIN_ARGS(SPerformanceReportWorstCaseRow) {}
+		SLATE_ARGUMENT(UWorstCasePerformanceReportInstance*, ReportInstance)
+		SLATE_ARGUMENT(SCustomizableObjecEditorPerformanceReport*, PerformanceReport)
+		SLATE_ARGUMENT(const FTableRowStyle*, Style)
+	SLATE_END_ARGS()
+
+	void Construct(const FArguments& InArgs, const TSharedRef<STableViewBase>& InOwnerTableView);
+
+	// FSerializableObject interface
+	void AddReferencedObjects(FReferenceCollector& Collector) override;
+	virtual FString GetReferencerName() const override
+	{
+		return TEXT("SPerformanceReportWorstCaseRow");
+	}
+
+private:
+	// Report instance we represent
+	UWorstCasePerformanceReportInstance* ReportInstance;
+
+	// Performance report we are from. Necessary to generate the parameter worst time case widgets when the row is a state worst case
+	SCustomizableObjecEditorPerformanceReport* PerformanceReport;
+
+	// Local copy of additional widget data
+	TArray<UWorstCasePerformanceReportInstance*> SubWidgetData;
+
+	virtual TSharedRef<SWidget> GenerateWidgetForColumn(const FName& ColumnName) override;
+	virtual TSharedPtr<SWidget> GenerateAdditionalWidgetForRow() override;
+	virtual EVisibility GetAdditionalWidgetDefaultVisibility() const override;
+	virtual void SetAdditionalWidgetVisibility(const EVisibility Visibility) override;
+	const FSlateBrush* GetSelectedRowHighlightColor() const;
+
+	TSharedRef<ITableRow> GenerateStyledRow(UWorstCasePerformanceReportInstance* ReportedInstance, const TSharedRef<STableViewBase>& OwnerTable);
+};
diff --git a/Engine/Plugins/Experimental/Mutable/Source/CustomizableObjectEditor/Private/SCustomizableObjectEditorTagExplorer.cpp b/Engine/Plugins/Experimental/Mutable/Source/CustomizableObjectEditor/Private/SCustomizableObjectEditorTagExplorer.cpp
new file mode 100644
--- /dev/null
+++ b/Engine/Plugins/Experimental/Mutable/Source/CustomizableObjectEditor/Private/SCustomizableObjectEditorTagExplorer.cpp
@@ -0,0 +1,440 @@
+// Copyright Epic Games, Inc. All Rights Reserved.
+
+#include "SCustomizableObjectEditorTagExplorer.h"
+
+#include "CustomizableObjectGraph.h"
+#include "GraphTraversal.h"
+
+#include "UnrealEditorPortabilityHelpers.h"
+#include "HAL/PlatformApplicationMisc.h"
+
+#include "Nodes/CustomizableObjectNode.h"
+#include "Nodes/CustomizableObjectNodeObject.h"
+#include "Nodes/CustomizableObjectNodeMaterial.h"
+#include "Nodes/CustomizableObjectNodeMeshClipMorph.h"
+#include "Nodes/CustomizableObjectNodeMeshClipWithMesh.h"
+#include "Nodes/CustomizableObjectNodeMaterialVariation.h"
+#include "Nodes/CustomizableObjectNodeExtendMaterial.h"
+
+#include "Widgets/Input/SButton.h"
+
+#define LOCTEXT_NAMESPACE "SCustomizableObjectEditorTagExplorer"
+
+void SCustomizableObjectEditorTagExplorer::Construct(const FArguments & InArgs)
+{
+	CustomizableObjectEditorPtr = InArgs._CustomizableObjectEditor;
+
+	ChildSlot
+	[
+		SNew(SVerticalBox)
+
+		+ SVerticalBox::Slot()
+		.AutoHeight()
+		.Padding(5.0f)
+		[
+			SNew(SHorizontalBox)
+
+			+ SHorizontalBox::Slot()
+			.AutoWidth()
+			[
+				SNew(SButton).Text(LOCTEXT("FindTags", "Find Tags")).OnClicked(this, &SCustomizableObjectEditorTagExplorer::GetCustomizableObjectTags)
+				.ToolTipText(LOCTEXT("FindTagToolTip", "Find all tags related to this Customizable Object."))
+			]
+		]
+
+		+ SVerticalBox::Slot()
+		.AutoHeight()
+		.Padding(5.0f)
+		[
+			SNew(SHorizontalBox)
+			+ SHorizontalBox::Slot()
+			.AutoWidth()
+			.Padding(5.0f, 2.0f, 0, 0)
+			[
+				SNew(STextBlock).Text(FText::FromString("Selected Tag:"))
+			]
+	
+			+ SHorizontalBox::Slot()
+			.Padding(5.0f, 0, 0, 0)
+			[
+				SAssignNew(TagComboBox, SComboBox<TSharedPtr<FString>>)
+				.OptionsSource(&ComboboxTags)
+				.OnGenerateWidget(this, &SCustomizableObjectEditorTagExplorer::MakeComboButtonItemWidget)
+				.OnSelectionChanged(this, &SCustomizableObjectEditorTagExplorer::OnComboBoxSelectionChanged)
+				.InitiallySelectedItem(TagComboBoxItem)
+				[
+					SNew(STextBlock)
+					.Text(this, &SCustomizableObjectEditorTagExplorer::GetCurrentItemLabel)
+				]
+			]
+	
+			+ SHorizontalBox::Slot()
+			.AutoWidth()
+			[
+				SNew(SButton).Text(LOCTEXT("CopyToClipboard", "Copy to Clipboard")).OnClicked(this, &SCustomizableObjectEditorTagExplorer::CopyTagToClipboard)
+				.ToolTipText(LOCTEXT("CopyToClipboardToolTip", "Copy tag name to clipboard."))
+			]
+		]
+
+		+ SVerticalBox::Slot()
+		.AutoHeight()
+		.Padding(5.0f)
+		[
+			SNew(STextBlock).Text(FText::FromString("Used in:"))
+		]
+
+		+ SVerticalBox::Slot()
+		.AutoHeight()
+		.Padding(2.0f)
+		[
+			SNew(SHorizontalBox)
+			
+			+ SHorizontalBox::Slot()
+			[
+				SAssignNew(ColumnMat, SListView<UCustomizableObjectNode*>)
+				.ListItemsSource(&MaterialNodes)
+				.OnGenerateRow(this, &SCustomizableObjectEditorTagExplorer::OnGenerateTableRow)
+				.OnSelectionChanged(this, &SCustomizableObjectEditorTagExplorer::OnTagTableSelectionChanged)
+				.SelectionMode(ESelectionMode::Single)
+				.HeaderRow
+				(
+					SNew(SHeaderRow)
+					+ SHeaderRow::Column("MaterialNode")
+					.DefaultLabel(LOCTEXT("MaterialNode", "Material Node"))
+					.FillWidth(0.3f)
+					.HAlignHeader(EHorizontalAlignment::HAlign_Center)
+					.HAlignCell(EHorizontalAlignment::HAlign_Center)
+				)
+			]
+			
+			+ SHorizontalBox::Slot()
+			[
+				SAssignNew(ColumnClipMesh, SListView<UCustomizableObjectNode*>)
+				.ListItemsSource(&ClipMeshNodes)
+				.OnGenerateRow(this, &SCustomizableObjectEditorTagExplorer::OnGenerateTableRow)
+				.OnSelectionChanged(this, &SCustomizableObjectEditorTagExplorer::OnTagTableSelectionChanged)
+				.HeaderRow
+				(
+					SNew(SHeaderRow)
+					+ SHeaderRow::Column("ClipMeshNode")
+					.DefaultLabel(LOCTEXT("ClipMeshNode", "Clip Mesh Node"))
+					.FillWidth(0.3f)
+					.HAlignHeader(EHorizontalAlignment::HAlign_Center)
+					.HAlignCell(EHorizontalAlignment::HAlign_Center)
+				)
+			]
+
+			+ SHorizontalBox::Slot()
+			[
+				SAssignNew(ColumnClipMorph, SListView<UCustomizableObjectNode*>)
+				.ListItemsSource(&ClipMorphNodes)
+				.OnGenerateRow(this, &SCustomizableObjectEditorTagExplorer::OnGenerateTableRow)
+				.OnSelectionChanged(this, &SCustomizableObjectEditorTagExplorer::OnTagTableSelectionChanged)
+				.HeaderRow
+				(
+					SNew(SHeaderRow)
+					+ SHeaderRow::Column("ClipMorphNode")
+					.DefaultLabel(LOCTEXT("ClipMorphNode", "Clip Morph Node"))
+					.FillWidth(0.3f)
+					.HAlignHeader(EHorizontalAlignment::HAlign_Center)
+					.HAlignCell(EHorizontalAlignment::HAlign_Center)
+				)
+			]
+	
+			+ SHorizontalBox::Slot()
+			[
+				SAssignNew(ColumnVar, SListView<UCustomizableObjectNode*>)
+				.ListItemsSource(&VariationNodes)
+				.OnGenerateRow(this, &SCustomizableObjectEditorTagExplorer::OnGenerateTableRow)
+				.OnSelectionChanged(this, &SCustomizableObjectEditorTagExplorer::OnTagTableSelectionChanged)
+				.HeaderRow
+				(
+					SNew(SHeaderRow)
+					+ SHeaderRow::Column("VariationNode")
+					.DefaultLabel(LOCTEXT("VariationNode", "Variation Node"))
+					.FillWidth(0.3f)
+					.HAlignHeader(EHorizontalAlignment::HAlign_Center)
+					.HAlignCell(EHorizontalAlignment::HAlign_Center)
+				)
+			]
+
+			+ SHorizontalBox::Slot()
+				[
+					SAssignNew(ColumnExtend, SListView<UCustomizableObjectNode*>)
+					.ListItemsSource(&ExtendNodes)
+					.OnGenerateRow(this, &SCustomizableObjectEditorTagExplorer::OnGenerateTableRow)
+					.OnSelectionChanged(this, &SCustomizableObjectEditorTagExplorer::OnTagTableSelectionChanged)
+					.HeaderRow
+					(
+						SNew(SHeaderRow)
+						+ SHeaderRow::Column("ExtendNode")
+						.DefaultLabel(LOCTEXT("ExtendNode", "Extend Node"))
+						.FillWidth(0.3f)
+						.HAlignHeader(EHorizontalAlignment::HAlign_Center)
+						.HAlignCell(EHorizontalAlignment::HAlign_Center)
+					)
+				]
+		]
+	];
+}
+
+FReply SCustomizableObjectEditorTagExplorer::GetCustomizableObjectTags()
+{
+	bool bMultipleBaseObjectsFound = false;
+	NodeTags.Empty();
+
+	UCustomizableObject* CustomizableObject = CustomizableObjectEditorPtr->GetCustomizableObject();
+
+	UCustomizableObjectNodeObject* Root = GetRootNode(CustomizableObject, bMultipleBaseObjectsFound);
+	TArray<UCustomizableObject*> CustomizableObjectFamily;
+	TArray<FString> Tags;
+
+	if (Root && !bMultipleBaseObjectsFound)
+	{
+		UCustomizableObject* AbsoluteCO = Root->ParentObject ? CustomizableObjectEditorPtr->GetAbsoluteCOParent(Root) : CustomizableObject;
+
+		if (AbsoluteCO)
+		{
+			CustomizableObjectFamily.Add(AbsoluteCO);
+			// Stores external customizable objects
+			CustomizableObjectEditorPtr->GetExternalChildObjects(AbsoluteCO, CustomizableObjectFamily, true);
+			ComboboxTags.Empty();
+
+			for (UCustomizableObject* CustObject : CustomizableObjectFamily)
+			{
+				FillTagInformation(CustObject, Tags);
+			}
+
+			for (FString tag : Tags)
+			{
+				ComboboxTags.Add(MakeShareable(new FString(tag)));
+			}
+
+			if (TagComboBox.IsValid())
+			{
+				TagComboBox.ToSharedRef()->RefreshOptions();
+			}
+		}
+	}
+	else
+	{
+		if (bMultipleBaseObjectsFound)
+		{
+			UE_LOG(LogTemp, Error, TEXT("Multiple base object nodes found. Can't check tags."));
+		}
+
+		if (!Root)
+		{
+			UE_LOG(LogTemp, Error, TEXT("No base object node found. Can't check tags."));
+		}
+	}
+
+	TagComboBox.ToSharedRef()->RefreshOptions();
+
+	return FReply::Handled();
+}
+
+void SCustomizableObjectEditorTagExplorer::FillTagInformation(UCustomizableObject * Object, TArray<FString>& Tags)
+{
+	if (Object->Source)
+	{
+		for (TObjectPtr<UEdGraphNode> Node : Object->Source->Nodes)
+		{
+			if (UCustomizableObjectNodeMaterial* TypedNodeMat = Cast<UCustomizableObjectNodeMaterial>(Node))
+			{
+				for (int i = 0; i < TypedNodeMat->Tags.Num(); ++i)
+				{
+					NodeTags.Add(TypedNodeMat->Tags[i], TypedNodeMat);
+					if (Tags.Find(TypedNodeMat->Tags[i]) == INDEX_NONE)
+					{
+						Tags.Add(TypedNodeMat->Tags[i]);
+					}
+				}
+			}
+
+			if (UCustomizableObjectNodeMaterialVariation* TypedNodeVariations = Cast<UCustomizableObjectNodeMaterialVariation>(Node))
+			{
+				for (int i = 0; i < TypedNodeVariations->Variations.Num(); ++i)
+				{
+					NodeTags.Add(TypedNodeVariations->Variations[i].Tag, TypedNodeVariations);
+					if (Tags.Find(TypedNodeVariations->Variations[i].Tag) == INDEX_NONE)
+					{
+						Tags.Add(TypedNodeVariations->Variations[i].Tag);
+					}
+				}
+			}
+
+			if (UCustomizableObjectNodeMeshClipMorph* TypedNodeClipMorph = Cast<UCustomizableObjectNodeMeshClipMorph>(Node))
+			{
+				for (int i = 0; i < TypedNodeClipMorph->Tags.Num(); ++i)
+				{
+					NodeTags.Add(TypedNodeClipMorph->Tags[i], TypedNodeClipMorph);
+					if (Tags.Find(TypedNodeClipMorph->Tags[i]) == INDEX_NONE)
+					{
+						Tags.Add(TypedNodeClipMorph->Tags[i]);
+					}
+				}
+			}
+
+			if (UCustomizableObjectNodeMeshClipWithMesh* TypedNodeClipMesh = Cast<UCustomizableObjectNodeMeshClipWithMesh>(Node))
+			{
+				for (int i = 0; i < TypedNodeClipMesh->Tags.Num(); ++i)
+				{
+					NodeTags.Add(TypedNodeClipMesh->Tags[i], TypedNodeClipMesh);
+					if (Tags.Find(TypedNodeClipMesh->Tags[i]) == INDEX_NONE)
+					{
+						Tags.Add(TypedNodeClipMesh->Tags[i]);
+					}
+				}
+			}
+
+			if (UCustomizableObjectNodeExtendMaterial* TypedNodeExtend = Cast<UCustomizableObjectNodeExtendMaterial>(Node))
+			{
+				for (int i = 0; i < TypedNodeExtend->Tags.Num(); ++i)
+				{
+					NodeTags.Add(TypedNodeExtend->Tags[i], TypedNodeExtend);
+					if (Tags.Find(TypedNodeExtend->Tags[i]) == INDEX_NONE)
+					{
+						Tags.Add(TypedNodeExtend->Tags[i]);
+					}
+				}
+			}
+		}
+	}
+}
+
+TSharedRef<SWidget> SCustomizableObjectEditorTagExplorer::MakeComboButtonItemWidget(TSharedPtr<FString> StringItem)
+{
+	return SNew(STextBlock).Text(FText::FromString(*StringItem));
+}
+
+FText SCustomizableObjectEditorTagExplorer::GetCurrentItemLabel() const
+{
+	if (TagComboBoxItem.IsValid())
+	{
+		return FText::FromString(*TagComboBoxItem);
+	}
+
+	return LOCTEXT("InvalidComboEntryText", "None");
+}
+
+FReply SCustomizableObjectEditorTagExplorer::CopyTagToClipboard()
+{
+	if (TagComboBoxItem.IsValid())
+	{
+		FString TagName = *TagComboBoxItem.Get();
+		FPlatformApplicationMisc::ClipboardCopy(*TagName);
+	}
+
+	return FReply::Handled();
+}
+
+void SCustomizableObjectEditorTagExplorer::OnComboBoxSelectionChanged(TSharedPtr<FString> NewValue, ESelectInfo::Type)
+{
+	TagComboBoxItem = NewValue;
+
+	for (TSharedPtr<FString> tag : ComboboxTags)
+	{
+		if (tag == NewValue)
+		{
+			TArray<UCustomizableObjectNode*> auxNodes;
+			NodeTags.MultiFind(*tag.Get(), auxNodes, false);
+
+			MaterialNodes.Empty();
+			VariationNodes.Empty();
+			ClipMeshNodes.Empty();
+			ClipMorphNodes.Empty();
+			ExtendNodes.Empty();
+
+			for (UCustomizableObjectNode* node : auxNodes)
+			{
+				if (Cast<UCustomizableObjectNodeMaterial>(node) && MaterialNodes.Find(node) == INDEX_NONE)
+				{
+					MaterialNodes.Add(node);
+				}
+
+				if (Cast<UCustomizableObjectNodeMaterialVariation>(node) && VariationNodes.Find(node) == INDEX_NONE)
+				{
+					VariationNodes.Add(node);
+				}
+
+				if (Cast<UCustomizableObjectNodeMeshClipMorph>(node) && ClipMorphNodes.Find(node) == INDEX_NONE)
+				{
+					ClipMorphNodes.Add(node);
+				}
+
+				if (Cast<UCustomizableObjectNodeMeshClipWithMesh>(node) && ClipMeshNodes.Find(node) == INDEX_NONE)
+				{
+					ClipMeshNodes.Add(node);
+				}
+				
+				if (Cast<UCustomizableObjectNodeExtendMaterial>(node) && ExtendNodes.Find(node) == INDEX_NONE)
+				{
+					ExtendNodes.Add(node);
+				}
+			}
+		}
+	}
+
+	if (ColumnMat.IsValid() && ColumnClipMesh.IsValid() && ColumnClipMorph.IsValid() && ColumnVar.IsValid() && ColumnExtend.IsValid())
+	{
+		ColumnMat.ToSharedRef()->RequestListRefresh();
+		ColumnClipMesh.ToSharedRef()->RequestListRefresh();
+		ColumnClipMorph.ToSharedRef()->RequestListRefresh();
+		ColumnVar.ToSharedRef()->RequestListRefresh();
+		ColumnExtend.ToSharedRef()->RequestListRefresh();
+	}
+}
+
+TSharedRef<ITableRow> SCustomizableObjectEditorTagExplorer::OnGenerateTableRow(UCustomizableObjectNode * Node, const TSharedRef<STableViewBase>& OwnerTable)
+{
+	bool multiRoot = false;
+	UCustomizableObject* NodeCO = Cast<UCustomizableObject>(Node->GetGraph()->GetOuter());
+	FString ObjectNode = Node->GetNodeTitle(ENodeTitleType::ListView).ToString();
+
+	if (NodeCO)
+	{
+		ObjectNode = GetRootNode(NodeCO, multiRoot)->GetNodeTitle(ENodeTitleType::MenuTitle).ToString();
+		int Indx = ObjectNode.Find("\n");
+		ObjectNode.RemoveAt(Indx, ObjectNode.Len() - Indx);
+	}
+
+
+
+	return SNew(STableRow<UCustomizableObjectNode*>, OwnerTable)
+		[
+			SNew(SHorizontalBox)
+			+ SHorizontalBox::Slot()
+		.HAlign(EHorizontalAlignment::HAlign_Center)
+		[
+			SNew(STextBlock)
+			.Text(FText::FromString(ObjectNode))
+		]
+		];
+}
+
+void SCustomizableObjectEditorTagExplorer::OnTagTableSelectionChanged(UCustomizableObjectNode* Entry, ESelectInfo::Type SelectInfo) const
+{
+	if (Entry)
+	{
+		UObject* Object = Entry->GetCustomizableObjectGraph()->GetOuter();
+
+		// Make sure the editor exists for this asset
+		GEditor->GetEditorSubsystem<UAssetEditorSubsystem>()->OpenEditorForAsset(Object);
+
+		// Find it
+		TSharedPtr<ICustomizableObjectEditor> Editor = Entry->GetGraphEditor();
+		Editor->SelectNode(Entry);
+
+		if (ColumnMat.IsValid() && ColumnClipMesh.IsValid() && ColumnClipMorph.IsValid() && ColumnVar.IsValid() && ColumnExtend.IsValid())
+		{
+			ColumnMat.ToSharedRef()->ClearSelection();
+			ColumnClipMesh.ToSharedRef()->ClearSelection();
+			ColumnClipMorph.ToSharedRef()->ClearSelection();
+			ColumnVar.ToSharedRef()->ClearSelection();
+			ColumnExtend.ToSharedRef()->ClearSelection();
+		}
+	}
+}
+
+#undef LOCTEXT_NAMESPACE 
\ No newline at end of file
diff --git a/Engine/Plugins/Experimental/Mutable/Source/CustomizableObjectEditor/Private/SCustomizableObjectEditorTagExplorer.h b/Engine/Plugins/Experimental/Mutable/Source/CustomizableObjectEditor/Private/SCustomizableObjectEditorTagExplorer.h
new file mode 100644
--- /dev/null
+++ b/Engine/Plugins/Experimental/Mutable/Source/CustomizableObjectEditor/Private/SCustomizableObjectEditorTagExplorer.h
@@ -0,0 +1,75 @@
+// Copyright Epic Games, Inc. All Rights Reserved.
+
+#pragma once
+
+#include "Widgets/SCompoundWidget.h"
+#include "CustomizableObjectEditor.h"
+
+class UCustomizableObjectNode;
+
+class SCustomizableObjectEditorTagExplorer : public SCompoundWidget
+{
+public:
+
+	SLATE_BEGIN_ARGS(SCustomizableObjectEditorTagExplorer){}
+		SLATE_ARGUMENT(FCustomizableObjectEditor*, CustomizableObjectEditor)
+	SLATE_END_ARGS()
+
+	void Construct(const FArguments& InArgs);
+
+private:
+
+	/** Button Callback to fill the tag list */
+	FReply GetCustomizableObjectTags();
+
+	/** Fills a list with all the tags found in the nodes of a graph */
+	void FillTagInformation(UCustomizableObject* Object, TArray<FString>& Tags);
+
+	/** Generates the combobox with all the tags  */
+	TSharedRef<SWidget> MakeComboButtonItemWidget(TSharedPtr<FString> StringItem);
+
+	/** Generates the text of the tags combobox */
+	FText GetCurrentItemLabel() const;
+
+	/** Copies the tag name to the clipboard */
+	FReply CopyTagToClipboard();
+
+	/** OnSelectionChanged callback of the tags combobox */
+	void OnComboBoxSelectionChanged(TSharedPtr<FString> NewValue, ESelectInfo::Type);
+
+	/** Tags table callbacks*/
+	TSharedRef<ITableRow> OnGenerateTableRow(UCustomizableObjectNode* Node, const TSharedRef<STableViewBase>& OwnerTable);
+	void OnTagTableSelectionChanged(UCustomizableObjectNode* Entry, ESelectInfo::Type SelectInfo) const;
+
+private:
+
+	/** Pointer back to the editor tool that owns us */
+	FCustomizableObjectEditor* CustomizableObjectEditorPtr;
+
+	/** Combobox for the Customizable Object Tags */
+	TSharedPtr<SComboBox<TSharedPtr<FString>>> TagComboBox;
+
+	/** Combobox string container */
+	TArray<TSharedPtr<FString>> ComboboxTags;
+
+	/** Combobox Selection */
+	TSharedPtr<FString> TagComboBoxItem;
+
+	/** List views for the nodes that contain the same tag */
+	TSharedPtr<SListView<UCustomizableObjectNode*>> ColumnMat;
+	TSharedPtr<SListView<UCustomizableObjectNode*>> ColumnVar;
+	TSharedPtr<SListView<UCustomizableObjectNode*>> ColumnClipMesh;
+	TSharedPtr<SListView<UCustomizableObjectNode*>> ColumnClipMorph;
+	TSharedPtr<SListView<UCustomizableObjectNode*>> ColumnExtend;
+
+	/** Arrays for each type of node that contains a tag */
+	TArray<UCustomizableObjectNode*> MaterialNodes;
+	TArray<UCustomizableObjectNode*> VariationNodes;
+	TArray<UCustomizableObjectNode*> ClipMeshNodes;
+	TArray<UCustomizableObjectNode*> ClipMorphNodes;
+	TArray<UCustomizableObjectNode*> ExtendNodes;
+
+	/** Multimap to relatre the nodes with the tags */
+	TMultiMap<FString, UCustomizableObjectNode*> NodeTags;
+
+};
\ No newline at end of file
diff --git a/Engine/Plugins/Experimental/Mutable/Source/CustomizableObjectEditor/Private/SCustomizableObjectEditorTextureAnalyzer.cpp b/Engine/Plugins/Experimental/Mutable/Source/CustomizableObjectEditor/Private/SCustomizableObjectEditorTextureAnalyzer.cpp
new file mode 100644
--- /dev/null
+++ b/Engine/Plugins/Experimental/Mutable/Source/CustomizableObjectEditor/Private/SCustomizableObjectEditorTextureAnalyzer.cpp
@@ -0,0 +1,234 @@
+// Copyright Epic Games, Inc. All Rights Reserved.
+
+#include "SCustomizableObjectEditorTextureAnalyzer.h"
+
+#include "IPropertyTable.h"
+#include "IPropertyTableColumn.h"
+#include "IPropertyTableCell.h"
+#include "IPropertyTableRow.h"
+
+#include "Engine/Texture.h"
+#include "Modules/ModuleManager.h"
+#include "Materials/MaterialInstance.h"
+#include "UnrealEditorPortabilityHelpers.h"
+#include "Widgets/Input/SButton.h"
+
+#define LOCTEXT_NAMESPACE "SCustomizableObjecEditorTextureAnalyzer"
+
+
+
+void SCustomizableObjecEditorTextureAnalyzer::Construct(const FArguments& InArgs)
+{
+	CustomizableObjectEditor = InArgs._CustomizableObjectEditor;
+	CustomizableObjectInstanceEditor = InArgs._CustomizableObjectInstanceEditor;
+
+	TotalSizeTextures = SNew(STextBlock);
+
+	ChildSlot
+	[
+		SNew(SVerticalBox)
+
+		+ SVerticalBox::Slot()
+		.HAlign(HAlign_Left)
+		.AutoHeight()
+		.Padding(1.0f, 5.0f, 0.0f, 10.0f)
+		[
+			SNew(SHorizontalBox)
+
+			+ SHorizontalBox::Slot()
+			.AutoWidth()
+			.Padding(1.0f, 5.0f, 0.0f, 0.0f)
+			[
+				SNew(SButton)
+				.Text(LOCTEXT("RefreshTable", "Refresh Table"))
+				.OnClicked(this, &SCustomizableObjecEditorTextureAnalyzer::RefreshTextureAnalyzerTable)
+				.Visibility(CustomizableObjectEditor || CustomizableObjectInstanceEditor ? EVisibility::Visible : EVisibility::Collapsed)
+			]
+
+			+ SHorizontalBox::Slot()
+			.AutoWidth()
+			.Padding(10.0f, 10.0f, 0.0f, 0.0f)
+			[
+				TotalSizeTextures.ToSharedRef()
+			]
+		]
+
+		+ SVerticalBox::Slot()
+		[
+			BuildTextureAnalyzerTable()
+		]
+	];
+}
+
+
+void SCustomizableObjecEditorTextureAnalyzer::AddReferencedObjects(FReferenceCollector& Collector)
+{
+	for (UObject* i : TabTextures)
+	{
+		Collector.AddReferencedObject(i);
+	}
+}
+
+void SCustomizableObjecEditorTextureAnalyzer::Tick(const FGeometry& AllottedGeometry, const double InCurrentTime, const float InDeltaTime)
+{
+	if (TextureAnalyzerTable.IsValid())
+	{
+		if (TextureAnalyzerTable->GetLastClickedCell().IsValid() && TextureAnalyzerTable->GetLastClickedCell()->InEditMode())
+		{
+			OnTextureTableSelectionChanged(TextureAnalyzerTable->GetLastClickedCell());
+		}
+	}
+
+	SCompoundWidget::Tick(AllottedGeometry, InCurrentTime, InDeltaTime);
+}
+
+
+TSharedRef<SWidget> SCustomizableObjecEditorTextureAnalyzer::BuildTextureAnalyzerTable()
+{
+	FPropertyEditorModule& PropertyEditorModule = FModuleManager::LoadModuleChecked<FPropertyEditorModule>("PropertyEditor");
+
+	TextureAnalyzerTable = PropertyEditorModule.CreatePropertyTable();
+	TextureAnalyzerTable->SetIsUserAllowedToChangeRoot(false);
+	TextureAnalyzerTable->SetSelectionMode(ESelectionMode::Single);
+	TextureAnalyzerTable->SetSelectionUnit(EPropertyTableSelectionUnit::Cell);
+	TextureAnalyzerTable->SetShowObjectName(false);
+
+	if (TextureAnalyzerTable.IsValid())
+	{
+		FillTextureAnalyzerTable();
+	}
+
+	return PropertyEditorModule.CreatePropertyTableWidget(TextureAnalyzerTable.ToSharedRef());
+}
+
+
+void SCustomizableObjecEditorTextureAnalyzer::FillTextureAnalyzerTable(UCustomizableObjectInstance* PreviewInstance)
+{
+	TabTextures.Empty();
+
+	uint32 TotalSize = 0;
+
+	if(!PreviewInstance)
+	{
+		if (CustomizableObjectEditor)
+		{
+			PreviewInstance = CustomizableObjectEditor->GetPreviewInstance();
+		}
+
+		if (CustomizableObjectInstanceEditor)
+		{
+			PreviewInstance = CustomizableObjectInstanceEditor->GetPreviewInstance();
+		}
+	}
+
+	if (PreviewInstance && PreviewInstance->HasAnySkeletalMesh())
+	{
+		for (int32 ComponentIndex = 0; ComponentIndex < PreviewInstance->SkeletalMeshes.Num(); ++ComponentIndex)
+		{
+			if (PreviewInstance->GetSkeletalMesh(ComponentIndex))
+			{
+				for (int i = 0; i < PreviewInstance->GetSkeletalMesh(ComponentIndex)->GetMaterials().Num(); ++i)
+				{
+					UMaterialInterface* MaterialInterface = PreviewInstance->GetSkeletalMesh(ComponentIndex)->GetMaterials()[i].MaterialInterface;
+					UMaterialInstance* Material = Cast<UMaterialInstance>(MaterialInterface);
+					FString MaterialPath = MaterialInterface->GetPathName();
+
+					// We just want the Transient materials
+					if (Material && MaterialPath.Contains("Transient"))
+					{
+						for (int32 j = 0; j < Material->TextureParameterValues.Num(); ++j)
+						{
+							UTexture* Texture = Material->TextureParameterValues[j].ParameterValue;
+							if (!Texture) continue;
+
+							FString TextPath = Texture->GetPathName();
+
+							// We just want the Transient textures
+							if (TextPath.Contains("Transient"))
+							{
+								//New Object for the table
+								UCustomizableObjectEditorTextureStats* Entry = NewObject<UCustomizableObjectEditorTextureStats>();
+								TabTextures.Add(Entry);
+
+								//Texture Info
+								Entry->Texture = Texture;
+								Entry->TextureName = Texture->GetName();
+								Entry->TextureParameterName = Material->TextureParameterValues[j].ParameterInfo.Name.ToString();
+								Entry->ResolutionX = Cast<UTexture2D>(Texture)->GetSizeX();
+								Entry->ResolutionY = Cast<UTexture2D>(Texture)->GetSizeY();
+								Entry->LODBias = Texture->GetCachedLODBias();
+								uint32 textureSize = Texture->CalcTextureMemorySizeEnum(TMC_AllMips);
+								Entry->Size = FString::Printf(TEXT("%.2f"), textureSize / 1024.0);
+								Entry->Format = Cast<UTexture2D>(Texture)->GetPixelFormat();
+								Entry->IsStreamed = Texture->IsCurrentlyVirtualTextured() ? "Virtual Streamed" : (!Texture->IsStreamable() ? "Not Streamed" : "Streamed");
+								Entry->LODGroup = Texture->LODGroup;
+
+								//Material Info
+								Entry->Material = MaterialInterface;
+								Entry->MaterialName = MaterialInterface->GetName();
+								Entry->ParentMaterial = Material->Parent;
+								Entry->MaterialParameterName = Material->Parent ? Material->Parent->GetName() : FString();
+								Entry->Component = ComponentIndex;
+
+								//Total size
+								TotalSize += textureSize;
+							}
+						}
+					}
+				}
+			}
+		}
+
+		TextureAnalyzerTable->SetObjects(TabTextures);
+	}
+
+	// Updating Texture Size Text
+	if (TotalSizeTextures.IsValid())
+	{
+		TotalSizeTextures->SetText(FText::FromString(FString::Printf(TEXT("Total Size (All LODs, only Mutable): %.2f Mb "), TotalSize / (1024.0 * 1024.0))));
+	}
+}
+
+
+void SCustomizableObjecEditorTextureAnalyzer::OnTextureTableSelectionChanged(TSharedPtr<class IPropertyTableCell> Cell)
+{
+	FString SelectedColumn = Cell->GetColumn()->GetDisplayName().ToString();
+
+	if (UCustomizableObjectEditorTextureStats* Current = Cast<UCustomizableObjectEditorTextureStats>(Cell->GetObject().Get()))
+	{
+		//Open Texture Editor Tab
+		if ((SelectedColumn == "Texture" || SelectedColumn == "Name") && Current->Texture)
+		{
+			GEditor->GetEditorSubsystem<UAssetEditorSubsystem>()->OpenEditorForAsset(Current->Texture);
+		}
+
+		//Open Material Editor Tab
+		if (SelectedColumn == "Material" && Current->Material)
+		{
+			GEditor->GetEditorSubsystem<UAssetEditorSubsystem>()->OpenEditorForAsset(Current->Material);
+		}
+
+		if (SelectedColumn == "Parent" && Current->ParentMaterial)
+		{
+			GEditor->GetEditorSubsystem<UAssetEditorSubsystem>()->OpenEditorForAsset(Current->ParentMaterial);
+		}
+	}
+
+	//Avoid edit mode (not needed)
+	Cell->ExitEditMode();
+}
+
+
+FReply SCustomizableObjecEditorTextureAnalyzer::RefreshTextureAnalyzerTable()
+{
+	return RefreshTextureAnalyzerTable(nullptr);
+}
+
+FReply SCustomizableObjecEditorTextureAnalyzer::RefreshTextureAnalyzerTable(UCustomizableObjectInstance* PreviewInstance)
+{
+	FillTextureAnalyzerTable(PreviewInstance);
+
+	return FReply::Handled();
+}
+
+#undef LOCTEXT_NAMESPACE 
\ No newline at end of file
diff --git a/Engine/Plugins/Experimental/Mutable/Source/CustomizableObjectEditor/Private/SCustomizableObjectEditorTransformViewportToolbar.cpp b/Engine/Plugins/Experimental/Mutable/Source/CustomizableObjectEditor/Private/SCustomizableObjectEditorTransformViewportToolbar.cpp
new file mode 100644
--- /dev/null
+++ b/Engine/Plugins/Experimental/Mutable/Source/CustomizableObjectEditor/Private/SCustomizableObjectEditorTransformViewportToolbar.cpp
@@ -0,0 +1,24 @@
+// Copyright Epic Games, Inc. All Rights Reserved.	
+
+#include "SCustomizableObjectEditorTransformViewportToolbar.h"
+#include "Framework/MultiBox/MultiBoxDefs.h"
+#include "Framework/MultiBox/MultiBoxBuilder.h"
+#include "Styling/AppStyle.h"
+#include "CustomizableObjectEditorViewportLODCommands.h"
+#include "SCustomizableObjectEditorViewport.h"
+#include "CustomizableObjectEditorViewportClient.h"
+
+#define LOCTEXT_NAMESPACE "SCustomizableObjectEditorTransformViewportToolbar"
+
+void SCustomizableObjectEditorTransformViewportToolbar::Construct( const FArguments& InArgs )
+{
+	Viewport = InArgs._Viewport;
+	ViewportTapBody = InArgs._ViewportTapBody;
+
+	const FCustomizableObjectEditorViewportLODCommands& ViewportLODMenuCommands = FCustomizableObjectEditorViewportLODCommands::Get();
+	UICommandList = MakeShareable(new FUICommandList);
+
+	SViewportToolBar::Construct(SViewportToolBar::FArguments());
+}
+
+#undef LOCTEXT_NAMESPACE 
