[mutable] Move Mutable plugin from the restricted folder to the experimental plugins folder.

diff --git a/Engine/Plugins/Experimental/Mutable/Source/CustomizableObjectEditor/Private/SCustomizableObjectEditorViewport.cpp b/Engine/Plugins/Experimental/Mutable/Source/CustomizableObjectEditor/Private/SCustomizableObjectEditorViewport.cpp
new file mode 100644
--- /dev/null
+++ b/Engine/Plugins/Experimental/Mutable/Source/CustomizableObjectEditor/Private/SCustomizableObjectEditorViewport.cpp
@@ -0,0 +1,1238 @@
+// Copyright Epic Games, Inc. All Rights Reserved.
+
+#include "SCustomizableObjectEditorViewport.h"
+#include "CustomizableObjectEditorActions.h"
+#include "SCustomizableObjectEditorViewportToolBar.h"
+#include "CustomizableObjectEditorViewportMenuCommands.h"
+#include "CustomizableObjectEditorViewportLODCommands.h"
+#include "CustomizableObjectEditorViewportClient.h"
+#include "SCustomizableObjectEditorTransformViewportToolbar.h"
+#include "CustomizableObjectEditorUtilities.h"
+#include "Nodes/CustomizableObjectNodeMeshClipMorph.h"
+#include "Nodes/CustomizableObjectNodeMeshClipWithMesh.h"
+
+#include "IPersonaPreviewScene.h"
+
+#include "MouseDeltaTracker.h"
+#include "PreviewScene.h"
+#include "Slate/SceneViewport.h"
+//#include "Widgets/Docking/SDockableTab.h"
+#include "Widgets/Input/STextComboBox.h"
+#include "Widgets/Input/SCheckBox.h"
+
+#include "CustomizableObjectEditor.h"
+#include "FbxMeshUtils.h"
+#include "BusyCursor.h"
+#include "MeshBuild.h"
+#include "ObjectTools.h"
+#include "Animation/PoseAsset.h"
+#include "DragAndDrop/AssetDragDropOp.h"
+#include "Viewports.h"
+#include "Framework/Notifications/NotificationManager.h"
+
+#include "UnrealPortabilityHelpers.h"
+#include "UnrealEditorPortabilityHelpers.h"
+#include "Editor/UnrealEdEngine.h"
+#include "UnrealEdGlobals.h"
+#include "Editor/EditorEngine.h"
+#include "SCustomizableObjectHighresScreenshot.h"
+#include "CustomizableObjectPreviewScene.h"
+
+#include "EditorViewportCommands.h"
+
+
+void SCustomizableObjectEditorViewport::Construct(const FArguments& InArgs, const FCustomizableObjectEditorViewportRequiredArgs& InRequiredArgs)
+{
+	//SkeletonTreePtr = InRequiredArgs.SkeletonTree;
+	PreviewScenePtr = InRequiredArgs.PreviewScene;
+	TabBodyPtr = InRequiredArgs.TabBody;
+	//AssetEditorToolkitPtr = InRequiredArgs.AssetEditorToolkit;
+	bShowStats = InArgs._ShowStats;
+
+	// Not necessary, since there is no undo on instance parameter change
+	//InRequiredArgs.OnPostUndo.Add(FSimpleDelegate::CreateSP(this, &SCustomizableObjectEditorViewport::OnUndoRedo));
+
+
+	SEditorViewport::Construct(
+		SEditorViewport::FArguments()
+		.IsEnabled(FSlateApplication::Get().GetNormalExecutionAttribute())
+		.AddMetaData<FTagMetaData>(TEXT("Persona.Viewport"))
+	);
+
+	Client->VisibilityDelegate.BindSP(this, &SCustomizableObjectEditorViewport::IsVisible);
+}
+
+TSharedRef<FEditorViewportClient> SCustomizableObjectEditorViewport::MakeEditorViewportClient()
+{
+	LevelViewportClient = MakeShareable(new FCustomizableObjectEditorViewportClient(TabBodyPtr.Pin()->CustomizableObjectEditorPtr, PreviewScenePtr.Pin().Get()));
+
+	LevelViewportClient->ViewportType = LVT_Perspective;
+	LevelViewportClient->bSetListenerPosition = false;
+	LevelViewportClient->SetViewLocation(EditorViewportDefs::DefaultPerspectiveViewLocation);
+	LevelViewportClient->SetViewRotation(EditorViewportDefs::DefaultPerspectiveViewRotation);
+
+	SceneViewport = MakeShareable(new FSceneViewport(LevelViewportClient.Get(), ViewportWidget));
+	//Viewport = MakeShareable(new FSceneViewport(LevelViewportClient.Get(), ViewportWidget));
+	//LevelViewportClient->Viewport = Viewport.Get();
+	//LevelViewportClient->Viewport->SetUserFocus(true);
+
+	// The viewport widget needs an interface so it knows what should render
+	//ViewportWidget->SetViewportInterface(Viewport.ToSharedRef());
+
+	return LevelViewportClient.ToSharedRef();
+}
+
+
+TSharedPtr<FSceneViewport>& SCustomizableObjectEditorViewport::GetSceneViewport()
+{
+	return SceneViewport;
+}
+
+
+TSharedPtr<SWidget> SCustomizableObjectEditorViewport::MakeViewportToolbar()
+{
+	return SNew(SCustomizableObjectEditorViewportToolBar, TabBodyPtr.Pin(), SharedThis(this))
+		.Cursor(EMouseCursor::Default);
+}
+
+void SCustomizableObjectEditorViewport::OnUndoRedo()
+{
+	LevelViewportClient->Invalidate();
+}
+
+
+//////////////////////////////////////////////////////////////////////////
+
+
+void SCustomizableObjectEditorViewportTabBody::Construct(const FArguments& InArgs)
+{
+	UICommandList = MakeShareable(new FUICommandList);
+
+	CustomizableObjectEditorPtr = InArgs._CustomizableObjectEditor;
+
+	CurrentViewMode = VMI_Lit;
+	LODSelection = 0;
+
+	FCustomizableObjectEditorViewportMenuCommands::Register();
+	//FAnimViewportShowCommands::Register();
+	FCustomizableObjectEditorViewportLODCommands::Register();
+	//FAnimViewportPlaybackCommands::Register();
+
+	FPreviewScene::ConstructionValues SceneConstructValues;
+	SceneConstructValues.bShouldSimulatePhysics = true;
+
+	PreviewScenePtr = MakeShareable(new FCustomizableObjectPreviewScene(SceneConstructValues));
+	//PreviewScenePtr->GetWorld()->WorldType = EWorldType::Editor; // Needed for custom depth pass for screenshot with no background
+
+	FCustomizableObjectEditorViewportRequiredArgs ViewportArgs(PreviewScenePtr.ToSharedRef(),SharedThis(this));
+	//(InSkeletonTree, InPreviewScene, SharedThis(this), InAssetEditorToolkit, InOnUndoRedo);
+
+	ViewportWidget = SNew(SCustomizableObjectEditorViewport, ViewportArgs);
+
+
+
+	this->ChildSlot
+		[
+			SNew(SVerticalBox)
+
+			+ SVerticalBox::Slot()
+			.Padding(2.0f)
+			.AutoHeight()
+			[
+				BuildToolBar()
+			]
+
+			+SVerticalBox::Slot()
+			.FillHeight(1)
+			[
+				ViewportWidget.ToSharedRef()
+			]
+		];
+
+	LevelViewportClient = StaticCastSharedPtr<FCustomizableObjectEditorViewportClient>(ViewportWidget->GetViewportClient());
+
+	//PreviewSkeletalMeshComponent = 0;
+
+	BindCommands();
+
+	AssetRegistryLoaded = false;
+}
+
+
+SCustomizableObjectEditorViewportTabBody::~SCustomizableObjectEditorViewportTabBody()
+{
+	if (LevelViewportClient.IsValid())
+	{
+		LevelViewportClient->Viewport = nullptr;
+	}
+
+	PreviewSkeletalMeshComponents.Reset();
+}
+
+
+void SCustomizableObjectEditorViewportTabBody::AddReferencedObjects( FReferenceCollector& Collector )
+{
+	for (UDebugSkelMeshComponent* PreviewSkeletalMeshComponent : PreviewSkeletalMeshComponents)
+	{
+		Collector.AddReferencedObject(PreviewSkeletalMeshComponent);
+	}
+}
+
+
+void SCustomizableObjectEditorViewportTabBody::SetClipMorphPlaneVisibility(bool bVisible, const FVector& Origin, const FVector& Normal, float MorphLength, const FBoxSphereBounds& Bounds, float Radius1, float Radius2, float RotationAngle)
+{
+	LevelViewportClient->SetClipMorphPlaneVisibility(bVisible, Origin, Normal, MorphLength, Bounds, Radius1, Radius2, RotationAngle);
+}
+
+
+void SCustomizableObjectEditorViewportTabBody::SetClipMorphPlaneVisibility(bool bVisible, UCustomizableObjectNodeMeshClipMorph* NodeMeshClipMorph)
+{
+	TransformWidgetVisibility.State.bClipMorphPlaneVisible = bVisible && bVisible != LevelViewportClient->IsNodeMeshClipMorphSelected();
+
+	EVisibility ChildVisibility = TransformWidgetVisibility.Value ? EVisibility::Visible : EVisibility::Hidden;
+	if (ViewportToolbarTransformWidget.IsValid())
+	{
+		ViewportToolbarTransformWidget.Pin()->SetVisibility(ChildVisibility);
+	}
+
+	LevelViewportClient->SetClipMorphPlaneVisibility(bVisible, NodeMeshClipMorph);
+}
+
+
+void SCustomizableObjectEditorViewportTabBody::SetClipMeshVisibility(bool bVisible, UStaticMesh* ClipMesh, UCustomizableObjectNodeMeshClipWithMesh* ClipMeshNode)
+{	
+	TransformWidgetVisibility.State.bClipMeshVisible = ClipMesh && ClipMeshNode && bVisible;
+
+	EVisibility ChildVisibility = TransformWidgetVisibility.Value ? EVisibility::Visible : EVisibility::Hidden;
+	if (ViewportToolbarTransformWidget.IsValid())
+	{
+		ViewportToolbarTransformWidget.Pin()->SetVisibility(ChildVisibility);
+	}
+	
+
+	LevelViewportClient->SetClipMeshVisibility(bVisible, ClipMesh, ClipMeshNode);
+}
+
+
+void SCustomizableObjectEditorViewportTabBody::SetProjectorVisibility(bool bVisible, FString ProjectorParameterName, FString ProjectorParameterNameWithIndex, int32 RangeIndex, const FCustomizableObjectProjector& Data, int32 ProjectorParameterIndex)
+{
+	if (LevelViewportClient->GetGizmoProxy().HasAssignedData &&
+		!LevelViewportClient->GetGizmoProxy().AssignedDataIsFromNode &&
+		LevelViewportClient->GetGizmoProxy().ProjectorParameterIndex == ProjectorParameterIndex &&
+		LevelViewportClient->GetGizmoProxy().ProjectorRangeIndex == RangeIndex)
+	{
+		return;
+	}
+	
+	TransformWidgetVisibility.State.bProjectorVisible = bVisible;
+
+	EVisibility ChildVisibility = TransformWidgetVisibility.Value ? EVisibility::Visible : EVisibility::Hidden;
+	if (ViewportToolbarTransformWidget.IsValid())
+	{
+		ViewportToolbarTransformWidget.Pin()->SetVisibility(ChildVisibility);
+	}
+
+	if ((LevelViewportClient->GetWidgetVisibility() != bVisible) || 
+		(LevelViewportClient->GetGizmoProxy().HasAssignedData &&
+		(
+			LevelViewportClient->GetGizmoProxy().AssignedDataIsFromNode ||
+			LevelViewportClient->GetGizmoProxy().ProjectorParameterIndex != ProjectorParameterIndex ||
+		    LevelViewportClient->GetGizmoProxy().ProjectorRangeIndex != RangeIndex)
+		))
+	{
+		LevelViewportClient->SetProjectorVisibility(bVisible, ProjectorParameterName, ProjectorParameterNameWithIndex, RangeIndex, Data, ProjectorParameterIndex);
+		
+	}
+}
+
+
+void SCustomizableObjectEditorViewportTabBody::SetProjectorType(bool bVisible, FString ProjectorParameterName, FString ProjectorParameterNameWithIndex, int32 RangeIndex, const FCustomizableObjectProjector& Data, int32 ProjectorParameterIndex)
+{
+	LevelViewportClient->SetProjectorType(bVisible, ProjectorParameterName, ProjectorParameterNameWithIndex, RangeIndex, Data, ProjectorParameterIndex);
+}
+
+
+void SCustomizableObjectEditorViewportTabBody::SetProjectorVisibility(bool bVisible, UCustomizableObjectNodeProjectorConstant* Projector)
+{
+	if (((Projector == nullptr) && !LevelViewportClient->GetGizmoProxy().HasAssignedData) ||
+	    (bVisible && (Projector != nullptr) && LevelViewportClient->GetGizmoProxy().HasAssignedData &&
+		(LevelViewportClient->GetGizmoProxy().DataOriginConstant == Projector)))
+	{
+		return;
+	}
+
+	TransformWidgetVisibility.State.bProjectorVisible = Projector && bVisible;
+	EVisibility ChildVisibility = TransformWidgetVisibility.Value ? EVisibility::Visible : EVisibility::Hidden;
+	if (ViewportToolbarTransformWidget.IsValid())
+	{
+		ViewportToolbarTransformWidget.Pin()->SetVisibility(ChildVisibility);
+	}
+
+	LevelViewportClient->SetProjectorVisibility(bVisible, Projector);
+}
+
+
+void SCustomizableObjectEditorViewportTabBody::SetProjectorParameterVisibility(bool bVisible, class UCustomizableObjectNodeProjectorParameter* ProjectorParameter)
+{
+	if (((ProjectorParameter == nullptr) && !LevelViewportClient->GetGizmoProxy().HasAssignedData) ||
+	    (bVisible && (ProjectorParameter != nullptr) && LevelViewportClient->GetGizmoProxy().HasAssignedData &&
+		(LevelViewportClient->GetGizmoProxy().DataOriginParameter == ProjectorParameter)))
+	{
+		return;
+	}
+
+	TransformWidgetVisibility.State.bProjectorVisible = ProjectorParameter && bVisible;
+
+	EVisibility ChildVisibility = TransformWidgetVisibility.Value ? EVisibility::Visible : EVisibility::Hidden;
+	if (ViewportToolbarTransformWidget.IsValid())
+	{
+		ViewportToolbarTransformWidget.Pin()->SetVisibility(ChildVisibility);
+	}
+
+	LevelViewportClient->SetProjectorParameterVisibility(bVisible, ProjectorParameter);
+}
+
+
+void SCustomizableObjectEditorViewportTabBody::ResetProjectorVisibility(bool OnlyNonNodeProjector)
+{
+	if(LevelViewportClient->GetIsManipulating())
+	{
+		return;
+	}
+
+	if (OnlyNonNodeProjector &&
+		LevelViewportClient->GetGizmoProxy().HasAssignedData &&
+		LevelViewportClient->GetGizmoProxy().AssignedDataIsFromNode)
+	{
+		return;
+	}
+
+	LevelViewportClient->ResetProjectorVisibility();
+
+	// Check if the projector's alpha is being modified.
+	// ResetProjectorVisibility is called every tick while modifying a projector's alpha,
+	// that makes the widget UI flash in and out.
+	bool bProjectorAlphaChange = false;
+	if (CustomizableObjectEditorPtr.IsValid())
+	{
+		UCustomizableObjectInstance* PreviewInstance = CustomizableObjectEditorPtr.Pin()->GetPreviewInstance();
+		if (PreviewInstance)
+		{
+			bProjectorAlphaChange = PreviewInstance->ProjectorAlphaChange;
+		}
+	}
+
+	if (ViewportToolbarTransformWidget.IsValid() && TransformWidgetVisibility.Value && !bProjectorAlphaChange)
+	{
+		ViewportToolbarTransformWidget.Pin()->SetVisibility(EVisibility::Hidden);
+		TransformWidgetVisibility.State.bProjectorVisible = false;
+	}
+}
+
+
+void SCustomizableObjectEditorViewportTabBody::Tick(const FGeometry& AllottedGeometry, const double InCurrentTime, const float InDeltaTime)
+{
+	SCompoundWidget::Tick(AllottedGeometry, InCurrentTime, InDeltaTime);
+
+	// Needed?
+	PreviewScenePtr->GetWorld()->Tick(LEVELTICK_All, InDeltaTime);
+
+	//if (PreviewSkeletalMeshComponent != nullptr)
+	//{
+	//	PreviewSkeletalMeshComponent->TickAnimation(InDeltaTime, false);
+	//	PreviewSkeletalMeshComponent->RefreshBoneTransforms();
+	//}
+
+	// Update the material list. Not ideal to do it every tick, but tracking changes on materials in the current instance is not easy right now.
+	if (PreviewSkeletalMeshComponents.Num()>0)
+	{
+		MaterialNames.Empty();
+
+		for (UDebugSkelMeshComponent* PreviewSkeletalMeshComponent : PreviewSkeletalMeshComponents)
+		{
+			const TArray<UMaterialInterface*> Materials = PreviewSkeletalMeshComponent->GetMaterials();
+			for (UMaterialInterface* m : Materials)
+			{
+				if (m)
+				{
+					const UMaterial* BaseMaterial = m->GetBaseMaterial();
+					MaterialNames.Add(MakeShareable(new FString(BaseMaterial->GetName())));
+				}
+			}
+		}
+	}
+	else
+	{
+		MaterialNames.Empty();
+	}
+
+}
+
+
+void SCustomizableObjectEditorViewportTabBody::SetPreviewComponents(TArray<UDebugSkelMeshComponent*>& InSkeletalMeshComponents)
+{
+	FTransform Transform = FTransform::Identity;
+
+	for (UDebugSkelMeshComponent* PreviewSkeletalMeshComponent : PreviewSkeletalMeshComponents)
+	{
+		if (PreviewSkeletalMeshComponent)
+		{
+			Transform = PreviewSkeletalMeshComponent->GetComponentToWorld();
+			PreviewScenePtr->RemoveComponent(PreviewSkeletalMeshComponent);
+			//PreviewSkeletalMeshComponent->DestroyComponent();
+		}
+	}
+
+	PreviewSkeletalMeshComponents = InSkeletalMeshComponents;
+
+	for (UDebugSkelMeshComponent* PreviewSkeletalMeshComponent : PreviewSkeletalMeshComponents)
+	{
+		// To support null components We check if the skeletal mesh is null in the viewport client
+		if (PreviewSkeletalMeshComponent)
+		{
+			PreviewSkeletalMeshComponent->bCastInsetShadow = true; // For better quality shadows in the editor previews, more similar to the in-game ones
+			PreviewSkeletalMeshComponent->bCanHighlightSelectedSections = false;
+			PreviewSkeletalMeshComponent->MarkRenderStateDirty();
+
+			PreviewScenePtr->AddComponent(PreviewSkeletalMeshComponent, Transform);
+		}
+	}
+
+	LevelViewportClient->SetPreviewComponents(PreviewSkeletalMeshComponents);
+}
+
+
+bool SCustomizableObjectEditorViewportTabBody::IsVisible() const
+{
+	return ViewportWidget.IsValid();
+}
+
+
+const TArray<UDebugSkelMeshComponent*>& SCustomizableObjectEditorViewportTabBody::GetSkeletalMeshComponents() const
+{
+	return PreviewSkeletalMeshComponents;
+}
+
+
+void SCustomizableObjectEditorViewportTabBody::BindCommands()
+{
+	FUICommandList& CommandList = *UICommandList;
+
+	const FCustomizableObjectEditorViewportCommands& Commands = FCustomizableObjectEditorViewportCommands::Get();
+
+	// Viewport commands
+	TSharedRef<FCustomizableObjectEditorViewportClient> EditorViewportClientRef = LevelViewportClient.ToSharedRef();
+
+	CommandList.MapAction(
+		Commands.SetCameraLock,
+		FExecuteAction::CreateSP( EditorViewportClientRef, &FCustomizableObjectEditorViewportClient::SetCameraLock ),
+		FCanExecuteAction(),
+		FIsActionChecked::CreateSP( EditorViewportClientRef, &FCustomizableObjectEditorViewportClient::IsCameraLocked ) );
+
+	CommandList.MapAction(
+		Commands.SetDrawUVs,
+		FExecuteAction::CreateSP( EditorViewportClientRef, &FCustomizableObjectEditorViewportClient::SetDrawUVOverlay ),
+		FCanExecuteAction(),
+		FIsActionChecked::CreateSP( EditorViewportClientRef, &FCustomizableObjectEditorViewportClient::IsSetDrawUVOverlayChecked ) );
+
+	CommandList.MapAction(
+		Commands.SetShowGrid,
+		FExecuteAction::CreateSP( EditorViewportClientRef, &FCustomizableObjectEditorViewportClient::SetShowGrid ),
+		FCanExecuteAction(),
+		FIsActionChecked::CreateSP( EditorViewportClientRef, &FCustomizableObjectEditorViewportClient::IsSetShowGridChecked ) );
+
+	CommandList.MapAction(
+		Commands.SetShowSky,
+		FExecuteAction::CreateSP(EditorViewportClientRef, &FCustomizableObjectEditorViewportClient::SetShowSky),
+		FCanExecuteAction(),
+		FIsActionChecked::CreateSP(EditorViewportClientRef, &FCustomizableObjectEditorViewportClient::IsSetShowSkyChecked));
+
+	CommandList.MapAction(
+		Commands.SetShowBounds,
+		FExecuteAction::CreateSP( EditorViewportClientRef, &FCustomizableObjectEditorViewportClient::SetShowBounds ),
+		FCanExecuteAction(),
+		FIsActionChecked::CreateSP( EditorViewportClientRef, &FCustomizableObjectEditorViewportClient::IsSetShowBoundsChecked ) );
+
+	CommandList.MapAction(
+		Commands.SetShowCollision,
+		FExecuteAction::CreateSP( EditorViewportClientRef, &FCustomizableObjectEditorViewportClient::SetShowCollision ),
+		FCanExecuteAction(),
+		FIsActionChecked::CreateSP( EditorViewportClientRef, &FCustomizableObjectEditorViewportClient::IsSetShowCollisionChecked ) );
+
+	// Menu
+	CommandList.MapAction(
+		Commands.SetShowPivot,
+		FExecuteAction::CreateSP( EditorViewportClientRef, &FCustomizableObjectEditorViewportClient::SetShowPivot ),
+		FCanExecuteAction(),
+		FIsActionChecked::CreateSP( EditorViewportClientRef, &FCustomizableObjectEditorViewportClient::IsSetShowPivotChecked ) );
+
+
+	CommandList.MapAction(
+		Commands.BakeInstance,
+		FExecuteAction::CreateSP(EditorViewportClientRef, &FCustomizableObjectEditorViewportClient::BakeInstance),
+		FCanExecuteAction(),
+		FIsActionChecked());
+
+	// In-viewport tool bar
+
+	//Bind menu commands
+	const FCustomizableObjectEditorViewportMenuCommands& MenuActions = FCustomizableObjectEditorViewportMenuCommands::Get();
+
+	//CommandList.MapAction(
+	//	MenuActions.CameraFollow,
+	//	FExecuteAction::CreateSP(this, &SCustomizableObjectEditorViewport::ToggleCameraFollow),
+	//	FCanExecuteAction::CreateSP(this, &SCustomizableObjectEditorViewport::CanChangeCameraMode),
+	//	FIsActionChecked::CreateSP(this, &SCustomizableObjectEditorViewport::IsCameraFollowEnabled));
+
+	//CommandList.MapAction(
+	//	MenuActions.PreviewSceneSettings,
+	//	FExecuteAction::CreateSP(this, &SCustomizableObjectEditorViewport::OpenPreviewSceneSettings));
+
+	//CommandList.MapAction(
+	//	MenuActions.SetCPUSkinning,
+	//	FExecuteAction::CreateSP(EditorViewportClientRef, &FCustomizableObjectEditorViewportClient::ToggleCPUSkinning),
+	//	FCanExecuteAction(),
+	//	FIsActionChecked::CreateSP(EditorViewportClientRef, &FCustomizableObjectEditorViewportClient::IsSetCPUSkinningChecked));
+
+	//CommandList.MapAction(
+	//	MenuActions.SetShowNormals,
+	//	FExecuteAction::CreateSP(EditorViewportClientRef, &FCustomizableObjectEditorViewportClient::ToggleShowNormals),
+	//	FCanExecuteAction(),
+	//	FIsActionChecked::CreateSP(EditorViewportClientRef, &FCustomizableObjectEditorViewportClient::IsSetShowNormalsChecked));
+
+	//CommandList.MapAction(
+	//	MenuActions.SetShowTangents,
+	//	FExecuteAction::CreateSP(EditorViewportClientRef, &FCustomizableObjectEditorViewportClient::ToggleShowTangents),
+	//	FCanExecuteAction(),
+	//	FIsActionChecked::CreateSP(EditorViewportClientRef, &FCustomizableObjectEditorViewportClient::IsSetShowTangentsChecked));
+
+	//CommandList.MapAction(
+	//	MenuActions.SetShowBinormals,
+	//	FExecuteAction::CreateSP(EditorViewportClientRef, &FCustomizableObjectEditorViewportClient::ToggleShowBinormals),
+	//	FCanExecuteAction(),
+	//	FIsActionChecked::CreateSP(EditorViewportClientRef, &FCustomizableObjectEditorViewportClient::IsSetShowBinormalsChecked));
+
+	//CommandList.MapAction(
+	//	MenuActions.AnimSetDrawUVs,
+	//	FExecuteAction::CreateSP(EditorViewportClientRef, &FCustomizableObjectEditorViewportClient::ToggleDrawUVOverlay),
+	//	FCanExecuteAction(),
+	//	FIsActionChecked::CreateSP(EditorViewportClientRef, &FCustomizableObjectEditorViewportClient::IsSetDrawUVOverlayChecked));
+
+
+	//Bind LOD preview menu commands
+	const FCustomizableObjectEditorViewportLODCommands& ViewportLODMenuCommands = FCustomizableObjectEditorViewportLODCommands::Get();
+
+	//LOD Auto
+	CommandList.MapAction(
+		ViewportLODMenuCommands.LODAuto,
+		FExecuteAction::CreateSP(this, &SCustomizableObjectEditorViewportTabBody::OnSetLODModel, 0),
+		FCanExecuteAction(),
+		FIsActionChecked::CreateSP(this, &SCustomizableObjectEditorViewportTabBody::IsLODModelSelected, 0));
+
+	// LOD 0
+	CommandList.MapAction(
+		ViewportLODMenuCommands.LOD0,
+		FExecuteAction::CreateSP(this, &SCustomizableObjectEditorViewportTabBody::OnSetLODModel, 1),
+		FCanExecuteAction(),
+		FIsActionChecked::CreateSP(this, &SCustomizableObjectEditorViewportTabBody::IsLODModelSelected, 1));
+
+	CommandList.MapAction(
+		ViewportLODMenuCommands.TranslateMode,
+		FExecuteAction::CreateSP(this, &SCustomizableObjectEditorViewportTabBody::ProjectorCheckboxStateChanged, UE::Widget::WM_Translate),
+		FCanExecuteAction(),
+		FIsActionChecked());
+
+	CommandList.MapAction(
+		ViewportLODMenuCommands.RotateMode,
+		FExecuteAction::CreateSP(this, &SCustomizableObjectEditorViewportTabBody::ProjectorCheckboxStateChanged, UE::Widget::WM_Rotate),
+		FCanExecuteAction(),
+		FIsActionChecked());
+
+	CommandList.MapAction(
+		ViewportLODMenuCommands.ScaleMode,
+		FExecuteAction::CreateSP(this, &SCustomizableObjectEditorViewportTabBody::ProjectorCheckboxStateChanged, UE::Widget::WM_Scale),
+		FCanExecuteAction(),
+		FIsActionChecked());
+
+	CommandList.MapAction(
+		ViewportLODMenuCommands.RotationGridSnap,
+		FExecuteAction::CreateSP(this, &SCustomizableObjectEditorViewportTabBody::RotationGridSnapClicked),
+		FCanExecuteAction(),
+		FIsActionChecked::CreateSP(this, &SCustomizableObjectEditorViewportTabBody::RotationGridSnapIsChecked)
+	);
+
+	CommandList.MapAction(
+		ViewportLODMenuCommands.HighResScreenshot,
+		FExecuteAction::CreateSP(this, &SCustomizableObjectEditorViewportTabBody::OnTakeHighResScreenshot),
+		FCanExecuteAction()
+	);
+
+	//Orbital Camera Mode
+	CommandList.MapAction(
+		ViewportLODMenuCommands.OrbitalCamera,
+		FExecuteAction::CreateSP(this, &SCustomizableObjectEditorViewportTabBody::SetCameraMode, true),
+		FCanExecuteAction(),
+		FIsActionChecked::CreateSP(this, &SCustomizableObjectEditorViewportTabBody::IsCameraModeActive,0)
+	);
+
+	//Free Camera Mode
+	CommandList.MapAction(
+		ViewportLODMenuCommands.FreeCamera,
+		FExecuteAction::CreateSP(this, &SCustomizableObjectEditorViewportTabBody::SetCameraMode, false),
+		FCanExecuteAction(),
+		FIsActionChecked::CreateSP(this, &SCustomizableObjectEditorViewportTabBody::IsCameraModeActive,1)
+	);
+
+	const FEditorViewportCommands& ViewportCommands = FEditorViewportCommands::Get();
+
+	// Camera Views
+	CommandList.MapAction(
+		ViewportCommands.Perspective,
+		FExecuteAction::CreateSP(EditorViewportClientRef, &FCustomizableObjectEditorViewportClient::SetViewportType, LVT_Perspective),
+		FCanExecuteAction(),
+		FIsActionChecked::CreateSP(EditorViewportClientRef, &FCustomizableObjectEditorViewportClient::IsActiveViewportType, LVT_Perspective));
+
+	CommandList.MapAction(
+		ViewportCommands.Front,
+		FExecuteAction::CreateSP(EditorViewportClientRef, &FCustomizableObjectEditorViewportClient::SetViewportType, LVT_OrthoNegativeYZ),
+		FCanExecuteAction(),
+		FIsActionChecked::CreateSP(EditorViewportClientRef, &FCustomizableObjectEditorViewportClient::IsActiveViewportType, LVT_OrthoNegativeYZ));
+
+	CommandList.MapAction(
+		ViewportCommands.Left,
+		FExecuteAction::CreateSP(EditorViewportClientRef, &FCustomizableObjectEditorViewportClient::SetViewportType, LVT_OrthoNegativeXZ),
+		FCanExecuteAction(),
+		FIsActionChecked::CreateSP(EditorViewportClientRef, &FCustomizableObjectEditorViewportClient::IsActiveViewportType, LVT_OrthoNegativeXZ));
+
+	CommandList.MapAction(
+		ViewportCommands.Top,
+		FExecuteAction::CreateSP(EditorViewportClientRef, &FCustomizableObjectEditorViewportClient::SetViewportType, LVT_OrthoXY),
+		FCanExecuteAction(),
+		FIsActionChecked::CreateSP(EditorViewportClientRef, &FCustomizableObjectEditorViewportClient::IsActiveViewportType, LVT_OrthoXY));
+
+	CommandList.MapAction(
+		ViewportCommands.Back,
+		FExecuteAction::CreateSP(EditorViewportClientRef, &FCustomizableObjectEditorViewportClient::SetViewportType, LVT_OrthoYZ),
+		FCanExecuteAction(),
+		FIsActionChecked::CreateSP(EditorViewportClientRef, &FCustomizableObjectEditorViewportClient::IsActiveViewportType, LVT_OrthoYZ));
+
+	CommandList.MapAction(
+		ViewportCommands.Right,
+		FExecuteAction::CreateSP(EditorViewportClientRef, &FCustomizableObjectEditorViewportClient::SetViewportType, LVT_OrthoXZ),
+		FCanExecuteAction(),
+		FIsActionChecked::CreateSP(EditorViewportClientRef, &FCustomizableObjectEditorViewportClient::IsActiveViewportType, LVT_OrthoXZ));
+
+	CommandList.MapAction(
+		ViewportCommands.Bottom,
+		FExecuteAction::CreateSP(EditorViewportClientRef, &FCustomizableObjectEditorViewportClient::SetViewportType, LVT_OrthoNegativeXY),
+		FCanExecuteAction(),
+		FIsActionChecked::CreateSP(EditorViewportClientRef, &FCustomizableObjectEditorViewportClient::IsActiveViewportType, LVT_OrthoNegativeXY));
+
+	// all other LODs will be added dynamically 
+
+}
+
+
+void SCustomizableObjectEditorViewportTabBody::OnTakeHighResScreenshot()
+{
+	CustomizableObjectHighresScreenshot = SCustomizableObjectHighresScreenshot::OpenDialog(ViewportWidget->GetSceneViewport(), LevelViewportClient, PreviewSkeletalMeshComponents[0], PreviewScenePtr);
+}
+
+
+bool SCustomizableObjectEditorViewportTabBody::GetIsManipulatingGizmo()
+{
+	return LevelViewportClient->GetIsManipulating();
+}
+
+
+void SCustomizableObjectEditorViewportTabBody::SetCameraMode(bool Value)
+{
+	LevelViewportClient->SetCameraMode(Value);
+}
+
+
+bool SCustomizableObjectEditorViewportTabBody::IsCameraModeActive(int Value)
+{
+	bool IsOrbitalCemareaActive = LevelViewportClient->IsOrbitalCameraActive();
+	return (Value == 0) ? IsOrbitalCemareaActive : !IsOrbitalCemareaActive;
+}
+
+
+void SCustomizableObjectEditorViewportTabBody::SetDrawDefaultUVMaterial()
+{
+	GenerateUVMaterialOptions();
+
+	FString CurrentDrawnMaterial = LevelViewportClient->GetMaterialToDrawInUVs();
+
+	if (!ArrayUVMaterialOptionString.IsEmpty())
+	{
+		bool bFound = false;
+
+		for (int32 StringIndex = 0; StringIndex < ArrayUVMaterialOptionString.Num(); ++StringIndex)
+		{
+			if (*(ArrayUVMaterialOptionString[StringIndex]) == CurrentDrawnMaterial)
+			{
+				bFound = true;
+				break;
+			}
+		}
+
+		if (!bFound)
+		{
+			LevelViewportClient->SetDrawUVOverlayMaterial(*(ArrayUVMaterialOptionString[0]), "0");
+		}
+	}
+}
+
+
+TSharedRef<SWidget> SCustomizableObjectEditorViewportTabBody::BuildToolBar()
+{
+	//FToolBarBuilder SaveThumbnailToolbarBuilder(GetToolkitCommands());
+	//{
+	//	SaveThumbnailToolbarBuilder.AddToolBarButton(FCustomizableObjectEditorViewportCommands::Get().SaveThumbnail);
+	//}
+
+
+	FToolBarBuilder CommandToolbarBuilder(UICommandList, FMultiBoxCustomization::None);
+	{
+		CommandToolbarBuilder.AddToolBarButton(FCustomizableObjectEditorViewportCommands::Get().SetShowGrid);
+		CommandToolbarBuilder.AddToolBarButton(FCustomizableObjectEditorViewportCommands::Get().SetShowSky);
+		//CommandToolbarBuilder.AddToolBarButton(FCustomizableObjectEditorViewportCommands::Get().SetShowBounds);
+		//CommandToolbarBuilder.AddToolBarButton(FCustomizableObjectEditorViewportCommands::Get().SetShowPivot);
+		//CommandToolbarBuilder.AddToolBarButton(FCustomizableObjectEditorViewportCommands::Get().SetShowNormals);
+		//CommandToolbarBuilder.AddToolBarButton(FCustomizableObjectEditorViewportCommands::Get().SetShowTangents);
+		//CommandToolbarBuilder.AddToolBarButton(FCustomizableObjectEditorViewportCommands::Get().SetShowBinormals);
+	}
+	CommandToolbarBuilder.BeginSection("Material UVs");
+	{
+		CommandToolbarBuilder.AddToolBarButton(FCustomizableObjectEditorViewportCommands::Get().SetDrawUVs);
+
+		CommandToolbarBuilder.AddComboButton(
+			FUIAction(),
+			FOnGetContent::CreateSP(
+				this,
+				&SCustomizableObjectEditorViewportTabBody::GenerateUVMaterialOptionsMenuContent),
+			TAttribute<FText>(),
+			TAttribute<FText>(),
+			FSlateIcon());
+	}
+	CommandToolbarBuilder.EndSection();
+
+	// Utilities
+	CommandToolbarBuilder.BeginSection("Utilities");
+	CommandToolbarBuilder.AddToolBarButton(FCustomizableObjectEditorViewportCommands::Get().BakeInstance);
+	CommandToolbarBuilder.EndSection();
+
+
+	return
+		SNew(SHorizontalBox)
+
+		+ SHorizontalBox::Slot()
+		.Padding(4, 0)
+		[
+			SNew(SBorder)
+			.Padding(0)
+		.BorderImage(FAppStyle::GetBrush("NoBorder"))
+		.IsEnabled(FSlateApplication::Get().GetNormalExecutionAttribute())
+		[
+			CommandToolbarBuilder.MakeWidget()
+		]
+		];
+}
+
+
+TSharedRef<SWidget> SCustomizableObjectEditorViewportTabBody::GenerateUVMaterialOptionsMenuContent()
+{
+	const bool bShouldCloseWindowAfterMenuSelection = false;
+	FMenuBuilder MenuBuilder(bShouldCloseWindowAfterMenuSelection, UICommandList);
+	MenuBuilder.BeginSection("ShowUV");
+	{
+		// Generating an array with all the options of the combobox 
+		GenerateUVMaterialOptions();
+
+		// Sorting array
+        ArrayUVMaterialOptionString.Sort(::CompareNames);
+
+		// Setting initial selected option
+		if (ArrayUVMaterialOptionString.Num())
+		{
+			bool bFound = false;
+
+			if (SelectedUVMaterial.IsValid())
+			{
+				for (int32 i = 0; i < ArrayUVMaterialOptionString.Num(); ++i)
+				{
+					if (*SelectedUVMaterial == *ArrayUVMaterialOptionString[i])
+					{
+						bFound = true;
+						SelectedUVMaterial = ArrayUVMaterialOptionString[i];
+						break;
+					}
+				}
+			}
+			else
+			{
+				FString CurrentDrawnMaterial = LevelViewportClient->GetMaterialToDrawInUVs();
+
+				for (int32 i = 0; i < ArrayUVMaterialOptionString.Num(); ++i)
+				{
+					if (CurrentDrawnMaterial == *ArrayUVMaterialOptionString[i])
+					{
+						bFound = true;
+						SelectedUVMaterial = ArrayUVMaterialOptionString[i];
+						break;
+					}
+				}
+			}
+
+			if (!bFound)
+			{
+				SelectedUVMaterial = ArrayUVMaterialOptionString[0];
+			}
+		}
+		else
+		{
+			SelectedUVMaterial = nullptr;
+		}
+	
+		UVMaterialOptionCombo = SNew(STextComboBox)
+			.OptionsSource(&ArrayUVMaterialOptionString)
+			.InitiallySelectedItem(SelectedUVMaterial)
+			.OnSelectionChanged(this, &SCustomizableObjectEditorViewportTabBody::OnMaterialChanged);
+
+		// Generating an array with all the options of the combobox 
+		GenerateUVChannelOptions(false);
+
+		// Setting initial selected option
+		if (ArrayUVChannelOptionString.Num())
+		{
+			if (SelectedUVChannel.IsValid())
+			{
+				bool bFound = false;
+
+				for (int32 i = 0; i < ArrayUVChannelOptionString.Num(); ++i)
+				{
+					if (*SelectedUVChannel == *ArrayUVChannelOptionString[i])
+					{
+						bFound = true;
+						SelectedUVChannel = ArrayUVChannelOptionString[i];
+						break;
+					}
+				}
+
+				if (!bFound)
+				{
+					SelectedUVChannel = ArrayUVChannelOptionString[0];
+				}
+			}
+			else
+			{
+				SelectedUVChannel = ArrayUVChannelOptionString[0];
+			}
+		}
+		else
+		{
+			SelectedUVChannel = nullptr;
+		}
+
+		UVChannelOptionCombo = SNew(STextComboBox)
+			.OptionsSource(&ArrayUVChannelOptionString)
+			.InitiallySelectedItem(SelectedUVChannel)
+			.OnSelectionChanged(this, &SCustomizableObjectEditorViewportTabBody::OnUVChannelChanged);
+		
+		MenuBuilder.AddWidget(UVMaterialOptionCombo.ToSharedRef(), FText::FromString(TEXT("Material")));
+		MenuBuilder.AddWidget(UVChannelOptionCombo.ToSharedRef(), FText::FromString(TEXT("UV Channel")));
+	}
+
+	MenuBuilder.EndSection();
+
+	return MenuBuilder.MakeWidget();
+}
+
+
+TSharedRef<SWidget> SCustomizableObjectEditorViewportTabBody::ShowStateTestData()
+{
+	const bool bShouldCloseWindowAfterMenuSelection = false;
+	FMenuBuilder MenuBuilder(bShouldCloseWindowAfterMenuSelection, UICommandList);
+
+	MenuBuilder.BeginSection("Objects to include");
+	{
+		MenuBuilder.AddMenuEntry(FCustomizableObjectEditorViewportCommands::Get().StateChangeShowData);
+		MenuBuilder.AddMenuEntry(FCustomizableObjectEditorViewportCommands::Get().StateChangeShowGeometryData);
+	}
+	MenuBuilder.EndSection();
+
+	return MenuBuilder.MakeWidget();
+}
+
+
+bool SCustomizableObjectEditorViewportTabBody::IsMaterialsComboEnabled() const
+{
+	return true;
+	//return LevelViewportClient.IsValid() && LevelViewportClient->IsSetDrawUVOverlayChecked();
+}
+
+
+void SCustomizableObjectEditorViewportTabBody::GenerateUVMaterialOptions()
+{
+	ArrayUVMaterialOptionString.Empty();
+	
+	int32 ComponentIndex = 0;
+
+	for (UDebugSkelMeshComponent* PreviewSkeletalMeshComponent : PreviewSkeletalMeshComponents)
+	{
+		if (PreviewSkeletalMeshComponent != nullptr && PreviewSkeletalMeshComponent->SkeletalMesh != nullptr && PreviewSkeletalMeshComponent->SkeletalMesh->GetResourceForRendering() != nullptr)
+		{
+			TMap<FString, int32> MaterialLODs;
+
+			// Add Suffix "__X" for materials with duplicated names and Suffing " LOD_X" for multiple LODs.
+			const FSkeletalMeshRenderData* MeshRes = PreviewSkeletalMeshComponent->SkeletalMesh->GetResourceForRendering();
+			const TArray<UMaterialInterface*> Materials = PreviewSkeletalMeshComponent->GetMaterials();
+			for (UMaterialInterface* m : Materials)
+			{
+				const UMaterial* BaseMaterial = m->GetBaseMaterial();
+
+				FString BaseMaterialName = BaseMaterial->GetName();
+
+				int32 LODCount = 0;
+				if (int32* Count = MaterialLODs.Find(BaseMaterialName))
+				{
+					if ((*Count) < MeshRes->LODRenderData.Num() - 1)
+					{
+						(*Count)++;
+						LODCount = (*Count);
+					}
+					else
+					{
+						bool bSearchEmptyLOD = true;
+						int MaterialIndexWithinLOD = 1;
+
+						while (bSearchEmptyLOD)
+						{
+							if (int32* DupeMaterialCount = MaterialLODs.Find(BaseMaterialName + FString::Printf(TEXT("__%d"), MaterialIndexWithinLOD)))
+							{
+								if ((*DupeMaterialCount) < MeshRes->LODRenderData.Num() - 1)
+								{
+									(*DupeMaterialCount)++;
+									LODCount = (*DupeMaterialCount);
+									BaseMaterialName += FString::Printf(TEXT("__%d"), MaterialIndexWithinLOD);
+									bSearchEmptyLOD = false;
+								}
+								else
+								{
+									MaterialIndexWithinLOD++;
+								}
+							}
+							else
+							{
+								MaterialLODs.Add(BaseMaterialName + FString::Printf(TEXT("__%d"), MaterialIndexWithinLOD), 0);
+
+								BaseMaterialName += FString::Printf(TEXT("__%d"), MaterialIndexWithinLOD);
+								bSearchEmptyLOD = false;
+							}
+						}
+					}
+				}
+				else
+				{
+					MaterialLODs.Add(BaseMaterialName, LODCount);
+				}
+
+				BaseMaterialName += FString::Printf(TEXT(" LOD_%d_Component_%d"), LODCount, ComponentIndex);
+
+				ArrayUVMaterialOptionString.Add(MakeShareable(new FString(BaseMaterialName)));
+			}
+		}
+
+		ComponentIndex++;
+	}
+}
+
+
+void SCustomizableObjectEditorViewportTabBody::OnMaterialChanged(TSharedPtr<FString> Selected, ESelectInfo::Type SelectInfo)
+{
+	if (Selected.IsValid() && SelectInfo == ESelectInfo::OnMouseClick)
+	{
+		SelectedUVMaterial = Selected;
+
+		//We need to update options for the new LOD
+		GenerateUVChannelOptions(true);
+
+		if (LevelViewportClient.IsValid() && SelectedUVChannel.IsValid())
+		{
+			LevelViewportClient->SetDrawUVOverlayMaterial(*Selected, *SelectedUVChannel);
+		}
+	}
+}
+
+
+void SCustomizableObjectEditorViewportTabBody::GenerateUVChannelOptions(bool bReset)
+{
+	ArrayUVChannelOptionString.Empty();
+
+	if (!SelectedUVMaterial)
+	{
+		return;
+	}
+
+	FString SelectedMat = *SelectedUVMaterial;
+
+	// From String to values
+	FString NameWithLOD, ComponentString;
+	SelectedMat.Split(FString("_Component_"), &NameWithLOD, &ComponentString);
+	check(ComponentString.IsNumeric());
+	int32 ComponentIndex = FCString::Atoi(*ComponentString);
+	check(PreviewSkeletalMeshComponents.IsValidIndex(ComponentIndex));
+
+	FString Name, LODString;
+	bool bSplit = NameWithLOD.Split(FString(" LOD_"), &Name, &LODString);
+	int32 LODIndex = FCString::Atoi(*LODString);
+
+	UDebugSkelMeshComponent* PreviewSkeletalMeshComponent = PreviewSkeletalMeshComponents[ComponentIndex];
+
+	if (PreviewSkeletalMeshComponent != nullptr && PreviewSkeletalMeshComponent->SkeletalMesh != nullptr
+		&& PreviewSkeletalMeshComponent->SkeletalMesh->GetResourceForRendering() != nullptr)
+	{
+		const FSkeletalMeshRenderData* MeshRes = PreviewSkeletalMeshComponent->SkeletalMesh->GetResourceForRendering();
+		
+		int32 UVChannels = MeshRes->LODRenderData[LODIndex].GetNumTexCoords();
+		
+		for (int32 UVChan = 0; UVChan < UVChannels; ++UVChan)
+		{
+			ArrayUVChannelOptionString.Add(MakeShareable(new FString(FString::FromInt(UVChan))));
+		}
+
+		// Resets the value to the first element of the array
+		if (bReset)
+		{
+			SelectedUVChannel = ArrayUVChannelOptionString[0];
+			UVChannelOptionCombo->SetSelectedItem(SelectedUVChannel);
+		}
+	}
+}
+
+
+void SCustomizableObjectEditorViewportTabBody::OnUVChannelChanged(TSharedPtr<FString> Selected, ESelectInfo::Type SelectInfo)
+{
+	if (Selected.IsValid() && SelectInfo == ESelectInfo::OnMouseClick)
+	{
+		SelectedUVChannel = Selected;
+		if (LevelViewportClient.IsValid() && SelectedUVMaterial.IsValid())
+		{
+			LevelViewportClient->SetDrawUVOverlayMaterial(*SelectedUVMaterial, *Selected);
+		}
+	}
+}
+
+
+void SCustomizableObjectEditorViewportTabBody::SetAnimation(UAnimationAsset* Animation, EAnimationMode::Type AnimationType)
+{
+	LevelViewportClient->SetAnimation(Animation, AnimationType);
+}
+
+
+FReply SCustomizableObjectEditorViewportTabBody::OnDrop(const FGeometry& MyGeometry, const FDragDropEvent& DragDropEvent)
+{
+	if (FAssetDragDropOp* DragDropOp = DragDropEvent.GetOperationAs<FAssetDragDropOp>().Get())
+	{
+		if (Helper_GetAssets(DragDropOp).Num())
+		{
+			UPoseAsset* PoseAsset = Cast<UPoseAsset>(Helper_GetAssets(DragDropOp)[0].GetAsset());
+			if (PoseAsset != nullptr)
+			{
+				CustomizableObjectEditorPtr.Pin()->SetPoseAsset(PoseAsset);
+				return FReply::Handled();
+			}
+
+			UAnimationAsset* AnimationAsset = Cast<UAnimationAsset>(Helper_GetAssets(DragDropOp)[0].GetAsset());
+			if (AnimationAsset)
+			{
+				LevelViewportClient->SetAnimation(AnimationAsset, EAnimationMode::AnimationSingleNode);
+
+				return FReply::Handled();
+			}
+		}
+	}
+
+	return FReply::Unhandled();
+}
+
+
+int32 SCustomizableObjectEditorViewportTabBody::GetLODModelCount() const
+{
+	int32 LODModelCount = 0;
+
+	for (UDebugSkelMeshComponent* PreviewComponent : GetSkeletalMeshComponents())
+	{
+		if (PreviewComponent && PreviewComponent->SkeletalMesh)
+		{
+			const TIndirectArray<FSkeletalMeshLODRenderData>& LODModels = Helper_GetLODData(PreviewComponent->SkeletalMesh);
+			LODModelCount = FMath::Max(LODModelCount, LODModels.Num());
+		}
+	}
+
+	return LODModelCount;
+}
+
+bool SCustomizableObjectEditorViewportTabBody::IsLODModelSelected(int32 LODSelectionType) const
+{
+	return (LODSelection == LODSelectionType) ? true : false;
+}
+
+
+void SCustomizableObjectEditorViewportTabBody::ProjectorCheckboxStateChanged(UE::Widget::EWidgetMode InNewMode)
+{
+	switch (InNewMode)
+	{
+		case UE::Widget::WM_Translate:
+		{
+			GetViewportClient()->SetWidgetMode(UE::Widget::WM_Translate);
+			GetViewportClient()->SetProjectorWidgetMode(UE::Widget::WM_Translate);
+			break;
+		}
+		case UE::Widget::WM_Rotate:
+		{
+			GetViewportClient()->SetWidgetMode(UE::Widget::WM_Rotate);
+			GetViewportClient()->SetProjectorWidgetMode(UE::Widget::WM_Rotate);
+			break;
+		}
+		case UE::Widget::WM_Scale:
+		{
+			GetViewportClient()->SetWidgetMode(UE::Widget::WM_Scale);
+			GetViewportClient()->SetProjectorWidgetMode(UE::Widget::WM_Scale);
+			break;
+		}
+	}
+}
+
+
+void SCustomizableObjectEditorViewportTabBody::RotationGridSnapClicked()
+{
+	GUnrealEd->Exec(GEditor->GetEditorWorldContext().World(), *FString::Printf(TEXT("MODE ROTGRID=%d"), !GetDefault<ULevelEditorViewportSettings>()->RotGridEnabled ? 1 : 0));
+}
+
+
+bool SCustomizableObjectEditorViewportTabBody::RotationGridSnapIsChecked() const
+{
+	return GetDefault<ULevelEditorViewportSettings>()->RotGridEnabled;
+}
+
+
+void SCustomizableObjectEditorViewportTabBody::SetRotationGridSize(int32 InIndex, ERotationGridMode InGridMode)
+{
+	GEditor->SetRotGridSize(InIndex, InGridMode);
+}
+
+
+bool SCustomizableObjectEditorViewportTabBody::IsRotationGridSizeChecked(int32 GridSizeIndex, ERotationGridMode GridMode)
+{
+	const ULevelEditorViewportSettings* ViewportSettings = GetDefault<ULevelEditorViewportSettings>();
+	return (ViewportSettings->CurrentRotGridSize == GridSizeIndex) && (ViewportSettings->CurrentRotGridMode == GridMode);
+}
+
+
+void SCustomizableObjectEditorViewportTabBody::OnSetLODModel(int32 LODSelectionType)
+{
+	LODSelection = LODSelectionType;
+
+	for (UDebugSkelMeshComponent* PreviewComponent : GetSkeletalMeshComponents())
+	{
+		if (PreviewComponent)
+		{
+#if (ENGINE_MAJOR_VERSION == 4 && ENGINE_MINOR_VERSION >= 24) || ENGINE_MAJOR_VERSION >= 5
+			PreviewComponent->SetForcedLOD(LODSelection);
+#else
+			PreviewComponent->ForcedLodModel = LODSelection;
+#endif
+			//PopulateUVChoices();
+		}
+	}
+}
+
+/*
+void SCustomizableObjectEditorViewportTabBody::SetEditorTransformViewportToolbar(TWeakPtr<class SCustomizableObjectEditorTransformViewportToolbar> EditorTransformViewportToolbarParam)
+{
+	EditorTransformViewportToolbar = EditorTransformViewportToolbarParam;
+}*/
+
+void SCustomizableObjectEditorViewportTabBody::SetViewportToolbarTransformWidget(TWeakPtr<class SWidget> InTransformWidget)
+{
+	ViewportToolbarTransformWidget = InTransformWidget;
+}
+
+void SCustomizableObjectEditorViewportTabBody::UpdateGizmoDataToOrigin()
+{
+	LevelViewportClient->UpdateGizmoDataToOrigin();
+}
+
+
+void SCustomizableObjectEditorViewportTabBody::SetGizmoCallUpdateSkeletalMesh(bool Value)
+{
+	LevelViewportClient->SetGizmoCallUpdateSkeletalMesh(Value);
+}
+
+
+FString SCustomizableObjectEditorViewportTabBody::GetGizmoProjectorParameterName()
+{
+	return LevelViewportClient->GetGizmoProjectorParameterName();
+}
+
+
+FString SCustomizableObjectEditorViewportTabBody::GetGizmoProjectorParameterNameWithIndex()
+{
+	return LevelViewportClient->GetGizmoProjectorParameterNameWithIndex();
+}
+
+
+bool SCustomizableObjectEditorViewportTabBody::GetGizmoHasAssignedData()
+{
+	return LevelViewportClient->GetGizmoHasAssignedData();
+}
+
+
+bool SCustomizableObjectEditorViewportTabBody::GetGizmoAssignedDataIsFromNode()
+{
+	return LevelViewportClient->GetGizmoAssignedDataIsFromNode();
+}
+
+
+void SCustomizableObjectEditorViewportTabBody::CopyTransformFromOriginData()
+{
+    LevelViewportClient->CopyTransformFromOriginData();
+}
+
+
+void SCustomizableObjectEditorViewportTabBody::ProjectorParameterChanged(UCustomizableObjectNodeProjectorParameter* Node)
+{
+	LevelViewportClient->ProjectorParameterChanged(Node);
+}
+
+
+void SCustomizableObjectEditorViewportTabBody::ProjectorParameterChanged(UCustomizableObjectNodeProjectorConstant* Node)
+{
+	LevelViewportClient->ProjectorParameterChanged(Node);
+}
+
+
+bool SCustomizableObjectEditorViewportTabBody::AnyProjectorNodeSelected()
+{
+	return LevelViewportClient->AnyProjectorNodeSelected();
+}
+
+
+void SCustomizableObjectEditorViewportTabBody::SetCustomizableObject(UCustomizableObject* CustomizableObjectParameter)
+{
+	LevelViewportClient->SetCustomizableObject(CustomizableObjectParameter);
+}
+
+
+void SCustomizableObjectEditorViewportTabBody::SetAssetRegistryLoaded(bool Value)
+{
+	AssetRegistryLoaded = Value;
+	LevelViewportClient->SetAssetRegistryLoaded(Value);
+}
+
+FLinearColor SCustomizableObjectEditorViewportTabBody::GetViewportBackgroundColor() const
+{
+	return LevelViewportClient->GetBackgroundColor();
+}
diff --git a/Engine/Plugins/Experimental/Mutable/Source/CustomizableObjectEditor/Private/SCustomizableObjectEditorViewport.h b/Engine/Plugins/Experimental/Mutable/Source/CustomizableObjectEditor/Private/SCustomizableObjectEditorViewport.h
new file mode 100644
--- /dev/null
+++ b/Engine/Plugins/Experimental/Mutable/Source/CustomizableObjectEditor/Private/SCustomizableObjectEditorViewport.h
@@ -0,0 +1,322 @@
+// Copyright Epic Games, Inc. All Rights Reserved.
+
+#pragma once
+
+#include "PreviewScene.h"
+#include "SEditorViewport.h"
+#include "Components/SkeletalMeshComponent.h"
+#include "Settings/LevelEditorViewportSettings.h"
+#include "ICustomizableObjectInstanceEditor.h"
+#include "Nodes/CustomizableObjectNodeProjectorConstant.h"
+
+
+struct FCustomizableObjectEditorViewportRequiredArgs
+{
+	FCustomizableObjectEditorViewportRequiredArgs(
+		//const TSharedRef<class ISkeletonTree>& InSkeletonTree, 
+		const TSharedRef<class FCustomizableObjectPreviewScene>& InPreviewScene,
+		TSharedRef<class SCustomizableObjectEditorViewportTabBody> InTabBody 
+		//TSharedRef<class FAssetEditorToolkit> InAssetEditorToolkit, 
+		//FSimpleMulticastDelegate& InOnPostUndo
+	)
+//		: SkeletonTree(InSkeletonTree)
+		: PreviewScene(InPreviewScene)
+		, TabBody(InTabBody)
+//		, AssetEditorToolkit(InAssetEditorToolkit)
+//		, OnPostUndo(InOnPostUndo)
+	{}
+
+	//TSharedRef<class ISkeletonTree> SkeletonTree;
+
+	TSharedRef<class FCustomizableObjectPreviewScene> PreviewScene;
+
+	TSharedRef<class SCustomizableObjectEditorViewportTabBody> TabBody;
+
+	//TSharedRef<class FAssetEditorToolkit> AssetEditorToolkit;
+
+	//FSimpleMulticastDelegate& OnPostUndo;
+};
+
+
+//////////////////////////////////////////////////////////////////////////
+// SCustomizableObjectEditorViewport: the pure viewport widget
+
+class SCustomizableObjectEditorViewport : public SEditorViewport
+{
+public:
+	SLATE_BEGIN_ARGS(SCustomizableObjectEditorViewport)
+		: _ShowStats(true)
+	{}
+
+		SLATE_ARGUMENT(bool, ShowStats)
+
+	SLATE_END_ARGS()
+
+	void Construct(const FArguments& InArgs, const FCustomizableObjectEditorViewportRequiredArgs& InRequiredArgs);
+
+	TSharedPtr<FSceneViewport>& GetSceneViewport();
+
+protected:
+	// SEditorViewport interface
+	virtual TSharedRef<FEditorViewportClient> MakeEditorViewportClient() override;
+	virtual TSharedPtr<SWidget> MakeViewportToolbar() override;
+	// End of SEditorViewport interface
+
+	/**  Handle undo/redo by refreshing the viewport */
+	void OnUndoRedo();
+
+protected:
+	// Viewport client
+	TSharedPtr<class FCustomizableObjectEditorViewportClient> LevelViewportClient;
+
+	// Pointer to the compound widget that owns this viewport widget
+	TWeakPtr<class SCustomizableObjectEditorViewportTabBody> TabBodyPtr;
+
+	// The preview scene that we are viewing
+	TWeakPtr<class FCustomizableObjectPreviewScene> PreviewScenePtr;
+
+	TWeakPtr<class SCustomizableObjectViewportToolBar> ToolbarPtr;
+
+	// The skeleton tree we are editing
+	//TWeakPtr<class ISkeletonTree> SkeletonTreePtr;
+
+	// The asset editor we are embedded in
+	//TWeakPtr<class FAssetEditorToolkit> AssetEditorToolkitPtr;
+
+	/** Whether we should show stats for this viewport */
+	bool bShowStats;
+};
+
+
+/**
+ * CustomizableObject Editor Preview viewport widget wwith toolbars, etc.
+ */
+class SCustomizableObjectEditorViewportTabBody : public SCompoundWidget, public FGCObject
+{
+public:
+	SLATE_BEGIN_ARGS(SCustomizableObjectEditorViewportTabBody){}
+		SLATE_ARGUMENT(TWeakPtr<class ICustomizableObjectInstanceEditor>, CustomizableObjectEditor)
+	SLATE_END_ARGS()
+
+	void Construct(const FArguments& InArgs);
+	virtual ~SCustomizableObjectEditorViewportTabBody();
+
+	void SetAnimation(class UAnimationAsset* Animation, EAnimationMode::Type AnimationType);
+
+	virtual FReply OnDrop(const FGeometry& MyGeometry, const FDragDropEvent& DragDropEvent) override;
+	
+	// FSerializableObject interface
+	virtual void AddReferencedObjects( FReferenceCollector& Collector ) override;
+	virtual FString GetReferencerName() const override
+	{
+		return TEXT("SCustomizableObjectEditorViewportTabBody");
+	}
+	// End of FSerializableObject interface
+
+	// SWidget interface
+	virtual void Tick(const FGeometry& AllottedGeometry, const double InCurrentTime, const float InDeltaTime) override;
+	// End of SWidget interface
+
+	/**
+	* @return The list of commands on the viewport that are bound to delegates
+	*/
+	const TSharedPtr<FUICommandList>& GetCommandList() const { return UICommandList; }
+
+
+	void SetClipMorphPlaneVisibility(bool bVisible, const FVector& Origin = FVector::ZeroVector, const FVector& Normal = FVector::ZeroVector, float MorphLength = 0.f, const FBoxSphereBounds& Bounds = FBoxSphereBounds(), float Radius1 = 3.f, float Radius2 = 2.f, float RotationAngle = 0.f);
+	void SetClipMorphPlaneVisibility(bool bVisible, class UCustomizableObjectNodeMeshClipMorph* NodeMorphClipMesh);
+	void SetClipMeshVisibility(bool bVisible, UStaticMesh* ClipMesh, class UCustomizableObjectNodeMeshClipWithMesh* ClipMeshNode);
+	void SetProjectorVisibility(bool bVisible, FString ProjectorParameterName, FString ProjectorParameterNameWithIndex, int32 RangeIndex, const FCustomizableObjectProjector& Data, int32 ProjectorParameterIndex);
+	void SetProjectorVisibility(bool bVisible, class UCustomizableObjectNodeProjectorConstant* Projector = nullptr);
+	void SetProjectorParameterVisibility(bool bVisible, class UCustomizableObjectNodeProjectorParameter* ProjectorParameter = nullptr);
+	void ResetProjectorVisibility(bool OnlyNonNodeProjector);
+	void SetProjectorType(bool bVisible, FString ProjectorParameterName, FString ProjectorParameterNameWithIndex, int32 RangeIndex, const FCustomizableObjectProjector& Data, int32 ProjectorParameterIndex);
+
+	/** Binds commands associated with the viewport client. */
+	void BindCommands();
+	TSharedRef<SWidget> BuildToolBar();
+
+	TSharedRef<SWidget> GenerateUVMaterialOptionsMenuContent();
+
+	/** Callback to show / hide the state and rutime parameter test in the viewport */
+	TSharedRef<SWidget> ShowStateTestData();
+
+	/** Retrieves the skeletal mesh component. */
+	const TArray<UDebugSkelMeshComponent*>& GetSkeletalMeshComponents() const;
+	
+	void SetPreviewComponents(TArray<UDebugSkelMeshComponent*>& InSkeletalMeshComponents);
+	
+	/** Function to get viewport's current background color */
+	FLinearColor GetViewportBackgroundColor() const;
+
+	/** Function to get the number of LOD models associated with the preview skeletal mesh*/
+	int32 GetLODModelCount() const;
+
+	/** LOD model selection checking function*/
+	bool IsLODModelSelected(int32 LODSelectionType) const;
+	int32 GetLODSelection() const { return LODSelection; };
+
+	/** Callback for when the projector checkbox changes */
+	void ProjectorCheckboxStateChanged(UE::Widget::EWidgetMode InNewMode);
+
+	/** Called when the rotation grid snap is toggled off and on */
+	void RotationGridSnapClicked();
+
+	/** Returns whether or not rotation grid snap is enabled */
+	bool RotationGridSnapIsChecked() const;
+
+	/** Sets the rotation grid size */
+	static void SetRotationGridSize(int32 InIndex, ERotationGridMode InGridMode);
+
+	/** Checks to see if the specified rotation grid angle is the current rotation grid angle */
+	static bool IsRotationGridSizeChecked(int32 GridSizeIndex, ERotationGridMode GridMode);
+
+	/** Function to set LOD model selection*/
+	void OnSetLODModel(int32 LODSelectionType);
+
+	// Getter of PreviewScenePtr
+	TSharedPtr<class FCustomizableObjectPreviewScene> GetPreviewScene() { return PreviewScenePtr; }
+
+	//
+	TSharedPtr<class FCustomizableObjectEditorViewportClient> GetViewportClient() { return LevelViewportClient; }
+
+	// Pointer back to the editor tool that owns us.
+	TWeakPtr<class ICustomizableObjectInstanceEditor> CustomizableObjectEditorPtr;
+
+	//void SetEditorTransformViewportToolbar(TWeakPtr<class SCustomizableObjectEditorTransformViewportToolbar> EditorTransformViewportToolbarParam);
+	void SetViewportToolbarTransformWidget(TWeakPtr<class SWidget> TransformWidget);
+
+	/** Update the current gizmo data to its corresponding data origin */
+	void UpdateGizmoDataToOrigin();
+
+	/** Setter for LevelViewportClient::GizmoProxy::CallUpdateSkeletalMesh */
+	void SetGizmoCallUpdateSkeletalMesh(bool Value);
+
+	/** Getter for LevelViewportClient::GizmoProxy::ProjectorParameterName */
+	FString GetGizmoProjectorParameterName();
+
+	/** Getter for LevelViewportClient::GizmoProxy::ProjectorParameterNameWithIndex */
+	FString GetGizmoProjectorParameterNameWithIndex();
+
+	/** Getter for LevelViewportClient::GizmoProxy::HasAssignedData */
+	bool GetGizmoHasAssignedData();
+
+	/** Getter for LevelViewportClient::GizmoProxy::AssignedDataIsFromNode */
+	bool GetGizmoAssignedDataIsFromNode();
+
+	void CopyTransformFromOriginData();
+
+	/** Returns true if any projector node is currently selected (in the case of the Customizable Object Editor) */
+	bool AnyProjectorNodeSelected();
+
+	/** To notify of changes in any parameter projector node property in case is the one being used, currently
+	* just the projector type parameter is taken into account */
+	void ProjectorParameterChanged(UCustomizableObjectNodeProjectorParameter* Node);
+
+	/** To notify of changes in any parameter projector node property in case is the one being used, currently
+	* just the projector type parameter is taken into account */
+	void ProjectorParameterChanged(UCustomizableObjectNodeProjectorConstant* Node);
+	
+	/** Setter of CustomizableObject */
+	void SetCustomizableObject(UCustomizableObject* CustomizableObjectParameter);
+
+	/** Setter of AssetRegistryLoaded and LevelViewportClient::AssetRegistryLoaded */
+	void SetAssetRegistryLoaded(bool Value);
+
+	/** Called to bring up the screenshot UI */
+	void OnTakeHighResScreenshot();
+
+	/** Getter for LevelViewportClient::bManipulating */
+	bool GetIsManipulatingGizmo();
+
+	/** Sets camera mode of the LevelViewportClient::bActivateOrbitalCamera */
+	void SetCameraMode(bool Value);
+
+	/** Returns whether or not the a camera mode is active (Orbital or Free) */
+	bool IsCameraModeActive(int Value);
+
+	/** Sets the first material of the instance as the default section to draw in the UVs Overlay (when there is no material selected in the combobox)*/
+	void SetDrawDefaultUVMaterial();
+
+private:
+	/** Determines the visibility of the viewport. */
+	bool IsVisible() const;
+
+	EActiveTimerReturnType HandleActiveTimer(double InCurrentTime, float InDeltaTime);
+
+	
+	// Components for the preview mesh.
+	TArray<UDebugSkelMeshComponent*> PreviewSkeletalMeshComponents;
+
+	// The scene for this viewport.
+	TSharedPtr<class FCustomizableObjectPreviewScene> PreviewScenePtr;
+
+	// Level viewport client
+	TSharedPtr<class FCustomizableObjectEditorViewportClient> LevelViewportClient;
+
+	// Viewport widget
+	TSharedPtr<class SCustomizableObjectEditorViewport> ViewportWidget;
+	TSharedPtr<class SCustomizableObjectEditorToolBar> ViewportToolbar;
+
+	/** Commands that are bound to delegates*/
+	TSharedPtr<FUICommandList> UICommandList;
+
+	// The currently selected view mode.
+	EViewModeIndex CurrentViewMode;
+
+	// Selection for the material to filter when showing UVs
+	TArray<TSharedPtr<FString>> ArrayUVMaterialOptionString;
+
+	// Combo box for UV material selection
+	TSharedPtr< STextComboBox> UVMaterialOptionCombo;
+
+	// Selection for the channel to filter when showing UVs
+	TArray<TSharedPtr<FString>> ArrayUVChannelOptionString;
+
+	// Combo box for UV channel selection
+	TSharedPtr< STextComboBox> UVChannelOptionCombo;
+
+	// Generates the Combobox options for each UV Channel
+	void GenerateUVMaterialOptions();
+
+	// Generates the Combobox options for each UV Channel
+	void GenerateUVChannelOptions(bool bReset);
+
+	// Selected option of the UV Material ComboBox
+	TSharedPtr<FString> SelectedUVMaterial;
+	
+	// Selected option of the UV Channel ComboBox
+	TSharedPtr<FString> SelectedUVChannel;
+
+	// Options for the above combo.
+	TArray< TSharedPtr< FString > > MaterialNames;
+
+	// Slate callbacks to manage the material combo
+	bool IsMaterialsComboEnabled() const;
+	void OnMaterialChanged(TSharedPtr<FString> Selected, ESelectInfo::Type SelectInfo);
+	void OnUVChannelChanged(TSharedPtr<FString> Selected, ESelectInfo::Type SelectInfo);
+
+	/** Current LOD selection*/
+	int32 LODSelection;
+
+	/** To be able to show and hide the constant projector RTS controls */
+	//TWeakPtr<class SCustomizableObjectEditorTransformViewportToolbar> EditorTransformViewportToolbar;
+	TWeakPtr<class SWidget> ViewportToolbarTransformWidget;
+
+	union ViewportToolbarTransformWidgetVisibility
+	{
+		uint32 Value;
+		struct 
+		{
+			uint32 bProjectorVisible      : 1;
+			uint32 bClipMeshVisible       : 1;
+			uint32 bClipMorphPlaneVisible : 1;
+		} State;
+	} TransformWidgetVisibility {0};
+
+	/** Flag to know when the asset registry initial loading has completed, value set by the Customizable Object / Customizable Object Instance editor */
+	bool AssetRegistryLoaded;
+
+	/** Weak pointer to the highres screenshot dialog if it's open. Will become invalid if UI is closed whilst the viewport is open */
+	TWeakPtr<class SWindow> CustomizableObjectHighresScreenshot;
+};
