[mutable] Move Mutable plugin from the restricted folder to the experimental plugins folder.

diff --git a/Engine/Plugins/Experimental/Mutable/Source/CustomizableObjectEditor/Private/SCustomizableObjectEditorViewportToolBar.cpp b/Engine/Plugins/Experimental/Mutable/Source/CustomizableObjectEditor/Private/SCustomizableObjectEditorViewportToolBar.cpp
new file mode 100644
--- /dev/null
+++ b/Engine/Plugins/Experimental/Mutable/Source/CustomizableObjectEditor/Private/SCustomizableObjectEditorViewportToolBar.cpp
@@ -0,0 +1,1435 @@
+// Copyright Epic Games, Inc. All Rights Reserved.
+
+#include "SCustomizableObjectEditorViewportToolBar.h"
+#include "SCustomizableObjectEditorViewport.h"
+#include "EditorViewportCommands.h"
+#include "CustomizableObjectEditorViewportMenuCommands.h"
+#include "CustomizableObjectEditorViewportLODCommands.h"
+#include "CustomizableObjectEditorViewportClient.h"
+#include "SEditorViewportToolBarMenu.h"
+#include "STransformViewportToolbar.h"
+#include "Editor/EditorPerProjectUserSettings.h"
+#include "SEditorViewportViewMenu.h"
+#include "Widgets/Colors/SColorPicker.h"
+#include "Widgets/Input/SNumericEntryBox.h"
+#include "Widgets/Input/STextComboBox.h"
+#include "AssetViewerSettings.h"
+#include "ISkeletonTree.h"
+#include "IEditableSkeleton.h"
+#include "Engine/PreviewMeshCollection.h"
+#include "Widgets/Input/SButton.h"
+#include "Widgets/Input/SCheckBox.h"
+#include "CustomizableObjectEditorViewportClient.h"
+#include "Preferences/PersonaOptions.h"
+#include "SCustomizableObjectEditorTransformViewportToolbar.h"
+#include "Widgets/Layout/SScaleBox.h"
+#include "SViewportToolBarComboMenu.h"
+
+#include "Editor/UnrealEdEngine.h"
+#include "UnrealEdGlobals.h"
+#include "Settings/LevelEditorViewportSettings.h"
+#include "Widgets/Layout/SUniformGridPanel.h"
+#include "Internationalization/Internationalization.h"
+#include "Editor/EditorEngine.h"
+#include "LevelEditor.h"
+#include "Modules/ModuleManager.h"
+
+#define LOCTEXT_NAMESPACE "CustomizableObjectEditorViewportToolBar"
+
+
+//Class definition which represents widget to modify viewport's background color.
+//class SBackgroundColorSettings : public SCompoundWidget
+//{
+//public:
+//
+//	SLATE_BEGIN_ARGS(SBackgroundColorSettings)
+//	{}
+//	SLATE_ARGUMENT(TWeakPtr<SCustomizableObjectEditorViewport>, AnimEditorViewport)
+//		SLATE_END_ARGS()
+//
+//		/** Constructs this widget from its declaration */
+//		void Construct(const FArguments& InArgs)
+//	{
+//		AnimViewportPtr = InArgs._AnimEditorViewport;
+//
+//		TSharedPtr<SWidget> ExtraWidget =
+//			SNew(SBorder)
+//			.BorderImage(FAppStyle::GetBrush("FilledBorder"))
+//			[
+//				SNew(SHorizontalBox)
+//				+ SHorizontalBox::Slot()
+//			.AutoWidth()
+//			.Padding(1)
+//			[
+//				SNew(SColorBlock)
+//				.Color(AnimViewportPtr.Pin().ToSharedRef(), &SCustomizableObjectEditorViewport::GetViewportBackgroundColor)
+//			.IgnoreAlpha(true)
+//			.ToolTipText(LOCTEXT("ColorBlock_ToolTip", "Select background color"))
+//			.OnMouseButtonDown(this, &SBackgroundColorSettings::OnColorBoxClicked)
+//			]
+//			];
+//
+//		this->ChildSlot
+//			[
+//				SNew(SAnimPlusMinusSlider)
+//				.Label(LOCTEXT("BrightNess", "Brightness:"))
+//			.IsEnabled(this, &SBackgroundColorSettings::IsBrightnessSliderEnabled)
+//			.OnMinusClicked(this, &SBackgroundColorSettings::OnDecreaseBrightness)
+//			.MinusTooltip(LOCTEXT("DecreaseBrightness_ToolTip", "Decrease brightness"))
+//			.SliderValue(AnimViewportPtr.Pin().ToSharedRef(), &SCustomizableObjectEditorViewport::GetBackgroundBrightness)
+//			.OnSliderValueChanged(AnimViewportPtr.Pin().ToSharedRef(), &SCustomizableObjectEditorViewport::SetBackgroundBrightness)
+//			.SliderTooltip(LOCTEXT("BackgroundBrightness_ToolTip", "Change background brightness"))
+//			.OnPlusClicked(this, &SBackgroundColorSettings::OnIncreaseBrightness)
+//			.PlusTooltip(LOCTEXT("IncreaseBrightness_ToolTip", "Increase brightness"))
+//			.ExtraWidget(ExtraWidget)
+//			];
+//	}
+//
+//protected:
+//
+//	/** Function to open color picker window when selected from context menu */
+//	FReply OnColorBoxClicked(const FGeometry& MyGeometry, const FPointerEvent& MouseEvent)
+//	{
+//		if (MouseEvent.GetEffectingButton() != EKeys::LeftMouseButton)
+//		{
+//			return FReply::Unhandled();
+//		}
+//
+//		FSlateApplication::Get().DismissAllMenus();
+//
+//		FLinearColor BackgroundColor = AnimViewportPtr.Pin()->GetViewportBackgroundColor();
+//		TArray<FLinearColor*> LinearColorArray;
+//		LinearColorArray.Add(&BackgroundColor);
+//
+//		FColorPickerArgs PickerArgs;
+//		PickerArgs.bIsModal = true;
+//		PickerArgs.ParentWidget = AnimViewportPtr.Pin();
+//		PickerArgs.bOnlyRefreshOnOk = true;
+//		PickerArgs.DisplayGamma = TAttribute<float>::Create(TAttribute<float>::FGetter::CreateUObject(GEngine, &UEngine::GetDisplayGamma));
+//		PickerArgs.LinearColorArray = &LinearColorArray;
+//		PickerArgs.OnColorCommitted = FOnLinearColorValueChanged::CreateSP(AnimViewportPtr.Pin().ToSharedRef(), &SCustomizableObjectEditorViewport::SetViewportBackgroundColor);
+//
+//		if (OpenColorPicker(PickerArgs))
+//		{
+//			AnimViewportPtr.Pin()->RefreshViewport();
+//		}
+//
+//		return FReply::Handled();
+//	}
+//
+//	/** Callback function for decreasing color brightness */
+//	FReply OnDecreaseBrightness()
+//	{
+//		const float DeltaValue = 0.05f;
+//		AnimViewportPtr.Pin()->SetBackgroundBrightness(AnimViewportPtr.Pin()->GetBackgroundBrightness() - DeltaValue);
+//		return FReply::Handled();
+//	}
+//
+//	/** Callback function for increasing color brightness */
+//	FReply OnIncreaseBrightness()
+//	{
+//		const float DeltaValue = 0.05f;
+//		AnimViewportPtr.Pin()->SetBackgroundBrightness(AnimViewportPtr.Pin()->GetBackgroundBrightness() + DeltaValue);
+//		return FReply::Handled();
+//	}
+//
+//	/** Callback function which determines whether this widget is enabled */
+//	bool IsBrightnessSliderEnabled() const
+//	{
+//		const UAssetViewerSettings* Settings = UAssetViewerSettings::Get();
+//		const UEditorPerProjectUserSettings* PerProjectUserSettings = GetDefault<UEditorPerProjectUserSettings>();
+//		const int32 ProfileIndex = Settings->Profiles.IsValidIndex(PerProjectUserSettings->AssetViewerProfileIndex) ? PerProjectUserSettings->AssetViewerProfileIndex : 0;
+//
+//		ensureMsgf(Settings->Profiles.IsValidIndex(PerProjectUserSettings->AssetViewerProfileIndex), TEXT("Invalid default settings pointer or current profile index"));
+//
+//		return !Settings->Profiles[ProfileIndex].bShowEnvironment;
+//	}
+//
+//protected:
+//
+//	/** The viewport hosting this widget */
+//	TWeakPtr<SCustomizableObjectEditorViewport> AnimViewportPtr;
+//};
+
+//Class definition which represents widget to modify strength of wind for clothing
+//class SClothWindSettings : public SCompoundWidget
+//{
+//public:
+//
+//	SLATE_BEGIN_ARGS(SClothWindSettings)
+//	{}
+//
+//	SLATE_ARGUMENT(TWeakPtr<SCustomizableObjectEditorViewport>, AnimEditorViewport)
+//		SLATE_END_ARGS()
+//
+//		/** Constructs this widget from its declaration */
+//		void Construct(const FArguments& InArgs)
+//	{
+//		AnimViewportPtr = InArgs._AnimEditorViewport;
+//
+//		TSharedPtr<SWidget> ExtraWidget = SNew(STextBlock)
+//			.Text(AnimViewportPtr.Pin().ToSharedRef(), &SCustomizableObjectEditorViewport::GetWindStrengthLabel)
+//			.Font(FAppStyle::GetFontStyle(TEXT("MenuItem.Font")));
+//
+//		this->ChildSlot
+//			[
+//				SNew(SAnimPlusMinusSlider)
+//				.IsEnabled(this, &SClothWindSettings::IsWindEnabled)
+//			.Label(LOCTEXT("WindStrength", "Wind Strength:"))
+//			.OnMinusClicked(this, &SClothWindSettings::OnDecreaseWindStrength)
+//			.MinusTooltip(LOCTEXT("DecreaseWindStrength_ToolTip", "Decrease Wind Strength"))
+//			.SliderValue(AnimViewportPtr.Pin().ToSharedRef(), &SCustomizableObjectEditorViewport::GetWindStrengthSliderValue)
+//			.OnSliderValueChanged(AnimViewportPtr.Pin().ToSharedRef(), &SCustomizableObjectEditorViewport::SetWindStrength)
+//			.SliderTooltip(LOCTEXT("WindStrength_ToolTip", "Change wind strength"))
+//			.OnPlusClicked(this, &SClothWindSettings::OnIncreaseWindStrength)
+//			.PlusTooltip(LOCTEXT("IncreasetWindStrength_ToolTip", "Increase Wind Strength"))
+//			.ExtraWidget(ExtraWidget)
+//			];
+//	}
+//
+//protected:
+//	/** Callback function for decreasing size */
+//	FReply OnDecreaseWindStrength()
+//	{
+//		const float DeltaValue = 0.1f;
+//		AnimViewportPtr.Pin()->SetWindStrength(AnimViewportPtr.Pin()->GetWindStrengthSliderValue() - DeltaValue);
+//		return FReply::Handled();
+//	}
+//
+//	/** Callback function for increasing size */
+//	FReply OnIncreaseWindStrength()
+//	{
+//		const float DeltaValue = 0.1f;
+//		AnimViewportPtr.Pin()->SetWindStrength(AnimViewportPtr.Pin()->GetWindStrengthSliderValue() + DeltaValue);
+//		return FReply::Handled();
+//	}
+//
+//	/** Callback function which determines whether this widget is enabled */
+//	bool IsWindEnabled() const
+//	{
+//		return AnimViewportPtr.Pin()->IsApplyingClothWind();
+//	}
+//
+//protected:
+//	/** The viewport hosting this widget */
+//	TWeakPtr<SCustomizableObjectEditorViewport> AnimViewportPtr;
+//};
+
+//Class definition which represents widget to modify gravity for preview
+//class SGravitySettings : public SCompoundWidget
+//{
+//public:
+//
+//	SLATE_BEGIN_ARGS(SGravitySettings)
+//	{}
+//
+//	SLATE_ARGUMENT(TWeakPtr<SCustomizableObjectEditorViewport>, AnimEditorViewport)
+//		SLATE_END_ARGS()
+//
+//		void Construct(const FArguments& InArgs)
+//	{
+//		AnimViewportPtr = InArgs._AnimEditorViewport;
+//
+//		TSharedPtr<SWidget> ExtraWidget = SNew(STextBlock)
+//			.Text(AnimViewportPtr.Pin().ToSharedRef(), &SCustomizableObjectEditorViewport::GetGravityScaleLabel)
+//			.Font(FAppStyle::GetFontStyle(TEXT("MenuItem.Font")));
+//
+//		this->ChildSlot
+//			[
+//				SNew(SAnimPlusMinusSlider)
+//				.Label(LOCTEXT("Gravity Scale", "Gravity Scale Preview:"))
+//			.OnMinusClicked(this, &SGravitySettings::OnDecreaseGravityScale)
+//			.MinusTooltip(LOCTEXT("DecreaseGravitySize_ToolTip", "Decrease Gravity Scale"))
+//			.SliderValue(AnimViewportPtr.Pin().ToSharedRef(), &SCustomizableObjectEditorViewport::GetGravityScaleSliderValue)
+//			.OnSliderValueChanged(AnimViewportPtr.Pin().ToSharedRef(), &SCustomizableObjectEditorViewport::SetGravityScale)
+//			.SliderTooltip(LOCTEXT("GravityScale_ToolTip", "Change Gravity Scale"))
+//			.OnPlusClicked(this, &SGravitySettings::OnIncreaseGravityScale)
+//			.PlusTooltip(LOCTEXT("IncreaseGravityScale_ToolTip", "Increase Gravity Scale"))
+//			.ExtraWidget(ExtraWidget)
+//			];
+//	}
+//
+//protected:
+//	FReply OnDecreaseGravityScale()
+//	{
+//		const float DeltaValue = 0.025f;
+//		AnimViewportPtr.Pin()->SetGravityScale(AnimViewportPtr.Pin()->GetGravityScaleSliderValue() - DeltaValue);
+//		return FReply::Handled();
+//	}
+//
+//	FReply OnIncreaseGravityScale()
+//	{
+//		const float DeltaValue = 0.025f;
+//		AnimViewportPtr.Pin()->SetGravityScale(AnimViewportPtr.Pin()->GetGravityScaleSliderValue() + DeltaValue);
+//		return FReply::Handled();
+//	}
+//
+//protected:
+//	TWeakPtr<SCustomizableObjectEditorViewport> AnimViewportPtr;
+//};
+
+///////////////////////////////////////////////////////////
+// SCustomizableObjectEditorViewportToolBar
+
+
+void SCustomizableObjectEditorViewportToolBar::Construct(const FArguments& InArgs, TSharedPtr<class SCustomizableObjectEditorViewportTabBody> InViewport, TSharedPtr<class SEditorViewport> InRealViewport)
+{
+	Viewport = InViewport;
+
+	TSharedRef<SCustomizableObjectEditorViewportTabBody> ViewportRef = Viewport.Pin().ToSharedRef();
+
+	//TSharedRef<SCustomizableObjectEditorTransformViewportToolbar> Toolbar = SNew(SCustomizableObjectEditorTransformViewportToolbar)
+	//	.Viewport(InRealViewport.ToSharedRef())
+	//	.ViewportTapBody(ViewportRef);
+
+	//TSharedRef<SHorizontalBox> RightToolbar = SNew(SHorizontalBox)
+	//	+ SHorizontalBox::Slot()
+	//	.AutoWidth()
+	//	.Padding(2.0f, 2.0f)
+	//	[
+	//		Toolbar
+	//	];
+
+	//ViewportRef->SetEditorTransformViewportToolbar(Toolbar);
+
+
+	TSharedRef<SHorizontalBox> LeftToolbar = SNew(SHorizontalBox)
+	//	// Generic viewport options
+	//	+ SHorizontalBox::Slot()
+	//	.AutoWidth()
+	//	.Padding(2.0f, 2.0f)
+	//	[
+	//		//Menu
+	//		SNew(SEditorViewportToolbarMenu)
+	//		.ParentToolBar(SharedThis(this))
+	//	.Image("EditorViewportToolBar.MenuDropdown")
+	//	.OnGetMenuContent(this, &SCustomizableObjectEditorViewportToolBar::GenerateViewMenu)
+	//	]
+
+	//// Camera Type (Perspective/Top/etc...)
+	+ SHorizontalBox::Slot()
+	.AutoWidth()
+	.Padding(2.0f, 2.0f)
+	[
+		SNew(SEditorViewportToolbarMenu)
+		.ParentToolBar(SharedThis(this))
+		.Label(this, &SCustomizableObjectEditorViewportToolBar::GetCameraMenuLabel)
+		.LabelIcon(this, &SCustomizableObjectEditorViewportToolBar::GetCameraMenuLabelIcon)
+		.OnGetMenuContent(this, &SCustomizableObjectEditorViewportToolBar::GenerateViewportTypeMenu)
+	]
+
+	+ SHorizontalBox::Slot()
+		.AutoWidth()
+		.Padding(FMargin(2.0f, 2.0f))
+		[
+			SNew( SEditorViewportToolbarMenu )
+			.ParentToolBar( SharedThis( this ) )
+			.Cursor( EMouseCursor::Default )
+			.Image( "EditorViewportToolBar.MenuDropdown" )
+			.AddMetaData<FTagMetaData>(FTagMetaData(TEXT("EditorViewportToolBar.MenuDropdown")))
+			.OnGetMenuContent( this, &SCustomizableObjectEditorViewportToolBar::GenerateOptionsMenu )
+		]
+
+	// View menu (lit, unlit, etc...)
+	+ SHorizontalBox::Slot()
+		.AutoWidth()
+		.Padding(2.0f, 2.0f)
+		[
+			SNew(SEditorViewportViewMenu, InRealViewport.ToSharedRef(), SharedThis(this))
+		]
+
+	// Show flags menu
+	//+ SHorizontalBox::Slot()
+	//	.AutoWidth()
+	//	.Padding(2.0f, 2.0f)
+	//	[
+	//		SNew(SEditorViewportToolbarMenu)
+	//		.ParentToolBar(SharedThis(this))
+	//	.Label(LOCTEXT("ShowMenu", "Show"))
+	//	.OnGetMenuContent(this, &SCustomizableObjectEditorViewportToolBar::GenerateShowMenu)
+	//	]
+
+	// LOD menu
+	+ SHorizontalBox::Slot()
+		.AutoWidth()
+		.Padding(2.0f, 2.0f)
+		[
+			//LOD
+			SNew(SEditorViewportToolbarMenu)
+				.ParentToolBar(SharedThis(this))
+				.Label(this, &SCustomizableObjectEditorViewportToolBar::GetLODMenuLabel)
+				.OnGetMenuContent(this, &SCustomizableObjectEditorViewportToolBar::GenerateLODMenu)
+		]
+
+	// Camera Mode menu (Free or Orbital camera)
+	+ SHorizontalBox::Slot()
+		.AutoWidth()
+		.Padding(2.0f, 2.0f)
+		[
+			//Camera Mode
+			SNew(SEditorViewportToolbarMenu)
+				.ParentToolBar(SharedThis(this))
+				.Label(this, &SCustomizableObjectEditorViewportToolBar::GetCameraModeMenuLabel)
+				.OnGetMenuContent(this, &SCustomizableObjectEditorViewportToolBar::GenerateCameraModeMenu)
+		];
+	
+	TSharedRef<SWidget> RTSButtons = GenerateRTSButtons();
+	ViewportRef->SetViewportToolbarTransformWidget(RTSButtons);
+
+	LeftToolbar->AddSlot()
+		.AutoWidth()
+		.Padding(2.0f, 2.0f)
+		[
+			RTSButtons
+		];
+
+	// Playback speed menu
+	//+ SHorizontalBox::Slot()
+	//	.AutoWidth()
+	//	.Padding(2.0f, 2.0f)
+	//	[
+	//		SNew(SEditorViewportToolbarMenu)
+	//		.ParentToolBar(SharedThis(this))
+	//	.Label(this, &SCustomizableObjectEditorViewportToolBar::GetPlaybackMenuLabel)
+	//	.LabelIcon(FAppStyle::GetBrush("AnimViewportMenu.PlayBackSpeed"))
+	//	.OnGetMenuContent(this, &SCustomizableObjectEditorViewportToolBar::GeneratePlaybackMenu)
+	//	]
+
+	//+ SHorizontalBox::Slot()
+	//	.Padding(3.0f, 1.0f)
+	//	.HAlign(HAlign_Right)
+	//	[
+	//		SNew(STransformViewportToolBar)
+	//		.Viewport(InRealViewport)
+	//	.CommandList(InRealViewport->GetCommandList())
+	//	.Visibility(this, &SCustomizableObjectEditorViewportToolBar::GetTransformToolbarVisibility)
+	//	];
+	//@TODO: Need clipping horizontal box: LeftToolbar->AddWrapButton();
+
+	static const FName DefaultForegroundName("DefaultForeground");
+
+	FLinearColor ButtonColor1 = FLinearColor(0.1f, 0.1f, 0.1f, 1.0f);
+	FLinearColor ButtonColor2 = FLinearColor(0.2f, 0.2f, 0.2f, 0.75f);
+	FLinearColor TextColor1 = FLinearColor(1.0f, 1.0f, 1.0f, 1.0f);
+	FLinearColor TextColor2 = FLinearColor(0.8f, 0.8f, 0.8f, 0.8f);
+	FSlateFontInfo Info = FAppStyle::GetFontStyle("BoldFont");
+	Info.Size += 26;
+
+	ChildSlot
+		[
+			SNew(SVerticalBox)
+			+ SVerticalBox::Slot()
+			.VAlign(VAlign_Top)
+			[
+				SNew(SBorder)
+				.BorderImage(FAppStyle::GetBrush("NoBorder"))
+#if ENGINE_MAJOR_VERSION < 5
+				// Color and opacity is changed based on whether or not the mouse cursor is hovering over the toolbar area
+				.ColorAndOpacity(this, &SViewportToolBar::OnGetColorAndOpacity)
+#endif
+				.ForegroundColor(FAppStyle::GetSlateColor(DefaultForegroundName))
+				[
+					SNew(SVerticalBox)
+					+ SVerticalBox::Slot()
+					.AutoHeight()
+					[
+						SNew(SHorizontalBox)
+						+ SHorizontalBox::Slot()
+						.HAlign(HAlign_Left)
+						[
+							LeftToolbar
+						]
+					]
+				]
+			]
+	
+			+ SVerticalBox::Slot()
+			.HAlign(HAlign_Center)
+			.VAlign(VAlign_Center)
+			.AutoHeight()
+			[
+				SNew(SScaleBox)
+				.Stretch(EStretch::ScaleToFit)
+				[
+					SAssignNew(CompileErrorLayout, SButton)
+					.ButtonStyle(&FButtonStyle::GetDefault())
+					.ButtonColorAndOpacity(ButtonColor2)
+					.ForegroundColor(ButtonColor2)
+					.HAlign(HAlign_Center)
+					.VAlign(VAlign_Center)
+					.Visibility(this, &SCustomizableObjectEditorViewportToolBar::GetShowCompileErrorOverlay)
+					.Content()
+					[
+						SNew(STextBlock)
+						.Text(this, &SCustomizableObjectEditorViewportToolBar::GetCompileErrorOverlayText)
+						.Justification(ETextJustify::Center)
+						.ColorAndOpacity(FSlateColor(TextColor2))
+						.Font(Info)
+					]
+				]
+			]
+		];
+
+	SViewportToolBar::Construct(SViewportToolBar::FArguments());
+
+	//CompileErrorLayout.Get()->SetVisibility(EVisibility::Hidden);
+}
+
+EVisibility SCustomizableObjectEditorViewportToolBar::GetShowCompileErrorOverlay() const
+{
+	if ((Viewport.IsValid()) && (Viewport.Pin().Get() != nullptr) &&
+		Viewport.Pin()->CustomizableObjectEditorPtr.IsValid() &&
+		Viewport.Pin().Get()->GetViewportClient().IsValid() &&
+		(Viewport.Pin().Get()->GetViewportClient().Get() != nullptr))
+	{
+		if (Viewport.Pin()->CustomizableObjectEditorPtr.Pin()->GetAssetRegistryLoaded())
+		{
+			UCustomizableObjectInstance* Instance = Viewport.Pin()->CustomizableObjectEditorPtr.Pin()->GetPreviewInstance();
+
+			if ((Instance != nullptr) && !Instance->bEditorPropertyChanged && !Viewport.Pin().Get()->GetViewportClient()->GetGizmoIsProjectorParameterSelected())
+			{
+				if ((Instance->SkeletalMeshStatus == ESkeletalMeshState::UpdateError) ||
+					(Instance->SkeletalMeshStatus == ESkeletalMeshState::PostUpdateError) ||
+					(Instance->SkeletalMeshStatus == ESkeletalMeshState::AsyncUpdatePending))
+				{
+					return EVisibility::Visible;
+				}
+				else
+				{
+					return EVisibility::Hidden;
+				}
+			}
+
+			if ((Instance == nullptr) &&
+				!Viewport.Pin().Get()->GetViewportClient()->GetGizmoIsProjectorParameterSelected())
+			{
+				return EVisibility::Visible;
+			}
+		}
+		else
+		{
+			return EVisibility::Visible;
+		}
+	}
+
+	return EVisibility::Hidden;
+}
+
+
+EVisibility SCustomizableObjectEditorViewportToolBar::GetShowCompileInfoOverlay() const
+{
+	if ((Viewport.IsValid()) && (Viewport.Pin().Get() != nullptr) &&
+		(Viewport.Pin()->CustomizableObjectEditorPtr.IsValid()) &&
+		Viewport.Pin()->CustomizableObjectEditorPtr.Pin()->GetAssetRegistryLoaded())
+	{
+		if (Viewport.Pin()->CustomizableObjectEditorPtr.Pin()->GetAssetRegistryLoaded())
+		{
+			UCustomizableObjectInstance* Instance = Viewport.Pin()->CustomizableObjectEditorPtr.Pin()->GetPreviewInstance();
+
+			if (Instance != nullptr)
+			{
+				if ((Instance->SkeletalMeshStatus == ESkeletalMeshState::UpdateError) ||
+					(Instance->SkeletalMeshStatus == ESkeletalMeshState::PostUpdateError) ||
+					(Instance->SkeletalMeshStatus == ESkeletalMeshState::AsyncUpdatePending))
+				{
+					return EVisibility::Visible;
+				}
+				else
+				{
+					return EVisibility::Hidden;
+				}
+			}
+		}
+		else
+		{
+			return EVisibility::Visible;
+		}
+	}
+
+	return EVisibility::Hidden;
+}
+
+
+FText SCustomizableObjectEditorViewportToolBar::GetCompileErrorOverlayText() const
+{
+	if ((Viewport.IsValid()) && (Viewport.Pin().Get() != nullptr) &&
+		(Viewport.Pin()->CustomizableObjectEditorPtr.IsValid()))
+	{
+		if (Viewport.Pin()->CustomizableObjectEditorPtr.Pin()->GetAssetRegistryLoaded())
+		{
+			UCustomizableObjectInstance* Instance = Viewport.Pin()->CustomizableObjectEditorPtr.Pin()->GetPreviewInstance();
+
+			if (Instance != nullptr)
+			{
+				if (Instance->SkeletalMeshStatus == ESkeletalMeshState::UpdateError)
+				{
+					return FText::FromString("Error updating skeletal mesh");
+				}
+				else if (Instance->SkeletalMeshStatus == ESkeletalMeshState::PostUpdateError)
+				{
+					return FText::FromString("Post Update Error: check the output log for more information.");
+				}
+				else if (Instance->SkeletalMeshStatus == ESkeletalMeshState::AsyncUpdatePending)
+				{
+					TSharedPtr<FCustomizableObjectEditorViewportClient> viewport = Viewport.Pin()->GetViewportClient();
+					if (viewport.IsValid())
+					{
+						const GizmoRTSProxy& gizmo = viewport->GetGizmoProxy();
+						if (gizmo.AnyGizmoSelected)
+						{
+							if (gizmo.AssignedDataIsFromNode)
+							{
+								return FText::FromString("Changing default projector");
+							}
+							else
+							{
+								return FText::FromString("Changing current projector instance");
+							}
+						}
+					}
+					return FText::FromString("Updating skeletal mesh");
+				}
+			}
+			else
+			{
+				return FText::FromString("No Skeletal Mesh generated");
+			}
+		}
+		else
+		{
+			return FText::FromString("Loading assets");
+		}
+	}
+	return FText::FromString("");
+}
+
+
+EVisibility SCustomizableObjectEditorViewportToolBar::GetTransformToolbarVisibility() const
+{
+	return EVisibility::Hidden;
+	//return Viewport.Pin()->CanUseGizmos() ? EVisibility::Visible : EVisibility::Hidden;
+}
+
+
+TSharedRef<SWidget> SCustomizableObjectEditorViewportToolBar::GenerateViewMenu() const
+{
+	const FCustomizableObjectEditorViewportMenuCommands& Actions = FCustomizableObjectEditorViewportMenuCommands::Get();
+
+	const bool bInShouldCloseWindowAfterMenuSelection = true;
+	FMenuBuilder ViewMenuBuilder(bInShouldCloseWindowAfterMenuSelection, Viewport.Pin()->GetCommandList());
+	//{
+	//	ViewMenuBuilder.BeginSection("AnimViewportSceneSetup", LOCTEXT("ViewMenu_SceneSetupLabel", "Scene Setup"));
+	//	{
+	//		ViewMenuBuilder.AddMenuEntry(FCustomizableObjectEditorViewportMenuCommands::Get().PreviewSceneSettings);
+
+	//		// We cant allow animation configuration via the viewport menu in 'old regular' Persona
+	//		if (GetDefault<UPersonaOptions>()->bUseStandaloneAnimationEditors)
+	//		{
+	//			ViewMenuBuilder.AddSubMenu(
+	//				LOCTEXT("SceneSetupLabel", "Scene Setup"),
+	//				LOCTEXT("SceneSetupTooltip", "Set up preview meshes, animations etc."),
+	//				FNewMenuDelegate::CreateRaw(this, &SCustomizableObjectEditorViewportToolBar::GenerateSceneSetupMenu),
+	//				false,
+	//				FSlateIcon(FAppStyle::GetAppStyleSetName(), "AnimViewportMenu.SceneSetup")
+	//			);
+	//		}
+
+	//		ViewMenuBuilder.AddSubMenu(
+	//			LOCTEXT("TurnTableLabel", "Turn Table"),
+	//			LOCTEXT("TurnTableTooltip", "Set up auto-rotation of preview."),
+	//			FNewMenuDelegate::CreateRaw(this, &SCustomizableObjectEditorViewportToolBar::GenerateTurnTableMenu),
+	//			false,
+	//			FSlateIcon(FAppStyle::GetAppStyleSetName(), "AnimViewportMenu.TurnTableSpeed")
+	//		);
+	//	}
+	//	ViewMenuBuilder.EndSection();
+
+	//	ViewMenuBuilder.BeginSection("AnimViewportCamera", LOCTEXT("ViewMenu_CameraLabel", "Camera"));
+	//	{
+	//		ViewMenuBuilder.AddMenuEntry(FCustomizableObjectEditorViewportMenuCommands::Get().CameraFollow);
+	//	}
+	//	ViewMenuBuilder.EndSection();
+	//}
+
+	return ViewMenuBuilder.MakeWidget();
+}
+
+//TSharedRef<SWidget> SCustomizableObjectEditorViewportToolBar::GenerateShowMenu() const
+//{
+//	const FAnimViewportShowCommands& Actions = FAnimViewportShowCommands::Get();
+//
+//	const bool bInShouldCloseWindowAfterMenuSelection = true;
+//	FMenuBuilder ShowMenuBuilder(bInShouldCloseWindowAfterMenuSelection, Viewport.Pin()->GetCommandList());
+//	{
+//		ShowMenuBuilder.BeginSection("AnimViewportFOV", LOCTEXT("Viewport_FOVLabel", "Field Of View"));
+//		{
+//			const float FOVMin = 5.f;
+//			const float FOVMax = 170.f;
+//
+//			TSharedPtr<SWidget> FOVWidget = SNew(SSpinBox<float>)
+//				.Font(FAppStyle::GetFontStyle(TEXT("MenuItem.Font")))
+//				.MinValue(FOVMin)
+//				.MaxValue(FOVMax)
+//				.Value(this, &SCustomizableObjectEditorViewportToolBar::OnGetFOVValue)
+//				.OnValueChanged(this, &SCustomizableObjectEditorViewportToolBar::OnFOVValueChanged)
+//				.OnValueCommitted(this, &SCustomizableObjectEditorViewportToolBar::OnFOVValueCommitted);
+//
+//			ShowMenuBuilder.AddWidget(FOVWidget.ToSharedRef(), FText());
+//		}
+//		ShowMenuBuilder.EndSection();
+//
+//		ShowMenuBuilder.BeginSection("AnimViewportAudio", LOCTEXT("Viewport_AudioLabel", "Audio"));
+//		{
+//			ShowMenuBuilder.AddMenuEntry(Actions.MuteAudio);
+//		}
+//		ShowMenuBuilder.EndSection();
+//
+//		ShowMenuBuilder.BeginSection("AnimViewportRootMotion", LOCTEXT("Viewport_RootMotionLabel", "Root Motion"));
+//		{
+//			ShowMenuBuilder.AddMenuEntry(Actions.ProcessRootMotion);
+//		}
+//		ShowMenuBuilder.EndSection();
+//
+//		ShowMenuBuilder.BeginSection("AnimViewportMesh", LOCTEXT("ShowMenu_Actions_Mesh", "Mesh"));
+//		{
+//			ShowMenuBuilder.AddMenuEntry(Actions.ShowRetargetBasePose);
+//			ShowMenuBuilder.AddMenuEntry(Actions.ShowBound);
+//			ShowMenuBuilder.AddMenuEntry(Actions.UseInGameBound);
+//			ShowMenuBuilder.AddMenuEntry(Actions.ShowPreviewMesh);
+//			ShowMenuBuilder.AddMenuEntry(Actions.ShowMorphTargets);
+//		}
+//		ShowMenuBuilder.EndSection();
+//
+//		ShowMenuBuilder.BeginSection("AnimViewportAnimation", LOCTEXT("ShowMenu_Actions_Asset", "Asset"));
+//		{
+//			ShowMenuBuilder.AddMenuEntry(Actions.ShowRawAnimation);
+//			ShowMenuBuilder.AddMenuEntry(Actions.ShowNonRetargetedAnimation);
+//			ShowMenuBuilder.AddMenuEntry(Actions.ShowAdditiveBaseBones);
+//			ShowMenuBuilder.AddMenuEntry(Actions.ShowSourceRawAnimation);
+//			ShowMenuBuilder.AddMenuEntry(Actions.ShowBakedAnimation);
+//		}
+//		ShowMenuBuilder.EndSection();
+//
+//		ShowMenuBuilder.BeginSection("AnimViewportPreviewBones", LOCTEXT("ShowMenu_Actions_Bones", "Hierarchy"));
+//		{
+//			ShowMenuBuilder.AddSubMenu(
+//				LOCTEXT("AnimViewportBoneDrawSubMenu", "Bone"),
+//				LOCTEXT("AnimViewportBoneDrawSubMenuToolTip", "Bone Draw Option"),
+//				FNewMenuDelegate::CreateRaw(this, &SCustomizableObjectEditorViewportToolBar::FillShowBoneDrawMenu)
+//			);
+//
+//			ShowMenuBuilder.AddMenuEntry(Actions.ShowSockets);
+//			ShowMenuBuilder.AddMenuEntry(Actions.ShowBoneNames);
+//
+//			ShowMenuBuilder.AddSubMenu(
+//				LOCTEXT("AnimViewportOverlayDrawSubMenu", "Mesh Overlay"),
+//				LOCTEXT("AnimViewportOverlayDrawSubMenuToolTip", "Options"),
+//				FNewMenuDelegate::CreateRaw(this, &SCustomizableObjectEditorViewportToolBar::FillShowOverlayDrawMenu)
+//			);
+//		}
+//		ShowMenuBuilder.EndSection();
+//
+//		ShowMenuBuilder.AddMenuSeparator();
+//		ShowMenuBuilder.AddSubMenu(
+//			LOCTEXT("AnimviewportInfo", "Display Info"),
+//			LOCTEXT("AnimviewportInfoSubMenuToolTip", "Display Mesh Info in Viewport"),
+//			FNewMenuDelegate::CreateRaw(this, &SCustomizableObjectEditorViewportToolBar::FillShowDisplayInfoMenu));
+//
+//
+//		ShowMenuBuilder.AddMenuSeparator();
+//
+//		ShowMenuBuilder.AddSubMenu(
+//			LOCTEXT("AnimViewportSceneSubMenu", "Scene Setup"),
+//			LOCTEXT("AnimViewportSceneSubMenuToolTip", "Options relating to the preview scene"),
+//			FNewMenuDelegate::CreateRaw(this, &SCustomizableObjectEditorViewportToolBar::FillShowSceneMenu));
+//
+//		ShowMenuBuilder.AddSubMenu(
+//			LOCTEXT("AnimViewportAdvancedSubMenu", "Advanced"),
+//			LOCTEXT("AnimViewportAdvancedSubMenuToolTip", "Advanced options"),
+//			FNewMenuDelegate::CreateRaw(this, &SCustomizableObjectEditorViewportToolBar::FillShowAdvancedMenu));
+//	}
+//
+//	return ShowMenuBuilder.MakeWidget();
+//
+//}
+
+//void SCustomizableObjectEditorViewportToolBar::FillShowSceneMenu(FMenuBuilder& MenuBuilder) const
+//{
+//	const FAnimViewportShowCommands& Actions = FAnimViewportShowCommands::Get();
+//
+//	MenuBuilder.BeginSection("AnimViewportAccessory", LOCTEXT("Viewport_AccessoryLabel", "Accessory"));
+//	{
+//		MenuBuilder.AddMenuEntry(Actions.AutoAlignFloorToMesh);
+//	}
+//	MenuBuilder.EndSection();
+//
+//	MenuBuilder.BeginSection("AnimViewportFloorOffset", LOCTEXT("Viewport_FloorOffsetLabel", "Floor Height Offset"));
+//	{
+//		TSharedPtr<SWidget> FloorOffsetWidget = SNew(SNumericEntryBox<float>)
+//			.Font(FAppStyle::GetFontStyle(TEXT("MenuItem.Font")))
+//			.Value(this, &SCustomizableObjectEditorViewportToolBar::OnGetFloorOffset)
+//			.OnValueChanged(this, &SCustomizableObjectEditorViewportToolBar::OnFloorOffsetChanged)
+//			.ToolTipText(LOCTEXT("FloorOffsetToolTip", "Height offset for the floor mesh (stored per-mesh)"));
+//
+//		MenuBuilder.AddWidget(FloorOffsetWidget.ToSharedRef(), FText());
+//	}
+//	MenuBuilder.EndSection();
+//
+//	MenuBuilder.BeginSection("AnimViewportGrid", LOCTEXT("Viewport_GridLabel", "Grid"));
+//	{
+//		MenuBuilder.AddMenuEntry(Actions.ToggleGrid);
+//	}
+//	MenuBuilder.EndSection();
+//
+//	MenuBuilder.BeginSection("AnimViewportBackground", LOCTEXT("Viewport_BackgroundLabel", "Background"));
+//	{
+//		TSharedPtr<SWidget> BackgroundColorWidget = SNew(SBackgroundColorSettings).AnimEditorViewport(Viewport);
+//		MenuBuilder.AddWidget(BackgroundColorWidget.ToSharedRef(), FText());
+//	}
+//	MenuBuilder.EndSection();
+//}
+
+//void SCustomizableObjectEditorViewportToolBar::FillShowBoneDrawMenu(FMenuBuilder& MenuBuilder) const
+//{
+//	const FAnimViewportShowCommands& Actions = FAnimViewportShowCommands::Get();
+//
+//	MenuBuilder.BeginSection("AnimViewportPreviewHierarchyBoneDraw", LOCTEXT("ShowMenu_Actions_HierarchyAxes", "Hierarchy Local Axes"));
+//	{
+//		MenuBuilder.AddMenuEntry(Actions.ShowBoneDrawAll);
+//		MenuBuilder.AddMenuEntry(Actions.ShowBoneDrawSelected);
+//		MenuBuilder.AddMenuEntry(Actions.ShowBoneDrawNone);
+//	}
+//	MenuBuilder.EndSection();
+//}
+//
+//void SCustomizableObjectEditorViewportToolBar::FillShowOverlayDrawMenu(FMenuBuilder& MenuBuilder) const
+//{
+//	const FAnimViewportShowCommands& Actions = FAnimViewportShowCommands::Get();
+//
+//	MenuBuilder.BeginSection("AnimViewportPreviewOverlayDraw", LOCTEXT("ShowMenu_Actions_Overlay", "Overlay Draw"));
+//	{
+//		MenuBuilder.AddMenuEntry(Actions.ShowOverlayNone);
+//		MenuBuilder.AddMenuEntry(Actions.ShowBoneWeight);
+//		MenuBuilder.AddMenuEntry(Actions.ShowMorphTargetVerts);
+//	}
+//	MenuBuilder.EndSection();
+//}
+//
+//void SCustomizableObjectEditorViewportToolBar::FillShowAdvancedMenu(FMenuBuilder& MenuBuilder) const
+//{
+//	const FAnimViewportShowCommands& Actions = FAnimViewportShowCommands::Get();
+//
+//	// Draw UVs
+//	MenuBuilder.BeginSection("UVVisualization", LOCTEXT("UVVisualization_Label", "UV Visualization"));
+//	{
+//		MenuBuilder.AddMenuEntry(FAnimViewportMenuCommands::Get().AnimSetDrawUVs);
+//		MenuBuilder.AddWidget(Viewport.Pin()->UVChannelCombo.ToSharedRef(), FText());
+//	}
+//	MenuBuilder.EndSection();
+//
+//	MenuBuilder.BeginSection("Skinning", LOCTEXT("Skinning_Label", "Skinning"));
+//	{
+//		MenuBuilder.AddMenuEntry(FAnimViewportMenuCommands::Get().SetCPUSkinning);
+//	}
+//	MenuBuilder.EndSection();
+//
+//
+//	MenuBuilder.BeginSection("ShowVertex", LOCTEXT("ShowVertex_Label", "Vertex Normal Visualization"));
+//	{
+//		// Vertex debug flags
+//		MenuBuilder.AddMenuEntry(FAnimViewportMenuCommands::Get().SetShowNormals);
+//		MenuBuilder.AddMenuEntry(FAnimViewportMenuCommands::Get().SetShowTangents);
+//		MenuBuilder.AddMenuEntry(FAnimViewportMenuCommands::Get().SetShowBinormals);
+//	}
+//	MenuBuilder.EndSection();
+//
+//	MenuBuilder.BeginSection("AnimViewportPreviewHierarchyLocalAxes", LOCTEXT("ShowMenu_Actions_HierarchyAxes", "Hierarchy Local Axes"));
+//	{
+//		MenuBuilder.AddMenuEntry(Actions.ShowLocalAxesAll);
+//		MenuBuilder.AddMenuEntry(Actions.ShowLocalAxesSelected);
+//		MenuBuilder.AddMenuEntry(Actions.ShowLocalAxesNone);
+//	}
+//	MenuBuilder.EndSection();
+//}
+//
+//void SCustomizableObjectEditorViewportToolBar::FillShowClothingMenu(FMenuBuilder& MenuBuilder) const
+//{
+//}
+//
+//void SCustomizableObjectEditorViewportToolBar::FillShowDisplayInfoMenu(FMenuBuilder& MenuBuilder) const
+//{
+//	const FAnimViewportShowCommands& Actions = FAnimViewportShowCommands::Get();
+//
+//	// display info levels
+//	{
+//		MenuBuilder.AddMenuEntry(Actions.ShowDisplayInfoBasic);
+//		MenuBuilder.AddMenuEntry(Actions.ShowDisplayInfoDetailed);
+//		MenuBuilder.AddMenuEntry(Actions.ShowDisplayInfoSkelControls);
+//		MenuBuilder.AddMenuEntry(Actions.HideDisplayInfo);
+//	}
+//}
+
+FText SCustomizableObjectEditorViewportToolBar::GetLODMenuLabel() const
+{
+	FText Label = LOCTEXT("LODMenu_AutoLabel", "LOD Auto");
+	if (Viewport.IsValid())
+	{
+		int32 LODSelectionType = Viewport.Pin()->GetLODSelection();
+
+		if (LODSelectionType > 0)
+		{
+			FString TitleLabel = FString::Printf(TEXT("LOD %d"), LODSelectionType - 1);
+			Label = FText::FromString(TitleLabel);
+		}
+	}
+	return Label;
+}
+
+TSharedRef<SWidget> SCustomizableObjectEditorViewportToolBar::GenerateRTSButtons()
+{
+	const FCustomizableObjectEditorViewportLODCommands& Actions = FCustomizableObjectEditorViewportLODCommands::Get();
+
+	TSharedPtr< FExtender > InExtenders;
+	FToolBarBuilder ToolbarBuilder(Viewport.Pin()->GetCommandList(), FMultiBoxCustomization::None, InExtenders);
+
+	// Use a custom style
+	FName ToolBarStyle = "EditorViewportToolBar";
+	ToolbarBuilder.SetStyle(&FAppStyle::Get(), ToolBarStyle);
+	ToolbarBuilder.SetLabelVisibility(EVisibility::Collapsed);
+
+	// Transform controls cannot be focusable as it fights with the press space to change transform mode feature
+	ToolbarBuilder.SetIsFocusable(false);
+
+	ToolbarBuilder.BeginSection("Transform");
+	ToolbarBuilder.BeginBlockGroup();
+	{
+		// Move Mode
+		static FName TranslateModeName = FName(TEXT("TranslateMode"));
+		ToolbarBuilder.AddToolBarButton(FCustomizableObjectEditorViewportLODCommands::Get().TranslateMode, NAME_None, TAttribute<FText>(), TAttribute<FText>(), TAttribute<FSlateIcon>(), TranslateModeName);
+
+		// Rotate Mode
+		static FName RotateModeName = FName(TEXT("RotateMode"));
+		ToolbarBuilder.AddToolBarButton(FCustomizableObjectEditorViewportLODCommands::Get().RotateMode, NAME_None, TAttribute<FText>(), TAttribute<FText>(), TAttribute<FSlateIcon>(), RotateModeName);
+
+		// Scale Mode
+		static FName ScaleModeName = FName(TEXT("ScaleMode"));
+		ToolbarBuilder.AddToolBarButton(FCustomizableObjectEditorViewportLODCommands::Get().ScaleMode, NAME_None, TAttribute<FText>(), TAttribute<FText>(), TAttribute<FSlateIcon>(), ScaleModeName);
+
+	}
+	ToolbarBuilder.EndBlockGroup();
+	ToolbarBuilder.EndSection();
+
+	ToolbarBuilder.BeginSection("RotationGridSnap");
+	{
+		// Grab the existing UICommand 
+		FUICommandInfo* Command = FCustomizableObjectEditorViewportLODCommands::Get().RotationGridSnap.Get();
+
+		static FName RotationSnapName = FName(TEXT("RotationSnap"));
+
+		// Setup a GridSnapSetting with the UICommand
+		ToolbarBuilder.AddWidget(SNew(SViewportToolBarComboMenu)
+			.Cursor(EMouseCursor::Default)
+			.Style(ToolBarStyle)
+			.IsChecked(this, &SCustomizableObjectEditorViewportToolBar::IsRotationGridSnapChecked)
+			.OnCheckStateChanged(this, &SCustomizableObjectEditorViewportToolBar::HandleToggleRotationGridSnap)
+			.Label(this, &SCustomizableObjectEditorViewportToolBar::GetRotationGridLabel)
+			.OnGetMenuContent(this, &SCustomizableObjectEditorViewportToolBar::FillRotationGridSnapMenu)
+			.ToggleButtonToolTip(Command->GetDescription())
+			.MenuButtonToolTip(LOCTEXT("RotationGridSnap_ToolTip", "Set the Rotation Grid Snap value"))
+			.Icon(Command->GetIcon())
+			.ParentToolBar(SharedThis(this))
+			, RotationSnapName);
+	}
+
+	ToolbarBuilder.EndSection();
+
+	ToolbarBuilder.SetIsFocusable(true);
+
+	return ToolbarBuilder.MakeWidget();
+}
+
+TSharedRef<SWidget> SCustomizableObjectEditorViewportToolBar::GenerateLODMenu() const
+{
+	const FCustomizableObjectEditorViewportLODCommands& Actions = FCustomizableObjectEditorViewportLODCommands::Get();
+
+	const bool bInShouldCloseWindowAfterMenuSelection = true;
+	FMenuBuilder ShowMenuBuilder(bInShouldCloseWindowAfterMenuSelection, Viewport.Pin()->GetCommandList());
+	{
+		// LOD Models
+		ShowMenuBuilder.BeginSection("AnimViewportPreviewLODs", LOCTEXT("ShowLOD_PreviewLabel", "Preview LODs"));
+		{
+			ShowMenuBuilder.AddMenuEntry(Actions.LODAuto);
+			ShowMenuBuilder.AddMenuEntry(Actions.LOD0);
+
+			int32 LODCount = Viewport.Pin()->GetLODModelCount();
+			for (int32 LODId = 1; LODId < LODCount; ++LODId)
+			{
+				FString TitleLabel = FString::Printf(TEXT("LOD %d"), LODId);
+
+				FUIAction Action(FExecuteAction::CreateSP(Viewport.Pin().ToSharedRef(), &SCustomizableObjectEditorViewportTabBody::OnSetLODModel, LODId + 1),
+					FCanExecuteAction(),
+					FIsActionChecked::CreateSP(Viewport.Pin().ToSharedRef(), &SCustomizableObjectEditorViewportTabBody::IsLODModelSelected, LODId + 1));
+
+				ShowMenuBuilder.AddMenuEntry(FText::FromString(TitleLabel), FText::GetEmpty(), FSlateIcon(), Action, NAME_None, EUserInterfaceActionType::RadioButton);
+			}
+		}
+		ShowMenuBuilder.EndSection();
+	}
+
+	return ShowMenuBuilder.MakeWidget();
+}
+
+TSharedRef<SWidget> SCustomizableObjectEditorViewportToolBar::GenerateViewportTypeMenu() const
+{
+	const bool bInShouldCloseWindowAfterMenuSelection = true;
+	FMenuBuilder CameraMenuBuilder(bInShouldCloseWindowAfterMenuSelection, Viewport.Pin()->GetCommandList());
+
+	// Camera types
+	CameraMenuBuilder.AddMenuEntry(FEditorViewportCommands::Get().Perspective);
+
+	CameraMenuBuilder.BeginSection("LevelViewportCameraType_Ortho", LOCTEXT("CameraTypeHeader_Ortho", "Orthographic"));
+	CameraMenuBuilder.AddMenuEntry(FEditorViewportCommands::Get().Top);
+	CameraMenuBuilder.AddMenuEntry(FEditorViewportCommands::Get().Bottom);
+	CameraMenuBuilder.AddMenuEntry(FEditorViewportCommands::Get().Left);
+	CameraMenuBuilder.AddMenuEntry(FEditorViewportCommands::Get().Right);
+	CameraMenuBuilder.AddMenuEntry(FEditorViewportCommands::Get().Front);
+	CameraMenuBuilder.AddMenuEntry(FEditorViewportCommands::Get().Back);
+	CameraMenuBuilder.EndSection();
+
+	return CameraMenuBuilder.MakeWidget();
+}
+
+TSharedRef<SWidget> SCustomizableObjectEditorViewportToolBar::GeneratePlaybackMenu() const
+{
+	//const FCustomizableObjectEditorViewportPlaybackCommands& Actions = FCustomizableObjectEditorViewportPlaybackCommands::Get();
+
+	const bool bInShouldCloseWindowAfterMenuSelection = true;
+
+	FMenuBuilder PlaybackMenuBuilder(bInShouldCloseWindowAfterMenuSelection, Viewport.Pin()->GetCommandList());
+	//{
+	//	// View modes
+	//	{
+	//		PlaybackMenuBuilder.BeginSection("AnimViewportPlaybackSpeed", LOCTEXT("PlaybackMenu_SpeedLabel", "Playback Speed"));
+	//		{
+	//			for (int i = 0; i < EAnimationPlaybackSpeeds::NumPlaybackSpeeds; ++i)
+	//			{
+	//				PlaybackMenuBuilder.AddMenuEntry(Actions.PlaybackSpeedCommands[i]);
+	//			}
+	//		}
+	//		PlaybackMenuBuilder.EndSection();
+	//	}
+	//}
+
+	return PlaybackMenuBuilder.MakeWidget();
+
+}
+
+void SCustomizableObjectEditorViewportToolBar::GenerateTurnTableMenu(FMenuBuilder& MenuBuilder) const
+{
+	//const FCustomizableObjectEditorViewportPlaybackCommands& Actions = FCustomizableObjectEditorViewportPlaybackCommands::Get();
+
+	//const bool bInShouldCloseWindowAfterMenuSelection = true;
+
+	//MenuBuilder.BeginSection("AnimViewportTurnTableMode", LOCTEXT("TurnTableMenu_ModeLabel", "Turn Table Mode"));
+	//{
+	//	MenuBuilder.AddMenuEntry(Actions.PersonaTurnTablePlay);
+	//	MenuBuilder.AddMenuEntry(Actions.PersonaTurnTablePause);
+	//	MenuBuilder.AddMenuEntry(Actions.PersonaTurnTableStop);
+	//}
+	//MenuBuilder.EndSection();
+
+	//MenuBuilder.BeginSection("AnimViewportTurnTableSpeed", LOCTEXT("TurnTableMenu_SpeedLabel", "Turn Table Speed"));
+	//{
+	//	for (int i = 0; i < EAnimationPlaybackSpeeds::NumPlaybackSpeeds; ++i)
+	//	{
+	//		MenuBuilder.AddMenuEntry(Actions.TurnTableSpeeds[i]);
+	//	}
+	//}
+	//MenuBuilder.EndSection();
+}
+
+void SCustomizableObjectEditorViewportToolBar::GenerateSceneSetupMenu(FMenuBuilder& MenuBuilder)
+{
+	//FPropertyEditorModule& PropertyEditorModule = FModuleManager::LoadModuleChecked<FPropertyEditorModule>("PropertyEditor");
+	//FDetailsViewArgs Args(false, false, false, FDetailsViewArgs::HideNameArea, true, nullptr, false, "PersonaPreviewSceneDescription");
+	//Args.bShowScrollBar = false;
+
+	//TSharedRef<IDetailsView> DetailsView = PropertyEditorModule.CreateDetailView(Args);
+
+	//DetailsView->RegisterInstancedCustomPropertyLayout(UPersonaPreviewSceneDescription::StaticClass(), FOnGetDetailCustomizationInstance::CreateSP(this, &SCustomizableObjectEditorViewportToolBar::CustomizePreviewSceneDescription));
+	//PropertyEditorModule.RegisterCustomPropertyTypeLayout(FPreviewMeshCollectionEntry::StaticStruct()->GetFName(), FOnGetPropertyTypeCustomizationInstance::CreateSP(this, &SCustomizableObjectEditorViewportToolBar::CustomizePreviewMeshCollectionEntry), nullptr, DetailsView);
+
+	//DetailsView->SetObject(StaticCastSharedRef<FAnimationEditorPreviewScene>(Viewport.Pin()->GetPreviewScene())->GetPreviewSceneDescription());
+
+	//MenuBuilder.AddWidget(DetailsView, FText(), true);
+}
+
+TSharedRef<SWidget> SCustomizableObjectEditorViewportToolBar::GenerateCameraModeMenu() const
+{
+	const FCustomizableObjectEditorViewportLODCommands& Actions = FCustomizableObjectEditorViewportLODCommands::Get();
+
+	const bool bInShouldCloseWindowAfterMenuSelection = true;
+
+	FMenuBuilder ShowMenuBuilder(bInShouldCloseWindowAfterMenuSelection, Viewport.Pin()->GetCommandList());
+	{
+		// LOD Models
+		ShowMenuBuilder.BeginSection("NAME_None");
+		{
+			ShowMenuBuilder.AddMenuEntry(Actions.OrbitalCamera);
+			ShowMenuBuilder.AddMenuEntry(Actions.FreeCamera);
+		}
+		ShowMenuBuilder.EndSection();
+	}
+
+	return ShowMenuBuilder.MakeWidget();
+}
+
+FText SCustomizableObjectEditorViewportToolBar::GetCameraModeMenuLabel() const
+{
+	FText Label = FText::FromString("Camera Mode");
+	if (Viewport.IsValid())
+	{
+		bool CameraMode = Viewport.Pin()->GetViewportClient()->IsOrbitalCameraActive();
+
+		if (CameraMode)
+		{
+			Label = FText::FromString("Camera Mode: Orbital");
+		}
+		else
+		{
+			Label = FText::FromString("Camera Mode: Free");
+		}
+	}
+	return Label;
+}
+
+//TSharedRef<class IDetailCustomization> SCustomizableObjectEditorViewportToolBar::CustomizePreviewSceneDescription()
+//{
+//	return MakeShareable(new FPreviewSceneDescriptionCustomization(FAssetData(&Viewport.Pin()->GetSkeletonTree()->GetEditableSkeleton()->GetSkeleton()).GetExportTextName(), Viewport.Pin()->GetPreviewScene()->GetPersonaToolkit()));
+//}
+//
+//TSharedRef<class IPropertyTypeCustomization> SCustomizableObjectEditorViewportToolBar::CustomizePreviewMeshCollectionEntry()
+//{
+//	return MakeShareable(new FPreviewMeshCollectionEntryCustomization(Viewport.Pin()->GetPreviewScene()));
+//}
+
+FSlateColor SCustomizableObjectEditorViewportToolBar::GetFontColor() const
+{
+	const UAssetViewerSettings* Settings = UAssetViewerSettings::Get();
+	const UEditorPerProjectUserSettings* PerProjectUserSettings = GetDefault<UEditorPerProjectUserSettings>();
+	const int32 ProfileIndex = Settings->Profiles.IsValidIndex(PerProjectUserSettings->AssetViewerProfileIndex) ? PerProjectUserSettings->AssetViewerProfileIndex : 0;
+
+	ensureMsgf(Settings->Profiles.IsValidIndex(PerProjectUserSettings->AssetViewerProfileIndex), TEXT("Invalid default settings pointer or current profile index"));
+
+	FLinearColor FontColor;
+	if (Settings->Profiles[ProfileIndex].bShowEnvironment)
+	{
+		FontColor = FLinearColor::White;
+	}
+	else
+	{
+		FLinearColor BackgroundColorInHSV = Viewport.Pin()->GetViewportBackgroundColor().LinearRGBToHSV();
+
+		// see if it's dark, if V is less than 0.2
+		if (BackgroundColorInHSV.B < 0.3f)
+		{
+			FontColor = FLinearColor::White;
+		}
+		else
+		{
+			FontColor = FLinearColor::Black;
+		}
+	}
+
+	return FontColor;
+}
+
+FText SCustomizableObjectEditorViewportToolBar::GetPlaybackMenuLabel() const
+{
+	FText Label = LOCTEXT("PlaybackError", "Error");
+	//if (Viewport.IsValid())
+	//{
+	//	for (int i = 0; i < EAnimationPlaybackSpeeds::NumPlaybackSpeeds; ++i)
+	//	{
+	//		if (Viewport.Pin()->IsPlaybackSpeedSelected(i))
+	//		{
+	//			Label = FText::FromString(FString::Printf(
+	//				(i == EAnimationPlaybackSpeeds::Quarter) ? TEXT("x%.2f") : TEXT("x%.1f"),
+	//				EAnimationPlaybackSpeeds::Values[i]
+	//			));
+	//			break;
+	//		}
+	//	}
+	//}
+	return Label;
+}
+
+FText SCustomizableObjectEditorViewportToolBar::GetCameraMenuLabel() const
+{
+	FText Label = LOCTEXT("Viewport_Default", "Camera");
+	TSharedPtr< SCustomizableObjectEditorViewportTabBody > PinnedViewport(Viewport.Pin());
+	if (PinnedViewport.IsValid())
+	{
+		switch (PinnedViewport->GetViewportClient()->ViewportType)
+		{
+		case LVT_Perspective:
+			Label = LOCTEXT("CameraMenuTitle_Perspective", "Perspective");
+			break;
+
+		case LVT_OrthoXY:
+			Label = LOCTEXT("CameraMenuTitle_Top", "Top");
+			break;
+
+		case LVT_OrthoNegativeXZ:
+			Label = LOCTEXT("CameraMenuTitle_Left", "Left");
+			break;
+
+		case LVT_OrthoNegativeYZ:
+			Label = LOCTEXT("CameraMenuTitle_Front", "Front");
+			break;
+
+		case LVT_OrthoNegativeXY:
+			Label = LOCTEXT("CameraMenuTitle_Bottom", "Bottom");
+			break;
+
+		case LVT_OrthoXZ:
+			Label = LOCTEXT("CameraMenuTitle_Right", "Right");
+			break;
+
+		case LVT_OrthoYZ:
+			Label = LOCTEXT("CameraMenuTitle_Back", "Back");
+			break;
+		case LVT_OrthoFreelook:
+			break;
+		}
+	}
+
+	return Label;
+}
+
+const FSlateBrush* SCustomizableObjectEditorViewportToolBar::GetCameraMenuLabelIcon() const
+{
+	FName Icon = NAME_None;
+	TSharedPtr< SCustomizableObjectEditorViewportTabBody > PinnedViewport(Viewport.Pin());
+	if (PinnedViewport.IsValid())
+	{
+		switch (PinnedViewport->GetViewportClient()->ViewportType)
+		{
+		case LVT_Perspective:
+			Icon = FName("EditorViewport.Perspective");
+			break;
+
+		case LVT_OrthoXY:
+			Icon = FName("EditorViewport.Top");
+			break;
+
+		case LVT_OrthoYZ:
+			Icon = FName("EditorViewport.Back");
+			break;
+
+		case LVT_OrthoXZ:
+			Icon = FName("EditorViewport.Right");
+			break;
+
+		case LVT_OrthoNegativeXY:
+			Icon = FName("EditorViewport.Bottom");
+			break;
+
+		case LVT_OrthoNegativeYZ:
+			Icon = FName("EditorViewport.Front");
+			break;
+
+		case LVT_OrthoNegativeXZ:
+			Icon = FName("EditorViewport.Left");
+			break;
+		case LVT_OrthoFreelook:
+			break;
+		}
+	}
+
+	return FAppStyle::GetBrush(Icon);
+}
+
+float SCustomizableObjectEditorViewportToolBar::OnGetFOVValue() const
+{
+	return Viewport.Pin()->GetViewportClient()->ViewFOV;
+}
+
+void SCustomizableObjectEditorViewportToolBar::OnFOVValueChanged(float NewValue) const
+{
+	TSharedPtr<FCustomizableObjectEditorViewportClient> ViewportClient = Viewport.Pin()->GetViewportClient();
+
+#if (ENGINE_MAJOR_VERSION == 4 && ENGINE_MINOR_VERSION >= 18) || ENGINE_MAJOR_VERSION >= 5
+	ViewportClient->FOVAngle = NewValue;
+	// \todo: this editor name should be somewhere else.
+	FString EditorName("CustomizableObjectEditor");
+	int ViewportIndex=0;
+	ViewportClient->ConfigOption->SetViewFOV(FName(*EditorName),NewValue,ViewportIndex);
+#else
+	ViewportClient->FOVAngle = NewValue;
+	ViewportClient->ConfigOption->SetViewFOV(NewValue);
+#endif
+
+	ViewportClient->ViewFOV = NewValue;
+	ViewportClient->Invalidate();
+}
+
+void SCustomizableObjectEditorViewportToolBar::OnFOVValueCommitted(float NewValue, ETextCommit::Type CommitInfo)
+{
+	//OnFOVValueChanged will be called... nothing needed here.
+}
+
+TOptional<float> SCustomizableObjectEditorViewportToolBar::OnGetFloorOffset() const
+{
+	TSharedPtr<FCustomizableObjectEditorViewportClient> ViewportClient = Viewport.Pin()->GetViewportClient();
+
+	return ViewportClient->GetFloorOffset();
+}
+
+void SCustomizableObjectEditorViewportToolBar::OnFloorOffsetChanged(float NewValue)
+{
+	TSharedPtr<FCustomizableObjectEditorViewportClient> ViewportClient = Viewport.Pin()->GetViewportClient();
+
+	ViewportClient->SetFloorOffset(NewValue);
+}
+
+
+ECheckBoxState SCustomizableObjectEditorViewportToolBar::IsRotationGridSnapChecked() const
+{
+	return GetDefault<ULevelEditorViewportSettings>()->RotGridEnabled ? ECheckBoxState::Checked : ECheckBoxState::Unchecked;
+}
+
+
+void SCustomizableObjectEditorViewportToolBar::HandleToggleRotationGridSnap(ECheckBoxState InState)
+{
+	GUnrealEd->Exec(GEditor->GetEditorWorldContext().World(), *FString::Printf(TEXT("MODE ROTGRID=%d"), !GetDefault<ULevelEditorViewportSettings>()->RotGridEnabled ? 1 : 0));
+}
+
+
+FText SCustomizableObjectEditorViewportToolBar::GetRotationGridLabel() const
+{
+	return FText::Format(LOCTEXT("GridRotation - Number - DegreeSymbol", "{0}\u00b0"), FText::AsNumber(GEditor->GetRotGridSize().Pitch));
+}
+
+
+TSharedRef<SWidget> SCustomizableObjectEditorViewportToolBar::FillRotationGridSnapMenu()
+{
+	const ULevelEditorViewportSettings* ViewportSettings = GetDefault<ULevelEditorViewportSettings>();
+
+	return SNew(SUniformGridPanel)
+
+		+ SUniformGridPanel::Slot(0, 0)
+		[
+			BuildRotationGridCheckBoxList("Common", LOCTEXT("RotationCommonText", "Common"), ViewportSettings->CommonRotGridSizes, GridMode_Common)
+		]
+
+		+ SUniformGridPanel::Slot(1, 0)
+		[
+			BuildRotationGridCheckBoxList("Div360", LOCTEXT("RotationDivisions360DegreesText", "Divisions of 360\u00b0"), ViewportSettings->DivisionsOf360RotGridSizes, GridMode_DivisionsOf360)
+		];
+}
+
+
+TSharedRef<SWidget> SCustomizableObjectEditorViewportToolBar::BuildRotationGridCheckBoxList(FName InExtentionHook, const FText& InHeading, const TArray<float>& InGridSizes, ERotationGridMode InGridMode) const
+{
+	const bool bShouldCloseWindowAfterMenuSelection = true;
+	FMenuBuilder RotationGridMenuBuilder(bShouldCloseWindowAfterMenuSelection, Viewport.Pin()->GetCommandList());
+
+	RotationGridMenuBuilder.BeginSection(InExtentionHook, InHeading);
+	for (int32 CurGridAngleIndex = 0; CurGridAngleIndex < InGridSizes.Num(); ++CurGridAngleIndex)
+	{
+		const float CurGridAngle = InGridSizes[CurGridAngleIndex];
+
+		FText MenuName = FText::Format(LOCTEXT("RotationGridAngle", "{0}\u00b0"), FText::AsNumber(CurGridAngle)); /*degree symbol*/
+		FText ToolTipText = FText::Format(LOCTEXT("RotationGridAngle_ToolTip", "Sets rotation grid angle to {0}"), MenuName); /*degree symbol*/
+
+		RotationGridMenuBuilder.AddMenuEntry(
+			MenuName,
+			ToolTipText,
+			FSlateIcon(),
+			FUIAction(FExecuteAction::CreateStatic(&SCustomizableObjectEditorViewportTabBody::SetRotationGridSize, CurGridAngleIndex, InGridMode),
+				FCanExecuteAction(),
+				FIsActionChecked::CreateStatic(&SCustomizableObjectEditorViewportTabBody::IsRotationGridSizeChecked, CurGridAngleIndex, InGridMode)),
+			NAME_None,
+			EUserInterfaceActionType::RadioButton);
+	}
+	RotationGridMenuBuilder.EndSection();
+
+	return RotationGridMenuBuilder.MakeWidget();
+}
+
+
+FReply SCustomizableObjectEditorViewportToolBar::OnMenuClicked()
+{
+	// If the menu button is clicked toggle the state of the menu anchor which will open or close the menu
+	if (MenuAnchor->ShouldOpenDueToClick())
+	{
+		MenuAnchor->SetIsOpen(true);
+		this->SetOpenMenu(MenuAnchor);
+	}
+	else
+	{
+		MenuAnchor->SetIsOpen(false);
+		TSharedPtr<SMenuAnchor> NullAnchor;
+		this->SetOpenMenu(MenuAnchor);
+	}
+
+	return FReply::Handled();
+}
+
+
+TSharedRef<SWidget> SCustomizableObjectEditorViewportToolBar::GenerateOptionsMenu() const
+{
+	const FCustomizableObjectEditorViewportLODCommands& LevelViewportActions = FCustomizableObjectEditorViewportLODCommands::Get();
+
+	// Get all menu extenders for this context menu from the level editor module
+	FLevelEditorModule& LevelEditorModule = FModuleManager::GetModuleChecked<FLevelEditorModule>( TEXT("LevelEditor") );
+	TArray<FLevelEditorModule::FLevelEditorMenuExtender> MenuExtenderDelegates = LevelEditorModule.GetAllLevelViewportOptionsMenuExtenders();
+	
+	TArray<TSharedPtr<FExtender>> Extenders;
+	for (int32 i = 0; i < MenuExtenderDelegates.Num(); ++i)
+	{
+		if (MenuExtenderDelegates[i].IsBound())
+		{
+			Extenders.Add(MenuExtenderDelegates[i].Execute(Viewport.Pin()->GetCommandList().ToSharedRef()));
+		}
+	}
+	TSharedPtr<FExtender> MenuExtender = FExtender::Combine(Extenders);
+
+	const bool bIsPerspective = Viewport.Pin()->GetViewportClient()->IsPerspective();
+	const bool bInShouldCloseWindowAfterMenuSelection = true;
+	FMenuBuilder OptionsMenuBuilder( bInShouldCloseWindowAfterMenuSelection, Viewport.Pin()->GetCommandList(), MenuExtender );
+	{
+		OptionsMenuBuilder.AddWidget( GenerateFOVMenu(), LOCTEXT("FOVAngle", "Field of View (H)") );
+		OptionsMenuBuilder.AddMenuEntry( LevelViewportActions.HighResScreenshot );
+	}
+
+	return OptionsMenuBuilder.MakeWidget();
+}
+
+
+TSharedRef<SWidget> SCustomizableObjectEditorViewportToolBar::GenerateFOVMenu() const
+{
+	const float FOVMin = 5.f;
+	const float FOVMax = 170.f;
+
+	return
+		SNew( SBox )
+		.HAlign( HAlign_Right )
+		[
+			SNew( SBox )
+			.Padding( FMargin(4.0f, 0.0f, 0.0f, 0.0f) )
+			.WidthOverride( 100.0f )
+			[
+				SNew(SSpinBox<float>)
+				.Font( FAppStyle::GetFontStyle( TEXT( "MenuItem.Font" ) ) )
+				.MinValue(FOVMin)
+				.MaxValue(FOVMax)
+				.Value( this, &SCustomizableObjectEditorViewportToolBar::OnGetFOVValue )
+				.OnValueChanged( this, &SCustomizableObjectEditorViewportToolBar::OnFOVValueChanged )
+			]
+		];
+}
+
+
+#undef LOCTEXT_NAMESPACE
diff --git a/Engine/Plugins/Experimental/Mutable/Source/CustomizableObjectEditor/Private/SCustomizableObjectLayoutGrid.cpp b/Engine/Plugins/Experimental/Mutable/Source/CustomizableObjectEditor/Private/SCustomizableObjectLayoutGrid.cpp
new file mode 100644
--- /dev/null
+++ b/Engine/Plugins/Experimental/Mutable/Source/CustomizableObjectEditor/Private/SCustomizableObjectLayoutGrid.cpp
@@ -0,0 +1,1182 @@
+// Copyright Epic Games, Inc. All Rights Reserved.
+
+#include "SCustomizableObjectLayoutGrid.h"
+#include "CustomizableObjectEditorModule.h"
+#include "CustomizableObjectLayout.h"
+#include "Styling/AppStyle.h"
+#include "Async/TaskGraphInterfaces.h"
+
+#include "CanvasItem.h"
+#include "CanvasTypes.h"
+
+#include "Framework/MultiBox/MultiBoxBuilder.h"
+#include "Framework/Application/SlateApplication.h"
+
+#include "Input/CursorReply.h"
+
+#include "Widgets/Input/SNumericEntryBox.h"
+
+// Required for engine branch preprocessor defines.
+#include "CustomizableObject/Public/UnrealPortabilityHelpers.h"
+
+#define LOCTEXT_NAMESPACE "CustomizableObjectEditor"
+
+/** Simple representation of the backbuffer for drawin UVs. */
+class FSlateCanvasRenderTarget : public FRenderTarget
+{
+public:
+	FIntPoint GetSizeXY() const override
+	{
+		return ViewRect.Size();
+	}
+
+	/** Sets the texture that this target renders to */
+	void SetRenderTargetTexture(FTexture2DRHIRef& InRHIRef)
+	{
+		RenderTargetTextureRHI = InRHIRef;
+	}
+
+	/** Clears the render target texture */
+	void ClearRenderTargetTexture()
+	{
+		RenderTargetTextureRHI.SafeRelease();
+	}
+
+	/** Sets the viewport rect for the render target */
+	void SetViewRect(const FIntRect& InViewRect)
+	{
+		ViewRect = InViewRect;
+	}
+
+	/** Gets the viewport rect for the render target */
+	const FIntRect& GetViewRect() const
+	{
+		return ViewRect;
+	}
+
+private:
+	FIntRect ViewRect;
+};
+
+
+/** Custom Slate drawing element. Holds a copy of all information required to draw UVs. */
+class FUVCanvasDrawer : public ICustomSlateElement
+{
+public:
+	~FUVCanvasDrawer()
+	{
+		delete RenderTarget;
+	};
+
+	/** Set the canvas area and all required data to paint the UVs.
+	 * 
+	 * All data will be copied.
+	 */
+	void Initialize(const FIntRect& InCanvasRect, const FVector2D& InOrigin, const FVector2D& InSize, const TArray<FVector2f>& InUVLayout);
+	
+private:
+	void DrawRenderThread(FRHICommandListImmediate& RHICmdList, const void* RenderTarget) override;
+
+	/** SlateElement initialized, can Draw during the DrawRenderThread call. */
+	bool Initialized = false;
+
+	/** Drawing origin. */
+	FVector2D Origin;
+
+	/** Drawing size. */
+	FVector2D Size;
+
+	/** Drawing UVLayout. */
+	TArray<FVector2D> UVLayout;
+
+	FSlateCanvasRenderTarget* RenderTarget = new FSlateCanvasRenderTarget();
+};
+
+
+void SCustomizableObjectLayoutGrid::Construct( const FArguments& InArgs )
+{
+	GridSize = InArgs._GridSize;
+	Blocks = InArgs._Blocks;
+	UVLayout = InArgs._UVLayout;
+	UnassignedUVLayoutVertices = InArgs._UnassignedUVLayoutVertices;
+	Mode = InArgs._Mode;
+	BlockChangedDelegate = InArgs._OnBlockChanged;
+	SelectionChangedDelegate = InArgs._OnSelectionChanged;
+	SelectionColor = InArgs._SelectionColor;
+	DeleteBlocksDelegate = InArgs._OnDeleteBlocks;
+	AddBlockAtDelegate = InArgs._OnAddBlockAt;
+	OnSetBlockPriority = InArgs._OnSetBlockPriority;
+
+	HasDragged = false;
+	Dragging = false;
+	Resizing = false;
+	ResizeCursor = false;
+	Selecting = false;
+
+	PaddingAmount = FVector2D::Zero();
+	DistanceFromOrigin = FVector2D::Zero();
+	Zoom = 1;
+
+	UVCanvasDrawer = TSharedPtr<FUVCanvasDrawer, ESPMode::ThreadSafe>(new FUVCanvasDrawer());
+}
+
+
+SCustomizableObjectLayoutGrid::~SCustomizableObjectLayoutGrid()
+{
+	// UVCanvasDrawer can only be destroyed after drawing the last command
+	ENQUEUE_RENDER_COMMAND(SafeDeletePreviewElement)(
+		[UVCanvasDrawer = UVCanvasDrawer](FRHICommandListImmediate& RHICmdList) mutable
+		{
+			UVCanvasDrawer.Reset();
+		}
+	);
+}
+
+
+int32 SCustomizableObjectLayoutGrid::OnPaint(const FPaintArgs& Args, const FGeometry& AllottedGeometry, const FSlateRect& MyClippingRect, FSlateWindowElementList& OutDrawElements, int32 LayerId, const FWidgetStyle& InWidgetStyle, bool bParentEnabled ) const
+{
+	int32 RetLayerId = SCompoundWidget::OnPaint(Args, AllottedGeometry, MyClippingRect, OutDrawElements, LayerId,InWidgetStyle, bParentEnabled );
+
+	bool bEnabled = ShouldBeEnabled( bParentEnabled );
+	const ESlateDrawEffect DrawEffects = bEnabled ? ESlateDrawEffect::None : ESlateDrawEffect::DisabledEffect;
+
+	// Paint inside the border only. 
+	const FVector2D BorderPadding = FVector2D(2,2);
+	FPaintGeometry ForegroundPaintGeometry = AllottedGeometry.ToInflatedPaintGeometry( -BorderPadding );
+	
+	const FIntPoint GridSizePoint = GridSize.Get();
+	const float OffsetX = BorderPadding.X;
+	const FVector2D AreaSize =  AllottedGeometry.GetLocalSize() - 2.0f * BorderPadding;
+	const float GridRatio = float(GridSizePoint.X) / float(GridSizePoint.Y);
+	FVector2D Size;
+	if ( AreaSize.X/GridRatio > AreaSize.Y )
+	{
+		Size.Y = AreaSize.Y;
+		Size.X = AreaSize.Y*GridRatio;
+	}
+	else
+	{
+		Size.X =  AreaSize.X;
+		Size.Y =  AreaSize.X/GridRatio;
+	}
+
+	FVector2D OldSize = Size;
+	Size *= Zoom;
+
+	float AuxCellSize = Size.X / GridSizePoint.X;
+	
+	// Drawing Offsets
+	FVector2D Offset = (AreaSize-Size)/2.0f;
+	FVector2D ZoomOffset = ((Size - OldSize) / 2.0f);
+	
+	// Drawing Origin
+	FVector2D Origin = BorderPadding + Offset + PaddingAmount - DistanceFromOrigin;
+
+	// Create line points
+	TArray< FVector2D > LinePoints;
+	LinePoints.SetNum(2);
+	
+	// Drawing Vertical Lines
+	for( int32 LineIndex = 0; LineIndex < GridSizePoint.X + 1; LineIndex++ )
+	{
+		LinePoints[0] = FVector2D( Origin.X + LineIndex * AuxCellSize, Origin.Y );
+		LinePoints[1] = FVector2D( Origin.X + LineIndex * AuxCellSize, Origin.Y + Size.Y );
+
+		FSlateDrawElement::MakeLines( OutDrawElements,RetLayerId,AllottedGeometry.ToPaintGeometry(),
+			LinePoints, ESlateDrawEffect::None,	FColor(150, 150, 150, 64), false, 2.0);
+	}
+
+	// Drawing Horizontal Lines
+	for( int32 LineIndex = 0; LineIndex < GridSizePoint.Y + 1; LineIndex++ )
+	{
+		LinePoints[0] = FVector2D( Origin.X, Origin.Y + LineIndex * AuxCellSize );
+		LinePoints[1] = FVector2D( Origin.X + Size.X, Origin.Y + LineIndex * AuxCellSize );
+
+		FSlateDrawElement::MakeLines( OutDrawElements, RetLayerId, AllottedGeometry.ToPaintGeometry(),
+			LinePoints,	ESlateDrawEffect::None, FColor(150, 150, 150, 64), false, 2.0);
+	}
+
+	RetLayerId++;
+
+	// Draw UV using a CustomSlateElement on the RenderThread
+	const float CanvasMinX = FMath::Max(0.0f, AllottedGeometry.AbsolutePosition.X);
+	const float CanvasMinY = FMath::Max(0.0f, AllottedGeometry.AbsolutePosition.Y);
+	const FIntRect CanvasRect(
+		FMath::RoundToInt(CanvasMinX),
+		FMath::RoundToInt(CanvasMinY),
+		FMath::RoundToInt(CanvasMinX + AllottedGeometry.GetLocalSize().X * AllottedGeometry.Scale),
+		FMath::RoundToInt(CanvasMinY + AllottedGeometry.GetLocalSize().Y * AllottedGeometry.Scale));
+	
+	UVCanvasDrawer->Initialize(CanvasRect, Origin * AllottedGeometry.Scale, Size * AllottedGeometry.Scale, UVLayout);
+	FSlateDrawElement::MakeCustom(OutDrawElements, RetLayerId, UVCanvasDrawer);
+
+	const auto MakeYellowSquareLine = [&](const TArray<FVector2D>& Points) -> void
+	{
+		FSlateDrawElement::MakeLines( OutDrawElements, RetLayerId, AllottedGeometry.ToPaintGeometry(),
+			Points,	ESlateDrawEffect::None,	FColor(250, 230, 43, 255), true, 2.0);
+	};
+
+	TArray<FVector2D> SquareLinePoints;
+	SquareLinePoints.SetNum(2);
+
+	const FVector2D CrossSize = Size*0.01;
+	for (const FVector2f& Vertex : UnassignedUVLayoutVertices)
+	{
+		SquareLinePoints[0] = (Origin + FVector2D(Vertex) * Size + FVector2D(CrossSize));
+		SquareLinePoints[1]	= (Origin + FVector2D(Vertex) * Size - FVector2D(CrossSize) * FVector2D(1.0f, -1.0f));
+		MakeYellowSquareLine(SquareLinePoints);
+		
+		SquareLinePoints[0]	= (Origin + FVector2D(Vertex) * Size - FVector2D(CrossSize));
+		MakeYellowSquareLine(SquareLinePoints);
+
+		SquareLinePoints[1]	= (Origin + FVector2D(Vertex) * Size + FVector2D(CrossSize) * FVector2D(1.0f, -1.0f));
+		MakeYellowSquareLine(SquareLinePoints);
+		
+		SquareLinePoints[0]	= (Origin + FVector2D(Vertex) * Size + FVector2D(CrossSize));
+		MakeYellowSquareLine(SquareLinePoints);
+	}
+
+
+	// Blocks
+	const FSlateBrush* BlockBrush = FAppStyle::GetBrush("TextBlock.HighlightShape");
+	const FVector2f PaintGeomPosition = AllottedGeometry.ToPaintGeometry().DrawPosition;
+	const float PaintGeomScale = AllottedGeometry.ToPaintGeometry().DrawScale;
+	
+	for (const FCustomizableObjectLayoutBlock& Block : Blocks.Get())
+	{
+		FSlateRenderTransform GeomTransform = FSlateRenderTransform(1.0f, FVector2D(PaintGeomPosition + BlockRects[Block.Id].Rect.Min * PaintGeomScale));
+		FPaintGeometry Geom(FSlateLayoutTransform(), GeomTransform, FVector2D(BlockRects[Block.Id].Rect.Size * PaintGeomScale), false );
+		
+		FSlateDrawElement::MakeBox(	OutDrawElements, RetLayerId, Geom, BlockBrush,
+			DrawEffects, SelectedBlocks.Find(Block.Id)>=0? SelectionColor :FColor(230,199,75,155) );	
+
+		if (Mode==ELGM_Edit)
+		{
+			FSlateRenderTransform HandleGeomTransform = FSlateRenderTransform(1.0f, FVector2D(PaintGeomPosition + BlockRects[Block.Id].HandleRect.Min * PaintGeomScale));
+			FPaintGeometry HandleGeom(FSlateLayoutTransform(), HandleGeomTransform, FVector2D(BlockRects[Block.Id].HandleRect.Size * PaintGeomScale), false);
+			
+			FColor HandleRectColor = FColor(255, 96, 96, 255);
+			bool bCanResize = SelectedBlocks.Num() == 1 && SelectedBlocks.Contains(Block.Id) && MouseOnBlock(Block.Id, CurrentMousePosition, true);
+
+			if (bCanResize)
+			{
+				HandleRectColor = FColor(200, 0, 0, 255);
+			}
+
+			FSlateDrawElement::MakeBox(	OutDrawElements, RetLayerId, HandleGeom, BlockBrush,
+				DrawEffects, HandleRectColor);
+
+
+			//Selected block outline
+			if (SelectedBlocks.Find(Block.Id) >= 0)
+			{
+				TArray<FVector2D> SelectionSquarePoints;
+				SelectionSquarePoints.SetNum(2);
+
+				FVector2D RectMin = FVector2D(BlockRects[Block.Id].Rect.Min);
+				FVector2D RectMax = FVector2D(BlockRects[Block.Id].Rect.Min + BlockRects[Block.Id].Rect.Size);
+
+				FVector2D TopRightCorner = FVector2D(RectMax.X, RectMin.Y);
+				FVector2D BottomLeftCorner = FVector2D(RectMin.X, RectMax.Y);
+
+				SelectionSquarePoints[0] = RectMin;
+				SelectionSquarePoints[1] = TopRightCorner;
+				MakeYellowSquareLine(SelectionSquarePoints);
+
+				SelectionSquarePoints[0] = RectMax;
+				MakeYellowSquareLine(SelectionSquarePoints);
+				
+				SelectionSquarePoints[1] = BottomLeftCorner;
+				MakeYellowSquareLine(SelectionSquarePoints);
+				
+				SelectionSquarePoints[0] = RectMin;
+				MakeYellowSquareLine(SelectionSquarePoints);
+			}
+		}
+	}
+
+	RetLayerId++;
+
+	// Drawing Multi-Selection rect
+	if (Mode == ELGM_Edit && Selecting)
+	{
+		TArray<FVector2D> SelectionSquarePoints;
+		SelectionSquarePoints.SetNum(2);
+
+		FVector2D RectMin = FVector2D(SelectionRect.Min);
+		FVector2D RectSize = FVector2D(SelectionRect.Size);
+
+		FVector2D TopLeft = RectMin;
+		FVector2D TopRight = RectMin + FVector2D(RectSize.X, 0.0f);
+		FVector2D BottomRight = RectMin + RectSize;
+		FVector2D BottomLeft = RectMin + FVector2D(0.0f, RectSize.Y);
+
+		SelectionSquarePoints[0] = TopLeft;
+		SelectionSquarePoints[1] = TopRight;
+		MakeYellowSquareLine(SelectionSquarePoints);
+
+		SelectionSquarePoints[0] = BottomRight;
+		MakeYellowSquareLine(SelectionSquarePoints);
+
+		SelectionSquarePoints[1] = BottomLeft;
+		MakeYellowSquareLine(SelectionSquarePoints);
+
+		SelectionSquarePoints[0] = TopLeft;
+		MakeYellowSquareLine(SelectionSquarePoints);
+	}
+
+	RetLayerId++;
+
+	return RetLayerId - 1;
+}
+
+
+void SCustomizableObjectLayoutGrid::Tick( const FGeometry& AllottedGeometry, const double InCurrentTime, const float InDeltaTime )
+{
+	const FVector2D BorderPadding = FVector2D(2,2);
+	const FVector2D AreaSize =  AllottedGeometry.Size - 2.0f * BorderPadding;
+	const float GridRatio = float(GridSize.Get().X)/float(GridSize.Get().Y);
+	FVector2D Size;
+	if ( AreaSize.X/GridRatio > AreaSize.Y )
+	{
+		Size.Y = AreaSize.Y;
+		Size.X = AreaSize.Y*GridRatio;
+	}
+	else
+	{
+		Size.X =  AreaSize.X;
+		Size.Y =  AreaSize.X/GridRatio;
+	}
+
+	FVector2D OldSize = Size;
+	Size *= Zoom;
+	CellSize = Size.X/GridSize.Get().X;
+	FVector2D Offset = (AreaSize-Size)/2.0f;
+	FVector2D ZoomOffset = (Size - OldSize) / 2.0f;
+	FVector2D Origin = BorderPadding + Offset + PaddingAmount - DistanceFromOrigin;
+	DrawOrigin = Origin;
+
+	BlockRects.Empty();
+
+	const TArray<FCustomizableObjectLayoutBlock>& CurrentBlocks = Blocks.Get();
+	for (const FCustomizableObjectLayoutBlock& Block : CurrentBlocks)
+	{
+		const FVector2f BlockMin(Block.Min);
+		const FVector2f BlockMax(Block.Max);
+
+		FBlockWidgetData BlockData;
+		BlockData.Rect.Min = FVector2f(Origin) + BlockMin * CellSize + CellSize * 0.1f;
+		BlockData.Rect.Size = (BlockMax - BlockMin) * CellSize - CellSize * 0.2f;
+
+		float HandleRectSize = FMath::Log2(float(GridSize.Get().X))/10.0f;
+		BlockData.HandleRect.Size = FVector2f(CellSize) * HandleRectSize;
+		BlockData.HandleRect.Min = BlockData.Rect.Min + BlockData.Rect.Size - BlockData.HandleRect.Size;
+
+		BlockRects.Add(Block.Id, BlockData );
+	}
+
+	// Update selection list
+	for (int i=0; i<SelectedBlocks.Num();)
+	{
+		bool Found = false;
+		for (const FCustomizableObjectLayoutBlock& Block : CurrentBlocks)
+		{
+			if (Block.Id == SelectedBlocks[i])
+			{
+				Found = true;
+			}
+		}
+
+		if ( !Found )
+		{
+			SelectedBlocks.RemoveAt(i);
+		}
+		else
+		{
+			++i;
+		}
+	}
+
+	if (Selecting)
+	{
+		CalculateSelectionRect();
+	}
+
+	SCompoundWidget::Tick( AllottedGeometry, InCurrentTime, InDeltaTime );
+}
+
+
+FReply SCustomizableObjectLayoutGrid::OnMouseButtonDown( const FGeometry& MyGeometry, const FPointerEvent& MouseEvent )
+{
+	if (Mode == ELGM_Edit)
+	{
+		if (MouseEvent.GetEffectingButton() == EKeys::LeftMouseButton)
+		{
+			HasDragged = false;
+			Dragging = false;
+			Resizing = false;
+
+			// To know if we clicked on a block
+			bool ClickOnBlock = false;
+
+			FVector2D Pos = MyGeometry.AbsoluteToLocal(MouseEvent.GetScreenSpacePosition());
+			InitSelectionRect = Pos;
+
+			//Reset Selection Rect
+			SelectionRect.Size = FVector2f::Zero();
+			SelectionRect.Min = FVector2f(Pos);
+
+			// Handles selection must be detected on mouse down
+			// We also check if we click on a block
+			TArray<FGuid> SelectedBlockHandles;
+
+			for (const FGuid& BlockId : SelectedBlocks)
+			{
+				if (MouseOnBlock(BlockId, Pos, true))
+				{
+					SelectedBlockHandles.Add(BlockId);
+				}
+
+				if (MouseOnBlock(BlockId, Pos))
+				{
+					if (SelectedBlocks.Contains(BlockId))
+					{
+						ClickOnBlock = true;
+					}
+				}
+			}
+
+			if (SelectedBlocks.Num() && ClickOnBlock)
+			{
+				Dragging = true;
+				DragStart = Pos;
+
+				if (SelectedBlocks.Num() == 1 && SelectedBlockHandles.Contains(SelectedBlocks[0]))
+				{
+					Resizing = true;
+				}
+			}
+		}
+		else if (MouseEvent.GetEffectingButton() == EKeys::RightMouseButton)
+		{
+			//Mouse position
+			FVector2D Pos = MyGeometry.AbsoluteToLocal(MouseEvent.GetScreenSpacePosition());
+			FVector2D CellDelta = (Pos - DrawOrigin) / CellSize;
+
+			// Create context menu
+			const bool CloseAfterSelection = true;
+			FMenuBuilder MenuBuilder(CloseAfterSelection, NULL, TSharedPtr<FExtender>(), false, &FCoreStyle::Get(), false);
+
+			MenuBuilder.BeginSection("Block Management", LOCTEXT("GridActionsTitle", "Grid Actions"));
+			{
+				if (SelectedBlocks.Num())
+				{
+					FUIAction DeleteAction(FExecuteAction::CreateSP(this, &SCustomizableObjectLayoutGrid::DeleteSelectedBlocks));
+					MenuBuilder.AddMenuEntry(LOCTEXT("DeleteBlocksLabel", "Delete"), LOCTEXT("DeleteBlocksTooltip", "Delete Selected Blocks"), FSlateIcon(), DeleteAction);
+
+					FUIAction DuplicateAction(FExecuteAction::CreateSP(this, &SCustomizableObjectLayoutGrid::DuplicateBlocks));
+					MenuBuilder.AddMenuEntry(LOCTEXT("DuplicateBlocksLabel", "Duplicate"), LOCTEXT("DuplicateBlocksTooltip", "Duplicate Selected Blocks"), FSlateIcon(), DuplicateAction);
+				}
+				else
+				{
+					FUIAction AddNewBlockAction(FExecuteAction::CreateSP(this, &SCustomizableObjectLayoutGrid::GenerateNewBlock, CellDelta));
+					MenuBuilder.AddMenuEntry(LOCTEXT("AddNewBlockLabel", "Add Block"), LOCTEXT("AddNewBlockTooltip", "Add New Block"), FSlateIcon(), AddNewBlockAction);
+				}
+			}
+			MenuBuilder.EndSection();
+
+			MenuBuilder.BeginSection("Fixed Layout Strategy", LOCTEXT("BlockActionsTitle", "Fixed Layout Actions"));
+			{
+				if (SelectedBlocks.Num() && LayoutStrategy == ECustomizableObjectTextureLayoutPackingStrategy::Fixed)
+				{
+					MenuBuilder.AddWidget(
+						SNew(SBox)
+						.WidthOverride(125.0f)
+						[
+							SNew(SNumericEntryBox<int32>)
+							.MinValue(0)
+						.MaxValue(INT_MAX)
+						.MaxSliderValue(100)
+						.AllowSpin(SelectedBlocks.Num() == 1)
+						.Value(this, &SCustomizableObjectLayoutGrid::GetBlockPriortyValue)
+						.UndeterminedString(LOCTEXT("MultipleValues", "Multiples Values"))
+						.OnValueChanged(this, &SCustomizableObjectLayoutGrid::OnBlockPriorityChanged)
+						.ToolTipText(LOCTEXT("SetBlockPriorityTooltip", "Sets the block priority for a Fixed Layout Strategy"))
+						.EditableTextBoxStyle(&FAppStyle::GetWidgetStyle<FEditableTextBoxStyle>("NormalEditableTextBox"))
+						]
+					, FText::FromString("Block Priority"), true);
+				}
+			}
+			MenuBuilder.EndSection();
+
+			FWidgetPath WidgetPath = MouseEvent.GetEventPath() != nullptr ? *MouseEvent.GetEventPath() : FWidgetPath();
+			FSlateApplication::Get().PushMenu(AsShared(), WidgetPath, MenuBuilder.MakeWidget(), FSlateApplication::Get().GetCursorPos(), FPopupTransitionEffect(FPopupTransitionEffect::ContextMenu));
+		}
+		else if (MouseEvent.GetEffectingButton() == EKeys::MiddleMouseButton)
+		{
+			if (Zoom == 2)
+			{
+				Padding = true;
+				PaddingStart = MyGeometry.AbsoluteToLocal(MouseEvent.GetScreenSpacePosition());
+			}
+		}
+	}
+
+	return SCompoundWidget::OnMouseButtonDown( MyGeometry, MouseEvent );
+}
+
+
+FReply SCustomizableObjectLayoutGrid::OnMouseButtonUp( const FGeometry& MyGeometry, const FPointerEvent& MouseEvent )
+{
+	if (Mode == ELGM_Show)
+	{
+		return SCompoundWidget::OnMouseButtonUp(MyGeometry, MouseEvent);
+	}
+
+	if ( MouseEvent.GetEffectingButton()==EKeys::LeftMouseButton)
+	{
+		Dragging = false;
+		Resizing = false;
+
+		// Left Shif is pressed for multi selection
+		bool bLeftShift = MouseEvent.GetModifierKeys().IsLeftShiftDown();
+
+		// Screen to Widget Position
+		FVector2D Pos = MyGeometry.AbsoluteToLocal(MouseEvent.GetScreenSpacePosition());
+		
+		// Selection before reset
+		TArray<FGuid> OldSelection = SelectedBlocks;
+		TArray<FGuid> OldPossibleSelection = PossibleSelectedBlocks;
+
+		PossibleSelectedBlocks.Reset();
+
+		// Reset selection if multi selection is not enabled
+		if (Mode == ELGM_Edit && !bLeftShift && !HasDragged)
+		{
+			// Only one selected block allowed in edit mode.
+			SelectedBlocks.Reset();
+		}
+
+		if (!Selecting)
+		{
+			if (!HasDragged)
+			{
+				//Backward iteration to select the block rendered in front of the rest
+				const TArray<FCustomizableObjectLayoutBlock>& CurrentBlocks = Blocks.Get();
+				for (int32 i = CurrentBlocks.Num() - 1; i > -1; --i)
+				{
+					if (MouseOnBlock(CurrentBlocks[i].Id, Pos))
+					{
+						PossibleSelectedBlocks.Add(CurrentBlocks[i].Id);
+					}
+				}
+
+				bool bSameSelection = PossibleSelectedBlocks == OldPossibleSelection;
+
+				for (int32 i = 0; i < PossibleSelectedBlocks.Num(); ++i)
+				{
+					if (bLeftShift || Mode == ELGM_Select)
+					{
+						if (PossibleSelectedBlocks.Num() == 1)
+						{
+							if (SelectedBlocks.Contains(PossibleSelectedBlocks[i]))
+							{
+								SelectedBlocks.Remove(PossibleSelectedBlocks[i]);
+							}
+							else
+							{
+								SelectedBlocks.Add(PossibleSelectedBlocks[i]);
+								break;
+							}
+						}
+						else
+						{
+							if (!SelectedBlocks.Contains(PossibleSelectedBlocks[i]))
+							{
+								SelectedBlocks.Add(PossibleSelectedBlocks[i]);
+								break;
+							}
+						}
+					}
+					else
+					{
+						if (OldSelection.Num() == 0)
+						{
+							SelectedBlocks.Add(PossibleSelectedBlocks[0]);
+						}
+
+						if (bSameSelection)
+						{
+							if (OldSelection.Contains(PossibleSelectedBlocks[i]))
+							{
+								SelectedBlocks.Remove(PossibleSelectedBlocks[i]);
+
+								if (i == PossibleSelectedBlocks.Num() - 1)
+								{
+									SelectedBlocks.Add(PossibleSelectedBlocks[0]);
+									break;
+								}
+								else
+								{
+									SelectedBlocks.Add(PossibleSelectedBlocks[i + 1]);
+								}
+							}
+						}
+						else
+						{
+							if (OldSelection.Contains(PossibleSelectedBlocks[i]) && PossibleSelectedBlocks.Num() > 1)
+							{
+								SelectedBlocks.Remove(PossibleSelectedBlocks[i]);
+							}
+							else
+							{
+								SelectedBlocks.AddUnique(PossibleSelectedBlocks[i]);
+								break;
+							}
+						}
+					}
+				}
+			}
+		}
+		else
+		{
+			FBox2D SelectRect(FVector2D(SelectionRect.Min), FVector2D(SelectionRect.Min + SelectionRect.Size) );
+			
+			const TArray<FCustomizableObjectLayoutBlock>& CurrentBlocks = Blocks.Get();
+			for (int32 i = 0; i < CurrentBlocks.Num(); ++i)
+			{
+				FBox2D CurrentBlock(FVector2D(BlockRects[CurrentBlocks[i].Id].Rect.Min), FVector2D(BlockRects[CurrentBlocks[i].Id].Rect.Min + BlockRects[CurrentBlocks[i].Id].Rect.Size));
+				
+				if (SelectedBlocks.Contains(CurrentBlocks[i].Id))
+				{
+					if (!SelectRect.Intersect(CurrentBlock) && !bLeftShift)
+					{
+						SelectedBlocks.Remove(CurrentBlocks[i].Id);
+					}
+				}
+				else
+				{
+					if (SelectRect.Intersect(CurrentBlock))
+					{
+						SelectedBlocks.Add(CurrentBlocks[i].Id);
+					}
+				}
+			}
+		}
+
+		// Executing selection delegate
+		if (OldSelection != SelectedBlocks)
+		{
+			SelectionChangedDelegate.ExecuteIfBound(SelectedBlocks);
+		}
+
+		HasDragged = false;
+		Selecting = false;
+	}
+	else if (MouseEvent.GetEffectingButton() == EKeys::MiddleMouseButton)
+	{
+		Padding = false;
+	}
+
+	return SCompoundWidget::OnMouseButtonUp( MyGeometry, MouseEvent );
+}
+
+
+FReply SCustomizableObjectLayoutGrid::OnMouseMove( const FGeometry& MyGeometry, const FPointerEvent& MouseEvent )
+{
+	CurrentMousePosition = MyGeometry.AbsoluteToLocal(MouseEvent.GetScreenSpacePosition());
+
+	if (Mode != ELGM_Edit)
+	{
+		return SCompoundWidget::OnMouseMove(MyGeometry, MouseEvent);
+	}
+
+	if(MouseEvent.IsMouseButtonDown(EKeys::LeftMouseButton))
+	{
+		FVector2D Pos = MyGeometry.AbsoluteToLocal(MouseEvent.GetScreenSpacePosition());
+
+		if (Dragging && SelectedBlocks.Num())
+		{
+			FVector2D CellDelta = (Pos - DragStart) / CellSize;
+
+			int CellDeltaX = CellDelta.X;
+			int CellDeltaY = CellDelta.Y;
+
+			DragStart += FVector2D(CellDeltaX * CellSize, CellDeltaY * CellSize);
+
+			if (CellDeltaX || CellDeltaY)
+			{
+				HasDragged = true;
+
+				const TArray<FCustomizableObjectLayoutBlock>& CurrentBlocks = Blocks.Get();
+
+				if (!Resizing)
+				{
+					FIntRect TotalBlock;
+					bool bFirstBlock = true;
+
+					for (const FCustomizableObjectLayoutBlock& B : CurrentBlocks)
+					{
+						FIntRect Block(B.Min, B.Max);
+
+						if (SelectedBlocks.Contains(B.Id))
+						{
+							if (bFirstBlock)
+							{
+								TotalBlock = Block;
+								bFirstBlock = false;
+							}
+
+							TotalBlock.Min.X = FMath::Min(TotalBlock.Min.X, Block.Min.X);
+							TotalBlock.Min.Y = FMath::Min(TotalBlock.Min.Y, Block.Min.Y);
+							TotalBlock.Max.X = FMath::Max(TotalBlock.Max.X, Block.Max.X);
+							TotalBlock.Max.Y = FMath::Max(TotalBlock.Max.Y, Block.Max.Y);
+						}
+					}
+
+					FIntPoint Grid = GridSize.Get();
+					FIntRect BlockMovement = TotalBlock;
+					BlockMovement.Min.X = FMath::Max(0, FMath::Min(TotalBlock.Min.X + CellDeltaX, Grid.X - TotalBlock.Size().X));
+					BlockMovement.Min.Y = FMath::Max(0, FMath::Min(TotalBlock.Min.Y + CellDeltaY, Grid.Y - TotalBlock.Size().Y));
+
+					BlockMovement.Max = BlockMovement.Min + TotalBlock.Size();
+
+					FIntRect AddMovement = BlockMovement - TotalBlock;
+
+					for (const FCustomizableObjectLayoutBlock& B : CurrentBlocks)
+					{
+						if (SelectedBlocks.Find(B.Id) != INDEX_NONE)
+						{
+							FIntRect ResultBlock(B.Min, B.Max);
+							ResultBlock.Max += AddMovement.Max;
+							ResultBlock.Min += AddMovement.Min;
+
+							BlockChangedDelegate.ExecuteIfBound(B.Id, ResultBlock);
+						}
+					}
+				}
+				else
+				{
+					for (const FCustomizableObjectLayoutBlock& B : CurrentBlocks)
+					{
+						FIntRect Block;
+						for (const FGuid& Id : SelectedBlocks)
+						{
+							if (B.Id != Id)
+							{
+								continue;
+							}
+
+							Block.Min = B.Min;
+							Block.Max = B.Max;
+
+							FIntRect InitialBlock = Block;
+
+							FIntPoint Grid = GridSize.Get();
+
+							FIntPoint BlockSize = Block.Size();
+							Block.Max.X = FMath::Max(Block.Min.X + 1, FMath::Min(Block.Max.X + CellDeltaX, Grid.X));
+							Block.Max.Y = FMath::Max(Block.Min.Y + 1, FMath::Min(Block.Max.Y + CellDeltaY, Grid.Y));
+
+							if (Block != InitialBlock)
+							{
+								BlockChangedDelegate.ExecuteIfBound(Id, Block);
+							}
+
+							break;
+						}
+					}
+				}
+			}
+		}
+
+		if (!Selecting && !Dragging)
+		{
+			bool ClickOnBlock = false;
+
+			for (const FGuid& BlockId : SelectedBlocks)
+			{
+				if (MouseOnBlock(BlockId, Pos))
+				{
+					if (SelectedBlocks.Contains(BlockId))
+					{
+						ClickOnBlock = true;
+					}
+				}
+			}
+
+			int32 MovementSensitivity = 4;
+			FVector2D MouseDiference = InitSelectionRect - Pos;
+			MouseDiference = MouseDiference.GetAbs();
+
+			if (!ClickOnBlock && (MouseDiference.X > MovementSensitivity || MouseDiference.Y > MovementSensitivity))
+			{
+				HasDragged = true;
+				Selecting = true;
+			}
+		}
+	}
+	
+	if (!Dragging && !Resizing && SelectedBlocks.Num()==1)
+	{
+		const TArray<FCustomizableObjectLayoutBlock>& CurrentBlocks = Blocks.Get();
+		for (int i = CurrentBlocks.Num() - 1; i > -1; --i)
+		{
+			// Check for new created blocks
+			if (BlockRects.Contains(CurrentBlocks[i].Id) && SelectedBlocks.Contains(CurrentBlocks[i].Id))
+			{
+				FVector2D Pos = MyGeometry.AbsoluteToLocal(MouseEvent.GetScreenSpacePosition());
+				if (MouseOnBlock(CurrentBlocks[i].Id, Pos, true))
+				{
+					ResizeCursor = true;
+					break;
+				}
+			}
+
+			ResizeCursor = false;
+		}
+	}
+
+	// In case we lose focus
+	if (Padding)
+	{
+		if (MouseEvent.IsMouseButtonDown(EKeys::MiddleMouseButton))
+		{
+			FVector2D Pos = MyGeometry.AbsoluteToLocal(MouseEvent.GetScreenSpacePosition());
+			PaddingAmount += Pos - PaddingStart;
+			PaddingStart = Pos;
+		}
+		else
+		{
+			Padding = false;
+		}
+	}
+
+	if (!MouseEvent.IsMouseButtonDown(EKeys::LeftMouseButton))
+	{
+		Selecting = false;
+		Dragging = false;
+
+		if (Resizing)
+		{
+			ResizeCursor = false;
+			Resizing = false;
+		}
+	}
+
+	return SCompoundWidget::OnMouseMove( MyGeometry, MouseEvent );
+}
+
+
+FReply SCustomizableObjectLayoutGrid::OnMouseWheel(const FGeometry& MyGeometry, const FPointerEvent& MouseEvent)
+{
+	if (Mode == ELGM_Edit)
+	{
+		if (MouseEvent.GetWheelDelta() > 0)
+		{
+			if (Zoom < 2)
+			{
+				FVector2D GridCenter = DrawOrigin + (FVector2D((float)GridSize.Get().X, (float)GridSize.Get().Y) / 2.0f) * CellSize;
+				DistanceFromOrigin = CurrentMousePosition - GridCenter;
+
+				Zoom++;
+			}
+		}
+		else
+		{
+			if (Zoom > 1)
+			{
+				DistanceFromOrigin = FVector2D::Zero();
+				PaddingAmount = FVector2D::Zero();
+
+				Zoom--;
+			}
+		}
+	}
+
+	return SCompoundWidget::OnMouseWheel(MyGeometry, MouseEvent);
+}
+
+
+FReply SCustomizableObjectLayoutGrid::OnKeyDown(const FGeometry& MyGeometry, const FKeyEvent& InKeyEvent)
+{
+	if (Mode != ELGM_Edit)
+	{
+		return SCompoundWidget::OnKeyDown(MyGeometry, InKeyEvent);
+	}
+	
+	if (InKeyEvent.IsLeftControlDown())
+	{
+		if (InKeyEvent.GetKey() == EKeys::D)
+		{
+			DuplicateBlocks();
+		}
+		else if (InKeyEvent.GetKey() == EKeys::N)
+		{
+			FVector2D MouseToCellPosition = (CurrentMousePosition - DrawOrigin) / CellSize;
+			GenerateNewBlock(MouseToCellPosition);
+		}
+	}
+
+	if (InKeyEvent.GetKey() == EKeys::Delete)
+	{
+		DeleteSelectedBlocks();
+	}
+
+	return SCompoundWidget::OnKeyDown(MyGeometry, InKeyEvent);
+}
+
+
+FCursorReply SCustomizableObjectLayoutGrid::OnCursorQuery(const FGeometry& MyGeometry, const FPointerEvent& CursorEvent) const
+{
+	if (ResizeCursor)
+	{
+		return FCursorReply::Cursor(EMouseCursor::ResizeSouthEast);
+	}
+	else
+	{
+		return FCursorReply::Cursor(EMouseCursor::Default);
+	}
+
+	return FCursorReply::Unhandled();
+}
+
+
+FVector2D SCustomizableObjectLayoutGrid::ComputeDesiredSize(float NotUsed) const
+{
+	return FVector2D(200.0f, 200.0f);
+}
+
+
+void SCustomizableObjectLayoutGrid::SetSelectedBlock(FGuid block )
+{
+	SelectedBlocks.Reset();
+	SelectedBlocks.Add( block );
+}
+
+
+void SCustomizableObjectLayoutGrid::SetSelectedBlocks( const TArray<FGuid>& blocks )
+{
+	SelectedBlocks = blocks;
+}
+
+
+TArray<FGuid> SCustomizableObjectLayoutGrid::GetSelectedBlocks() const
+{
+	return SelectedBlocks;
+}
+
+
+void SCustomizableObjectLayoutGrid::DeleteSelectedBlocks()
+{
+	DeleteBlocksDelegate.ExecuteIfBound();
+}
+
+
+void SCustomizableObjectLayoutGrid::GenerateNewBlock(FVector2D MousePosition)
+{
+	if (MousePosition.X > 0 && MousePosition.Y > 0 && MousePosition.X < GridSize.Get().X && MousePosition.Y < GridSize.Get().Y)
+	{
+		FIntPoint Min = FIntPoint(MousePosition.X, MousePosition.Y);
+		FIntPoint Max = Min + FIntPoint(1, 1);
+
+		AddBlockAtDelegate.ExecuteIfBound(Min, Max);
+
+		SelectedBlocks.Add(Blocks.Get().Last().Id);
+	}
+}
+
+
+void SCustomizableObjectLayoutGrid::DuplicateBlocks()
+{
+	if (SelectedBlocks.Num())
+	{
+		for (const FCustomizableObjectLayoutBlock& Block : Blocks.Get())
+		{
+			if (SelectedBlocks.Find(Block.Id) != INDEX_NONE)
+			{
+				AddBlockAtDelegate.ExecuteIfBound(Block.Min, Block.Max);
+			}
+		}
+	}
+}
+
+
+void SCustomizableObjectLayoutGrid::CalculateSelectionRect()
+{
+	if (InitSelectionRect.X <= CurrentMousePosition.X)
+	{
+		if (InitSelectionRect.Y <= CurrentMousePosition.Y)
+		{
+			SelectionRect.Min = FVector2f(InitSelectionRect);
+			SelectionRect.Size = FVector2f(CurrentMousePosition - InitSelectionRect);
+		}
+		else
+		{
+			SelectionRect.Min = FVector2f(InitSelectionRect.X, CurrentMousePosition.Y);
+
+			FVector2f AuxVector(CurrentMousePosition.X, InitSelectionRect.Y);
+			SelectionRect.Size = AuxVector - SelectionRect.Min;
+		}
+	}
+	else
+	{
+		if (InitSelectionRect.Y <= CurrentMousePosition.Y)
+		{
+			SelectionRect.Min = FVector2f(CurrentMousePosition.X, InitSelectionRect.Y);
+
+			FVector2f AuxVector(InitSelectionRect.X, CurrentMousePosition.Y);
+			SelectionRect.Size = AuxVector - SelectionRect.Min;
+		}
+		else
+		{
+			SelectionRect.Min = FVector2f(CurrentMousePosition);
+			SelectionRect.Size = FVector2f(InitSelectionRect - CurrentMousePosition);
+		}
+	}
+
+}
+
+
+void SCustomizableObjectLayoutGrid::SetBlocks(const FIntPoint& InGridSize, const TArray<FCustomizableObjectLayoutBlock>& InBlocks)
+{
+	GridSize = InGridSize;
+	Blocks = InBlocks;
+}
+
+
+bool SCustomizableObjectLayoutGrid::MouseOnBlock(FGuid BlockId, FVector2D MousePosition, bool CheckResizeBlock) const
+{
+	FVector2f Min, Max;
+	if (CheckResizeBlock)
+	{
+		Min = BlockRects[BlockId].HandleRect.Min;
+		Max = Min + BlockRects[BlockId].HandleRect.Size;
+	}
+	else
+	{
+		Min = BlockRects[BlockId].Rect.Min;
+		Max = Min + BlockRects[BlockId].Rect.Size;
+	}
+
+	if (MousePosition.X > Min.X && MousePosition.X<Max.X && MousePosition.Y>Min.Y && MousePosition.Y < Max.Y)
+	{
+		return true;
+	}
+
+	return false;
+}
+
+
+TOptional<int32> SCustomizableObjectLayoutGrid::GetBlockPriortyValue() const
+{
+	if (SelectedBlocks.Num())
+	{
+		TArray<FCustomizableObjectLayoutBlock> CurrentSelectedBlocks;
+
+		for (const FCustomizableObjectLayoutBlock& Block : Blocks.Get())
+		{
+			if (SelectedBlocks.Contains(Block.Id))
+			{
+				CurrentSelectedBlocks.Add(Block);
+			}
+		}
+
+		int32 BlockPriority = CurrentSelectedBlocks[0].Priority;
+		bool bSamePriority = true;
+
+		for (const FCustomizableObjectLayoutBlock& Block : CurrentSelectedBlocks)
+		{
+			if (Block.Priority != BlockPriority)
+			{
+				bSamePriority = false;
+				break;
+			}
+		}
+
+		if (bSamePriority)
+		{
+			return BlockPriority;
+		}
+	}
+
+	return TOptional<int32>();
+}
+
+
+void SCustomizableObjectLayoutGrid::OnBlockPriorityChanged(int32 InValue)
+{
+	if (SelectedBlocks.Num())
+	{
+		OnSetBlockPriority.ExecuteIfBound(InValue);
+	}
+}
+
+
+void SCustomizableObjectLayoutGrid::SetLayoutStrategy(ECustomizableObjectTextureLayoutPackingStrategy Strategy)
+{
+	LayoutStrategy = Strategy;
+}
+
+
+// Canvas Drawer --------------------------------------------------------------
+
+void FUVCanvasDrawer::Initialize(const FIntRect& InCanvasRect, const FVector2D& InOrigin, const FVector2D& InSize, const TArray<FVector2f>& InUVLayout)
+{
+	Initialized = InCanvasRect.Size().X > 0 && InCanvasRect.Size().Y > 0;
+	if (Initialized)
+	{
+		RenderTarget->SetViewRect(InCanvasRect);
+
+		Origin = InOrigin;
+		Size = InSize;
+
+		// Convert data
+		UVLayout.SetNum(InUVLayout.Num());
+		for (int32 Index=0; Index<InUVLayout.Num(); ++Index)
+		{
+			UVLayout[Index] = FVector2D(InUVLayout[Index]);
+		}
+	}
+}
+
+
+void FUVCanvasDrawer::DrawRenderThread(class FRHICommandListImmediate& RHICmdList, const void* InWindowBackBuffer)
+{
+	//UE crashes if there are no UVs to draw due to DX
+	if (Initialized && UVLayout.Num())
+	{
+		RenderTarget->SetRenderTargetTexture(*(FTexture2DRHIRef*)InWindowBackBuffer);
+
+#if MUTABLE_CLEAN_ENGINE_BRANCH
+		FCanvas Canvas(RenderTarget, nullptr, FGameTime(), GMaxRHIFeatureLevel);
+#else
+		FCanvas Canvas(RenderTarget, nullptr, 0, 0, 0, GMaxRHIFeatureLevel);
+#endif
+
+		const uint32 NumEdges = UVLayout.Num() / 2;
+
+		Canvas.SetRenderTargetRect(RenderTarget->GetViewRect());
+
+		FBatchedElements* BatchedElements = Canvas.GetBatchedElements(FCanvas::ET_Line);
+		BatchedElements->AddReserveLines(NumEdges);
+
+		const FLinearColor LineColor = FLinearColor::White;
+		const FHitProxyId HitProxyId = Canvas.GetHitProxyId();
+
+		FVector LinePoints[2];
+		for (uint32 LineIndex = 0; LineIndex < NumEdges; ++LineIndex)
+		{
+			LinePoints[0] = FVector(Origin + UVLayout[LineIndex * 2 + 0] * Size, 0.0f);
+			LinePoints[1] = FVector(Origin + UVLayout[LineIndex * 2 + 1] * Size, 0.0f);
+
+			BatchedElements->AddLine(LinePoints[0], LinePoints[1], LineColor, HitProxyId);
+		}
+
+		Canvas.Flush_RenderThread(RHICmdList, true);
+
+		RenderTarget->ClearRenderTargetTexture();
+	}
+}
+
+#undef LOCTEXT_NAMESPACE
diff --git a/Engine/Plugins/Experimental/Mutable/Source/CustomizableObjectEditor/Private/SCustomizableObjectLayoutGrid.h b/Engine/Plugins/Experimental/Mutable/Source/CustomizableObjectEditor/Private/SCustomizableObjectLayoutGrid.h
new file mode 100644
--- /dev/null
+++ b/Engine/Plugins/Experimental/Mutable/Source/CustomizableObjectEditor/Private/SCustomizableObjectLayoutGrid.h
@@ -0,0 +1,196 @@
+// Copyright Epic Games, Inc. All Rights Reserved.
+
+#pragma once
+
+#include "Nodes/CustomizableObjectNodeLayoutBlocks.h"
+#include "Widgets/SCompoundWidget.h"
+#include "Widgets/DeclarativeSyntaxSupport.h"
+//#include "PreviewScene.h"
+
+
+struct FRect2D
+{
+	FVector2f Min;
+	FVector2f Size;
+};
+
+typedef enum
+{
+	ELGM_Show,
+	ELGM_Edit,
+	ELGM_Select
+} ELayoutGridMode;
+
+
+/** */
+class SCustomizableObjectLayoutGrid : public SCompoundWidget
+{
+
+public:
+	DECLARE_DELEGATE_TwoParams(FBlockChangedDelegate, FGuid /*BlockId*/, FIntRect /*Block*/);
+	DECLARE_DELEGATE_OneParam(FBlockSelectionChangedDelegate, const TArray<FGuid>& );
+	DECLARE_DELEGATE(FDeleteBlockDelegate);
+	DECLARE_DELEGATE_TwoParams(FAddBlockAtDelegate, FIntPoint, FIntPoint);
+	DECLARE_DELEGATE_OneParam(FSetBlockPriority, int32);
+
+	SLATE_BEGIN_ARGS( SCustomizableObjectLayoutGrid ){}
+
+		SLATE_ATTRIBUTE( FIntPoint, GridSize )
+		SLATE_ATTRIBUTE( TArray<FCustomizableObjectLayoutBlock>, Blocks )
+		SLATE_ARGUMENT( TArray<FVector2f>, UVLayout  )
+		SLATE_ARGUMENT( TArray<FVector2f>, UnassignedUVLayoutVertices )
+		SLATE_ARGUMENT( ELayoutGridMode, Mode  )
+		SLATE_ARGUMENT( FColor, SelectionColor  )
+		SLATE_EVENT( FBlockChangedDelegate, OnBlockChanged )
+		SLATE_EVENT( FBlockSelectionChangedDelegate, OnSelectionChanged )
+		SLATE_EVENT(FDeleteBlockDelegate, OnDeleteBlocks)
+		SLATE_EVENT(FAddBlockAtDelegate, OnAddBlockAt)
+		SLATE_EVENT(FSetBlockPriority, OnSetBlockPriority)
+
+	SLATE_END_ARGS()
+
+	void Construct( const FArguments& InArgs );
+	~SCustomizableObjectLayoutGrid();
+
+	// SWidgetInterface
+	int32 OnPaint(const FPaintArgs& Args, const FGeometry& AllottedGeometry, const FSlateRect& MyClippingRect, FSlateWindowElementList& OutDrawElements, int32 LayerId, const FWidgetStyle& InWidgetStyle, bool bParentEnabled ) const override;
+	void Tick( const FGeometry& AllottedGeometry, const double InCurrentTime, const float InDeltaTime ) override;
+
+	FReply OnMouseButtonDown( const FGeometry& MyGeometry, const FPointerEvent& MouseEvent ) override;
+	FReply OnMouseButtonUp( const FGeometry& MyGeometry, const FPointerEvent& MouseEvent ) override;
+	FReply OnMouseMove( const FGeometry& MyGeometry, const FPointerEvent& MouseEvent ) override;
+	FReply OnMouseWheel(const FGeometry& MyGeometry, const FPointerEvent& MouseEvent) override;
+
+	bool SupportsKeyboardFocus() const override { return true; }
+	FReply OnKeyDown(const FGeometry& MyGeometry, const FKeyEvent& InKeyEvent) override;
+
+	FCursorReply OnCursorQuery(const FGeometry& MyGeometry, const FPointerEvent& CursorEvent) const override;
+	FVector2D ComputeDesiredSize(float) const override;
+
+	// Own interface
+
+	/** Set the currently selected block */
+	void SetSelectedBlock( FGuid block );
+	void SetSelectedBlocks( const TArray<FGuid>& blocks );
+
+	/** */
+	TArray<FGuid> GetSelectedBlocks() const;
+
+	/** Calls the delegate to delete the selected blocks */
+	void DeleteSelectedBlocks();
+
+	/** Generates a new block at mouse position */
+	void GenerateNewBlock(FVector2D MousePosition);
+
+	/** Duplicates the selected blocks */
+	void DuplicateBlocks();
+
+	void CalculateSelectionRect();
+
+	FColor SelectionColor;
+
+	/** Set the grid and blocks to show in the widget. */
+	void SetBlocks( const FIntPoint& GridSize, const TArray<FCustomizableObjectLayoutBlock>& Blocks);
+
+	/** Gets the priority value of the selected blocks */
+	TOptional<int32> GetBlockPriortyValue() const;
+
+	/** Callback when the priority of a block changes */
+	void OnBlockPriorityChanged(int32 InValue);
+
+	/** Sets the packing strategy of the current layout */
+	void SetLayoutStrategy(ECustomizableObjectTextureLayoutPackingStrategy Strategy);
+
+private:
+
+	bool MouseOnBlock(FGuid BlockId, FVector2D MousePosition, bool CheckResizeBlock = false) const;
+
+private:
+
+	/** A delegate to report block changes */
+	FBlockChangedDelegate BlockChangedDelegate;
+	FBlockSelectionChangedDelegate SelectionChangedDelegate;
+	FDeleteBlockDelegate DeleteBlocksDelegate;
+	FAddBlockAtDelegate AddBlockAtDelegate;
+	FSetBlockPriority OnSetBlockPriority;
+
+	/** Size of the grid in blocks */
+	TAttribute<FIntPoint> GridSize;
+
+	/** Array with all the blocks of the layout */
+	TAttribute< TArray<FCustomizableObjectLayoutBlock> > Blocks;
+
+	/** Array with all the UVs to draw in the layout */
+	TArray<FVector2f> UVLayout;
+
+	/** Array with all the unassigned UVs */
+	TArray<FVector2f> UnassignedUVLayoutVertices;
+
+	/** Layout mode */
+	ELayoutGridMode Mode;
+
+	float CellSize;
+
+	struct FBlockWidgetData
+	{
+		FRect2D Rect;
+		FRect2D HandleRect;
+	};
+	TMap<FGuid,FBlockWidgetData> BlockRects;
+
+	/** Interaction status. */
+	TArray<FGuid> SelectedBlocks;
+	TArray<FGuid> PossibleSelectedBlocks;
+
+	/** Bools needed for the Block Management */
+	/** Indicates when we have dragged the mouse after click */
+	bool HasDragged;
+
+	/** Indicates when we are dragging the mouse */
+	bool Dragging;
+	
+	/** Indicates when we are resizing a block */
+	bool Resizing;
+	
+	/** Indicates when we have to change the mouse cursor */
+	bool ResizeCursor;
+	
+	/** Indicates when we are making a selection */
+	bool Selecting;
+	
+	/** Indicates when we are padding */
+	bool Padding;
+
+	/** Position where the drag started */
+	FVector2D DragStart;
+
+	/** Position where the layout grid starts to be drawn */
+	FVector2D DrawOrigin;
+
+	/** Ammount of padding since start draging */
+	FVector2D PaddingAmount;
+
+	/** Position where the padding started */
+	FVector2D PaddingStart;
+
+	/** Distance from the origin in the padding movement */
+	FVector2D DistanceFromOrigin;
+
+	/** Level of zoom */
+	int32 Zoom;
+
+	/** Selection Rectangle */
+	FRect2D SelectionRect;
+
+	/** Position where the Selection Rectangle started */
+	FVector2D InitSelectionRect;
+
+	/** Current mouse position */
+	FVector2D CurrentMousePosition;
+
+	/** Custom Slate drawing element. Used to improve the UVs drawing performance. */
+	TSharedPtr<class FUVCanvasDrawer, ESPMode::ThreadSafe> UVCanvasDrawer;
+
+	/** Current packing strategy of the represented layout */
+	ECustomizableObjectTextureLayoutPackingStrategy LayoutStrategy;
+};
