[mutable] Move Mutable plugin from the restricted folder to the experimental plugins folder.

diff --git a/Engine/Plugins/Experimental/Mutable/Source/CustomizableObjectEditor/Private/SCustomizableObjectNodeLayoutBlocksEditor.cpp b/Engine/Plugins/Experimental/Mutable/Source/CustomizableObjectEditor/Private/SCustomizableObjectNodeLayoutBlocksEditor.cpp
new file mode 100644
--- /dev/null
+++ b/Engine/Plugins/Experimental/Mutable/Source/CustomizableObjectEditor/Private/SCustomizableObjectNodeLayoutBlocksEditor.cpp
@@ -0,0 +1,581 @@
+// Copyright Epic Games, Inc. All Rights Reserved.
+
+#include "SCustomizableObjectNodeLayoutBlocksEditor.h"
+#include "CustomizableObjectEditorModule.h"
+#include "CustomizableObjectEditorActions.h"
+#include "Nodes/CustomizableObjectNodeLayoutBlocksDetails.h"
+#include "CustomizableObjectCompiler.h"
+#include "CustomizableObjectEditorStyle.h"
+#include "CustomizableObjectLayout.h"
+
+#include "GenerateMutableSource/GenerateMutableSource.h"
+
+#include "CustomizableObjectCompiler.h"
+
+#include "MouseDeltaTracker.h"
+#include "SCustomizableObjectLayoutGrid.h"
+
+#include "Widgets/Input/STextComboBox.h"
+#include "Widgets/Input/SEditableTextBox.h"
+#include "Widgets/Input/SButton.h"
+#include "Framework/MultiBox/MultiBoxBuilder.h"
+#include "FbxMeshUtils.h"
+#include "BusyCursor.h"
+#include "MeshBuild.h"
+#include "ObjectTools.h"
+
+#include "MutableTools/Public/NodeLayout.h"
+
+#include "../Private/GeomFitUtils.h"
+
+
+#define LOCTEXT_NAMESPACE "CustomizableObjectEditor"
+
+/**
+ * 
+ */
+class FLayoutEditorCommands : public TCommands<FLayoutEditorCommands>
+{
+
+public:
+	FLayoutEditorCommands() : TCommands<FLayoutEditorCommands>
+	(
+		"LayoutEditorCommands", // Context name for fast lookup
+		NSLOCTEXT( "Contexts", "LayoutEditorCommands", "Layout Editor" ), // Localized context name for displaying
+		NAME_None, // Parent
+		FCustomizableObjectEditorStyle::GetStyleSetName()
+	)
+	{
+	}	
+	
+	/**  */
+	TSharedPtr< FUICommandInfo > AddBlock;
+	TSharedPtr< FUICommandInfo > RemoveBlock;
+	TSharedPtr< FUICommandInfo > GenerateBlocks;
+
+	/**
+	 * Initialize commands
+	 */
+	virtual void RegisterCommands() override
+	{
+		UI_COMMAND( AddBlock, "Add Block", "Add a new block to the layout.", EUserInterfaceActionType::Button, FInputChord() );
+		UI_COMMAND( RemoveBlock, "Remove Block", "Remove a block from the layout.", EUserInterfaceActionType::Button, FInputChord() );
+		UI_COMMAND( GenerateBlocks, "Generate Blocks", "Generate Blocks automatically from UVs", EUserInterfaceActionType::Button, FInputChord() );
+	}
+};
+
+
+
+void SCustomizableObjectNodeLayoutBlocksEditor::Construct(const FArguments& InArgs)
+{
+	CustomizableObjectEditorPtr = InArgs._CustomizableObjectEditor;
+	CurrentLayout = 0;
+	
+	BindCommands();
+}
+
+
+void SCustomizableObjectNodeLayoutBlocksEditor::SetCurrentLayout(UCustomizableObjectLayout* Layout )
+{
+	CurrentLayout = Layout;
+
+	// Try to locate the source mesh
+	TArray<FVector2f> UVs;
+	TArray<FVector2f> UnassignedUVs;
+
+	if (CurrentLayout)
+	{
+		CurrentLayout->GetUVChannel(UVs);
+
+		UnassignedUVs = TArray<FVector2f>();
+		
+		if (CurrentLayout->UnassignedUVs.Num())
+		{
+			UnassignedUVs = CurrentLayout->UnassignedUVs[0];
+		}
+	}
+
+	this->ChildSlot
+	[
+		SNew(SVerticalBox)
+
+		+ SVerticalBox::Slot()
+		.Padding( 2.0f )
+		.AutoHeight()
+		[
+			BuildLayoutToolBar()
+		]
+		.AutoHeight()
+		[
+			BuildLayoutToolBar()
+		]
+
+		+ SVerticalBox::Slot()
+		.FillHeight(1)
+		[
+			SAssignNew(LayoutGridWidget, SCustomizableObjectLayoutGrid)
+			.Mode(ELGM_Edit)
+			.GridSize(this, &SCustomizableObjectNodeLayoutBlocksEditor::GetGridSize)
+			.Blocks(this, &SCustomizableObjectNodeLayoutBlocksEditor::GetBlocks)
+			.UVLayout(UVs)
+			.UnassignedUVLayoutVertices(UnassignedUVs)
+			.SelectionColor(FColor(75, 106, 230, 155))
+			.OnBlockChanged(this, &SCustomizableObjectNodeLayoutBlocksEditor::OnBlockChanged)
+			.OnDeleteBlocks(this, &SCustomizableObjectNodeLayoutBlocksEditor::OnRemoveBlock)
+			.OnAddBlockAt(this, &SCustomizableObjectNodeLayoutBlocksEditor::OnAddBlockAt)
+			.OnSetBlockPriority(this, &SCustomizableObjectNodeLayoutBlocksEditor::OnSetBlockPriority)
+		]
+	];
+
+	if (CurrentLayout && LayoutGridWidget.IsValid())
+	{
+		LayoutGridWidget->SetLayoutStrategy(CurrentLayout->GetPackingStrategy());
+	}	
+}
+
+
+SCustomizableObjectNodeLayoutBlocksEditor::~SCustomizableObjectNodeLayoutBlocksEditor()
+{
+}
+
+
+void SCustomizableObjectNodeLayoutBlocksEditor::AddReferencedObjects( FReferenceCollector& Collector )
+{
+	Collector.AddReferencedObject( CurrentLayout );
+}
+
+
+TSharedRef<SWidget> SCustomizableObjectNodeLayoutBlocksEditor::BuildLayoutToolBar()
+{
+	TSharedPtr<FUICommandList> CommandList = CustomizableObjectEditorPtr.Pin()->GetToolkitCommands();
+	FToolBarBuilder LayoutToolbarBuilder(CommandList, FMultiBoxCustomization::None);
+
+	//Getting toolbar style
+	const ISlateStyle* const StyleSet = &FCoreStyle::Get();
+	const FName& StyleName = "ToolBar";
+	
+	// Build toolbar widgets
+	LayoutGridSizes.Empty();
+	int32 MaxGridSize = 32;
+	for(int32 Size = 1; Size <= MaxGridSize; Size*=2 )
+	{
+		LayoutGridSizes.Add( MakeShareable( new FString( FString::Printf(TEXT("%d x %d"), Size, Size ) ) ) );
+	}
+	
+	LayoutGridSizeWidget = SNew(SVerticalBox);
+
+	LayoutGridSizeWidget->AddSlot()
+	.AutoHeight()
+	.Padding(0.0f,0.0f,0.0f,10.0f)
+	[
+		SAssignNew(LayoutGridSizeCombo, STextComboBox)
+		.OptionsSource(&LayoutGridSizes)
+		.OnSelectionChanged(this, &SCustomizableObjectNodeLayoutBlocksEditor::OnGridSizeChanged)
+		.IsEnabled(FSlateApplication::Get().GetNormalExecutionAttribute())
+	];
+
+	LayoutGridSizeWidget->AddSlot()
+	.AutoHeight()
+	.HAlign(EHorizontalAlignment::HAlign_Center)
+	.VAlign(EVerticalAlignment::VAlign_Bottom)
+	[
+		SNew(STextBlock).Text(LOCTEXT("LayoutGridSizeText", "Grid Size"))
+		.ShadowOffset(FVector2D::UnitVector)
+		.ColorAndOpacity(FLinearColor::Gray)
+	];
+
+	if (CurrentLayout)
+	{
+		int grid = CurrentLayout->GetGridSize().X;
+		int option = -1;
+		while ( grid )
+		{
+			option++;
+			grid >>= 1;
+		}
+
+		if ( LayoutGridSizes.IsValidIndex(option) )
+		{
+			LayoutGridSizeCombo->SetSelectedItem( LayoutGridSizes[option] );
+		}
+	}
+
+	LayoutToolbarBuilder.BeginSection("Blocks");
+	{
+		LayoutToolbarBuilder.AddToolBarButton(FLayoutEditorCommands::Get().AddBlock);
+		LayoutToolbarBuilder.AddToolBarButton(FLayoutEditorCommands::Get().RemoveBlock);
+		LayoutToolbarBuilder.AddToolBarButton(FLayoutEditorCommands::Get().GenerateBlocks);
+	}
+	LayoutToolbarBuilder.EndSection();
+
+	LayoutToolbarBuilder.BeginSection("Grid");
+	{		
+		//LayoutToolbarBuilder.AddComboButton(FUIAction(),FOnGetContent::CreateSP(
+		//		this,&SCustomizableObjectNodeLayoutBlocksEditor::GenerateLayoutGridOptionsMenuContent),
+		//	TAttribute<FText>(),TAttribute<FText>(),FSlateIcon());
+
+		LayoutToolbarBuilder.AddWidget(LayoutGridSizeWidget.ToSharedRef());
+	}
+	LayoutToolbarBuilder.EndSection();
+
+	
+	StrategyWidget = BuildLayoutStrategyWidgets(StyleSet, StyleName);
+	
+	LayoutToolbarBuilder.BeginSection("Packing Strategy");
+	{
+		LayoutToolbarBuilder.AddWidget(StrategyWidget.ToSharedRef());
+	}
+	LayoutToolbarBuilder.EndSection();
+
+	return
+		SNew(SHorizontalBox)
+
+		+SHorizontalBox::Slot()
+		.Padding(4,0)
+		[
+			SNew(SBorder)
+			.Padding(0)
+			.BorderImage(FAppStyle::GetBrush("NoBorder"))
+			.IsEnabled( FSlateApplication::Get().GetNormalExecutionAttribute() )
+			[
+				LayoutToolbarBuilder.MakeWidget()
+			]
+		]
+		;
+}
+
+
+TSharedRef<SWidget> SCustomizableObjectNodeLayoutBlocksEditor::BuildLayoutStrategyWidgets(const ISlateStyle* Style, const FName& StyleName)
+{
+	MaxLayoutGridSizes.Empty();
+	int32 MaxGridSize = 32;
+	for (int32 Size = 1; Size <= MaxGridSize; Size *= 2)
+	{
+		MaxLayoutGridSizes.Add(MakeShareable(new FString(FString::Printf(TEXT("%d x %d"), Size, Size))));
+	}
+
+	MaxLayoutGridSizeCombo = SNew(STextComboBox)
+		.OptionsSource(&MaxLayoutGridSizes)
+		.OnSelectionChanged(this,&SCustomizableObjectNodeLayoutBlocksEditor::OnMaxGridSizeChanged)
+		.IsEnabled(FSlateApplication::Get().GetNormalExecutionAttribute());
+
+	LayoutPackingStrategies.Empty();
+	LayoutPackingStrategies.Add(MakeShareable(new FString("Resizable Layout")));
+	LayoutPackingStrategies.Add(MakeShareable(new FString("Fixed Layout")));
+
+	LayoutPackingStrategyCombo = SNew(STextComboBox)
+		.OptionsSource(&LayoutPackingStrategies)
+		.OnSelectionChanged(this, &SCustomizableObjectNodeLayoutBlocksEditor::OnLayoutPackingStrategyChanged)
+		.IsEnabled(FSlateApplication::Get().GetNormalExecutionAttribute());
+
+
+	LayoutStrategyWidget = SNew(SHorizontalBox)
+		+ SHorizontalBox::Slot().Padding(2.0f).AutoWidth()
+		[
+			SNew(STextBlock)
+			.Text(LOCTEXT("Layout Strategy:", "Layout Strategy:"))
+			.ToolTipText(LOCTEXT("LayoutStrategyTooltup","Selects the packing strategy."))
+			.Margin(FMargin(0, 4, 0, 0))
+			.ShadowOffset(FVector2D::UnitVector)
+			.ColorAndOpacity(FLinearColor::Gray)
+		]
+
+		+ SHorizontalBox::Slot().Padding(2.0f).AutoWidth()
+		[
+			SNew(SVerticalBox)
+			+ SVerticalBox::Slot()
+			.AutoHeight()
+			[
+				LayoutPackingStrategyCombo.ToSharedRef()
+			]
+		];
+
+
+	FixedLayoutWidget = SNew(SHorizontalBox)
+		+ SHorizontalBox::Slot().Padding(2.0f).AutoWidth()
+		[
+			SNew(STextBlock)
+			.Text(LOCTEXT("Max Layout Size:", "Max Layout Size:"))
+			.Margin(FMargin(0, 4, 0, 0))
+			.ShadowOffset(FVector2D::UnitVector)
+			.ColorAndOpacity(FLinearColor::Gray)
+		]
+		+ SHorizontalBox::Slot().Padding(2.0f).AutoWidth()
+		[
+			MaxLayoutGridSizeCombo.ToSharedRef()
+		];
+
+	//Modify after the creation of all the widgets
+	if (CurrentLayout)
+	{
+		int grid = CurrentLayout->GetMaxGridSize().X;
+		int option = -1;
+		while (grid)
+		{
+			option++;
+			grid >>= 1;
+		}
+
+		if (MaxLayoutGridSizes.IsValidIndex(option))
+		{
+			MaxLayoutGridSizeCombo->SetSelectedItem(MaxLayoutGridSizes[option]);
+		}
+
+		if (LayoutPackingStrategies.IsValidIndex((uint32)CurrentLayout->GetPackingStrategy()))
+		{
+			LayoutPackingStrategyCombo->SetSelectedItem(LayoutPackingStrategies[(uint32)CurrentLayout->GetPackingStrategy()]);
+		}
+	}
+
+	return SNew(SVerticalBox)
+		//.IsEnabled(CurrentLayout ? !CurrentLayout->GetNode()->GetGraphEditor()->GetCustomizableObject()->bDisableTextureLayoutManagement : false)
+		+ SVerticalBox::Slot()
+		[
+			LayoutStrategyWidget.ToSharedRef()
+		]
+		+SVerticalBox::Slot()
+		[
+			FixedLayoutWidget.ToSharedRef()
+		]
+	;
+}
+
+
+void SCustomizableObjectNodeLayoutBlocksEditor::OnGridSizeChanged( TSharedPtr<FString> NewSelection, ESelectInfo::Type SelectInfo )
+{
+	if (CurrentLayout)
+	{
+		int Size = 1 << LayoutGridSizes.Find(NewSelection);
+
+		if (CurrentLayout->GetGridSize().X!=Size || CurrentLayout->GetGridSize().Y!=Size )
+		{
+			CurrentLayout->SetGridSize(FIntPoint(Size));
+
+			// Adjust all the blocks sizes
+			for ( int b=0; b< CurrentLayout->Blocks.Num(); ++b )
+			{
+				CurrentLayout->Blocks[b].Min.X = FMath::Min( CurrentLayout->Blocks[b].Min.X, Size-1 );
+				CurrentLayout->Blocks[b].Min.Y = FMath::Min( CurrentLayout->Blocks[b].Min.Y, Size-1 );
+				CurrentLayout->Blocks[b].Max.X = FMath::Min( CurrentLayout->Blocks[b].Max.X, Size );
+				CurrentLayout->Blocks[b].Max.Y = FMath::Min( CurrentLayout->Blocks[b].Max.Y, Size );
+			}
+
+			CurrentLayout->MarkPackageDirty();
+		}
+	}
+}
+
+
+void SCustomizableObjectNodeLayoutBlocksEditor::OnAddBlock()
+{
+	if (CurrentLayout)
+	{
+		FCustomizableObjectLayoutBlock block;
+		block.Min = FIntPoint( 0, 0 );
+		block.Max = FIntPoint(1, 1);
+		block.Id = FGuid::NewGuid();
+		block.Priority = 0;
+		CurrentLayout->Blocks.Add( block );
+		CurrentLayout->MarkPackageDirty();
+
+		if (LayoutGridWidget.IsValid())
+		{
+			LayoutGridWidget->SetSelectedBlock(block.Id);
+		}
+	}
+}
+
+
+void SCustomizableObjectNodeLayoutBlocksEditor::OnAddBlockAt(const FIntPoint Min, const FIntPoint Max)
+{
+	if (CurrentLayout)
+	{
+		FCustomizableObjectLayoutBlock block;
+		block.Min = Min;
+		block.Max = Max;
+		block.Id = FGuid::NewGuid();
+		block.Priority = 0;
+
+		CurrentLayout->Blocks.Add(block);
+		CurrentLayout->MarkPackageDirty();
+	}
+}
+
+
+
+void SCustomizableObjectNodeLayoutBlocksEditor::OnRemoveBlock()
+{
+	if (CurrentLayout)
+	{
+		if (LayoutGridWidget.IsValid())
+		{
+			bool Change = false;
+
+			TArray<FGuid> selected = LayoutGridWidget->GetSelectedBlocks();
+			for (int i=0; i<CurrentLayout->Blocks.Num();)
+			{
+				if (selected.Contains(CurrentLayout->Blocks[i].Id))
+				{
+					Change = true;
+					CurrentLayout->Blocks.RemoveAt(i);
+				}
+				else
+				{
+					++i;
+				}
+			}
+
+			if (Change)
+			{
+				CurrentLayout->MarkPackageDirty();
+			}
+		}
+	}
+}
+
+
+void SCustomizableObjectNodeLayoutBlocksEditor::OnGenerateBlocks()
+{
+	if (CurrentLayout)
+	{
+		CurrentLayout->GenerateBlocksFromUVs();
+	}
+}
+
+
+FIntPoint SCustomizableObjectNodeLayoutBlocksEditor::GetGridSize() const
+{
+	if ( CurrentLayout )
+	{
+		return CurrentLayout->GetGridSize();
+	}
+	return FIntPoint(1);
+}
+
+
+void SCustomizableObjectNodeLayoutBlocksEditor::OnBlockChanged( FGuid BlockId, FIntRect Block )
+{
+	if (CurrentLayout)
+	{
+		for (FCustomizableObjectLayoutBlock& B : CurrentLayout->Blocks)
+		{
+			if (B.Id == BlockId)
+			{
+				B.Min = Block.Min;
+				B.Max = Block.Max;
+			}
+		}
+	}
+}
+
+
+void SCustomizableObjectNodeLayoutBlocksEditor::OnSetBlockPriority(int32 InValue)
+{
+	if (CurrentLayout)
+	{
+		if (LayoutGridWidget.IsValid())
+		{
+			TArray<FGuid> SelectedBlocks = LayoutGridWidget->GetSelectedBlocks();
+
+			for (int i = 0; i < CurrentLayout->Blocks.Num(); ++i)
+			{
+				if (SelectedBlocks.Contains(CurrentLayout->Blocks[i].Id))
+				{
+					CurrentLayout->Blocks[i].Priority = InValue;
+					CurrentLayout->MarkPackageDirty();
+				}
+			}
+		}
+	}	
+}
+
+
+void SCustomizableObjectNodeLayoutBlocksEditor::OnLayoutPackingStrategyChanged(TSharedPtr<FString> NewSelection, ESelectInfo::Type SelectInfo)
+{
+	if (CurrentLayout && FixedLayoutWidget)
+	{
+		uint32 selection = LayoutPackingStrategies.IndexOfByKey(NewSelection);
+
+		if (!selection)
+		{
+			FixedLayoutWidget->SetVisibility(EVisibility::Collapsed);
+		}
+		else
+		{
+			FixedLayoutWidget->SetVisibility(EVisibility::Visible);
+		}
+
+		if (CurrentLayout->GetPackingStrategy() != (ECustomizableObjectTextureLayoutPackingStrategy)selection)
+		{
+			CurrentLayout->SetPackingStrategy((ECustomizableObjectTextureLayoutPackingStrategy)selection);
+
+			if (LayoutGridWidget.IsValid())
+			{
+				LayoutGridWidget->SetLayoutStrategy(CurrentLayout->GetPackingStrategy());
+			}
+
+			CurrentLayout->MarkPackageDirty();
+		}
+	}
+}
+
+void SCustomizableObjectNodeLayoutBlocksEditor::OnMaxGridSizeChanged(TSharedPtr<FString> NewSelection, ESelectInfo::Type SelectInfo)
+{
+	if (CurrentLayout)
+	{
+		int Size = 1 << MaxLayoutGridSizes.Find(NewSelection);
+
+		if (CurrentLayout->GetMaxGridSize().X != Size || CurrentLayout->GetMaxGridSize().Y != Size)
+		{
+			CurrentLayout->SetMaxGridSize(FIntPoint(Size));
+			
+			CurrentLayout->MarkPackageDirty();
+		}
+	}
+}
+
+
+TArray<FCustomizableObjectLayoutBlock> SCustomizableObjectNodeLayoutBlocksEditor::GetBlocks() const
+{
+	TArray<FCustomizableObjectLayoutBlock> Blocks;
+
+	if (CurrentLayout)
+	{
+		Blocks = CurrentLayout->Blocks;
+	}
+
+	return Blocks;
+}
+
+
+void SCustomizableObjectNodeLayoutBlocksEditor::BindCommands()
+{
+	// Register our commands. This will only register them if not previously registered
+	FLayoutEditorCommands::Register();
+
+	const FLayoutEditorCommands& Commands = FLayoutEditorCommands::Get();
+
+	const TSharedRef<FUICommandList>& UICommandList = CustomizableObjectEditorPtr.Pin()->GetToolkitCommands();
+
+	UICommandList->MapAction(
+		Commands.AddBlock,
+		FExecuteAction::CreateSP( this, &SCustomizableObjectNodeLayoutBlocksEditor::OnAddBlock ),
+		FCanExecuteAction(),
+		FIsActionChecked() );
+
+	UICommandList->MapAction(
+		Commands.RemoveBlock,
+		FExecuteAction::CreateSP( this, &SCustomizableObjectNodeLayoutBlocksEditor::OnRemoveBlock ),
+		FCanExecuteAction(),
+		FIsActionChecked() );
+
+	UICommandList->MapAction(
+		Commands.GenerateBlocks,
+		FExecuteAction::CreateSP(this, &SCustomizableObjectNodeLayoutBlocksEditor::OnGenerateBlocks),
+		FCanExecuteAction(),
+		FIsActionChecked());
+}
+
+
+#undef LOCTEXT_NAMESPACE
+
diff --git a/Engine/Plugins/Experimental/Mutable/Source/CustomizableObjectEditor/Private/SCustomizableObjectNodeLayoutBlocksEditor.h b/Engine/Plugins/Experimental/Mutable/Source/CustomizableObjectEditor/Private/SCustomizableObjectNodeLayoutBlocksEditor.h
new file mode 100644
--- /dev/null
+++ b/Engine/Plugins/Experimental/Mutable/Source/CustomizableObjectEditor/Private/SCustomizableObjectNodeLayoutBlocksEditor.h
@@ -0,0 +1,101 @@
+// Copyright Epic Games, Inc. All Rights Reserved.
+
+#pragma once
+
+#include "Widgets/SCompoundWidget.h"
+#include "ICustomizableObjectInstanceEditor.h"
+#include "Nodes/CustomizableObjectNodeLayoutBlocks.h"
+#include "CustomizableObjectLayout.h"
+
+
+/**
+ * CustomizableObject Editor Preview viewport widget
+ */
+class SCustomizableObjectNodeLayoutBlocksEditor : public SCompoundWidget, public FGCObject
+{
+public:
+	SLATE_BEGIN_ARGS( SCustomizableObjectNodeLayoutBlocksEditor ){}
+		SLATE_ARGUMENT(TWeakPtr<ICustomizableObjectInstanceEditor>, CustomizableObjectEditor)
+	SLATE_END_ARGS()
+
+	void Construct(const FArguments& InArgs);
+	~SCustomizableObjectNodeLayoutBlocksEditor();
+	
+	// FSerializableObject interface
+	void AddReferencedObjects( FReferenceCollector& Collector ) override;
+	virtual FString GetReferencerName() const override
+	{
+		return TEXT("SCustomizableObjectNodeLayoutBlocksEditor");
+	}
+	// End of FSerializableObject interface
+
+
+	/** Binds commands associated with the viewport client. */
+	void BindCommands();
+
+	/**  */
+	void SetCurrentLayout( class UCustomizableObjectLayout* Layout );
+
+private:
+
+	/** Pointer back to the editor tool that owns us */
+	TWeakPtr<ICustomizableObjectInstanceEditor> CustomizableObjectEditorPtr;
+
+	/** */
+	class UCustomizableObjectLayout* CurrentLayout;
+
+	/** */
+	TSharedPtr<class SCustomizableObjectLayoutGrid> LayoutGridWidget;
+
+	TSharedPtr< class SVerticalBox > LayoutGridSizeWidget;
+
+	/** Widget for displaying the available layout block grid sizes. */
+	TSharedPtr< class STextComboBox > LayoutGridSizeCombo;
+	TSharedPtr< class STextComboBox > MaxLayoutGridSizeCombo;
+
+	/** List of available layout grid sizes. */
+	TArray< TSharedPtr< FString > > LayoutGridSizes;
+
+	/** List of available layout max grid sizes. */
+	TArray< TSharedPtr< FString > > MaxLayoutGridSizes;
+
+	/** List of available layout packing strategies. */
+	TArray< TSharedPtr< FString > > LayoutPackingStrategies;
+
+	/** Widget for displaying the available layout packing strategies. */
+	TSharedPtr< class STextComboBox > LayoutPackingStrategyCombo;
+
+	/** Widget to select the layout packing strategy */
+	TSharedPtr<class SHorizontalBox> LayoutStrategyWidget;
+
+	/** Widget to select the fixed layout properties */
+	TSharedPtr<class SHorizontalBox> FixedLayoutWidget;
+
+	TSharedPtr<SWidget> StrategyWidget;
+
+	/** */
+	FIntPoint GetGridSize() const;
+	void OnBlockChanged(FGuid BlockId, FIntRect Block );
+	TArray<FCustomizableObjectLayoutBlock> GetBlocks() const;
+
+	/** Callbacks from the layout block editor. */
+	TSharedRef<SWidget> BuildLayoutToolBar();
+	TSharedRef<SWidget> BuildLayoutStrategyWidgets(const ISlateStyle* Style, const FName& StyleName);
+	void OnAddBlock();
+	void OnAddBlockAt(const FIntPoint Min, const FIntPoint Max);
+	void OnRemoveBlock();
+
+	/** Generate layout blocks using UVs*/
+	void OnGenerateBlocks();
+		
+	/** . */
+	void OnGridSizeChanged( TSharedPtr<FString> NewSelection, ESelectInfo::Type SelectInfo );
+	void OnMaxGridSizeChanged(TSharedPtr<FString> NewSelection, ESelectInfo::Type SelectInfo);
+
+	/** Sets the block priority from the input text. */
+	void OnSetBlockPriority(int32 InValue);
+
+	/** Called when the packing strategy has changed. */
+	void OnLayoutPackingStrategyChanged(TSharedPtr<FString> NewSelection, ESelectInfo::Type SelectInfo);
+
+};
diff --git a/Engine/Plugins/Experimental/Mutable/Source/CustomizableObjectEditor/Private/SCustomizableObjectNodeLayoutBlocksSelector.cpp b/Engine/Plugins/Experimental/Mutable/Source/CustomizableObjectEditor/Private/SCustomizableObjectNodeLayoutBlocksSelector.cpp
new file mode 100644
--- /dev/null
+++ b/Engine/Plugins/Experimental/Mutable/Source/CustomizableObjectEditor/Private/SCustomizableObjectNodeLayoutBlocksSelector.cpp
@@ -0,0 +1,300 @@
+// Copyright Epic Games, Inc. All Rights Reserved.
+
+#include "SCustomizableObjectNodeLayoutBlocksSelector.h"
+#include "CustomizableObjectEditorModule.h"
+#include "CustomizableObjectEditorActions.h"
+#include "CustomizableObjectEditorModule.h"
+#include "CustomizableObjectEditorStyle.h"
+#include "CustomizableObjectLayout.h"
+
+#include "MouseDeltaTracker.h"
+#include "SCustomizableObjectLayoutGrid.h"
+
+#include "CustomizableObjectClasses.h"
+#include "CustomizableObjectEditor.h"
+#include "FbxMeshUtils.h"
+#include "BusyCursor.h"
+#include "GraphTraversal.h"
+#include "MeshBuild.h"
+#include "ObjectTools.h"
+
+#include "../Private/GeomFitUtils.h"
+
+#define LOCTEXT_NAMESPACE "CustomizableObjectEditor"
+
+
+/**
+ * 
+ */
+class FLayoutBlockSelectorCommands : public TCommands<FLayoutBlockSelectorCommands>
+{
+
+public:
+	FLayoutBlockSelectorCommands() : TCommands<FLayoutBlockSelectorCommands>
+	(
+		"LayoutBlockSelector", // Context name for fast lookup
+		NSLOCTEXT( "Contexts", "LayoutBlockSelector", "Layout Block Selector" ), // Localized context name for displaying
+		NAME_None, // Parent
+		FCustomizableObjectEditorStyle::GetStyleSetName()
+	)
+	{
+	}	
+	
+	/**  */
+	TSharedPtr< FUICommandInfo > SelectAll;
+	TSharedPtr< FUICommandInfo > SelectNone;
+
+	/**
+	 * Initialize commands
+	 */
+	virtual void RegisterCommands() override
+	{
+		UI_COMMAND( SelectAll, "Select All", "Select all the blocks in the layout.", EUserInterfaceActionType::Button, FInputChord() );
+		UI_COMMAND( SelectNone, "Unselect All", "Unselect all the blocks in the layout.", EUserInterfaceActionType::Button, FInputChord() );
+	}
+};
+
+
+void SCustomizableObjectNodeLayoutBlocksSelector::Construct(const FArguments& InArgs)
+{
+	CustomizableObjectEditorPtr = InArgs._CustomizableObjectEditor;
+	CurrentNode = nullptr;
+	
+	BindCommands();
+}
+
+
+void SCustomizableObjectNodeLayoutBlocksSelector::SetSelectedNode(UCustomizableObjectNodeEditLayoutBlocks* Node )
+{
+	CurrentNode = Node;
+
+	UCustomizableObjectNodeMaterialBase* ParentMaterialNode = Node ? Node->GetParentMaterialNode() : nullptr;
+
+	// Try to locate the source mesh
+	TArray<FVector2f> UVs;
+	if (CurrentNode && ParentMaterialNode)
+	{
+		const UEdGraphPin* SourceMeshPin = FindMeshBaseSource(*ParentMaterialNode->OutputPin(), false);
+		if (SourceMeshPin)
+		{
+			const UCustomizableObjectNodeSkeletalMesh* MeshNode = Cast<UCustomizableObjectNodeSkeletalMesh>(SourceMeshPin->GetOwningNode());
+			if (MeshNode)
+			{
+				if (ParentMaterialNode->GetLayouts().Num() > 0 && ParentMaterialNode->GetLayouts().IsValidIndex(CurrentNode->ParentLayoutIndex))
+				{
+					MeshNode->GetUVChannelForPin(SourceMeshPin, UVs, CurrentNode->ParentLayoutIndex);
+				}
+			}
+
+			const UCustomizableObjectNodeTable* TableNode = Cast<UCustomizableObjectNodeTable>(SourceMeshPin->GetOwningNode());
+			if (TableNode)
+			{
+				if (ParentMaterialNode->GetLayouts().Num() > 0 && ParentMaterialNode->GetLayouts().IsValidIndex(CurrentNode->ParentLayoutIndex))
+				{
+					TableNode->GetUVChannelForPin(SourceMeshPin, UVs, CurrentNode->ParentLayoutIndex);
+				}
+			}
+		}
+	}
+
+	this->ChildSlot
+		[
+			SNew(SVerticalBox)
+
+			+ SVerticalBox::Slot()
+				.Padding(2.0f)
+				.AutoHeight()
+				[
+					BuildLayoutToolBar()
+				]
+
+			+ SVerticalBox::Slot()
+				.FillHeight(1)
+				[
+					SAssignNew(LayoutGridWidget, SCustomizableObjectLayoutGrid)
+					.Mode(ELGM_Select)
+					.GridSize(this, &SCustomizableObjectNodeLayoutBlocksSelector::GetGridSize)
+					.Blocks(this, &SCustomizableObjectNodeLayoutBlocksSelector::GetBlocks)
+					.UVLayout(UVs)
+					.SelectionColor(FColor(178, 134, 73, 155))
+					.OnSelectionChanged(this, &SCustomizableObjectNodeLayoutBlocksSelector::OnSelectionChanged)
+				]
+		];
+
+	LayoutGridWidget->SetSelectedBlocks(CurrentNode->BlockIds);
+	BlocksLabel->SetText(FText::FromString(FString::Printf(TEXT("%d blocks selected"), CurrentNode->BlockIds.Num())));
+}
+
+
+SCustomizableObjectNodeLayoutBlocksSelector::~SCustomizableObjectNodeLayoutBlocksSelector()
+{
+}
+
+
+void SCustomizableObjectNodeLayoutBlocksSelector::AddReferencedObjects( FReferenceCollector& Collector )
+{
+	Collector.AddReferencedObject( CurrentNode );
+}
+
+
+TSharedRef<SWidget> SCustomizableObjectNodeLayoutBlocksSelector::BuildLayoutToolBar()
+{
+	TSharedPtr<FUICommandList> CommandList = CustomizableObjectEditorPtr.Pin()->GetToolkitCommands();
+	FToolBarBuilder LayoutToolbarBuilder(CommandList, FMultiBoxCustomization::None);
+
+	// Build toolbar widgets
+
+	LayoutToolbarBuilder.BeginSection("Blocks");
+	{
+		LayoutToolbarBuilder.AddToolBarButton(FLayoutBlockSelectorCommands::Get().SelectAll);
+		LayoutToolbarBuilder.AddToolBarButton(FLayoutBlockSelectorCommands::Get().SelectNone);
+	}
+	LayoutToolbarBuilder.EndSection();
+		
+	BlocksLabel = SNew(STextBlock)
+		.Text( LOCTEXT("0 blocks selected","0 blocks selected") );
+
+	LayoutToolbarBuilder.BeginSection("Label");
+	{
+		LayoutToolbarBuilder.AddWidget(BlocksLabel.ToSharedRef());
+	}
+	LayoutToolbarBuilder.EndSection();
+
+	return
+		SNew(SHorizontalBox)
+
+		+SHorizontalBox::Slot()
+		.Padding(4,0)
+		.AutoWidth()
+		[
+			SNew(SBorder)
+			.Padding(0)
+			.BorderImage(FAppStyle::GetBrush("NoBorder"))
+			.IsEnabled( FSlateApplication::Get().GetNormalExecutionAttribute() )
+			[
+				LayoutToolbarBuilder.MakeWidget()
+			]
+		]
+		;
+}
+
+
+void SCustomizableObjectNodeLayoutBlocksSelector::OnSelectAll()
+{
+	UCustomizableObjectNodeMaterialBase* ParentMaterialNode = CurrentNode ? CurrentNode->GetParentMaterialNode() : nullptr;
+
+	if ( CurrentNode && ParentMaterialNode)
+	{
+		TArray<UCustomizableObjectLayout*> Layouts = ParentMaterialNode->GetLayouts();
+
+		if (Layouts.Num() && Layouts.IsValidIndex(CurrentNode->ParentLayoutIndex))
+		{
+			UCustomizableObjectLayout* Layout = Layouts[CurrentNode->ParentLayoutIndex];
+
+			for( const FCustomizableObjectLayoutBlock& Block : Layout->Blocks )
+			{
+				CurrentNode->BlockIds.Add(Block.Id);
+			}
+
+			BlocksLabel->SetText(FText::FromString(FString::Printf(TEXT("%d blocks selected"), CurrentNode->BlockIds.Num())));
+
+			LayoutGridWidget->SetSelectedBlocks( CurrentNode->BlockIds );
+			CurrentNode->GetGraph()->MarkPackageDirty();
+		}
+	}
+}
+
+
+void SCustomizableObjectNodeLayoutBlocksSelector::OnSelectNone()
+{
+	if ( CurrentNode )
+	{
+		CurrentNode->BlockIds.Reset();
+		LayoutGridWidget->SetSelectedBlocks( CurrentNode->BlockIds );
+
+		BlocksLabel->SetText(LOCTEXT("0 blocks selected", "0 blocks selected"));
+
+		CurrentNode->GetGraph()->MarkPackageDirty();
+	}
+}
+
+
+FIntPoint SCustomizableObjectNodeLayoutBlocksSelector::GetGridSize() const
+{
+	FIntPoint Result(1, 1);
+	UCustomizableObjectNodeMaterialBase* ParentMaterialNode = CurrentNode ? CurrentNode->GetParentMaterialNode() : nullptr;
+
+	if ( CurrentNode && ParentMaterialNode)
+	{
+		TArray<UCustomizableObjectLayout*> Layouts = ParentMaterialNode->GetLayouts();
+
+		if (Layouts.Num() && Layouts.IsValidIndex(CurrentNode->ParentLayoutIndex))
+		{
+			UCustomizableObjectLayout* Layout = Layouts[CurrentNode->ParentLayoutIndex];
+			Result = Layout->GetGridSize();
+		}
+	}
+
+	return Result;
+}
+
+
+void SCustomizableObjectNodeLayoutBlocksSelector::OnSelectionChanged( const TArray<FGuid>& selected )
+{
+	if ( CurrentNode )
+	{
+		CurrentNode->BlockIds = selected;
+
+		BlocksLabel->SetText( FText::FromString( FString::Printf( TEXT("%d blocks selected"),selected.Num() ) ) );
+
+		CurrentNode->GetGraph()->MarkPackageDirty();
+	}
+}
+
+
+TArray<FCustomizableObjectLayoutBlock> SCustomizableObjectNodeLayoutBlocksSelector::GetBlocks() const
+{
+	TArray<FCustomizableObjectLayoutBlock> Blocks;
+
+	UCustomizableObjectNodeMaterialBase* ParentMaterialNode = CurrentNode ? CurrentNode->GetParentMaterialNode() : nullptr;
+
+	if (ParentMaterialNode)
+	{
+		TArray<UCustomizableObjectLayout*> Layouts = ParentMaterialNode->GetLayouts();
+
+		if (Layouts.IsValidIndex(CurrentNode->ParentLayoutIndex))
+		{
+			UCustomizableObjectLayout* Layout = Layouts[CurrentNode->ParentLayoutIndex];
+			Blocks = Layout->Blocks;
+		}
+	}
+
+	return Blocks;
+}
+
+
+void SCustomizableObjectNodeLayoutBlocksSelector::BindCommands()
+{
+	// Register our commands. This will only register them if not previously registered
+	FLayoutBlockSelectorCommands::Register();
+
+	const FLayoutBlockSelectorCommands& Commands = FLayoutBlockSelectorCommands::Get();
+
+	const TSharedRef<FUICommandList>& UICommandList = CustomizableObjectEditorPtr.Pin()->GetToolkitCommands();
+
+	UICommandList->MapAction(
+		Commands.SelectAll,
+		FExecuteAction::CreateSP( this, &SCustomizableObjectNodeLayoutBlocksSelector::OnSelectAll ),
+		FCanExecuteAction(),
+		FIsActionChecked() );
+
+	UICommandList->MapAction(
+		Commands.SelectNone,
+		FExecuteAction::CreateSP( this, &SCustomizableObjectNodeLayoutBlocksSelector::OnSelectNone ),
+		FCanExecuteAction(),
+		FIsActionChecked() );
+}
+
+
+#undef LOCTEXT_NAMESPACE
+
diff --git a/Engine/Plugins/Experimental/Mutable/Source/CustomizableObjectEditor/Private/SCustomizableObjectNodeLayoutBlocksSelector.h b/Engine/Plugins/Experimental/Mutable/Source/CustomizableObjectEditor/Private/SCustomizableObjectNodeLayoutBlocksSelector.h
new file mode 100644
--- /dev/null
+++ b/Engine/Plugins/Experimental/Mutable/Source/CustomizableObjectEditor/Private/SCustomizableObjectNodeLayoutBlocksSelector.h
@@ -0,0 +1,61 @@
+// Copyright Epic Games, Inc. All Rights Reserved.
+
+#pragma once
+
+#include "Widgets/SCompoundWidget.h"
+#include "ICustomizableObjectInstanceEditor.h"
+#include "Nodes/CustomizableObjectNodeLayoutBlocks.h"
+
+
+/**
+ * CustomizableObject Editor Preview viewport widget
+ */
+class SCustomizableObjectNodeLayoutBlocksSelector : public SCompoundWidget, public FGCObject
+{
+public:
+	SLATE_BEGIN_ARGS( SCustomizableObjectNodeLayoutBlocksSelector ){}
+		SLATE_ARGUMENT(TWeakPtr<ICustomizableObjectInstanceEditor>, CustomizableObjectEditor)
+	SLATE_END_ARGS()
+
+	void Construct(const FArguments& InArgs);
+	~SCustomizableObjectNodeLayoutBlocksSelector();
+	
+	// FSerializableObject interface
+	void AddReferencedObjects( FReferenceCollector& Collector ) override;
+	virtual FString GetReferencerName() const override
+	{
+		return TEXT("SCustomizableObjectNodeLayoutBlocksSelector");
+	}
+	// End of FSerializableObject interface
+
+
+	/** Binds commands associated with the viewport client. */
+	void BindCommands();
+
+	// These two types of nodes can use this widget
+	void SetSelectedNode(class UCustomizableObjectNodeEditLayoutBlocks* Node);
+
+private:
+
+	/** Pointer back to the editor tool that owns us */
+	TWeakPtr<ICustomizableObjectInstanceEditor> CustomizableObjectEditorPtr;
+
+	// 
+	class UCustomizableObjectNodeEditLayoutBlocks* CurrentNode = nullptr;
+
+	/** */
+	TSharedPtr<class SCustomizableObjectLayoutGrid> LayoutGridWidget;
+	
+	/** */
+	TSharedPtr<class STextBlock> BlocksLabel;
+
+	/** */
+	FIntPoint GetGridSize() const;
+	void OnSelectionChanged( const TArray<FGuid>& selected );
+	TArray<FCustomizableObjectLayoutBlock> GetBlocks() const;
+
+	/** Callbacks from the layout block editor. */
+	TSharedRef<SWidget> BuildLayoutToolBar();
+	void OnSelectAll();
+	void OnSelectNone();		
+};
diff --git a/Engine/Plugins/Experimental/Mutable/Source/CustomizableObjectEditor/Private/SCustomizableObjectNodeObjectStatesView.cpp b/Engine/Plugins/Experimental/Mutable/Source/CustomizableObjectEditor/Private/SCustomizableObjectNodeObjectStatesView.cpp
new file mode 100644
--- /dev/null
+++ b/Engine/Plugins/Experimental/Mutable/Source/CustomizableObjectEditor/Private/SCustomizableObjectNodeObjectStatesView.cpp
@@ -0,0 +1,514 @@
+// Copyright Epic Games, Inc. All Rights Reserved.
+
+#include "SCustomizableObjectNodeObjectStatesView.h"
+#include "CustomizableObjectEditorModule.h"
+#include "PropertyCustomizationHelpers.h"
+#include "CustomizableObjectEditorUtilities.h"
+#include "CustomizableObject.h"
+#include "Nodes/CustomizableObjectNodeObjectGroup.h"
+#include "Nodes/CustomizableObjectNodeObjectDetails.h"
+
+#include "DetailLayoutBuilder.h"
+#include "DetailCategoryBuilder.h"
+
+#include "Widgets/Input/SButton.h"
+#include "Widgets/Text/STextBlock.h"
+#include "Widgets/Input/STextComboBox.h"
+#include "Widgets/Input/SSearchBox.h"
+#include "Widgets/Input/SCheckBox.h"
+
+#include "Input/Reply.h"
+
+
+#define LOCTEXT_NAMESPACE "CustomizableObjectNodeObjectStatesView"
+
+
+void SCustomizableObjectNodeObjectSatesView::Construct(const FArguments& InArgs)
+{
+	Node = InArgs._Node;
+
+	if (Node)
+	{
+		SAssignNew(VerticalSlots, SVerticalBox);
+
+		for (int32 i = 0; i < Node->States.Num(); i++)
+		{
+			// Images and Runtime parameters widgets are stored in arrays to control the visibility
+			RuntimeParametersListWidgets.Add(SNew(SCustomizableObjectRuntimeParameterList).Node(Node).StateIndex(i));
+			RuntimeParametersListWidgets[i]->SetCollapsed(true);
+			RuntimeParametersListWidgets[i]->SetVisibility(GetCollapsed(i));
+
+			CollapsedArrows.Add(SNew(SImage).Image(GetExpressionPreviewArrow(i)));
+
+			VerticalSlots->AddSlot()
+			.AutoHeight()
+			[
+				SNew(SVerticalBox)
+				+ SVerticalBox::Slot()
+				.Padding(0.0f, 0.0f, 3.0f, 5.0f)
+				[
+					SNew(SHorizontalBox)
+
+					// State variable label
+					+ SHorizontalBox::Slot()
+					.HAlign(EHorizontalAlignment::HAlign_Left)
+					.AutoWidth()
+					.Padding(0.0f, 0.0f, 3.0f, 0.0f)
+					[
+						SNew(STextBlock)
+						.Text(FText::FromString("State:"))
+					]
+					
+					// State name
+					+ SHorizontalBox::Slot()
+					.HAlign(EHorizontalAlignment::HAlign_Left)
+					.AutoWidth()
+					.Padding(2.0f, 0.0f, 0.0f, 0.0f)
+					[
+						SNew(STextBlock)
+						.Text(FText::FromString(Node->States[i].Name))
+					]
+				]
+
+				+ SVerticalBox::Slot()
+				.AutoHeight()
+				.Padding(0.0f, 0.0f, 0.0f, 2.0f)
+				[
+					SNew(SVerticalBox)
+					+ SVerticalBox::Slot()
+					.HAlign(EHorizontalAlignment::HAlign_Left)
+					.Padding(10.0f, 0.0f, 3.0f, 3.0f)
+					.AutoHeight()
+					[
+						SNew(SHorizontalBox)
+
+						// Collapsed Arrow checkbox
+						+SHorizontalBox::Slot()
+						.HAlign(EHorizontalAlignment::HAlign_Left)
+						.AutoWidth()
+						.Padding(0.0f, 0.0f, 3.0f, 0.0f)
+						[
+							SNew(SCheckBox)
+							.OnCheckStateChanged(this, &SCustomizableObjectNodeObjectSatesView::OnCollapseChanged, i)
+							.IsChecked(ECheckBoxState::Unchecked)
+							.Cursor(EMouseCursor::Default)
+							.Style(FAppStyle::Get(), "Graph.Node.AdvancedView")
+							[
+								SNew(SHorizontalBox)
+								+ SHorizontalBox::Slot()
+								.VAlign(VAlign_Center)
+								.HAlign(HAlign_Center)
+								[
+									CollapsedArrows[i].ToSharedRef()
+								]
+							]
+						]
+						
+						// Number of runtime parameters
+						+ SHorizontalBox::Slot()
+						.HAlign(EHorizontalAlignment::HAlign_Left)
+						.AutoWidth()
+						.Padding(0.0f, 0.0f, 3.0f, 0.0f)
+						[
+							SNew(STextBlock)
+							.Text(FText::FromString(FString::Printf(TEXT("Runtime Parameters:  %d"), Node->States[i].RuntimeParameters.Num())))
+						]
+
+						+ SHorizontalBox::Slot()
+						.HAlign(EHorizontalAlignment::HAlign_Left)
+						.AutoWidth()
+						.Padding(0.0f, 0.0f, 3.0f, 0.0f)
+						[
+							SNew(SButton)
+							.OnClicked(this, &SCustomizableObjectNodeObjectSatesView::OnAddRuntimeParameterPressed,i)
+							.ToolTipText(LOCTEXT("AddRuntimeParameter", "Add Runtime Parameter"))
+							[
+								SNew(SHorizontalBox)
+								+ SHorizontalBox::Slot()
+								.VAlign(VAlign_Center)
+								.HAlign(HAlign_Center)
+								[
+									SNew(SImage)
+									.Image(FAppStyle::GetBrush(TEXT("Plus")))
+								]
+							]
+						]
+					]
+					
+					// Runtime parameters widget
+					+ SVerticalBox::Slot()
+					.HAlign(EHorizontalAlignment::HAlign_Left)
+					.Padding(20.0f, 0.0f, 0.0f, 0.0f)
+					.AutoHeight()
+					[
+						RuntimeParametersListWidgets[i].ToSharedRef()
+					]
+				]
+			];
+		}
+	}
+
+	// Add the widget to the child slot
+	ChildSlot
+	[
+		VerticalSlots.ToSharedRef()
+	];
+}
+
+
+void SCustomizableObjectNodeObjectSatesView::OnCollapseChanged(const ECheckBoxState NewCheckedState, int32 StateIndex)
+{
+	bool bCollapse = (NewCheckedState != ECheckBoxState::Checked);
+
+	RuntimeParametersListWidgets[StateIndex]->SetCollapsed(bCollapse);
+	RuntimeParametersListWidgets[StateIndex]->SetVisibility(GetCollapsed(StateIndex));
+	CollapsedArrows[StateIndex]->SetImage(GetExpressionPreviewArrow(StateIndex));
+}
+
+
+EVisibility SCustomizableObjectNodeObjectSatesView::GetCollapsed(int32 StateIndex)
+{
+	return RuntimeParametersListWidgets[StateIndex]->IsCollapsed() ? EVisibility::Collapsed : EVisibility::Visible;
+}
+
+
+const FSlateBrush* SCustomizableObjectNodeObjectSatesView::GetExpressionPreviewArrow(int32 StateIndex) const
+{
+	return FAppStyle::GetBrush(RuntimeParametersListWidgets[StateIndex]->IsCollapsed() ? TEXT("SurfaceDetails.PanUPositive") : TEXT("SurfaceDetails.PanVPositive"));
+}
+
+
+FReply SCustomizableObjectNodeObjectSatesView::OnAddRuntimeParameterPressed(int32 StateIndex)
+{
+	Node->States[StateIndex].RuntimeParameters.Add("NONE");
+	RuntimeParametersListWidgets[StateIndex]->BuildList();
+
+	return FReply::Handled();
+}
+
+
+// Widget for a list ofruntime parameters --------------------------------------------------------------------------
+
+
+void SCustomizableObjectRuntimeParameterList::Construct(const FArguments& InArgs)
+{
+	Node = InArgs._Node;
+	StateIndex = InArgs._StateIndex;
+
+	if (Node)
+	{
+		UCustomizableObject* CustomizableObject = Cast<UCustomizableObject>(Node->GetCustomizableObjectGraph()->GetOuter());
+
+		if (CustomizableObject)
+		{
+			SAssignNew(VerticalSlots, SVerticalBox);
+
+			if (VerticalSlots.IsValid())
+			{
+				BuildList();
+
+				ChildSlot
+				[
+					VerticalSlots.ToSharedRef()
+				];
+			}
+		}
+	}
+
+}
+
+
+void SCustomizableObjectRuntimeParameterList::BuildList()
+{
+	if (VerticalSlots.IsValid())
+	{
+		VerticalSlots->ClearChildren();
+
+		for (int32 i = 0; i < Node->States[StateIndex].RuntimeParameters.Num(); ++i)
+		{
+			VerticalSlots->AddSlot()
+				.Padding(0.0f, 2.0f, 0.0f, 2.0f)
+				[
+					SNew(SHorizontalBox)
+					+ SHorizontalBox::Slot()
+				.AutoWidth()
+				[
+					SNew(SCustomizableObjectRuntimeParameter)
+					.Node(Node)
+				.StateIndex(StateIndex)
+				.RuntimeParameterIndex(i)
+				]
+			+ SHorizontalBox::Slot()
+				.AutoWidth()
+				[
+					SNew(SButton)
+					.OnClicked(this, &SCustomizableObjectRuntimeParameterList::OnDeleteRuntimeParameter, i)
+				.ToolTipText(LOCTEXT("RemoveRuntimeParameter", "Remove Runtime Parameter"))
+				[
+					SNew(SHorizontalBox)
+					+ SHorizontalBox::Slot()
+				.VAlign(VAlign_Center)
+				.HAlign(HAlign_Center)
+				[
+					SNew(SImage)
+					.Image(FAppStyle::GetBrush(TEXT("Cross")))
+				]
+				]
+				]
+				];
+		}
+	}
+}
+
+
+FReply SCustomizableObjectRuntimeParameterList::OnDeleteRuntimeParameter(int32 ParameterIndex)
+{
+	Node->States[StateIndex].RuntimeParameters.RemoveAt(ParameterIndex);
+	BuildList();
+
+	return FReply::Handled();
+}
+
+
+// Widget for each Runtime parameter -------------------------------------------------------------------------------
+
+
+void SCustomizableObjectRuntimeParameter::Construct(const FArguments& InArgs)
+{
+	Node = InArgs._Node;
+	StateIndex = InArgs._StateIndex;
+	RuntimeParameterIndex = InArgs._RuntimeParameterIndex;
+
+	if (Node)
+	{
+		UCustomizableObject* CustomizableObject = Cast<UCustomizableObject>(Node->GetCustomizableObjectGraph()->GetOuter());
+
+		if (CustomizableObject)
+		{
+			ListViewOptions.Empty();
+			int32 ParameterCount = CustomizableObject->GetParameterCount();
+
+			for (int32 i = 0; i < ParameterCount; ++i)
+			{
+				ListViewOptions.Add(MakeShareable(new FString(CustomizableObject->GetParameterName(i))));
+			}
+
+			for (int32 i = 0; i < Node->States[StateIndex].RuntimeParameters.Num(); ++i)
+			{
+				bool bFound = false;
+
+				for (int32 j = 0; j < ListViewOptions.Num(); ++j)
+				{
+					if (*ListViewOptions[j].Get() == Node->States[StateIndex].RuntimeParameters[i])
+					{
+						bFound = true;
+						break;
+					}
+				}
+
+				if (!bFound)
+				{
+					ListViewOptions.Add(MakeShareable(new FString(Node->States[StateIndex].RuntimeParameters[i])));
+				}
+			}
+
+
+			// Alphabetical Order
+			ListViewOptions.Sort([](const TSharedPtr<FString>& A, const TSharedPtr<FString>& B)
+			{
+				FString StringA = *A.Get();
+				FString StringB = *B.Get();
+
+				int32 length = StringA.Len() > StringB.Len() ? StringB.Len() : StringA.Len();
+
+				for (int32 i = 0; i <length; ++i)
+				{
+					if (StringA[i] != StringB[i])
+					{
+						return StringA[i] < StringB[i];
+					}
+				}
+
+				return true;
+			}
+			);
+
+			ChildSlot
+			[
+				SNew(SHorizontalBox)
+
+				+ SHorizontalBox::Slot()
+				.AutoWidth()
+				.Padding(2.0f, 4.0f, 2.0f, 0.0f)
+				[
+					SNew(STextBlock)
+					.Text(FText::FromString(FString::Printf(TEXT("Parameter %d:"), RuntimeParameterIndex)))
+				]
+
+				+ SHorizontalBox::Slot()
+				.AutoWidth()
+				.Padding(8.0f, 0.0f, 0.0f, 0.0f)
+				[
+					SAssignNew(ComboButton,SComboButton)
+					.OnGetMenuContent(this, &SCustomizableObjectRuntimeParameter::GetComboButtonContent)
+					.ButtonContent()
+					[
+						SNew(STextBlock)
+						.Text(this, &SCustomizableObjectRuntimeParameter::GetCurrentItemLabel)
+					]
+				]
+			];
+		}
+	}
+}
+
+
+TSharedRef<SWidget> SCustomizableObjectRuntimeParameter::GetComboButtonContent()
+{
+	SearchItem.Reset();
+
+	// Listview Init
+	SAssignNew(RowNameComboListView, SListView<TSharedPtr<FString> >)
+		.ListItemsSource(&ListViewOptions)
+		.OnSelectionChanged(this, &SCustomizableObjectRuntimeParameter::OnComboButtonSelectionChanged)
+		.OnGenerateRow(this, &SCustomizableObjectRuntimeParameter::RowNameComboButtonGenarateWidget)
+		.SelectionMode(ESelectionMode::Single);
+
+	// Serachbox Init
+	SearchBoxWidget = SNew(SSearchBox)
+		.OnTextChanged(this, &SCustomizableObjectRuntimeParameter::OnSearchBoxFilterTextChanged)
+		.OnTextCommitted(this, &SCustomizableObjectRuntimeParameter::OnSearchBoxFilterTextCommitted);
+	
+	//Setting the focus to the Searchbox when the Comob button is oppen
+	ComboButton->SetMenuContentWidgetToFocus(SearchBoxWidget);
+
+	// Creating a widget that gives navigation to the searchbox and the listview
+	return  SNew(SVerticalBox)
+			+SVerticalBox::Slot()
+			.AutoHeight()
+			[
+				SearchBoxWidget.ToSharedRef()
+			]
+			+SVerticalBox::Slot()
+			.MaxHeight(100.0f)
+			[
+				RowNameComboListView.ToSharedRef()
+			];
+}
+
+
+void SCustomizableObjectRuntimeParameter::OnComboButtonSelectionChanged(TSharedPtr<FString> SelectedItem, ESelectInfo::Type SelectInfo)
+{
+	if (SelectInfo != ESelectInfo::OnNavigation)
+	{
+		if (SelectedItem.IsValid())
+		{
+			// Sets the value of the displayed name of the combo button
+			Node->States[StateIndex].RuntimeParameters[RuntimeParameterIndex] = *(SelectedItem.Get());
+
+			//Close the combobox when a selection is made
+			ComboButton->SetIsOpen(false);
+		}
+	}
+}
+
+
+TSharedRef<ITableRow> SCustomizableObjectRuntimeParameter::RowNameComboButtonGenarateWidget(TSharedPtr<FString> InItem, const TSharedRef<STableViewBase>& OwnerTable)
+{
+	//This is needed because the filter made in the function OnSearchBocFilterTextChanged Only works for the rendered items
+	const EVisibility WidgetVisibility = IsItemVisible(InItem) ? EVisibility::Visible : EVisibility::Collapsed;
+
+	return
+		SNew(STableRow<TSharedPtr<FString>>, OwnerTable)
+		.Visibility(WidgetVisibility)
+		[
+			SNew(STextBlock).Text(FText::FromString(*InItem))
+		];
+}
+
+
+FText SCustomizableObjectRuntimeParameter::GetRowNameComboButtonContentText() const
+{
+	TSharedPtr<FString> SelectedRowName = ComboButtonSelection;
+
+	if (SelectedRowName.IsValid())
+	{
+		return FText::FromString(*SelectedRowName);
+	}
+	else
+	{
+		return LOCTEXT("None", "None");
+	}
+}
+
+
+FText SCustomizableObjectRuntimeParameter::GetCurrentItemLabel() const
+{
+	return FText::FromString(Node->States[StateIndex].RuntimeParameters[RuntimeParameterIndex]);
+}
+
+
+void SCustomizableObjectRuntimeParameter::OnSearchBoxFilterTextChanged(const FText& InText)
+{
+	SearchItem = InText.ToString();
+
+	//This filter is just aplied to the items that are rendered of the ListView
+	for (int32 i = 0; i < ListViewOptions.Num(); i++)
+	{
+		TSharedPtr<ITableRow> Row = RowNameComboListView->WidgetFromItem(ListViewOptions[i]);
+		if (Row)
+		{
+			Row->AsWidget()->SetVisibility(IsItemVisible(ListViewOptions[i]) ? EVisibility::Visible : EVisibility::Collapsed);
+		}
+	}
+
+	RowNameComboListView->RequestListRefresh();
+}
+
+
+void SCustomizableObjectRuntimeParameter::OnSearchBoxFilterTextCommitted(const FText& InText, ETextCommit::Type CommitInfo)
+{
+	if (CommitInfo == ETextCommit::OnEnter)
+	{
+		bool bExists = false;
+
+		for (int32 i = 0; i < ListViewOptions.Num(); i++)
+		{
+			TSharedPtr<ITableRow> Row = RowNameComboListView->WidgetFromItem(ListViewOptions[i]);
+			if (Row && ListViewOptions[i].Get()->Equals(InText.ToString(), ESearchCase::IgnoreCase))
+			{
+				bExists = true;
+				Node->States[StateIndex].RuntimeParameters[RuntimeParameterIndex] = *ListViewOptions[i].Get();
+				break;
+			}
+		}
+
+		if (!bExists)
+		{			
+			UCustomizableObject* CustomizableObject = Cast<UCustomizableObject>(Node->GetCustomizableObjectGraph()->GetOuter());
+
+			if (CustomizableObject)
+			{
+				Node->States[StateIndex].RuntimeParameters[RuntimeParameterIndex] = InText.ToString();
+				ListViewOptions.Add(MakeShareable(new FString(InText.ToString())));
+			}
+		}
+
+		ComboButton->SetIsOpen(false);
+	}
+}
+
+
+bool SCustomizableObjectRuntimeParameter::IsItemVisible(TSharedPtr<FString> Item)
+{
+	bool bVisible = false;
+
+	if (SearchItem == "" || Item.Get()->Contains(SearchItem, ESearchCase::IgnoreCase))
+	{
+		bVisible = true;
+	}
+	
+	return bVisible;
+}
+
+#undef LOCTEXT_NAMESPACE
diff --git a/Engine/Plugins/Experimental/Mutable/Source/CustomizableObjectEditor/Private/SCustomizableObjectNodeObjectStatesView.h b/Engine/Plugins/Experimental/Mutable/Source/CustomizableObjectEditor/Private/SCustomizableObjectNodeObjectStatesView.h
new file mode 100644
--- /dev/null
+++ b/Engine/Plugins/Experimental/Mutable/Source/CustomizableObjectEditor/Private/SCustomizableObjectNodeObjectStatesView.h
@@ -0,0 +1,154 @@
+// Copyright Epic Games, Inc. All Rights Reserved.
+
+#pragma once
+
+#include "Widgets/SCompoundWidget.h"
+#include "SListViewSelectorDropdownMenu.h"
+#include "Widgets/Input/SSearchBox.h"
+#include "Widgets/Input/SComboBox.h"
+#include "Widgets/Input/SComboButton.h"
+#include "Widgets/Images/SImage.h"
+
+class SCustomizableObjectRuntimeParameter : public SCompoundWidget
+{
+public:
+
+	SLATE_BEGIN_ARGS(SCustomizableObjectRuntimeParameter) {}
+	SLATE_ARGUMENT(class UCustomizableObjectNodeObject*, Node)
+	SLATE_ARGUMENT(int32, StateIndex)
+	SLATE_ARGUMENT(int32, RuntimeParameterIndex)
+	SLATE_END_ARGS()
+
+	void Construct(const FArguments& InArgs);
+
+private:
+
+	/** Creates the content of the Combo Button */
+	TSharedRef<SWidget> GetComboButtonContent();
+
+	/** Generates the text of the combobox option */
+	FText GetCurrentItemLabel() const;
+
+	/** Callback for the Combo Button selection */
+	void OnComboButtonSelectionChanged(TSharedPtr<FString> SelectedItem, ESelectInfo::Type SelectInfo);
+
+	/** Generates the lables of the list view for the combo button */
+	TSharedRef<ITableRow> RowNameComboButtonGenarateWidget(TSharedPtr<FString> InItem, const TSharedRef<STableViewBase>& OwnerTable);
+
+	/** Generate the Combo Button selected label */
+	FText GetRowNameComboButtonContentText() const;
+
+	/** Callback for the OnTextChanged of the Searchbox */
+	void OnSearchBoxFilterTextChanged(const FText& InText);
+
+	/** Callback for the OnTextCommited of the Searchbox */
+	void OnSearchBoxFilterTextCommitted(const FText& InText, ETextCommit::Type CommitInfo);
+
+	// Returns true if the item should be visible in the Combo button
+	bool IsItemVisible(TSharedPtr<FString> Item);
+
+private:
+
+	/** Node with all the information */
+	class UCustomizableObjectNodeObject* Node;
+
+	/** Index to identify which parameter this widget is modifying */
+	int32 StateIndex;
+
+	/** Index to identify which parameter this widget is modifying */
+	int32 RuntimeParameterIndex;
+
+	/** Combobox for runtime parameter options */
+	TSharedPtr<SComboBox<TSharedPtr<FString>>> RuntimeParamCombobox;
+
+	/** Options shown in the ListView widget */
+	TArray<TSharedPtr<FString>> ListViewOptions;
+
+	/** ComboButton Selection */
+	TSharedPtr<FString> ComboButtonSelection;
+
+	/** ComboButton Widget */
+	TSharedPtr<SComboButton> ComboButton;
+	TSharedPtr< SListView< TSharedPtr< FString >>> RowNameComboListView;
+
+	/** Search box of the ComboButton */
+	TSharedPtr<SSearchBox> SearchBoxWidget;
+
+	/** Stores the input of the Searchboc widget*/
+	FString SearchItem;
+};
+
+
+/* Widget that represents List of Runtime Parameters Widgets */
+
+class SCustomizableObjectRuntimeParameterList : public SCompoundWidget
+{
+public:
+
+	SLATE_BEGIN_ARGS(SCustomizableObjectRuntimeParameterList) {}
+	SLATE_ARGUMENT(class UCustomizableObjectNodeObject*, Node)
+	SLATE_ARGUMENT(int32, StateIndex)
+	SLATE_END_ARGS()
+
+	void Construct(const FArguments& InArgs);
+	void BuildList();
+
+	/** Removes a Runtime parameter and rebuilds the runtimeparameter widgets */
+	FReply OnDeleteRuntimeParameter(int32 ParameterIndex);
+
+	bool IsCollapsed() { return bCollapsed; }
+	void SetCollapsed(bool Collapse ) { bCollapsed = Collapse; }
+	
+private:
+
+	/** Node with all the information */
+	class UCustomizableObjectNodeObject* Node;
+
+	/** Index to identify which parameter this widget is modifying */
+	int32 StateIndex;
+
+	/** Bool that determines if a RuntimeParameter should be collapsed or not*/
+	bool bCollapsed;
+
+	/** Vertical box widget for the RuntimeParameter Widgets  */
+	TSharedPtr<SVerticalBox> VerticalSlots;
+};
+
+class SCustomizableObjectNodeObjectSatesView : public SCompoundWidget
+{
+public:
+
+	SLATE_BEGIN_ARGS(SCustomizableObjectNodeObjectSatesView) {}
+	SLATE_ARGUMENT(UCustomizableObjectNodeObject*, Node)
+	SLATE_END_ARGS()
+
+	void Construct(const FArguments& InArgs);
+
+	/** Callback for the collapsing arrow checkbox */
+	void OnCollapseChanged(const ECheckBoxState NewCheckedState, int32 StateIndex);
+
+	/** Gets if a runtime parameter widget is collapsed */
+	EVisibility GetCollapsed(int32 StateIndex);
+
+	/** Return the brush for the collapsed arrow */
+	const FSlateBrush* GetExpressionPreviewArrow(int32 StateIndex) const;
+
+	/** Adds a new runtime parameter and rebuild the runtimeparameterListWidgets */
+	FReply OnAddRuntimeParameterPressed(int32 StateIndex);
+
+private:
+
+	/** Pointer to the current Node */
+	UCustomizableObjectNodeObject* Node;
+
+	/** Vertical Boxes to store the widget of each state */
+	TSharedPtr<SVerticalBox> VerticalSlots;
+
+	/** Arrays to edit widgets when needed */
+	TArray<TSharedPtr<SCustomizableObjectRuntimeParameterList>> RuntimeParametersListWidgets;
+
+	/** Array of SImages to control the collapsing buttons */
+	TArray<TSharedPtr<SImage>> CollapsedArrows;
+
+};
+
diff --git a/Engine/Plugins/Experimental/Mutable/Source/CustomizableObjectEditor/Private/SCustomizableObjectProperties.cpp b/Engine/Plugins/Experimental/Mutable/Source/CustomizableObjectEditor/Private/SCustomizableObjectProperties.cpp
new file mode 100644
--- /dev/null
+++ b/Engine/Plugins/Experimental/Mutable/Source/CustomizableObjectEditor/Private/SCustomizableObjectProperties.cpp
@@ -0,0 +1,445 @@
+// Copyright Epic Games, Inc. All Rights Reserved.
+
+#include "SCustomizableObjectProperties.h"
+#include "CustomizableObjectEditorModule.h"
+
+#include "Widgets/Input/SComboButton.h"
+
+#include "PackageTools.h"
+#include "AssetRegistry/AssetRegistryModule.h"
+#include "AssetSelection.h"
+
+#define LOCTEXT_NAMESPACE "SCustomizableObjectProperties"
+
+//DEFINE_LOG_CATEGORY_STATIC(LogEditorCustomizableObjectProperties, Log, All);
+
+class FStatePropertiesNode
+{
+public:
+
+	FStatePropertiesNode()
+	{
+	}
+
+	FStatePropertiesNode( const FString& InName )
+	{
+		Name = MakeShareable( new FString(InName) );
+	}
+
+
+	TSharedPtr<FString> Name;
+
+	TSharedPtr<FString> GetName()
+	{
+		return Name;
+	}
+
+
+	TArray<TSharedPtr<FStatePropertiesNode>> GetChildrenList()
+	{
+		TArray<TSharedPtr<FStatePropertiesNode>> ChildrenList;
+
+		return ChildrenList;
+	}
+
+};
+
+
+/** The item used for visualizing the class in the tree. */
+class SStatePropertiesItem : public STableRow< TSharedPtr<FString> >
+{
+public:
+	
+	SLATE_BEGIN_ARGS( SStatePropertiesItem )
+		: _StateName()
+		, _OwnerTableView()
+		, _HighlightText(NULL)
+		, _TextColor(FLinearColor(1.0f, 1.0f, 1.0f, 1.0f))
+		{}
+
+		/** The classname this item contains. */
+		SLATE_ARGUMENT( TSharedPtr<FString>, StateName )
+		/** The table that owns this item. */
+		SLATE_ARGUMENT( TSharedPtr<STableViewBase>, OwnerTableView )
+		/** The text this item should highlight, if any. */
+		SLATE_ARGUMENT( const FString*, HighlightText )
+		/** The color text this item will use. */
+		SLATE_ARGUMENT( FSlateColor, TextColor )
+		/** The node this item is associated with. */
+		SLATE_ARGUMENT( TSharedPtr<FStatePropertiesNode>, AssociatedNode)
+
+	SLATE_END_ARGS()
+
+	/**
+		* Construct the widget
+		*
+		* @param InArgs   A declaration from which to construct the widget
+		*/
+	void Construct( const FArguments& InArgs )
+	{
+		StateName = InArgs._StateName;
+		AssociatedNode = InArgs._AssociatedNode;
+
+		this->ChildSlot
+		[
+			SNew(SHorizontalBox)
+				
+			+SHorizontalBox::Slot()
+				[
+					SNew( SExpanderArrow, SharedThis(this) )
+				]
+
+			+SHorizontalBox::Slot()
+				.FillWidth(1.0f)
+				.Padding( 0.0f, 3.0f, 6.0f, 3.0f )
+				.VAlign(VAlign_Center)
+				[
+					SNew( STextBlock )
+						.Text( FText::FromString(*StateName.Get() ) )
+						//.HighlightString(*InArgs._HighlightText)
+						//.ColorAndOpacity( this, &SClassItem::GetTextColor)
+				]
+
+			+SHorizontalBox::Slot()
+				.HAlign(HAlign_Right)
+				.VAlign(VAlign_Center)
+				.Padding( 0.0f, 0.0f, 6.0f, 0.0f )
+				[
+					SNew( SComboButton )
+						.ContentPadding(FMargin(2.0f))
+						//.OnGetMenuContent(this, &SClassItem::GenerateDropDown)
+				]
+		];
+		
+		TextColor = InArgs._TextColor;
+
+		STableRow< TSharedPtr<FString> >::ConstructInternal(
+			STableRow::FArguments()
+				.ShowSelection(true),
+			InArgs._OwnerTableView.ToSharedRef() );	
+	}
+
+private:
+
+	/**
+		* Generates the drop down menu for the item.
+		*
+		* @return		The drop down menu widget.
+		*/
+	/*
+	TSharedRef<SWidget> GenerateDropDown()
+	{
+		// Empty list of commands.
+		TSharedPtr< FUICommandList > Commands;
+
+		const bool bShouldCloseWindowAfterMenuSelection = true;	// Set the menu to automatically close when the user commits to a choice
+		FMenuBuilder MenuBuilder( bShouldCloseWindowAfterMenuSelection, Commands);
+		{ 
+			if (UClass* Class = AssociatedNode->Class.Get())
+			{
+				bool bIsBlueprint(false);
+				bool bHasBlueprint(false);
+
+				ClassViewer::Helpers::GetClassInfo(Class, bIsBlueprint, bHasBlueprint);
+
+				bHasBlueprint = AssociatedNode->Blueprint.IsValid();
+
+				if (bIsBlueprint)
+				{
+					FUIAction Action( FExecuteAction::CreateStatic( &ClassViewer::Helpers::CreateBlueprint, Class ) );
+					MenuBuilder.AddMenuEntry(TEXT("Create Blueprint..."), TEXT("Creates a new Blueprint from the current one."), NAME_None, Action);
+				}
+
+				if (bHasBlueprint)
+				{
+					MenuBuilder.AddMenuSeparator();
+
+					FUIAction Action( FExecuteAction::CreateStatic( &ClassViewer::Helpers::OpenBlueprintTool, ClassViewer::Helpers::GetBlueprint(Class) ) );
+					MenuBuilder.AddMenuEntry(TEXT("Open Blueprint Editor"), TEXT("Open the Blueprint in the editor."), NAME_None, Action);
+				}
+
+				if (bHasBlueprint)
+				{
+					MenuBuilder.AddMenuSeparator();
+
+					FUIAction Action( FExecuteAction::CreateStatic( &ClassViewer::Helpers::FindInContentBrowser, ClassViewer::Helpers::GetBlueprint(Class), Class ) );
+					MenuBuilder.AddMenuEntry(TEXT("Find In Content Browser"), TEXT("Find in Content Browser"), NAME_None, Action);
+				}
+			}
+		}
+
+		return MenuBuilder.MakeWidget();
+	}
+	*/
+
+	/** Returns the text color for the item based on if it is selected or not. */
+	FSlateColor GetTextColor() const
+	{
+		const TSharedPtr< ITypedTableView< TSharedPtr<FString> > > OwnerWidget = OwnerTablePtr.Pin();
+		const TSharedPtr<FString>* MyItem = OwnerWidget->Private_ItemFromWidget( this );
+		const bool bIsSelected = OwnerWidget->Private_IsItemSelected( *MyItem );
+
+		if(bIsSelected)
+		{
+			return FSlateColor::UseForeground();
+		}
+
+		return TextColor;
+	}
+
+private:
+
+	/** The class name for which this item is associated with. */
+	TSharedPtr<FString> StateName;
+
+	/** The text color for this item. */
+	FSlateColor TextColor;
+
+	/** The Class Viewer Node this item is associated with. */
+	TSharedPtr< FStatePropertiesNode > AssociatedNode;
+
+};
+
+
+void SCustomizableObjectProperties::Construct(const FArguments& InArgs )
+{
+	bDisablePopulate = false;
+	bSaveExpansionStates = true;
+
+	Object = InArgs._Object;
+
+	this->ChildSlot
+	[
+		SNew(SVerticalBox)
+
+		+SVerticalBox::Slot()
+		.FillHeight(1.0f)
+		[
+			SNew(SVerticalBox)
+
+			+SVerticalBox::Slot()
+			.FillHeight(1.0f)
+			[
+				SAssignNew(StatesTree, STreeView<TSharedPtr< FStatePropertiesNode > >)
+					//.Visibility(EVisibility::Collapsed)
+					.SelectionMode(ESelectionMode::Single)
+					.TreeItemsSource( &RootTreeItems )
+					// Called to child items for any given parent item
+					.OnGetChildren( this, &SCustomizableObjectProperties::OnGetChildrenForTree )
+					// Generates the actual widget for a tree item
+					.OnGenerateRow( this, &SCustomizableObjectProperties::OnGenerateRow ) 
+
+					// Generates the right click menu.
+					//.OnContextMenuOpening(this, &SClassViewer::BuildMenuWidget)
+
+					// Find out when the user selects something in the tree
+					//.OnSelectionChanged( this, &SClassViewer::OnClassViewerSelectionChanged )
+
+					// Allow for some spacing between items with a larger item height.
+					.ItemHeight(20.0f)
+
+					.HeaderRow
+					(
+						SNew(SHeaderRow)
+						.Visibility(EVisibility::Collapsed)
+						+ SHeaderRow::Column(TEXT("State"))
+						.DefaultLabel(NSLOCTEXT("CustomizableObjectDetails","State","State"))
+					)
+			]
+		]
+	];
+
+	
+
+	// Register to have Populate called when doing a Hot Reload or when a Blueprint is compiled.
+	//GEditor->OnHotReload().AddSP(this, &SCustomizableObjectProperties::Populate);
+	//GEditor->OnBlueprintCompiled().AddSP(this, &SCustomizableObjectProperties::Populate);
+
+	Populate();
+}
+
+void SCustomizableObjectProperties::SetObject( UCustomizableObject* InObject )
+{
+	Object = InObject;
+	Populate();
+}
+
+TSharedRef<SWidget> SCustomizableObjectProperties::GetContent()
+{
+	return SharedThis( this );
+}
+
+SCustomizableObjectProperties::~SCustomizableObjectProperties()
+{
+	// Unregister to have Populate called when doing a Hot Reload or when a Blueprint is compiled.
+	//GEditor->OnHotReload().RemoveAll( this );
+	//GEditor->OnBlueprintCompiled().RemoveAll( this );
+	//GEditor->OnClassPackageLoadedOrUnloaded().RemoveAll( this );
+}
+
+void SCustomizableObjectProperties::OnGetChildrenForTree( TSharedPtr<FStatePropertiesNode> InParent, TArray< TSharedPtr< FStatePropertiesNode > >& OutChildren )
+{
+	// Simply return the children, it's already setup.
+	OutChildren = InParent->GetChildrenList();
+}
+
+
+TSharedRef< ITableRow > SCustomizableObjectProperties::OnGenerateRow( TSharedPtr<FStatePropertiesNode> Item, const TSharedRef< STableViewBase >& OwnerTable )
+{	
+	// If the item was accepted by the filter, leave it bright, otherwise dim it.
+	float AlphaValue = 1.0f;
+	TSharedRef< ITableRow > ReturnRow = SNew(SStatePropertiesItem)
+		.StateName(Item->GetName())
+		.TextColor(FLinearColor(1.0f, 1.0f, 1.0f, AlphaValue))
+		.AssociatedNode(Item)
+		.OwnerTableView( OwnerTable )
+		;
+
+	// Expand the item if needed.
+	bool* bIsExpanded = ExpansionStateMap.Find( Item->GetName() );
+	if( bIsExpanded && *bIsExpanded )
+	{
+		StatesTree->SetItemExpansion( Item, *bIsExpanded );
+	}
+
+	return ReturnRow;
+}
+
+
+void SCustomizableObjectProperties::SetAllExpansionStates(bool bInExpansionState)
+{
+	// Go through all the items in the root of the tree and recursively visit their children to set every item in the tree.
+	for(int32 ChildIndex = 0; ChildIndex < RootTreeItems.Num(); ChildIndex++)
+	{
+		SetAllExpansionStates_Helper( RootTreeItems[ChildIndex], bInExpansionState );
+	}
+}
+
+void SCustomizableObjectProperties::SetAllExpansionStates_Helper(TSharedPtr< FStatePropertiesNode > InNode, bool bInExpansionState)
+{
+	StatesTree->SetItemExpansion(InNode, bInExpansionState);
+
+	// Recursively go through the children.
+	for(int32 ChildIndex = 0; ChildIndex < InNode->GetChildrenList().Num(); ChildIndex++)
+	{
+		SetAllExpansionStates_Helper( InNode->GetChildrenList()[ChildIndex], bInExpansionState );
+	}
+}
+
+void SCustomizableObjectProperties::ResetExpansionStates()
+{
+	// Empty the map of states, do not save the states, and refresh the tree with a call to Populate so it will build the tree fresh.
+	ExpansionStateMap.Empty();
+	bSaveExpansionStates = false;
+	Populate();
+}
+
+
+void SCustomizableObjectProperties::MapExpansionStatesInTree( TSharedPtr<FStatePropertiesNode> InItem )
+{
+	ExpansionStateMap.Add( InItem->GetName(), StatesTree->IsItemExpanded( InItem ) );
+
+	// Map out all the children, this will be done recursively.
+	for( int32 ChildIdx(0); ChildIdx < InItem->GetChildrenList().Num(); ++ChildIdx )
+	{
+		MapExpansionStatesInTree( InItem->GetChildrenList()[ChildIdx] );
+	}
+}
+
+void SCustomizableObjectProperties::SetExpansionStatesInTree( TSharedPtr<FStatePropertiesNode> InItem )
+{
+	bool* bIsExpanded = ExpansionStateMap.Find( InItem->GetName() );
+	if( bIsExpanded )
+	{
+		StatesTree->SetItemExpansion( InItem, *bIsExpanded );
+
+		// No reason to set expansion states if the parent is not expanded, it does not seem to do anything.
+		if( *bIsExpanded )
+		{
+			for( int32 ChildIdx(0); ChildIdx < InItem->GetChildrenList().Num(); ++ChildIdx )
+			{
+				SetExpansionStatesInTree( InItem->GetChildrenList()[ChildIdx] );
+			}
+		}
+	}
+	else
+	{
+		// Default to no expansion.
+		StatesTree->SetItemExpansion( InItem, false );
+	}
+}
+
+void SCustomizableObjectProperties::Populate()
+{
+	// Populating the tree/list is disabled during loads to stop needless re-population.
+	if(bDisablePopulate)
+	{
+		return;
+	}
+
+	// If showing a class tree, we may need to save expansion states.
+	if( StatesTree->GetVisibility() == EVisibility::Visible )
+	{
+		if( bSaveExpansionStates )
+		{
+			for( int32 ChildIdx(0); ChildIdx < RootTreeItems.Num(); ++ChildIdx )
+			{
+				// Check if the item is actually expanded or if it's only expanded because it is root level.
+				bool* bIsExpanded = ExpansionStateMap.Find( RootTreeItems[ChildIdx]->GetName() );
+				if( (bIsExpanded && !*bIsExpanded) || !bIsExpanded)
+				{
+					StatesTree->SetItemExpansion( RootTreeItems[ChildIdx], false );
+				}
+
+				// Recursively map out the expansion state of the tree-node.
+				MapExpansionStatesInTree( RootTreeItems[ChildIdx] );
+			}
+		}
+
+		// This is set to false before the call to populate when it is not desired.
+		bSaveExpansionStates = true;
+	}
+
+	// Empty the tree out so it can be redone.
+	RootTreeItems.Empty();
+
+	// Check if we will restore expansion states, we will not if there is filtering happening.
+	const bool bRestoreExpansionState = true;
+
+	if ( Object )
+	{
+		/*
+		for (int o=0; o<Object->States.Num(); ++o)
+		{
+			TSharedPtr<FStatePropertiesNode> node = MakeShareable( new FStatePropertiesNode );
+			node->Name = MakeShareable( new FString(Object->States[o].Name) );
+			RootTreeItems.Add( node );
+
+			if( bRestoreExpansionState )
+			{
+				SetExpansionStatesInTree( node );
+			}
+		}
+		*/
+	}
+
+	// All root level items get expanded automatically.
+	//StatesTree->SetItemExpansion(RootNode, true);
+
+	// Now that new items are in the tree, we need to request a refresh.
+	StatesTree->RequestTreeRefresh();
+}
+
+TSharedPtr<FStatePropertiesNode> SCustomizableObjectProperties::CreateNoneOption()
+{
+	TSharedPtr<FStatePropertiesNode> NoneItem = MakeShareable( new FStatePropertiesNode("None") );
+
+	// The item "passes" the filter so it does not appear grayed out.
+	//NoneItem->bPassesFilter = true;
+
+	return NoneItem;
+}
+
+
+#undef LOCTEXT_NAMESPACE
+
diff --git a/Engine/Plugins/Experimental/Mutable/Source/CustomizableObjectEditor/Private/SCustomizableObjectProperties.h b/Engine/Plugins/Experimental/Mutable/Source/CustomizableObjectEditor/Private/SCustomizableObjectProperties.h
new file mode 100644
--- /dev/null
+++ b/Engine/Plugins/Experimental/Mutable/Source/CustomizableObjectEditor/Private/SCustomizableObjectProperties.h
@@ -0,0 +1,107 @@
+// Copyright Epic Games, Inc. All Rights Reserved.
+
+#pragma once
+
+
+#include "CustomizableObjectEditorModule.h"
+#include "Widgets/Views/STreeView.h"
+
+
+//////////////////////////////////////////////////////////////////////////
+// SClassViewer
+
+class FStatePropertiesNode;
+
+class SCustomizableObjectProperties : public SCompoundWidget
+{
+public:
+	SLATE_BEGIN_ARGS( SCustomizableObjectProperties )
+		: _Object((UCustomizableObject*)NULL)
+	{
+	}
+
+	SLATE_ARGUMENT( UCustomizableObject*, Object )
+
+	SLATE_END_ARGS()
+
+	/**
+	 * Construct the widget
+	 *
+	 * @param	InArgs			A declaration from which to construct the widget
+	 */
+	void Construct(const FArguments& InArgs);
+
+	/** Gets the widget contents of the app */
+	virtual TSharedRef<SWidget> GetContent();
+
+	virtual ~SCustomizableObjectProperties();
+
+	/** */
+	void SetObject( UCustomizableObject* InObject );
+
+private:
+	/** Retrieves the children for the input node.
+	 *	@param InParent				The parent node to retrieve the children from.
+	 *	@param OutChildren			List of children for the parent node.
+	 *
+	 */
+	void OnGetChildrenForTree( TSharedPtr<FStatePropertiesNode> InParent, TArray< TSharedPtr< FStatePropertiesNode > >& OutChildren );
+
+	/** Creates the row widget when called by Slate when an item appears on the tree. */
+	TSharedRef< ITableRow > OnGenerateRow( TSharedPtr<FStatePropertiesNode> Item, const TSharedRef< STableViewBase >& OwnerTable );
+
+	/** Called by Slate when an item is selected from the tree/list. */
+	void OnClassViewerSelectionChanged( TSharedPtr<FStatePropertiesNode> Item, ESelectInfo::Type SelectInfo );
+
+	/** 
+	 *	Sets all expansion states in the tree.
+	 *
+	 *	@param bInExpansionState			The expansion state to set the tree to.
+	 */
+	void SetAllExpansionStates(bool bInExpansionState);
+
+	/**
+	 *	A helper function to recursively set the tree.
+	 *
+	 *	@param	InNode						The current node in the tree.
+	 *	@param	bInExpansionState			The expansion state to set the tree to.
+	 */
+	void SetAllExpansionStates_Helper(TSharedPtr< FStatePropertiesNode > InNode, bool bInExpansionState);
+
+	/** Resets the expansion states on the tree to default. */
+	void ResetExpansionStates();
+
+	/** Recursive function to map the expansion states of items in the tree. 
+	 *	@param InItem		The current item to examine the expansion state of.
+	 */
+	void MapExpansionStatesInTree( TSharedPtr<FStatePropertiesNode> InItem );
+
+	/** Recursive function to set the expansion states of items in the tree. 
+	 *	@param InItem		The current item to set the expansion state of.
+	 */
+	void SetExpansionStatesInTree( TSharedPtr<FStatePropertiesNode> InItem );
+
+	/** Populates the tree with items based on the current filter. */
+	void Populate();
+
+	/** Creates a "None" option for the tree/list. */
+	TSharedPtr<FStatePropertiesNode> CreateNoneOption();
+
+private:
+
+	UCustomizableObject* Object;
+
+	TSharedPtr<STreeView<TSharedPtr< FStatePropertiesNode > >> StatesTree;
+
+	/** The root items to be displayed in the state tree. */
+	TArray< TSharedPtr< FStatePropertiesNode > > RootTreeItems;
+
+	/** Turns off the re-population for a single call. */
+	bool bDisablePopulate;
+
+	/** true if expansions states should be saved when compiling. */
+	bool bSaveExpansionStates;
+
+	/** The map holding the expansion state map for the tree. */
+	TMap< TSharedPtr< FString >, bool > ExpansionStateMap;
+};
diff --git a/Engine/Plugins/Experimental/Mutable/Source/CustomizableObjectEditor/Private/SCustomizableObjectSystemVersion.h b/Engine/Plugins/Experimental/Mutable/Source/CustomizableObjectEditor/Private/SCustomizableObjectSystemVersion.h
new file mode 100644
--- /dev/null
+++ b/Engine/Plugins/Experimental/Mutable/Source/CustomizableObjectEditor/Private/SCustomizableObjectSystemVersion.h
@@ -0,0 +1,30 @@
+// Copyright Epic Games, Inc. All Rights Reserved.
+
+#pragma once
+
+#include "CustomizableObjectSystem.h"
+#include "Widgets/SCompoundWidget.h"
+#include "Widgets/DeclarativeSyntaxSupport.h"
+
+
+class CUSTOMIZABLEOBJECTEDITOR_API SCustomizableObjectSystemVersion : public SCompoundWidget
+{
+public:
+
+	SLATE_BEGIN_ARGS(SCustomizableObjectSystemVersion)
+		{}
+	SLATE_END_ARGS()
+
+	/** */
+	void Construct( const FArguments& InArgs );
+
+public:
+
+	// SWidget interface
+	void Tick( const FGeometry& AllottedGeometry, const double InCurrentTime, const float InDeltaTime ) override;
+
+private:
+
+
+};
+
