[mutable] Move Mutable plugin from the restricted folder to the experimental plugins folder.

diff --git a/Engine/Plugins/Experimental/Mutable/Source/CustomizableObjectEditor/Private/SMutableCodeViewer.cpp b/Engine/Plugins/Experimental/Mutable/Source/CustomizableObjectEditor/Private/SMutableCodeViewer.cpp
new file mode 100644
--- /dev/null
+++ b/Engine/Plugins/Experimental/Mutable/Source/CustomizableObjectEditor/Private/SMutableCodeViewer.cpp
@@ -0,0 +1,2351 @@
+// Copyright Epic Games, Inc. All Rights Reserved.
+
+#include "SMutableCodeViewer.h"
+
+#include "SMutableImageViewer.h"
+#include "SMutableMeshViewer.h"
+#include "CustomizableObjectCompiler.h"
+#include "MutableTools/Private/ErrorLogPrivate.h"
+#include "MutableTools/Public/Streams.h"
+#include "CustomizableObject/Private/CustomizableInstancePrivateData.h"
+#include "MutableRuntime/Private/ModelPrivate.h"
+#include "MutableRuntime/Private/SystemPrivate.h"
+#include "MutableRuntime/Public/Image.h"
+#include "MutableRuntime/Public/Mesh.h"
+#include "MutableRuntime/Public/MutableString.h"
+
+#include "Framework/MultiBox/MultiBoxBuilder.h"
+#include "Widgets/Input/SButton.h"
+#include "Widgets/Input/STextComboBox.h"
+#include "Widgets/Views/STreeView.h"
+#include "Widgets/Images/SImage.h"
+#include "Widgets/Layout/SSplitter.h"
+#include "Widgets/Layout/SScrollBox.h"
+#include "Widgets/Input/SNumericEntryBox.h"
+#include "Widgets/Docking/SDockTab.h"
+#include "EditorDirectories.h"
+#include "DesktopPlatformModule.h"
+#include "SEnumCombo.h"
+#include "HAL/PlatformApplicationMisc.h"
+#include "Internationalization/Text.h"
+
+#include "SMutableColorPreviewBox.h"
+#include "SMutableConstantsWidget.h"
+
+#include "Widgets/Views/SHeaderRow.h"
+#include "SMutableColorPreviewBox.h"
+
+
+#define LOCTEXT_NAMESPACE "SMutableDebugger"
+
+
+namespace MutableCodeTreeViewColumns
+{
+	static const FName OperationsColumnID("Operations");
+	static const FName AdditionalDataColumnID("Additional Data");
+};
+
+/**
+ * Mutable tree row used to display the operations held on the Mutable model object. 
+ */
+class SMutableCodeTreeRow final : public SMultiColumnTableRow<TSharedPtr<FMutableCodeTreeElement>>
+{
+public:
+
+	void Construct(const FArguments& Args, const TSharedRef<STableViewBase>& InOwnerTableView, const TSharedPtr<FMutableCodeTreeElement>& InRowItem)
+	{
+		RowItem = InRowItem;
+		
+		SMultiColumnTableRow< TSharedPtr<FMutableCodeTreeElement> >::Construct(
+			STableRow::FArguments()
+			.ShowSelection(true)
+			, InOwnerTableView
+		);
+	}
+
+	/** Method intended with the generation of the wanted objects for each column*/
+	virtual TSharedRef<SWidget> GenerateWidgetForColumn(const FName& ColumnName) override
+	{
+		// Primary column showing the name of the operation and tye type
+		if (ColumnName == MutableCodeTreeViewColumns::OperationsColumnID)
+		{
+			const mu::PROGRAM& Program = RowItem->MutableModel->GetPrivate()->m_program;
+			const mu::OP_TYPE Type = Program.GetOpType(RowItem->MutableOperation);
+			const TCHAR* OpName = mu::s_opNames[int32(Type)];
+
+			// Prepare the text shown on the UI side of the operation tree
+			FString MainLabel = FString::Printf(TEXT("%s (%d) : %s"), *RowItem->Caption, int32(RowItem->MutableOperation), OpName);
+			if (RowItem->DuplicatedOf)
+			{
+				MainLabel.Append(TEXT(" (duplicated)"));
+			}
+
+			// TODO:
+			const FSlateBrush* IconBrush = nullptr;
+			
+			// Prepare a ui container for all the UI objects required by this row element
+			TSharedRef<SHorizontalBox> RowContainer = SNew(SHorizontalBox)
+				
+			// First coll showing operation name and type
+			+ SHorizontalBox::Slot()
+			.HAlign(EHorizontalAlignment::HAlign_Left)
+			[
+				SNew(SOverlay)
+
+				+ SOverlay::Slot()
+				[
+					SAssignNew(this->HighlightingColorBox, SMutableColorPreviewBox)
+					.BoxColor(HighlightBoxDefaultColor)
+				]
+
+				+ SOverlay::Slot()
+				[
+					SNew(SHorizontalBox)
+					+ SHorizontalBox::Slot()
+					.VAlign(VAlign_Center)
+					.AutoWidth()
+					[
+						SNew(SExpanderArrow, SharedThis(this))
+					]
+
+					+ SHorizontalBox::Slot()
+					.VAlign(VAlign_Center)
+					.Padding(FMargin(5.f, 0.f, 5.f, 0.f))
+					.AutoWidth()
+					[
+						SNew(SOverlay)
+						+ SOverlay::Slot()
+						[
+							SNew(SImage)
+							.Image(IconBrush ? IconBrush : FCoreStyle::Get().GetDefaultBrush())
+							.ColorAndOpacity(IconBrush ? FLinearColor::White : FLinearColor::Transparent)
+						]
+					]
+
+					+ SHorizontalBox::Slot()
+					[
+						SNew(STextBlock)
+						.Text(FText::FromString(MainLabel))
+					]
+				]
+			];
+
+			return RowContainer;
+		}
+
+		// Second column showing some extra data related with the operation being displayed
+		if (ColumnName == MutableCodeTreeViewColumns::AdditionalDataColumnID)
+		{
+			// Determine what text will be shown on the extra data column alongside with the background color used for it
+			FText OperationExtraData = FText::FromString(FString(""));
+			const FSlateColor* OperationDataBoxColor = &ExtraDataBackgroundBoxDefaultColor;
+			if (RowItem->bIsDynamicResource)
+			{
+				OperationExtraData = DynamicResourceText;
+				OperationDataBoxColor = &DynamicResourceBoxColor;
+			}
+			else if (RowItem->bIsStateConstant)
+			{
+				OperationExtraData = StateConstantText;
+				OperationDataBoxColor = &StateConstantBoxColor;
+			}
+			
+			TSharedRef<SHorizontalBox> RowContainer =  SNew(SHorizontalBox)
+			
+			+ SHorizontalBox::Slot()
+			.HAlign(EHorizontalAlignment::HAlign_Left)
+			[
+				SNew(SHorizontalBox)
+
+				// Cheap way of doing a fixed space with a color
+				+SHorizontalBox::Slot()
+				.AutoWidth()
+				[
+					 SNew(SOverlay)
+					 +SOverlay::Slot()
+					 [
+						SNew(SMutableColorPreviewBox)
+						.BoxColor(*OperationDataBoxColor)
+					 ]
+
+					 // Hacky way of setting a fixed size for the color box
+					 + SOverlay::Slot()
+					 [
+						SNew(STextBlock)
+						.Text(EmptyText)
+					 ]
+				]
+
+				+ SHorizontalBox::Slot()
+				.Padding(4,1)
+				.FillWidth(0.9f)
+				[
+					SNew(STextBlock)
+					.Text(OperationExtraData)
+				]
+			];
+			
+			return RowContainer;
+		}
+
+		// Invalid column name so no widget will be produced 
+		return SNullWidget::NullWidget;
+	}
+	
+	/** Set the row background to one or another highlighting color depending if it is a unique or a duplicated row */
+	void Highlight() const
+	{
+		if (RowItem->DuplicatedOf)
+		{
+			HighlightingColorBox->SetColor(HighlightedDuplicatedBoxColor);
+		}
+		else
+		{
+			HighlightingColorBox->SetColor(HighlightedUniqueRowBoxColor);
+		}
+	}
+
+	/** Sets the background of the row to the default state */
+	void ResetHighlight() const
+	{
+		HighlightingColorBox->SetColor(HighlightBoxDefaultColor);
+	}
+
+	/** Returns a reference to the Element this row is representing */
+	TSharedPtr<FMutableCodeTreeElement>& GetItem()
+	{
+		return RowItem;
+	}
+
+private:
+
+	/** Pointer to the element that did spawn this row */
+	TSharedPtr<FMutableCodeTreeElement> RowItem = nullptr;
+
+	/** Transparent color */
+	const FLinearColor TransparentColor = FLinearColor(0,0,0,0);
+
+	/*
+	 * Operation Highlighting 
+	 */
+	
+	/** Custom Widget used to display a color. Used as the background of the text on the row to serve as highlighting Visual Element*/
+	TSharedPtr<SMutableColorPreviewBox> HighlightingColorBox = nullptr;
+	
+	/** The color used to highlight the row if duplicated from another row */
+	const FSlateColor HighlightedDuplicatedBoxColor = FSlateColor(FLinearColor(1, 1, 1, 0.15));
+
+	/** The color used to highlight elements that are originals (not duplicates)  */
+	const FSlateColor HighlightedUniqueRowBoxColor = FSlateColor(FLinearColor(1, 1, 1, 0.28));
+
+	/** Default color used when the row is not highlighted */
+	const FSlateColor HighlightBoxDefaultColor = FSlateColor(TransparentColor);
+
+	/*
+	 * Extra data objects
+	 */
+
+	// Text used to set the width of the color area in front of the extra data
+	const FText EmptyText = FText(INVTEXT(" "));
+	
+	/** String printed on the UI when the operation is shown to be dynamic resource */
+	const FText DynamicResourceText = FText::FromString(FString("DYNAMIC RESOURCE"));
+
+	/** String printed on the UI when the operation is shown to be state constant */
+	const FText StateConstantText = FText::FromString(FString("STATE CONSTANT"));
+	
+	/** Color used on the extra data column when no extra data is shown */
+	const FSlateColor ExtraDataBackgroundBoxDefaultColor =  FSlateColor(TransparentColor);
+
+	/** Color shown on the extra data column when the resource is found to be Dynamic */
+	const FSlateColor DynamicResourceBoxColor = FSlateColor(FLinearColor(0,0,1,0.8));
+
+	/** Color shown on the extra data column when the resource is found to be State Constant */
+	const FSlateColor StateConstantBoxColor = FSlateColor(FLinearColor(1,0,0,0.8));
+
+};
+
+
+void SMutableCodeViewer::AddReferencedObjects(FReferenceCollector& Collector)
+{
+	// Add UObjects here if we own any at some point
+	//Collector.AddReferencedObject(CustomizableObject);
+}
+
+
+FString SMutableCodeViewer::GetReferencerName() const
+{
+	return TEXT("SMutableCodeViewer");
+}
+
+
+void SMutableCodeViewer::Construct(const FArguments& InArgs, const mu::ModelPtr& InMutableModel/*, const TSharedPtr<SDockTab>& ConstructUnderMajorTab*/)
+{
+	MutableModel = InMutableModel;
+	PreviewParameters = MutableModel->NewParameters();
+	
+	// Min width allowed for the column. Needed to avoid having issues with the constants space being to small
+	// and then getting too tall on the y axis crashing the drawer.
+	constexpr float MinParametersCollWidth = 400;
+	
+	// create & initialize tab manager
+	//TabManager = FGlobalTabmanager::Get()->NewTabManager(ConstructUnderMajorTab.ToSharedRef());
+
+	const mu::Model::Private* ModelPrivate = MutableModel->GetPrivate();
+	const int32 StateCount = int32( ModelPrivate->m_program.m_states.size() );
+	for ( int32 StateIndex=0; StateIndex<StateCount; ++StateIndex )
+	{
+		const mu::PROGRAM::STATE& State = ModelPrivate->m_program.m_states[StateIndex];
+		FString Caption = FString::Printf( TEXT("state [%s]"), ANSI_TO_TCHAR(State.m_name.c_str()) );
+		RootNodes.Add(MakeShareable(new FMutableCodeTreeElement(MutableModel, State.m_root, Caption)));
+	}
+
+	// Setup Navigation system
+	{
+		// Store the addresses of the root nodes for later search operations
+		CacheRootNodeAddresses();
+
+		// Cache the operation types that are present on the model
+		CacheOperationTypesPresentOnModel();
+	
+		// Get an array of mutable types as an array of FStrings for the UI
+		GenerateNavigationOpTypeStrings();
+	}
+
+	
+	FToolBarBuilder ToolbarBuilder(TSharedPtr<const FUICommandList>(), FMultiBoxCustomization::None, TSharedPtr<FExtender>(), true);
+	ToolbarBuilder.SetLabelVisibility(EVisibility::Visible);
+	ToolbarBuilder.SetStyle(&FAppStyle::Get(), "SlimToolBar");
+
+	ToolbarBuilder.BeginSection("Export");
+
+	// Tree Sizes
+	constexpr float OperationsColumnWidth = 0.75f;
+	constexpr float ExtraDataColumnWidth = 0.25f;
+	
+	// Export
+	ToolbarBuilder.AddToolBarButton(
+		FUIAction(
+			FExecuteAction::CreateLambda([InMutableModel]()
+				{
+					TArray<FString> SaveFilenames;
+					IDesktopPlatform* DesktopPlatform = FDesktopPlatformModule::Get();
+					bool bSave = false;
+					if (DesktopPlatform)
+					{
+						const FString LastExportPath = FEditorDirectories::Get().GetLastDirectory(ELastDirectory::GENERIC_EXPORT);
+						const FString FileTypes = TEXT("Mutable compiled data files|*.mutable_compiled|All files|*.*");
+						bSave = DesktopPlatform->SaveFileDialog(
+							FSlateApplication::Get().FindBestParentWindowHandleForDialogs(nullptr),
+							TEXT("Export Mutable compiled object"),
+							*LastExportPath,
+							TEXT("exported.mutable_compiled"),
+							*FileTypes,
+							EFileDialogFlags::None,
+							SaveFilenames
+						);
+					}
+
+					if (bSave)
+					{
+						const FString SaveFileName = FString(SaveFilenames[0]);
+
+						mu::OutputFileStream Stream(TCHAR_TO_ANSI(*SaveFileName));
+						Stream.Write(MUTABLE_COMPILED_MODEL_FILETAG, 4);
+						mu::OutputArchive Archive(&Stream);
+						mu::Model::Serialise(InMutableModel.get(), Archive);
+						const uint32_t CodeVersion = MUTABLE_COMPILED_MODEL_CODE_VERSION;
+						Stream.Write((const uint8*)&CodeVersion, sizeof(uint32_t));
+						Stream.Flush();
+
+						FEditorDirectories::Get().SetLastDirectory(ELastDirectory::GENERIC_EXPORT, SaveFileName);
+					}
+				})
+			),
+			NAME_None,
+			LOCTEXT("ExportMutableCode", "Export"),
+			LOCTEXT("ExportMutableCodeTooltip", "Export a debug mutable compiled file."),
+			FSlateIcon(),
+			EUserInterfaceActionType::Button
+		);
+		
+	ToolbarBuilder.EndSection();
+
+	ToolbarBuilder.AddWidget(SNew(STextBlock).Text(FText::FromString(InArgs._DataTag)));
+	
+	ChildSlot
+	[
+		SNew(SVerticalBox)
+		+ SVerticalBox::Slot()
+		.AutoHeight()
+		.VAlign(VAlign_Center)
+		[
+			ToolbarBuilder.MakeWidget()
+		]
+		+ SVerticalBox::Slot()
+		.VAlign(VAlign_Fill)
+		[
+			SNew(SSplitter)
+			.Orientation(EOrientation::Orient_Horizontal)
+			
+			+ SSplitter::Slot()
+			.Value(0.35f)
+			[
+				SNew(SVerticalBox)
+		
+				// Operation type filtering slot
+				+ SVerticalBox::Slot()
+				.AutoHeight()
+				.HAlign(HAlign_Right)
+				[
+					// Border containing navigation elements
+					SNew(SBorder)
+					.HAlign(HAlign_Right)
+					[
+						SNew(SVerticalBox)
+
+						+ SVerticalBox::Slot()
+						.Padding(2,4)
+						.AutoHeight()
+						[
+							SNew(SHorizontalBox)
+						
+							+ SHorizontalBox::Slot()
+							.AutoWidth()
+							[
+								SNew(SHorizontalBox)
+
+								+ SHorizontalBox::Slot()
+								.AutoWidth()
+								[
+									SNew(STextBlock)
+									.Text(LOCTEXT("SelectedOperationTypeLabel","Search Operation Type :"))
+								]
+
+								+ SHorizontalBox::Slot()
+								.AutoWidth()
+								[
+									// TODO: Try to use a SComboButton so you are able to use a search box (Max)
+									SAssignNew(TargetedTypeSelector, STextComboBox)
+									.OptionsSource(&OperationTypesStrings)
+									.InitiallySelectedItem(OperationTypesStrings[0])	
+									.OnSelectionChanged(this,&SMutableCodeViewer::OnOptionTypeSelectionChanged)
+								]
+							]
+
+							+ SHorizontalBox::Slot()
+							.Padding(4,0)
+							.AutoWidth()
+							[
+								SNew(SButton)
+								.Text(LOCTEXT("GoToPreviousOperationButton"," < "))
+								.OnClicked(this,&SMutableCodeViewer::OnGoToPreviousOperationButtonPressed)
+								.IsEnabled(this,&SMutableCodeViewer::CanInteractWithPreviousOperationButton)
+							]
+
+							+ SHorizontalBox::Slot()
+							.AutoWidth()
+							[
+								SNew(SButton)
+								.Text(LOCTEXT("GoToNextOperationButton"," > "))
+								.OnClicked(this,&SMutableCodeViewer::OnGoToNextOperationButtonPressed)
+								.IsEnabled(this,&SMutableCodeViewer::CanInteractWithNextOperationButton)
+							]
+						]
+
+						+ SVerticalBox::Slot()
+						.AutoHeight()
+						.Padding(0,4)
+						[
+							SNew(STextBlock)
+							.Text(this,&SMutableCodeViewer::OnPrintNavigableObjectAddressesCount)
+							.Justification(ETextJustify::Right)
+						]
+					]
+					
+				]
+					
+				
+				// Tree operations slot
+				+ SVerticalBox::Slot()
+				.FillHeight(1.0f)
+				[
+					SNew(SBorder)
+					.BorderImage(FAppStyle::GetBrush("ToolPanel.GroupBorder"))
+					.Padding(FMargin(4.0f, 4.0f))
+					[
+						SAssignNew(TreeView, STreeView<TSharedPtr<FMutableCodeTreeElement>>)
+						.TreeItemsSource(&RootNodes)
+						.OnGenerateRow(this, &SMutableCodeViewer::GenerateRowForNodeTree)
+						.OnRowReleased(this, &SMutableCodeViewer::OnRowReleased)
+						.OnGetChildren(this, &SMutableCodeViewer::GetChildrenForInfo)
+						.OnSelectionChanged(this, &SMutableCodeViewer::OnSelectionChanged)
+						.OnSetExpansionRecursive(this, &SMutableCodeViewer::TreeExpandRecursive)
+						.OnContextMenuOpening(this, &SMutableCodeViewer::OnTreeContextMenuOpening)
+						.OnExpansionChanged(this, &SMutableCodeViewer::OnExpansionChanged)
+						.SelectionMode(ESelectionMode::Single)
+						.HeaderRow
+						(
+							SNew(SHeaderRow)
+
+							+ SHeaderRow::Column(MutableCodeTreeViewColumns::OperationsColumnID)
+								.DefaultLabel(LOCTEXT("Operation", "Operation"))
+								.FillWidth(OperationsColumnWidth)
+								
+							+ SHeaderRow::Column(MutableCodeTreeViewColumns::AdditionalDataColumnID)
+								.DefaultLabel(LOCTEXT("AdditionalData", "Additional Data"))
+								.FillWidth(ExtraDataColumnWidth)
+						)
+					]
+				]
+				
+			]
+			+ SSplitter::Slot()
+			.Value(0.75f)
+			[
+				SNew(SSplitter)
+				.Orientation(EOrientation::Orient_Horizontal)
+				
+				+ SSplitter::Slot()
+				.Value(0.28f)
+				.MinSize(MinParametersCollWidth)
+				[
+					// Splitter managing both parameter and constant panels
+					SNew(SSplitter)
+					.Orientation(EOrientation::Orient_Vertical)
+					+SSplitter::Slot()
+					[
+						SNew(SVerticalBox)
+
+						+ SVerticalBox::Slot()
+						.AutoHeight()
+						[
+							SNew(SHorizontalBox)
+
+							+ SHorizontalBox::Slot()
+							.AutoWidth()
+							[
+								SNew(STextBlock)
+								.Text(LOCTEXT("SkipMipsLabel", "Skip mips on generate :"))
+								.Visibility(this, &SMutableCodeViewer::IsMipSkipVisible)
+							]
+
+							+ SHorizontalBox::Slot()
+							[
+								SNew(SNumericEntryBox<int32>)
+								.Visibility(this, &SMutableCodeViewer::IsMipSkipVisible)
+								.AllowSpin(true)
+								.MinValue(0)
+								.MaxValue(16)
+								.MinSliderValue(0)
+								.MaxSliderValue(16)
+								.Value(this, &SMutableCodeViewer::GetCurrentMipSkip)
+								.OnValueChanged(this, &SMutableCodeViewer::OnCurrentMipSkipChanged)
+							]
+						]
+
+						+ SVerticalBox::Slot()
+						[
+							SNew(SScrollBox)
+							+ SScrollBox::Slot()
+							[
+								SAssignNew(ParametersWidget, SMutableParametersWidget)
+								.OnParametersValueChanged(this, &SMutableCodeViewer::OnPreviewParameterValueChanged)
+							]
+						]
+					]
+
+					+ SSplitter::Slot()
+					[
+						// Generate a new Constants panel to show the data stored on the current mutable program
+						SAssignNew(ConstantsWidget, SMutableConstantsWidget,
+							&(ModelPrivate->m_program),
+							SharedThis(this))
+					]
+
+				]
+				+ SSplitter::Slot()
+				.Value(0.72f)
+				[
+					SAssignNew(PreviewBorder, SBorder)
+					.BorderImage(FAppStyle::GetBrush("ToolPanel.GroupBorder"))
+					.Padding(FMargin(4.0f, 4.0f))					
+				]
+			]
+		]
+	];
+	
+	// Set the tree expanded by default
+	TreeExpandInstance();
+} 
+
+
+EVisibility SMutableCodeViewer::IsMipSkipVisible() const 
+{ 
+	return bSelectedOperationIsImage ? EVisibility::Visible : EVisibility::Hidden;
+}
+
+
+TOptional<int32> SMutableCodeViewer::GetCurrentMipSkip() const 
+{ 
+	return MipsToSkip; 
+}
+
+
+void SMutableCodeViewer::OnCurrentMipSkipChanged(int32 NewValue)
+{
+	MipsToSkip = NewValue;
+	bIsPreviewPendingUpdate = true;
+}
+
+
+#pragma region CodeTree operation search
+
+void SMutableCodeViewer::GenerateNavigationOpTypeStrings()
+{
+	// Grab only the names from the operation types located during the caching of operation types of the model
+	for (const mu::OP_TYPE LocatedOperationType : ModelOperationTypes)
+	{
+		// Find the name
+		const uint16 OperationIndex = static_cast<uint16>(LocatedOperationType);
+		const TCHAR* OpName = mu::s_opNames[OperationIndex];
+		
+		// Operation name to be saved
+		TSharedPtr<FString> OperationNamePointer = MakeShared<FString>(OpName);
+		OperationTypesStrings.Add( OperationNamePointer);
+	}
+}
+
+void SMutableCodeViewer::OnSelectedOperationTypeFromTree()
+{
+	// We require to have only 1 element selected to avoid having inconsistencies during operation
+	check(TreeView->GetNumItemsSelected() == 1);
+	
+	const TSharedPtr<FMutableCodeTreeElement> ReferenceOperationElement = TreeView->GetSelectedItems()[0];
+	
+	const mu::OP_TYPE OperationType =
+		MutableModel->GetPrivate()->m_program.GetOpType(ReferenceOperationElement->MutableOperation);
+	const int32 OperationIndex = ModelOperationTypes.IndexOfByKey(OperationType);
+
+	// Failing the next check would mean that we are not caching all the types present on the current operation's tree
+	check (OperationIndex != -1);
+	
+	const TSharedPtr<FString> NewOperationString = OperationTypesStrings[OperationIndex];
+
+	// Set the type operation type to be looking for -> Will invoke OnOptionTypeSelectionChanged
+	TargetedTypeSelector->SetSelectedItem(NewOperationString);
+}
+
+void SMutableCodeViewer::OnOptionTypeSelectionChanged(TSharedPtr<FString, ESPMode::ThreadSafe> NesSelectedOperationString,
+	ESelectInfo::Type SelectionType)
+{
+	// Cache the currently selected operation set on the UI by the user
+	const int32 NewOperationIndex = OperationTypesStrings.IndexOfByKey(NesSelectedOperationString);
+	const mu::OP_TYPE NewOperationType = ModelOperationTypes[NewOperationIndex];
+	if (NewOperationType == OperationTypeToSearch)
+	{
+		return;
+	}
+	OperationTypeToSearch = NewOperationType;
+	
+	// Only do the internal work if the type is one that makes sense searching
+	if (OperationTypeToSearch != mu::OP_TYPE::NONE)
+	{
+		// Reset the navigation element
+		SetCurrentNavigationElement(nullptr);
+		
+		// Locate all operations on the mutable operations tree (not the visual one) that do share the same operation type
+		// as the one selected. This will fill the array with the elements we should be looking for during the navigation operation
+		CacheAddressesOfOperationsOfType(OperationTypeToSearch);
+		
+		// Experimental
+		TreeView->ScrollToTop();
+	}
+	// None can be set by the user or be an indication that we are navigating over constant related operations
+	// todo: Separate both operations in some way on the UI to avoid complications in the code and in the UI's UX
+	else
+	{
+		// Clear all the elements on the navigation addresses 
+		NavigationOPAddresses.Empty();
+	}
+	
+}
+
+void SMutableCodeViewer::CacheAddressesOfOperationsOfType(mu::OP_TYPE TargetedOperationType)
+{
+	// Clear previous data
+	NavigationOPAddresses.Empty();
+	
+	// Locate all addresses of operations of the provided operation type
+	TSet<mu::OP::ADDRESS> ProcessedAddresses;
+	TSet<mu::OP::ADDRESS> OperationsWithTargetedType;
+
+	// Main update procedure run for the targeted state and the targeted parameter values
+	const mu::PROGRAM& Program = MutableModel->GetPrivate()->m_program;
+	for	(const mu::OP::ADDRESS& RootNodeAddress : RootNodeAddresses)
+	{
+		GetOperationsOfType(TargetedOperationType,RootNodeAddress,Program,OperationsWithTargetedType,ProcessedAddresses);
+	}
+	
+	if (!OperationsWithTargetedType.IsEmpty())
+	{
+		// Cache the navigation addresses so we are able to navigate over them
+		NavigationOPAddresses = OperationsWithTargetedType.Array();
+
+		// Tell the system changes have been made over Navigation Addresses
+		bUpdatedNavigationAddresses = true;
+	}
+
+}
+
+void SMutableCodeViewer::GetOperationsOfType ( const mu::OP_TYPE& TargetOperationType, const mu::OP::ADDRESS& InParentAddress,const mu::PROGRAM& InProgram,
+	TSet<mu::OP::ADDRESS>& OutAddressesOfType,
+	TSet<mu::OP::ADDRESS>& AlreadyProcessedAddresses)
+{
+	// Generic case for unnamed children traversal.
+	mu::ForEachReference(InProgram, InParentAddress, [this, &InProgram, &OutAddressesOfType, &AlreadyProcessedAddresses, &TargetOperationType]( mu::OP::ADDRESS ChildAddress)
+	{
+		// If the parent does have a child then process it 
+		if (ChildAddress && !AlreadyProcessedAddresses.Contains(ChildAddress))
+		{
+			// Cache if same data type and we share the same address (means this op is pointing at the provided resource)
+			if (InProgram.GetOpType(ChildAddress) == TargetOperationType)
+			{
+				OutAddressesOfType.Add(ChildAddress);
+			}
+		
+			// Cache to avoid processing it again later
+			AlreadyProcessedAddresses.Add(ChildAddress);
+		
+			// Process the children of this object
+			GetOperationsOfType(TargetOperationType,ChildAddress,InProgram,OutAddressesOfType,AlreadyProcessedAddresses);
+		}
+	});
+}
+
+
+
+void SMutableCodeViewer::CacheOperationTypesPresentOnModel()
+{
+	check(MutableModel)
+	check(RootNodeAddresses.Num());
+
+	// Set of operation types to be filed 
+	TSet<mu::OP_TYPE> OperationTypes;
+	
+	const mu::PROGRAM& Program = MutableModel->GetPrivate()->m_program;
+	TSet<mu::OP::ADDRESS> AlreadyVisitedAddresses;
+	for	(const mu::OP::ADDRESS& RootOperationAddress : RootNodeAddresses)
+	{
+		GetOperationTypesPresentOnModel(RootOperationAddress,Program,OperationTypes,AlreadyVisitedAddresses);
+	}
+
+	// After using the set object save it as an array for later reference
+	ModelOperationTypes = OperationTypes.Array();
+	
+	// Sort the contents of the array alphabetically
+	ModelOperationTypes.StableSort([&](const mu::OP_TYPE& A, const mu::OP_TYPE& B)
+	{
+		// Find the name
+		FString AString;
+		{
+			const uint16 OperationIndex = static_cast<uint16>(A);
+			const TCHAR* OpName = mu::s_opNames[OperationIndex];
+			AString = FString(OpName);
+		}
+		
+		// Find out the name of the first element
+		FString BString;
+		{
+			const uint16 OperationIndex = static_cast<uint16>(B);
+			const TCHAR* OpName = mu::s_opNames[OperationIndex];
+			BString = FString(OpName);
+		}
+		
+		// Then the name of the second element
+		return AString < BString;
+	});	
+
+	// Add the "None" option as a default value at index 0 to the sorted array.
+	// the NONE option serves as a way of telling the user "hey, you are not navigating over a type of operation".
+	// We currently use it to also tell the user that if he is navigating it is over constant related operations
+	ModelOperationTypes.Insert(mu::OP_TYPE::NONE,0);
+
+	// ModelOperationTypes is now an array starting with None and all the types found on the operations tree in alphabetical order
+}
+
+
+void SMutableCodeViewer::GetOperationTypesPresentOnModel(
+	const mu::OP::ADDRESS& InParentAddress, const mu::PROGRAM& InProgram, TSet<mu::OP_TYPE>& OutLocatedOperations,
+	TSet<mu::OP::ADDRESS>& AlreadyProcessedAddresses)
+{
+	// Generic case for unnamed children traversal.
+	mu::ForEachReference(InProgram, InParentAddress, [this, &InProgram, &OutLocatedOperations, &AlreadyProcessedAddresses]( mu::OP::ADDRESS ChildAddress)
+	{
+		// Avoid processing items more than once
+		if (ChildAddress && !AlreadyProcessedAddresses.Contains(ChildAddress))
+		{
+			// Cache the operation type for later usage
+			OutLocatedOperations.Add(InProgram.GetOpType(ChildAddress));
+			
+			// Cache to avoid processing it again later
+			AlreadyProcessedAddresses.Add(ChildAddress);
+		
+			// Process the children of this object
+			GetOperationTypesPresentOnModel(ChildAddress,InProgram,OutLocatedOperations,AlreadyProcessedAddresses);
+		}
+	});
+}
+
+
+void SMutableCodeViewer::LocateNavigationElements()
+{
+	NavigationFoundElements.Empty();
+	NavigationFoundElements.Reserve(TreeElements.Num());
+	
+	for (TSharedPtr<FMutableCodeTreeElement> TreeElement : TreeElements)
+	{
+		// Consider that you have found elements if their address corresponds to the address on our cached navigation items
+		const bool bIsNavigableElement = NavigationOPAddresses.Contains(TreeElement->MutableOperation);
+		if (bIsNavigableElement && !TreeElement->DuplicatedOf)
+		{
+			// Add the element to our list of elements we want to navigate over
+			NavigationFoundElements.Add(TreeElement);
+		}
+	}
+	
+	NavigationFoundElements.Shrink();
+
+#if UE_BUILD_DEBUG
+	UE_LOG(LogTemp, Log, TEXT("Located %d elements of %d for the current operation type"), NavigationFoundElements.Num(),TreeElements.Num());
+#endif
+
+	if (NavigationFoundElements.IsEmpty())
+	{
+		// No elements have been found for the newly selected type, set the selected element to be null.
+		SetCurrentNavigationElement(nullptr);
+	}
+}
+
+FText SMutableCodeViewer::OnPrintNavigableObjectAddressesCount() const
+{
+	// Depending on the amount of navigable objects (addresses, not actual elements) display the amount there are
+	return FText::Format(LOCTEXT("AmountOfNavigableOps"," Found Operations : {0} "), NavigationOPAddresses.Num());
+}
+
+void SMutableCodeViewer::SetCurrentNavigationElement(TSharedPtr<FMutableCodeTreeElement> NewelyFoundElement)
+{
+	if (CurrentNavigationElement && NewelyFoundElement)
+	{
+		// Avoid operating if the new selected element is the same than before
+		if (CurrentNavigationElement->MutableOperation == NewelyFoundElement->MutableOperation)
+		{
+#if UE_BUILD_DEBUG
+			UE_LOG(LogTemp,Warning,TEXT("Skipping CurrentSelectedElement update since it is the same as previously set."));
+#endif
+			return;
+		}
+	}
+
+	// Internally change the selected element for the navigation system to use as starting point for the next movement
+	// operation
+	CurrentNavigationElement = NewelyFoundElement;
+	
+	// Clear the selection and select the new object
+	// TreeView->ClearSelection();
+	if (CurrentNavigationElement)
+	{
+		// Perform the UI selection operation so the viewer gets triggered as if the user did select it manually
+		TreeView->SetSelection(CurrentNavigationElement);
+
+#if UE_BUILD_DEBUG 
+		// Report currently selected one (not actually selected in UI, only internally at the moment)
+		const int32 SelectedOperationIndex = NavigationFoundElements.IndexOfByKey(CurrentNavigationElement);
+		UE_LOG(LogTemp,Log,TEXT("\tSelected element %d of %d"),SelectedOperationIndex +1,NavigationFoundElements.Num());
+#endif
+
+		// TODO: Mind making the system automatically focus on the selected object. Currently is not always on the view
+		// and that could be improved. Scrolling over the tree may require a tree refresh, this means you may need to
+		// wait until that refresh is performed.
+		
+	}
+#if UE_BUILD_DEBUG 
+	else
+	{
+		UE_LOG(LogTemp,Log,TEXT("\tNo selected element has been found"));
+	}
+#endif
+	
+}
+
+
+bool SMutableCodeViewer::CanInteractWithPreviousOperationButton() const
+{
+	// Only navigable if there are more than 0 elements to traverse and we are not scrolling
+	return !bIsScrolling && NavigationOPAddresses.Num() > 0;
+}
+
+bool SMutableCodeViewer::CanInteractWithNextOperationButton() const
+{
+	// Only navigable if there are more than 0 elements to traverse and we are not scrolling
+	return !bIsScrolling && NavigationOPAddresses.Num() > 0;
+}
+
+
+void SMutableCodeViewer::SortNavigationElementsArray()
+{
+	// Do not waste time sorting an empty array of elements
+	// If only one element is present there is no reason for performing a sort operation
+	if (NavigationFoundElements.Num() < 2)
+	{
+		return;
+	}
+
+	// Cache the position on the tree for each of the navigation elements on view
+	TMap<int32,TSharedPtr<FMutableCodeTreeElement>> PositionPerElement;
+	PositionPerElement.Reserve(NavigationFoundElements.Num());
+	for	(int32 i = 0; i < NavigationFoundElements.Num(); i++)
+	{
+		const int32 PositionOnTree = TreeView->WidgetFromItem(NavigationFoundElements[i])->GetIndexInList();
+		PositionPerElement.Add(PositionOnTree, NavigationFoundElements[i]);
+	}
+
+	// By working with a TMap we can ensure the indices provided get automatically sorted on Addition.
+	
+	// Grab the values as an array to be later be used as the new array for NavigationElements
+	TArray<TSharedPtr< FMutableCodeTreeElement>> SortedNavigationElements;
+	PositionPerElement.GenerateValueArray(SortedNavigationElements);
+
+	// Make the targeted elements of type the new array with the correct values but sorted out
+	NavigationFoundElements = MoveTemp(SortedNavigationElements);
+}
+
+
+FReply SMutableCodeViewer::OnGoToPreviousOperationButtonPressed()
+{
+	// Filter out not valid type
+	if (NavigationOPAddresses.IsEmpty())
+	{
+		return FReply::Handled();
+	}
+
+	// Sort the array of navigation elements just in case the order of the elements on the view changed
+	// This can happen if the user scrolls the view and new elements of the targeted type come to view
+	SortNavigationElementsArray();
+	
+	// Make a copy of the array to later be able to compare it to look for new elements of the targeted type
+	NavigationPreviouslyFoundElements = NavigationFoundElements;
+
+	// Fully expand tree if required to uncover all possible elements of the targeted type.
+	// The expansion does not discern between element types, all non duplicated elements will be expanded
+	if (bTreeWasExpanded == false && !bIsScrolling)
+	{
+		// Expand all the elements by calling the standard expansion operation (ignoring duplicates of course)
+		TreeExpandElements(RootNodes,false);
+
+		// Notify the navigation system that we just expanded the tree
+		bTreeWasExpanded = true;
+
+		// It will trigger an automatic scroll operation if possible depending on the tree
+		bIsSearchingForPreviousElement = true;
+			
+		return FReply::Handled();
+	}
+
+	// At this point we know the tree has been fully expanded and we should be able to traverse it without the
+	// need of expanding it again unless we get to another object of the targeted type with should, if used as
+	// start for another search, perform another cautionary tree expansion.
+	
+	// Get the index of the currently selected element on the array of operations of the current type
+	// If it is < 0 means that the current object is not part of the actual elements of the targeted type
+	int32 SelectedOperationIndex = NavigationFoundElements.IndexOfByKey(CurrentNavigationElement);
+
+	if ( CurrentNavigationElement )
+	{
+		// A current element is selected and is not the first one. Select the previous element
+		if ( SelectedOperationIndex > 0)
+		{
+			// Get the previous element on the array (the one on top of the current one on the tree structure)
+			SetCurrentNavigationElement(NavigationFoundElements[SelectedOperationIndex - 1]);
+			return FReply::Handled();
+		}
+	}
+	// No current element is currently selected
+	else
+	{
+		// We have no current element selected but there are elements there to be used. Select the last one
+		if ( !NavigationFoundElements.IsEmpty() )
+		{
+			// Get the last object and use it as the current object
+			SelectedOperationIndex = NavigationFoundElements.Num() -1;
+			SetCurrentNavigationElement(NavigationFoundElements[SelectedOperationIndex]);
+			return FReply::Handled();
+		}
+	}
+
+	// If no other element of the type we want is found it means we should look for other elements by
+	// scrolling the view up to the top or until we find another element
+	{
+		// It will trigger an automatic scroll operation if possible depending on the tree.
+		bIsSearchingForPreviousElement = true;
+	}
+
+	return FReply::Handled();
+}
+
+void SMutableCodeViewer::GoToPreviousOperationAfterRefresh()
+{
+	// Is mandatory for this system to be stable to operate after the tree refresh
+	check(TreeView->IsPendingRefresh() == false);
+	
+	// Required for all the possible contexts where this method could  partake
+	check(bIsSearchingForPreviousElement);
+	
+	// Check that we did change the amount of elements of our current targeted type.
+	// If the tree was fully expanded on the previous update then check the changes by entering the if block
+	const bool bNewElementFound = HaveNewElementsBeenFound();
+	if ( bNewElementFound || bTreeWasExpanded)
+	{
+		// Reset the flag since we are already being run after the tree refresh
+		bIsSearchingForPreviousElement = false;
+		
+		// Continue the operation now that the tree view has been refreshes and the new rows are there to be used
+		OnGoToPreviousOperationButtonPressed();
+		
+		// Reset the just expanded flag if found to be true
+		bTreeWasExpanded = false;
+		
+		// Reset the scrolling flag in case it was active
+		bIsScrolling = false;
+		return;
+	}
+
+	// We do know that no element of the targeted type has been found. So, we should start searching
+	
+	// If scrollable try to scroll over the tree to locate new elements
+	if (TreeView->IsScrollbarNeeded() &&
+		TreeView->GetScrollOffset() > 0 )
+	{
+		// Flag the system so it is able to tell if the scrolling is being performed or not
+		bIsScrolling = true;
+
+		// Add the max possible offset to be fast but secure on the translation process over the graph
+		ComputedMaxViewScrollStep = FMath::Max( TreeElements.Num() - 1,MinTreeScrollStep);
+		TreeView->AddScrollOffset(-ComputedMaxViewScrollStep,true);
+
+		// Check to avoid children drawing crash when having negative offset values
+		if (TreeView->GetScrollOffset() < 0)
+		{
+			TreeView->ScrollToTop();
+		}
+			
+		return;
+	}
+
+	// If not scrollable then just ignore this operation since it means no new elements have been added and
+	// no targeted elements can be retrieved from the graph
+	bIsSearchingForPreviousElement = false;
+	bIsScrolling = false;
+}
+
+
+FReply SMutableCodeViewer::OnGoToNextOperationButtonPressed()
+{
+	// Go forward one index on the array TargetedElementsOfType.
+	// If in last index try to expand other rows
+	
+	// Filter out not valid type
+	if (NavigationOPAddresses.IsEmpty())
+	{
+		return FReply::Handled();
+	}
+
+	// Sort the array of navigation elements just in case the order of the elements on the view changed
+	// This can happen if the user scrolls the view and new elements of the targeted type come to view
+	SortNavigationElementsArray();
+
+	// Make a copy of the array to be able to later compare it to the tree with the element expanded
+	NavigationPreviouslyFoundElements = NavigationFoundElements;
+
+	// Fully expand tree if required to uncover all possible elements of the targeted type.
+	// The expansion does not discern between element types, all non duplicated elements will be expanded
+	if (bTreeWasExpanded == false && !bIsScrolling)
+	{
+		// Expand all elements to avoid leaving any of them behind when searching for the next element
+		TreeExpandElements(RootNodes,false);
+
+		// Tell the system to continue operating after the tree refresh
+		bIsSearchingForNextElement = true;
+
+		// Save a flag to avoid redoing this again later (only required once per selected object)
+		bTreeWasExpanded = true;
+			
+		return FReply::Handled();
+	}
+
+	// At this point we know the tree has been fully expanded and we should be able to traverse it without the
+	// need of expanding it again unless we get to another object of the targeted type with should, if used as
+	// start for another search, perform another cautionary tree expansion.
+	
+	// Get the index of the currently selected element on the array of operations of the current type
+	// If it is < 0 means that the current object is not part of the actual elements of the targeted type
+	int32 SelectedOperationIndex = NavigationFoundElements.IndexOfByKey(CurrentNavigationElement);
+	
+	// If we have an element already set then use it to compute the values
+	if (CurrentNavigationElement)
+	{
+		// If we are not the last element of the array then move on to the next one normally
+		// Since we have expanded the current element we should be able to access contents inside of it (children of it)
+		if (SelectedOperationIndex < NavigationFoundElements.Num() - 1)
+		{
+			// Select the new element
+			SetCurrentNavigationElement(NavigationFoundElements[SelectedOperationIndex + 1]);
+			return FReply::Handled();
+		}
+	}
+	// We do not have a current element to use as starting point for our search.
+	else
+	{
+		// Since we are moving downwards, grab the first element visible to make it be the current object.
+		if (NavigationFoundElements.Num() > 0)
+		{
+			// Get the first object and store its direction
+			SelectedOperationIndex = 0;
+			SetCurrentNavigationElement(NavigationFoundElements[SelectedOperationIndex]);
+			return FReply::Handled();
+		}
+	}
+
+	// At this point we are not able to reach any other element of the current type. We must then
+	// start scrolling down in order to get to the next element or the end of the tree view.
+	{
+		// Enable this flag to make the system search for new elements on the tree on the next tick
+		bIsSearchingForNextElement = true;
+	}
+	
+	return FReply::Handled();
+}
+
+
+void SMutableCodeViewer::GoToNextOperationAfterRefresh()
+{
+	// Is mandatory for this system to be stable to operate after the tree refresh
+	check(TreeView->IsPendingRefresh() == false);
+	
+	// Required for all the possible contexts where this method could  partake
+	check(bIsSearchingForNextElement);
+	
+	// Check if the new CurrentArray does have elements not already set on the previous array
+	// and also check if we expanded the tree on the previous tick. If so, enter the if block
+	const bool bNewElementFound = HaveNewElementsBeenFound();
+	if (bNewElementFound || bTreeWasExpanded)
+	{
+		// Reset the flag since we are already being run after the tree refresh
+		bIsSearchingForNextElement = false;
+		
+		// Recreate the button press to move to the next element after expanding the tree
+		OnGoToNextOperationButtonPressed();
+		
+		// Reset the just expanded flag
+		// Setting this to false makes the system able to fully expand the tree later when we start searching again
+		bTreeWasExpanded = false;
+		
+		// Reset the scrolling flag in case it was active
+		bIsScrolling = false;
+		return;
+	}
+
+	// We do know that no element of the targeted type has been found. So, we should start searching
+	
+	// If scrollable try to scroll over the tree to locate new elements
+	if (TreeView->IsScrollbarNeeded()
+		&& TreeView->GetScrollDistanceRemaining().Y > 0)
+	{
+		// Flag the system so it is able to tell if the scrolling is being performed or not
+		bIsScrolling = true;
+			
+		// Add a bit of scroll
+		ComputedMaxViewScrollStep = FMath::Max( TreeElements.Num() - 1,MinTreeScrollStep);
+		TreeView->AddScrollOffset(ComputedMaxViewScrollStep,true);
+
+		// Do not trust the tree to be able to handle this kind of artificial scrolling and avoid surpassing the end
+		if (TreeView->GetScrollDistanceRemaining().Y < 0 )
+		{
+			TreeView->ScrollToBottom();
+		}
+			
+		return;
+	}
+
+	// If not scrollable then just ignore this operation since it means no new elements have been added and
+	// no targeted elements can be retrieved from the graph
+	bIsSearchingForNextElement = false;
+	bIsScrolling = false;
+}
+
+bool SMutableCodeViewer::HaveNewElementsBeenFound() const
+{
+	// Check if the new CurrentArray does have elements not already set on the previous array
+	bool NewElementAdded = false;
+
+	// If the current array does have data and the old not then we directly know that new elements have been added.
+	if (NavigationFoundElements.Num() && NavigationPreviouslyFoundElements.IsEmpty())
+	{
+		return true;
+	}
+	
+	// Iterate the current array to 
+	for (int32 CurrentArrayIndex = 0; CurrentArrayIndex < NavigationFoundElements.Num(); CurrentArrayIndex++)
+	{
+		// If the element is not found on the previous array it means it is new, and therefore a valid element to select
+		if (!NavigationPreviouslyFoundElements.Contains(NavigationFoundElements[CurrentArrayIndex]))
+		{
+			NewElementAdded = true;
+			break; 
+		}
+	}
+
+	return NewElementAdded;
+}
+
+
+#pragma endregion 
+
+#pragma region CodeTree Callbacks
+
+
+TSharedRef<ITableRow> SMutableCodeViewer::GenerateRowForNodeTree(TSharedPtr<FMutableCodeTreeElement> InTreeNode, const TSharedRef<STableViewBase>& InOwnerTable)
+{
+	// Save the node for later access
+	TreeElements.Add(InTreeNode);
+	
+	// Generate a row element
+	TSharedRef<SMutableCodeTreeRow> Row = SNew(SMutableCodeTreeRow, InOwnerTable, InTreeNode);
+	
+	// Only add it to our collection of navigation elements if we know it has an address we want to be able to navigate to
+	if (NavigationOPAddresses.Contains(InTreeNode->MutableOperation))
+	{
+		// Ignore duplicates
+		if (!InTreeNode->DuplicatedOf)
+		{
+			NavigationFoundElements.Add(InTreeNode);
+		}
+	}
+	
+	// Determine if a row should be painted as highlighted based on the selected item
+	if (TreeView->GetNumItemsSelected())
+	{
+		const TSharedPtr<FMutableCodeTreeElement> SelectedElement = TreeView->GetSelectedItems()[0];
+		if (InTreeNode->MutableOperation == SelectedElement->MutableOperation)
+		{
+			Row->Highlight();
+		}
+	}
+	
+	return Row;
+}
+
+
+void SMutableCodeViewer::GetChildrenForInfo(TSharedPtr<FMutableCodeTreeElement> InInfo, TArray<TSharedPtr<FMutableCodeTreeElement>>& OutChildren)
+{
+	if (!InInfo->MutableModel)
+	{
+		return;
+	}
+
+	// If this is a duplicated of another row, don't provide its children.
+	//if (InInfo->DuplicatedOf)
+	//{
+	//	return;
+	//}
+
+	const mu::PROGRAM& Program = MutableModel->GetPrivate()->m_program;
+
+	mu::OP::ADDRESS ParentAddress = InInfo->MutableOperation;
+
+	// Generic case for unnamed children traversal.
+	uint32 ChildIndex = 0;
+	mu::ForEachReference(Program, InInfo->MutableOperation, [this, ParentAddress, &ChildIndex, &OutChildren](mu::OP::ADDRESS ChildAddress)
+	{
+		if (ChildAddress)
+		{
+			const FItemCacheKey Key = { ParentAddress, ChildAddress, ChildIndex };
+			const TSharedPtr<FMutableCodeTreeElement>* CachedItem = ItemCache.Find(Key);
+
+			if (CachedItem)
+			{
+				OutChildren.Add(*CachedItem);
+			}
+			else
+			{
+				const TSharedPtr<FMutableCodeTreeElement>* MainItemPtr = MainItemPerOp.Find(ChildAddress);
+				const TSharedPtr<FMutableCodeTreeElement> Item = MakeShareable(new FMutableCodeTreeElement(MutableModel, ChildAddress, TEXT(""), MainItemPtr));
+
+				OutChildren.Add(Item);
+				ItemCache.Add(Key, Item);
+
+				if (!MainItemPtr)
+				{
+					MainItemPerOp.Add(ChildAddress, Item);
+				}
+			}
+		}
+		++ChildIndex;
+	});
+}
+
+void SMutableCodeViewer::OnExpansionChanged(TSharedPtr<FMutableCodeTreeElement> InItem, bool bInExpanded)
+{
+	// If an element gets expanded then contract (if found) the other element that uses the same address
+	if (bInExpanded)
+	{
+		const mu::OP::ADDRESS MutableOperation = InItem->MutableOperation;
+		const TSharedPtr<FMutableCodeTreeElement>* ExpandedElement = ExpandedElements.Find(MutableOperation);
+		if (ExpandedElement)
+		{
+			TreeView->SetItemExpansion(*ExpandedElement, false);
+		}
+		
+		ExpandedElements.Add(MutableOperation, InItem);
+	}
+	else
+	{
+		ExpandedElements.Remove(InItem->MutableOperation);
+	}
+}
+
+
+void SMutableCodeViewer::OnSelectionChanged(TSharedPtr<FMutableCodeTreeElement> InNode, ESelectInfo::Type InSelectInfo)
+{
+	if (bIsElementHighlighted)
+	{
+		ClearHighlightedItems();
+	}
+
+	TArray<TSharedPtr<FMutableCodeTreeElement>> SelectedNodes;
+	TreeView->GetSelectedItems(SelectedNodes);
+
+	PreviewBorder->ClearContent();
+	ParametersWidget->SetParameters(nullptr);
+
+	SelectedOperationAddress = 0;
+	bSelectedOperationIsImage = false;
+
+	if (SelectedNodes.IsEmpty())
+	{
+		return;
+	}
+
+	// Find the duplicates for the selected tree element element and highlight them
+	if (InNode)
+	{
+		HighlightDuplicatesOfEntry(InNode);
+	}
+
+	bIsPreviewPendingUpdate = true;
+
+	SelectedOperationAddress = SelectedNodes[0]->MutableOperation;
+	const mu::OP_TYPE OperationType = MutableModel->GetPrivate()->m_program.GetOpType(SelectedOperationAddress);
+	const mu::DATATYPE OperationDataType = mu::GetOpDataType(OperationType);
+
+	ParametersWidget->SetParameters(PreviewParameters);
+	
+	switch (OperationDataType)
+	{
+	case mu::DT_LAYOUT:
+	{
+		// Create or reuse the UI
+		PrepareLayoutViewer();
+		break;
+	}
+
+	case mu::DT_IMAGE:
+	{
+		// Create or reuse the UI
+		bSelectedOperationIsImage = true;
+		PrepareImageViewer();
+		break;
+	}
+
+	case mu::DT_MESH:
+	{
+		// Create or reuse the UI
+		PrepareMeshViewer();
+		break;
+	}
+
+	case mu::DT_SCALAR:
+	{
+		// Create or reuse the UI
+		if (!PreviewScalarViewer)
+		{
+			PreviewScalarViewer = SNew(SMutableScalarViewer);
+		}
+
+		PreviewBorder->SetContent(PreviewScalarViewer.ToSharedRef());
+		break;
+	}
+
+	case mu::DT_STRING:
+	{
+		// Create or reuse the UI
+		PrepareStringViewer();
+		break;
+	}
+
+	case mu::DT_COLOUR:
+	{
+		// Create or reuse the UI
+		if (!PreviewColorViewer)
+		{
+			PreviewColorViewer = SNew(SMutableColorViewer);
+		}
+
+		PreviewBorder->SetContent(PreviewColorViewer.ToSharedRef());
+		break;
+	}
+
+	case mu::DT_INT:
+	{
+		// Create or reuse the UI
+		if (!PreviewIntViewer)
+		{
+			PreviewIntViewer = SNew(SMutableIntViewer);
+		}
+
+		PreviewBorder->SetContent(PreviewIntViewer.ToSharedRef());
+		break;
+	}
+
+	case mu::DT_BOOL:
+	{
+		// Create or reuse the UI
+		if (!PreviewBoolViewer)
+		{
+			PreviewBoolViewer = SNew(SMutableBoolViewer);
+		}
+
+		PreviewBorder->SetContent(PreviewBoolViewer.ToSharedRef());
+		break;
+	}
+
+	case mu::DT_PROJECTOR:
+	{
+		// Create or reuse the UI
+		PrepareProjectorViewer();
+		break;
+	}
+
+	
+	default:
+		// There is no viewer for this type yet.
+		break;
+
+	}
+}
+
+TSharedPtr<SWidget> SMutableCodeViewer::OnTreeContextMenuOpening()
+{
+	FMenuBuilder MenuBuilder(true, nullptr);
+	
+	if (TreeView->GetSelectedItems().Num())
+	{
+		if (TreeView->GetSelectedItems().Num() == 1)
+		{
+			MenuBuilder.AddMenuEntry(
+		LOCTEXT("Set_as_search_operation_type","Set as search Operation"),
+		LOCTEXT("Set_as_search_operation_type_tip", "Sets the type of this operation as the type to be looking for when searching for operations on the tree view"),
+				FSlateIcon(),
+				FUIAction(FExecuteAction::CreateSP(this, &SMutableCodeViewer::OnSelectedOperationTypeFromTree)
+				)	
+			);
+		}
+
+		MenuBuilder.AddMenuSeparator();
+		
+		MenuBuilder.AddMenuEntry(
+			LOCTEXT("Expand_Selected", "Expand Selected Operation"),
+			LOCTEXT("Expand_Selected", "Expands only the selected Operation and leaves the other as they are."),
+			FSlateIcon(),
+			FUIAction(FExecuteAction::CreateSP(this, &SMutableCodeViewer::TreeExpandSelected)
+			)
+		);
+	}
+
+
+	MenuBuilder.AddMenuEntry(
+		LOCTEXT("Expand_Instance", "Expand Instance-Level Operations"),
+		LOCTEXT("Expand_Instance", "Expands all the operations in the tree that are instance operations (not images, meshes, booleans, etc.)."),
+		FSlateIcon(),
+		FUIAction(FExecuteAction::CreateSP(this, &SMutableCodeViewer::TreeExpandInstance)
+			//, FCanExecuteAction::CreateSP(this, &SMutableCodeViewer::HasAnyItemInPalette)
+		)
+	);
+
+	MenuBuilder.AddMenuEntry(
+		LOCTEXT("Expand_Unique", "Expand All Unique Operations"),
+		LOCTEXT("Expand_Unique", "Expands all the operations in the tree that have not been expanded yet."),
+		FSlateIcon(),
+		FUIAction(FExecuteAction::CreateSP(this, &SMutableCodeViewer::TreeExpandUnique)
+		)
+	);
+
+	return MenuBuilder.MakeWidget();
+}
+
+void SMutableCodeViewer::TreeExpandRecursive(TSharedPtr<FMutableCodeTreeElement> InInfo, bool bExpand)
+{
+	if (bExpand)
+	{
+		TreeExpandUnique();
+	}
+}
+
+void SMutableCodeViewer::OnRowReleased(const TSharedRef<ITableRow>& InTreeRow)
+{
+	SMutableCodeTreeRow* CastedTableRow = static_cast<SMutableCodeTreeRow*>(&InTreeRow.Get());
+	const TSharedPtr<FMutableCodeTreeElement>& RowElement = CastedTableRow->GetItem();
+	TreeElements.Remove(RowElement);
+	
+	// Remove the element from the array of elements of the current type
+	if (NavigationFoundElements.Contains(RowElement))
+	{
+		NavigationFoundElements.Remove(RowElement);
+		if (CurrentNavigationElement == RowElement)
+		{
+			SetCurrentNavigationElement(nullptr);
+		}
+	}
+}
+#pragma endregion
+
+#pragma region Highlight Methods
+void SMutableCodeViewer::HighlightDuplicatesOfEntry(const TSharedPtr<FMutableCodeTreeElement>& InTargetEntry)
+{
+	if (bIsElementHighlighted)
+	{
+		ClearHighlightedItems();
+	}
+
+	// Highlight the elements related to the currently selected item of the tree
+	HighlightedOperation = InTargetEntry->MutableOperation;
+
+	for (const TSharedPtr<FMutableCodeTreeElement>& TreeItem : TreeElements)
+	{
+		if (TreeItem.Get() != InTargetEntry.Get() && TreeItem->MutableOperation == HighlightedOperation)
+		{
+			TSharedPtr<ITableRow> TableRow = TreeView->WidgetFromItem(TreeItem);
+			const SMutableCodeTreeRow* MutableRow = static_cast<SMutableCodeTreeRow*>(TableRow.Get());
+			MutableRow->Highlight();
+		}
+	}
+
+	bIsElementHighlighted = true;
+}
+
+void SMutableCodeViewer::ClearHighlightedItems()
+{
+	// Clear the previously highlighted elements
+	for (const TSharedPtr<FMutableCodeTreeElement>& HighlightedElement : TreeElements)
+	{
+		if (HighlightedElement->MutableOperation == HighlightedOperation)
+		{
+			TSharedPtr<ITableRow> TableRow = TreeView->WidgetFromItem(HighlightedElement);
+
+			if (TableRow.IsValid())
+			{
+				const SMutableCodeTreeRow* MutableRow = static_cast<SMutableCodeTreeRow*>(TableRow.Get());
+				MutableRow->ResetHighlight();
+			}
+
+		}
+	}
+
+	bIsElementHighlighted = false;
+}
+#pragma endregion
+
+#pragma region Element Expansion Llogic
+
+
+void SMutableCodeViewer::TreeExpandElements(TArray<TSharedPtr<FMutableCodeTreeElement>>& InElementsToExpand, 
+	bool bForceExpandDuplicates /*= false*/,
+	mu::DATATYPE FilteringDataType /*= mu::DATATYPE::DT_NONE*/,
+	TSharedPtr<FProcessedOperationsBuffer> InExpandedOperationsBuffer /* = nullptr */)
+{
+	if (InElementsToExpand.IsEmpty())
+	{
+		return;
+	}
+	
+	// Initialization of recursive elements if this is the first invocation of method
+	{
+		if (!InExpandedOperationsBuffer)
+		{
+			InExpandedOperationsBuffer = MakeShared<FProcessedOperationsBuffer>();
+		}
+	}
+
+	// Load references to the arrays containing all the operations already worked on during another recursive call to this
+	// method
+	TArray<mu::OP::ADDRESS>& AlreadyExpandedOriginalOperations = InExpandedOperationsBuffer->ExpandedOriginalOperations;
+	TArray<mu::OP::ADDRESS>& AlreadyExpandedDuplicatedOperations = InExpandedOperationsBuffer->ExpandedDuplicatedOperations;
+	
+	// Array containing the children object found on Item.
+	TArray<TSharedPtr<FMutableCodeTreeElement>> Children;
+	
+	// Index of the current element being processed
+	int32 CurrentElementIndex = 0;
+	while (CurrentElementIndex < InElementsToExpand.Num() )
+	{
+		// Grab the current element to process and move the index forward once
+		TSharedPtr<FMutableCodeTreeElement> Item = InElementsToExpand[CurrentElementIndex++];
+		check(Item);
+
+		// Identifier of the element. May be repeated if there are elements duplicating another element
+		const mu::OP::ADDRESS OperationAddress = Item->MutableOperation;
+		
+		// Filter the elements being expanded if the user has defined a desired DATATYPE
+		if (FilteringDataType != mu::DATATYPE::DT_NONE)
+		{
+			const mu::OP_TYPE OperationType = Item->MutableModel->GetPrivate()->m_program.GetOpType(OperationAddress);
+			const mu::DATATYPE OperationDataType = mu::GetOpDataType(OperationType);
+
+			// If it is not of the desired type then ignore it and continue to the next pending element
+			if (OperationDataType != FilteringDataType)
+			{
+				continue;
+			}
+		}
+
+		// Reset the children array
+		Children.SetNum(0);
+		
+		// If not duplicated expand it and grab the children to be also expanded on the next loop
+		if (Item->DuplicatedOf == nullptr )
+		{
+			// Was this unique element expanded before (only valid if also expanding duplicates)
+			bool bHasBeenExpandedPreviously = false;
+			
+			// Mind duplicated original elements if dealing with duplicated operation expansions.
+			if (bForceExpandDuplicates)
+			{
+				// Make sure we have not already expanded this item to avoid recursive expansions of the same item and
+				// children
+				bHasBeenExpandedPreviously = AlreadyExpandedOriginalOperations.Contains(OperationAddress);
+			}
+			
+			// Only check for duplicated original elements when working with duplicates
+			if (!bHasBeenExpandedPreviously )
+			{
+				// Get the children of this unique element and prepare them for processing
+				GetChildrenForInfo(Item, Children);
+
+				// Call for the expansion of the children first
+				TreeExpandElements(Children, bForceExpandDuplicates, FilteringDataType,InExpandedOperationsBuffer);
+
+				// At this point all the children objects that needed expansion are already expanded so we can proceed with
+				// the expansion of this element
+				
+				{
+					// If we do expect to expand duplicates make sure we record this object as being expanded to be later able
+					// to block the expansion of duplicates of this object
+					if (bForceExpandDuplicates)
+					{
+						// Register this node as expanded so other nodes are able to check if it has already been worked with
+						AlreadyExpandedOriginalOperations.Add(OperationAddress);
+					}
+
+					// Only ask for the expansion of the element if we know it can be expanded due to it having
+					// children
+					if (!Children.IsEmpty())
+					{
+						// Expand this unique element
+						TreeView->SetItemExpansion(Item, true);
+					}
+				}
+			}
+			
+		}
+		// If it is a duplicated node
+		else
+		{
+			// Special behavior where we expand duplicates if parent is not found to be expanded
+			if (bForceExpandDuplicates)
+			{
+				// Was this element expanded as an original operation? We only want to expand the duplicate if the original
+				// was not duplicated before
+				bool bOriginalElementHasBeenExpanded = false;
+
+				// Was this element expanded on a duplicated element? we only want to expand the first duplicate!
+				const bool bOtherDuplicateOfSameOpWasExpandedBefore = AlreadyExpandedDuplicatedOperations.Contains(OperationAddress);
+
+				// Only check if there is another original element with the same operation if we know that there is not another
+				// duplicated element using this operation
+				if (!bOtherDuplicateOfSameOpWasExpandedBefore)
+				{
+					bOriginalElementHasBeenExpanded = AlreadyExpandedOriginalOperations.Contains(OperationAddress);
+				}
+				
+				// Was this operation expanded before?
+				const bool bWasOperationExpandedPreviously = bOriginalElementHasBeenExpanded || bOtherDuplicateOfSameOpWasExpandedBefore;
+
+				// If this operation have not yet been expanded then expand it!
+				// Duplicates do not have priority over original elements.
+				if ( !bWasOperationExpandedPreviously )
+				{
+					// Mark the children to be expanded later if conditions are met
+					GetChildrenForInfo(Item, Children);
+
+					// Expand the children objects
+					TreeExpandElements(Children, bForceExpandDuplicates, FilteringDataType,InExpandedOperationsBuffer);
+
+					// At this point all the children objects that needed expansion are already expanded so we can proceed with
+					// the expansion of this element
+					
+					{
+						// Record this node being expanded
+						AlreadyExpandedDuplicatedOperations.Add(OperationAddress);
+
+						// Only ask for the expansion of the element if we know it can be expanded due to it having
+						// children
+						if (!Children.IsEmpty())
+						{
+							// Expand the current element since we know it is from a operation not yet expanded
+							TreeView->SetItemExpansion(Item, true);
+						}
+					}
+				}
+			}
+		}
+	}
+}
+
+
+void SMutableCodeViewer::TreeExpandSelected()
+{
+	// Get the selected items and expand them excluding the duplicates
+	TArray<TSharedPtr<FMutableCodeTreeElement>> SelectedItems = TreeView->GetSelectedItems();
+	TreeExpandElements(SelectedItems,true);
+}
+
+void SMutableCodeViewer::TreeExpandUnique()
+{
+	// Expand the tree from the root and do not expand the duplicated elements
+	TreeExpandElements(RootNodes);
+}
+
+void SMutableCodeViewer::TreeExpandInstance()
+{
+	// Expand only the items that match the datatype provided
+	TreeExpandElements(RootNodes,false, mu::DT_INSTANCE);
+}
+#pragma endregion
+
+#pragma region Caching of operations related to constant resource
+
+void SMutableCodeViewer::CacheRootNodeAddresses()
+{
+	check (MutableModel);
+	check (RootNodeAddresses.IsEmpty())
+
+	TArray<mu::OP::ADDRESS> FoundRootNodeAddresses;
+	
+	const mu::Model::Private* ModelPrivate = MutableModel->GetPrivate();
+	const int32 StateCount = int32( ModelPrivate->m_program.m_states.size() );
+	for ( int32 StateIndex=0; StateIndex < StateCount; ++StateIndex )
+	{
+		const mu::PROGRAM::STATE& State = ModelPrivate->m_program.m_states[StateIndex];
+		FoundRootNodeAddresses.Add(State.m_root);
+	}
+
+	RootNodeAddresses = MoveTemp(FoundRootNodeAddresses);
+}
+
+ void SMutableCodeViewer::CacheAddressesRelatedWithConstantResource(const mu::DATATYPE ConstantDataType,
+	const int32 IndexOnConstantsArray)
+{
+	check(MutableModel);
+	if (IndexOnConstantsArray < 0)
+	{
+		// Not valid index.
+		UE_LOG(LogTemp,Error,TEXT("The provided index [%d] is not valid."),IndexOnConstantsArray );
+		return;
+	}
+	
+	// Iterate over all the operations found on this model. For that use as start the root elements of this graph and 
+	// move deeper in it
+	
+	TSet<mu::OP::ADDRESS> ProcessedAddresses;
+	TSet<mu::OP::ADDRESS> ConstantResourcesAddresses; 
+
+	// Main update procedure run for the targeted state and the targeted parameter values
+	const mu::PROGRAM& Program = MutableModel->GetPrivate()->m_program;
+	for (const mu::OP::ADDRESS& RootOperationAddress : RootNodeAddresses)
+	{
+		GetOperationsReferencingConstantResource(ConstantDataType, IndexOnConstantsArray,Program, RootOperationAddress,
+											 ConstantResourcesAddresses,ProcessedAddresses);
+	}
+	
+	// At this point we did get all the addresses of operations that do involve the usage of our resource
+	if (ConstantResourcesAddresses.Num() > 0)
+	{
+		// Managed by TargetedTypeSelector->SetSelectedItem callback method
+		// SetCurrentNavigationElement(nullptr);
+		
+		// Reset operation selection object to show no element to navigate to since we are not navigating over op types
+		{
+			// When setting it to none we are telling the system to not navigate over the operations
+			const int32 OperationIndex =  StaticCast<int32>(mu::OP_TYPE::NONE);
+			const TSharedPtr<FString> NewOperationString = OperationTypesStrings[OperationIndex];
+
+			// Set the type operation type to be looking for to be none.
+			TargetedTypeSelector->SetSelectedItem(NewOperationString);
+		}
+
+		// Dump the located resources array onto the navigation array since we have content to navigate over
+		NavigationOPAddresses = ConstantResourcesAddresses.Array();
+
+		// Move the view to the top if possible
+		// this will take some time to finish -> wait until update end
+		TreeView->ScrollToTop();
+		
+		// Tell the system changes have been made over Navigation Addresses
+		bUpdatedNavigationAddresses = true;
+	}
+	else
+	{
+		UE_LOG(LogTemp,Error,TEXT("The provided constant index does not seem to be used anywere : Make sure the index is valid and that IsConstantResourceUsedByOperation() switch is up to date"));
+	}
+
+}
+
+void SMutableCodeViewer::GetOperationsReferencingConstantResource(const mu::DATATYPE ConstantDataType, const int32 IndexOnConstantsArray, const mu::PROGRAM& InProgram,
+	const mu::OP::ADDRESS& InParentAddress,
+	TSet<mu::OP::ADDRESS>& OutAddressesWithPresence,
+	TSet<mu::OP::ADDRESS>& AlreadyProcessedAddresses)
+{
+	// Generic case for unnamed children traversal.
+	mu::ForEachReference(InProgram, InParentAddress, [this, &OutAddressesWithPresence, &ConstantDataType, &IndexOnConstantsArray, &AlreadyProcessedAddresses, &InProgram]( mu::OP::ADDRESS ChildAddress)
+	{
+		// If the parent does have a child then process it 
+		if (ChildAddress && !AlreadyProcessedAddresses.Contains(ChildAddress))
+		{
+			// Cache if same data type and we share the same address (means this op is pointing at the provided resource)
+			if (IsConstantResourceUsedByOperation(IndexOnConstantsArray,ConstantDataType,ChildAddress,InProgram))
+			{
+				OutAddressesWithPresence.Add(ChildAddress);
+			}
+		
+			// Cache to avoid processing it again later
+			AlreadyProcessedAddresses.Add(ChildAddress);
+		
+			// Process the children of this object
+			GetOperationsReferencingConstantResource(ConstantDataType,IndexOnConstantsArray,InProgram,ChildAddress,OutAddressesWithPresence,AlreadyProcessedAddresses);
+		}
+	});
+}
+
+bool SMutableCodeViewer::IsConstantResourceUsedByOperation(const int32 IndexOnConstantsArray,
+	const mu::DATATYPE ConstantDataType, const mu::OP::ADDRESS OperationAddress, const mu::PROGRAM& InProgram) const
+{
+	// Cache the current operation type to know where to look and what to check
+	const mu::OP_TYPE OperationType = InProgram.GetOpType(OperationAddress);
+
+	// Making usage of the operation data type is not valid since some operations while return one type do, in fact,
+	// contain data from other types (like the mesh constant for example that contains mesh, skeleton and physics asset)
+	// const mu::DATATYPE DataType = mu::GetOpDataType(OperationType);
+	// if (DataType != ConstantDataType)
+	// {
+	// 	return false;
+	// }
+
+	// Is this operation referencing (by an index) the index we are providing from a constants array
+	bool bResourceLocated = false;
+	
+	// Check if the operation data type is compatible with the type of resources we are providing
+	switch (ConstantDataType)
+	{
+		case mu::DATATYPE::DT_STRING:
+			{
+				// TIP: To know if they represent a constant value check the code on code runner to see if they read from the constants array
+				if (OperationType == mu::OP_TYPE::ST_CONSTANT)
+				{
+					bResourceLocated = IndexOnConstantsArray == InProgram.GetOpArgs<mu::OP::ResourceConstantArgs>(OperationAddress).value;
+				}
+				else if (OperationType == mu::OP_TYPE::IN_ADDSTRING)
+				{
+					bResourceLocated = IndexOnConstantsArray == InProgram.GetOpArgs<mu::OP::InstanceAddArgs>(OperationAddress).name;
+				}
+				else if (OperationType == mu::OP_TYPE::IN_ADDMESH)
+				{
+					bResourceLocated = IndexOnConstantsArray == InProgram.GetOpArgs<mu::OP::InstanceAddArgs>(OperationAddress).name;
+				}
+				else if(OperationType == mu::OP_TYPE::IN_ADDIMAGE)
+				{
+					bResourceLocated = IndexOnConstantsArray == InProgram.GetOpArgs<mu::OP::InstanceAddArgs>(OperationAddress).name;
+				}
+				else if (OperationType == mu::OP_TYPE::IN_ADDVECTOR)
+				{
+					bResourceLocated = IndexOnConstantsArray == InProgram.GetOpArgs<mu::OP::InstanceAddArgs>(OperationAddress).name;
+				}
+				else if (OperationType == mu::OP_TYPE::IN_ADDSCALAR)
+				{
+					bResourceLocated = IndexOnConstantsArray == InProgram.GetOpArgs<mu::OP::InstanceAddArgs>(OperationAddress).name;
+				}
+				else if (OperationType == mu::OP_TYPE::IN_ADDCOMPONENT)
+				{
+					bResourceLocated = IndexOnConstantsArray == InProgram.GetOpArgs<mu::OP::InstanceAddArgs>(OperationAddress).name;
+				}
+				else if (OperationType == mu::OP_TYPE::IN_ADDSURFACE)
+				{
+					bResourceLocated = IndexOnConstantsArray == InProgram.GetOpArgs<mu::OP::InstanceAddArgs>(OperationAddress).name;
+				}
+				else if (OperationType == mu::OP_TYPE::ME_CLIPMORPHPLANE)
+				{
+					const mu::OP::MeshClipMorphPlaneArgs Arguments = InProgram.GetOpArgs<mu::OP::MeshClipMorphPlaneArgs>(OperationAddress);
+
+					// treat the data as if it was a bone name
+					if (Arguments.vertexSelectionType == mu::OP::MeshClipMorphPlaneArgs::VS_BONE_HIERARCHY)
+					{
+						bResourceLocated = IndexOnConstantsArray == InProgram.GetOpArgs<mu::OP::MeshClipMorphPlaneArgs>(OperationAddress).vertexSelectionShapeOrBone;
+					}
+				}
+				else if (OperationType == mu::OP_TYPE::ME_BINDSHAPE)
+				{
+					mu::OP::MeshBindShapeArgs Arguments = InProgram.GetOpArgs<mu::OP::MeshBindShapeArgs>(OperationAddress);
+					const uint8_t* Data = InProgram.GetOpArgsPointer(OperationAddress);
+
+					// Bones are stored after the args
+					Data += sizeof(Arguments);
+
+					// Iterate over the bones and check if they point to the same index on the string constants array
+					int32 NumBones;
+					FMemory::Memcpy(&NumBones, Data, sizeof(int32)); 
+					Data += sizeof(int32);
+				
+					for (int32 Bone = 0; Bone < NumBones; ++Bone)
+					{
+						// Exit once we know that the data is pointing to the index provided
+						if (*Data == IndexOnConstantsArray)
+						{
+							bResourceLocated = true;
+							break;
+						}
+						
+						Data += sizeof(int32);
+					}
+				}
+				
+				break;
+			}
+
+		case mu::DATATYPE::DT_IMAGE:
+			{
+				if (OperationType == mu::OP_TYPE::IM_CONSTANT)
+				{
+					bResourceLocated = IndexOnConstantsArray == InProgram.GetOpArgs<mu::OP::ResourceConstantArgs>(OperationAddress).value;
+				} 
+				
+				break;
+			}
+
+		case mu::DATATYPE::DT_MESH:
+			{
+				if (OperationType == mu::OP_TYPE::ME_CONSTANT)
+				{
+					bResourceLocated = IndexOnConstantsArray == InProgram.GetOpArgs<mu::OP::MeshConstantArgs>(OperationAddress).value;
+				}
+				break;
+			}
+
+		case mu::DATATYPE::DT_LAYOUT:
+			{
+				if (OperationType == mu::OP_TYPE::LA_CONSTANT)
+				{
+					bResourceLocated = IndexOnConstantsArray == InProgram.GetOpArgs<mu::OP::ResourceConstantArgs>(OperationAddress).value;
+				}
+				break;
+			}
+
+		case mu::DATATYPE::DT_PROJECTOR:
+			{
+				if (OperationType == mu::OP_TYPE::PR_CONSTANT)
+				{
+					bResourceLocated = IndexOnConstantsArray == InProgram.GetOpArgs<mu::OP::ResourceConstantArgs>(OperationAddress).value;
+				}
+				break;
+			}
+
+		case mu::DATATYPE::DT_MATRIX:
+			{
+				if (OperationType == mu::OP_TYPE::ME_TRANSFORM)
+				{
+					bResourceLocated = IndexOnConstantsArray == InProgram.GetOpArgs<mu::OP::MeshTransformArgs>(OperationAddress).matrix;
+				}
+				break;
+			}
+
+		case mu::DATATYPE::DT_SHAPE:
+			{
+				if (OperationType == mu::OP_TYPE::ME_CLIPMORPHPLANE)
+				{
+					const mu::OP::MeshClipMorphPlaneArgs Arguments =  InProgram.GetOpArgs<mu::OP::MeshClipMorphPlaneArgs>(OperationAddress);
+					
+					// Morph shape
+					bResourceLocated = IndexOnConstantsArray == Arguments.morphShape;
+					if (bResourceLocated)
+					{
+						break;
+					}
+
+					if (Arguments.vertexSelectionType == mu::OP::MeshClipMorphPlaneArgs::VS_SHAPE)
+					{
+						// Selection Shape
+						bResourceLocated = IndexOnConstantsArray == InProgram.GetOpArgs<mu::OP::MeshClipMorphPlaneArgs>(OperationAddress).vertexSelectionShapeOrBone;
+					}
+				}
+				break;
+			}
+		
+		case mu::DATATYPE::DT_CURVE:
+			{
+				if (OperationType == mu::OP_TYPE::SC_CURVE)
+				{
+					bResourceLocated = IndexOnConstantsArray == InProgram.GetOpArgs<mu::OP::ScalarCurveArgs>(OperationAddress).curve;
+				}
+				break;
+			}
+
+		case mu::DATATYPE::DT_SKELETON:
+			{
+				if (OperationType == mu::OP_TYPE::ME_CONSTANT)
+				{
+					bResourceLocated = IndexOnConstantsArray == InProgram.GetOpArgs<mu::OP::MeshConstantArgs>(OperationAddress).skeleton;
+				}
+				break;
+			}
+
+		// Invalid types
+		case mu::DATATYPE::DT_NONE:
+		default:
+			{
+				checkNoEntry();
+			}
+	}
+
+
+	return bResourceLocated;
+}
+
+
+
+#pragma endregion 
+
+
+void SMutableCodeViewer::Tick(const FGeometry& AllottedGeometry, const double InCurrentTime, const float InDeltaTime)
+{
+	SCompoundWidget::Tick(AllottedGeometry, InCurrentTime, InDeltaTime);
+	
+	// After the tick we do know the tree has been refreshed, so all expansion and contraction operations have been
+	// completed and the new data has been loaded onto our listening arrays. Then its safe to expect the widgets to be
+	// there to be selected or inspected.
+	if (!TreeView->IsPendingRefresh())
+	{
+		// If the addresses set have changed then start looking for the first element that matches the first element on
+		// the addresses array
+		if ( bUpdatedNavigationAddresses)
+		{
+			// Reset the flag
+			bUpdatedNavigationAddresses = false;
+
+			// Locate the elements that can be currently navigated over.
+			LocateNavigationElements();
+
+			// Simulate a button press to go to the next operation. At this point we should have no operation set so it
+			// will find and select the first one
+			// Commented out to avoid the jumping around behaviour of the tree view when selecting a constant resource.
+			// It also avoids situations where we select one constant that invokes a viewer (of type StringViewport for example) and close after, due
+			// to the operation that hosts that resource being selected automatically its that operation previewer witch gets open therefore
+			// clouding the initial viewport we were using.
+			// OnGoToNextOperationButtonPressed();
+
+			return;
+		}
+
+		// Are we searching upwards?
+		if (bIsSearchingForPreviousElement)
+		{
+			GoToPreviousOperationAfterRefresh();
+			return;
+		}
+
+		// Are we searching downwards?
+		if (bIsSearchingForNextElement)
+		{
+			GoToNextOperationAfterRefresh();
+			return;
+		}
+	}
+	
+	
+	if (!bIsPreviewPendingUpdate)
+	{
+		return;
+	}
+
+	bIsPreviewPendingUpdate = false;
+
+	const mu::OP_TYPE OperationType = MutableModel->GetPrivate()->m_program.GetOpType(SelectedOperationAddress);
+	const mu::DATATYPE OperationDataType = mu::GetOpDataType(OperationType);
+
+	const mu::SettingsPtr Settings = new mu::Settings();
+	const mu::SystemPtr System = new mu::System(Settings);
+	System->GetPrivate()->BeginBuild(MutableModel);
+
+	switch (OperationDataType)
+	{
+	case mu::DT_LAYOUT:
+	{
+		check(PreviewLayoutViewer);
+		mu::LayoutPtrConst MutableLayout = System->GetPrivate()->BuildLayout(MutableModel, PreviewParameters.get(), SelectedOperationAddress);
+		PreviewLayoutViewer->SetLayout(MutableLayout);
+		break;
+	}
+
+	case mu::DT_IMAGE:
+	{
+		check(PreviewImageViewer);
+		mu::ImagePtrConst MutableImage = System->GetPrivate()->BuildImage(MutableModel, PreviewParameters.get(), SelectedOperationAddress, MipsToSkip);
+		PreviewImageViewer->SetImage(MutableImage, 0);
+		break;
+	}
+
+	case mu::DT_MESH:
+	{
+		check(PreviewMeshViewer);
+		mu::MeshPtrConst MutableMesh = System->GetPrivate()->BuildMesh(MutableModel, PreviewParameters.get(), SelectedOperationAddress);
+		PreviewMeshViewer->SetMesh(MutableMesh);
+		break;
+	}
+
+	case mu::DT_BOOL:
+	{
+		check(PreviewBoolViewer);
+		const bool MutableBool = System->GetPrivate()->BuildBool(MutableModel, PreviewParameters.get(), SelectedOperationAddress);
+		PreviewBoolViewer->SetBool(MutableBool);
+		break;
+	}
+
+	case mu::DT_INT:
+	{
+		check(PreviewIntViewer);
+		const int MutableInt = System->GetPrivate()->BuildInt(MutableModel, PreviewParameters.get(), SelectedOperationAddress);
+		PreviewIntViewer->SetInt(MutableInt);
+		break;
+	}
+
+	case mu::DT_SCALAR:
+	{
+		check(PreviewScalarViewer);
+		const float MutableScalar = System->GetPrivate()->BuildScalar(MutableModel, PreviewParameters.get(), SelectedOperationAddress);
+		PreviewScalarViewer->SetScalar(MutableScalar);
+		break;
+	}
+
+	case mu::DT_STRING:
+	{
+		check(PreviewStringViewer);
+		mu::Ptr<const mu::String> MutableString = System->GetPrivate()->BuildString(MutableModel, PreviewParameters.get(), SelectedOperationAddress);
+		const FText MutableText = FText::FromString(FString(MutableString->GetValue()));
+		PreviewStringViewer->SetString(MutableText);
+		break;
+	}
+
+	case mu::DT_COLOUR:
+	{
+		check(PreviewColorViewer);
+		float RedValue;
+		float GreenValue;
+		float BlueValue;
+		System->GetPrivate()->BuildColour(MutableModel, PreviewParameters.get(), SelectedOperationAddress, &RedValue, &GreenValue, &BlueValue);
+		PreviewColorViewer->SetColor(RedValue, GreenValue, BlueValue);
+		break;
+	}
+
+	case mu::DT_PROJECTOR:
+	{
+		check (PreviewProjectorViewer);
+		const mu::Ptr<const mu::Projector> ProjectorPtr = System->GetPrivate()->BuildProjector(
+				MutableModel, PreviewParameters.get(), SelectedOperationAddress);
+		PreviewProjectorViewer->SetProjector(ProjectorPtr->m_value);
+	}
+	
+	default:
+#if UE_BUILD_DEBUG || UE_BUILD_DEVELOPMENT
+		UE_LOG(LogMutable,Display,TEXT("There is no previewer for the selected type of Mutable object"))
+#endif
+		// There is no viewer for this type.
+		break;
+	}
+
+	System->GetPrivate()->EndBuild();
+
+}
+
+
+void SMutableCodeViewer::OnPreviewParameterValueChanged(int32 ParamIndex)
+{
+	// This is deferred to the tick to avoid multiple updates per frame.
+	bIsPreviewPendingUpdate = true;
+}
+
+
+void SMutableCodeViewer::PrepareStringViewer()
+{
+	if (!PreviewStringViewer)
+	{
+		PreviewStringViewer = SNew(SMutableStringViewer);
+	}
+
+	PreviewBorder->SetContent(PreviewStringViewer.ToSharedRef());
+}
+
+void SMutableCodeViewer::PrepareImageViewer()
+{
+	if (!PreviewImageViewer)
+	{
+		PreviewImageViewer = SNew(SMutableImageViewer)
+			.GridSize(FIntPoint(8, 8));
+	}
+
+	PreviewBorder->SetContent(PreviewImageViewer.ToSharedRef());
+}
+
+void SMutableCodeViewer::PrepareMeshViewer()
+{
+	if (!PreviewMeshViewer)
+	{
+		PreviewMeshViewer = SNew(SMutableMeshViewer);
+	}
+
+	PreviewBorder->SetContent(PreviewMeshViewer.ToSharedRef());
+}
+
+void SMutableCodeViewer::PrepareLayoutViewer()
+{
+	if (!PreviewLayoutViewer)
+	{
+		PreviewLayoutViewer = SNew(SMutableLayoutViewer);
+	}
+
+	PreviewBorder->SetContent(PreviewLayoutViewer.ToSharedRef());
+}
+
+
+void SMutableCodeViewer::PrepareProjectorViewer()
+{
+	if (!PreviewProjectorViewer)
+	{
+		PreviewProjectorViewer = SNew(SMutableProjectorViewer);
+	}
+
+	PreviewBorder->SetContent(PreviewProjectorViewer.ToSharedRef());
+}
+
+
+void SMutableCodeViewer::PreviewMutableString(const mu::string* InStringPtr)
+{
+	if (!InStringPtr)
+	{
+		UE_LOG(LogTemp,Error,TEXT("Unable to preview data on null String pointer."))
+		return;
+	}
+	
+	// Prepare the previewer object to receive data 
+	PrepareStringViewer();
+	
+	//  Provide the desired data to the previewer object
+	const FText TextToShow = FText::FromString(FString(InStringPtr->c_str()));
+	PreviewStringViewer->SetString(TextToShow);
+}
+
+
+void SMutableCodeViewer::PreviewMutableImage(mu::ImagePtrConst InImagePtr)
+{
+	PrepareImageViewer();
+	PreviewImageViewer->SetImage(InImagePtr,0);
+}
+
+
+void SMutableCodeViewer::PreviewMutableMesh(mu::MeshPtrConst InMeshPtr)
+{
+	PrepareMeshViewer();
+	PreviewMeshViewer->SetMesh(InMeshPtr);
+}
+
+
+void SMutableCodeViewer::PreviewMutableLayout(mu::LayoutPtrConst Layout)
+{
+	PrepareLayoutViewer();
+	PreviewLayoutViewer->SetLayout(Layout);
+}
+
+
+void SMutableCodeViewer::PreviewMutableProjector(const mu::PROJECTOR* Projector)
+{
+	if (!Projector)
+	{
+		UE_LOG(LogTemp,Error,TEXT("Unable to preview data on null Projector pointer."))
+		return;
+	}
+	
+	PrepareProjectorViewer();
+	PreviewProjectorViewer->SetProjector(*Projector);
+}
+
+
+void SMutableCodeViewer::PreviewMutableSkeleton(mu::SkeletonPtrConst Skeleton)
+{
+	if (!PreviewSkeletonViewer)
+	{
+		PreviewSkeletonViewer = SNew(SMutableSkeletonViewer);
+	}
+
+	PreviewBorder->SetContent(PreviewSkeletonViewer.ToSharedRef());
+	
+	PreviewSkeletonViewer->SetSkeleton(Skeleton);
+}
+
+
+void SMutableCodeViewer::PreviewMutableCurve(const mu::Curve* Curve)
+{
+	if (!Curve)
+	{
+		UE_LOG(LogTemp,Error,TEXT("Unable to preview data on null Curve pointer."))
+		return;
+	}
+	
+	if (!PreviewCurveViewer)
+	{
+		PreviewCurveViewer = SNew(SMutableCurveViewer);
+	}
+
+	PreviewBorder->SetContent(PreviewCurveViewer.ToSharedRef());
+	
+	PreviewCurveViewer->SetCurve(*Curve);
+}
+
+// TODO: Implement matrix viewer
+void SMutableCodeViewer::PreviewMutableMatrix(const mu::mat4f* Mat)
+{
+	UE_LOG(LogMutable,Warning,TEXT("Previewer for Mutable Matrices not yet implemented"))
+}
+
+// TODO: Implement shape viewer
+void SMutableCodeViewer::PreviewMutableShape(const mu::SHAPE* Shape)
+{
+	UE_LOG(LogMutable,Warning,TEXT("Previewer for Mutable Shapes not yet implemented"))
+}
+
+#undef LOCTEXT_NAMESPACE 
