[mutable] Move Mutable plugin from the restricted folder to the experimental plugins folder.

diff --git a/Engine/Plugins/Experimental/Mutable/Source/CustomizableObjectEditor/Private/SMutableCodeViewer.h b/Engine/Plugins/Experimental/Mutable/Source/CustomizableObjectEditor/Private/SMutableCodeViewer.h
new file mode 100644
--- /dev/null
+++ b/Engine/Plugins/Experimental/Mutable/Source/CustomizableObjectEditor/Private/SMutableCodeViewer.h
@@ -0,0 +1,634 @@
+// Copyright Epic Games, Inc. All Rights Reserved.
+
+#pragma once
+
+#include "CustomizableObject.h"
+#include "SMutableBoolViewer.h"
+#include "SMutableColorViewer.h"
+#include "SMutableImageViewer.h"
+#include "SMutableIntViewer.h"
+#include "SMutableLayoutViewer.h"
+#include "SMutableMeshViewer.h"
+#include "SMutableParametersWidget.h"
+#include "SMutableProjectorViewer.h"
+#include "SMutableScalarViewer.h"
+#include "SMutableStringViewer.h"
+#include "SMutableSkeletonViewer.h"
+#include "SMutableCurveViewer.h"
+#include "Framework/Docking/TabManager.h"
+#include "MutableRuntime/Private/Operations.h"
+#include "MutableRuntime/Public/Model.h"
+#include "Widgets/SCompoundWidget.h"
+#include "Widgets/Views/STreeView.h"
+
+#include "MutableRuntime/Private/ModelPrivate.h"
+#include "MutableRuntime/Private/ParametersPrivate.h"
+
+// Forward declarations
+class FMutableCodeTreeElement;
+class SMutableConstantsWidget;
+
+/** This widget shows the internal Mutable Code for debugging purposes. 
+ * This is not the Unreal source graph in the UCustomizableObject, but the actual Mutable virtual machine graph.
+ */
+class SMutableCodeViewer :
+	public SCompoundWidget,
+	public FGCObject
+{
+public:
+
+	SLATE_BEGIN_ARGS(SMutableCodeViewer) {}
+
+	/** User-visible tag to identify the source of the data shown. */
+	SLATE_ARGUMENT(FString, DataTag)
+		
+	SLATE_END_ARGS()
+
+	void Construct(const FArguments& InArgs, const mu::ModelPtr& InMutableModel /*, const TSharedPtr<SDockTab>& ConstructUnderMajorTab*/);
+
+	// SWidget interface
+	virtual void Tick(const FGeometry& AllottedGeometry, const double InCurrentTime, const float InDeltaTime) override;
+
+	// FGCObject interface
+	virtual void AddReferencedObjects(FReferenceCollector& Collector) override;
+	virtual FString GetReferencerName() const override;
+
+private:
+
+	/** The Mutable Model that we are showing. */
+	mu::ModelPtr MutableModel;
+	
+	/** Selected model operation for preview. */
+	mu::OP::ADDRESS SelectedOperationAddress = 0;
+
+	/** Mutable parameters used in the preview. */
+	mu::ParametersPtr PreviewParameters;
+
+	/** Widget showing the parameters that affect the current preview. */
+	TSharedPtr<SMutableParametersWidget> ParametersWidget;
+
+	/** Widget showing the constants found on the mutable Model program */
+	TSharedPtr<SMutableConstantsWidget> ConstantsWidget;
+
+	/** If true, the parameters have changed and we need to update the preview. */
+	bool bIsPreviewPendingUpdate = false;
+
+	/** Widget container where different previews will be created. */
+	TSharedPtr<SBorder> PreviewBorder;
+
+	/*
+	* Preview windows for the data types exposed by Mutable
+	*/
+
+	/** Widget used to show the preview of layout operation results. Once created it is reused to preserve the settings. */
+	TSharedPtr<SMutableLayoutViewer> PreviewLayoutViewer;
+
+	/** Widget used to show the preview of image operation results. Once created it is reused to preserve the settings. */
+	TSharedPtr<SMutableImageViewer> PreviewImageViewer;
+
+	/** Widget used to show a preview of the mesh and the metadata it holds */
+	TSharedPtr<SMutableMeshViewer> PreviewMeshViewer;
+
+	/** Widget used to show a preview of a mutable bool value */
+	TSharedPtr<SMutableBoolViewer> PreviewBoolViewer;
+
+	/** Widget used to show a preview of a mutable int value */
+	TSharedPtr<SMutableIntViewer> PreviewIntViewer;
+
+	/** Widget used to show a preview of a mutable float value */
+	TSharedPtr<SMutableScalarViewer> PreviewScalarViewer;
+
+	/** Widget used to show a preview of a mutable string value */
+	TSharedPtr<SMutableStringViewer> PreviewStringViewer;
+	
+	/** Widget used to show a preview of a mutable color value */
+	TSharedPtr<SMutableColorViewer> PreviewColorViewer;
+
+	/** Widget used to display the data held on the mutable projector objects */
+	TSharedPtr<SMutableProjectorViewer> PreviewProjectorViewer;
+
+	/** Widget used to display the data held on the mutable skeleton objects */
+	TSharedPtr<SMutableSkeletonViewer> PreviewSkeletonViewer;
+
+	/** Widget used to display the data held on the mutable curve objects */
+	TSharedPtr<SMutableCurveViewer> PreviewCurveViewer;
+	
+	/** Tree widget showing the code hierarchically. */
+	TSharedPtr<STreeView<TSharedPtr<FMutableCodeTreeElement>>> TreeView;
+
+	/** Root nodes of the tree widget. */
+	TArray<TSharedPtr<FMutableCodeTreeElement>> RootNodes;
+
+	/** Cache of tree elements matching the operations that have been generated so far. 
+	* We store both the parent and the operation in the key, because a single operation may appear multiple times if it has different parents.
+	*/
+	struct FItemCacheKey
+	{
+		mu::OP::ADDRESS Parent=0;
+		mu::OP::ADDRESS Child=0;
+		uint32 ChildIndexInParent=0;
+
+		friend FORCEINLINE bool operator == (const FItemCacheKey& A, const FItemCacheKey& B)
+		{
+			return A.Parent == B.Parent && A.Child == B.Child && A.ChildIndexInParent == B.ChildIndexInParent;
+		}
+
+		friend FORCEINLINE uint32 GetTypeHash(const FItemCacheKey& Key)
+		{
+			return HashCombine(Key.Parent, HashCombine(Key.Child, Key.ChildIndexInParent));
+		}
+	};
+
+	/*
+	* Tree widget objects
+	*/
+
+	TMap< FItemCacheKey, TSharedPtr<FMutableCodeTreeElement>> ItemCache;
+
+	/** Main tree item for each op. An op can be represented with multiple tree nodes if it is reachable from different paths. */
+	TMap< mu::OP::ADDRESS, TSharedPtr<FMutableCodeTreeElement>> MainItemPerOp;
+
+	/** List with all the elements related to nodes displayed on the tree. */
+	TArray<TSharedPtr<FMutableCodeTreeElement>> TreeElements;
+
+	/** Array with all the elements that have been manually expanded by the user */
+	TMap< mu::OP::ADDRESS, TSharedPtr<FMutableCodeTreeElement>> ExpandedElements;
+
+	/** Control boolean to know if there are any highlighted elements on the tree */
+	bool bIsElementHighlighted = false;
+
+	/** Operation that is currently being highlighted */
+	mu::OP::ADDRESS HighlightedOperation{};
+
+	/*
+	 * Navigation : Operation type navigation Type selection object
+	 */
+	
+	/** Slate object that provides the user a way of selecting what kind of operation it wants to navigate.*/
+	TSharedPtr<STextComboBox> TargetedTypeSelector;
+	
+	/** Operation type we are using to search for tree nodes. Driven primarily by the UI */
+	mu::OP_TYPE OperationTypeToSearch = mu::OP_TYPE::NONE;
+
+	/** Array with all the names for each of the operations available on mutable. Used by the UI  */
+	TArray<TSharedPtr<FString>> OperationTypesStrings;
+	
+	/** Stores a list of strings based on the possible types of operations mutable define to be used by the UI */
+	void GenerateNavigationOpTypeStrings();
+
+	/** Method called when the user changes the targeted type using the UI. It is only half part of the operation since
+	 * "OnOptionTypeSelectionChangedAfterRefresh" will be the one that at the end will perform the operations over the tree.
+	 */
+	void OnOptionTypeSelectionChanged(TSharedPtr<FString, ESPMode::ThreadSafe> NesSelectedOperationString, ESelectInfo::Type SelectionType);
+	
+	/** Update the array with all the elements using the current operation. Required when changing the targeted operation type */
+	void LocateNavigationElements();
+	
+	/*
+	 * Navigation : UI callback methods
+	 */
+
+	FText OnPrintNavigableObjectAddressesCount() const;
+	
+	/** Method that tells the UI if the bottom to go to the previous element can be interacted
+	 * @return It returns true if the UI element should be interactable, false if not.
+	 */
+	bool CanInteractWithPreviousOperationButton() const;
+
+	/** Method that tells the UI if the bottom to go to the next element can be interacted
+	 * @return It returns true if the UI element should be interactable, false if not.
+	 */
+	bool CanInteractWithNextOperationButton() const;
+	
+
+	/** Callback method used to allow the user to directly set the type of operation to be scanning directly by selecting a
+	 * operation on the tree and using that operation type as the type to search for
+	 */
+	void OnSelectedOperationTypeFromTree();
+
+	/*
+	 * Navigation : control flags
+	 */
+	
+	/** Flag that tells the system the user does want to select the previous operation of the targeted type. */
+	bool bIsSearchingForPreviousElement = false;
+	
+	/** Flag that tells the system the user does want to select the next operation of the targeted type. */
+	bool bIsSearchingForNextElement = false;
+	
+	/** Flag used by the navigation system to know when it is required the full expansion of the tree.
+	 * It allows the system to safely jump from one element to another without the risk of leaving elements in the
+	 * middle without being accessed.
+	 */
+	bool bTreeWasExpanded = false;
+	
+	/** Flag set to true when the navigation system is automatically scrolling over the graph */
+	bool bIsScrolling = false;
+
+	/** Flag that tells the system when the navigation addresses have been updated and therefore action should be
+	 * taken after the tree refresh. Currently we look for the first element present on the array of addresses */
+	bool bUpdatedNavigationAddresses = false;
+	
+	/*
+	 * Navigation : OP Address based navigation
+	 */
+	
+	/** Collection of Mutable addresses. Cache used during tree navigation.
+	 * They serve as the navigation system backend and determine if an element will or not be stopped at during
+	 * the navigation process.
+	 * @note Not set as a TSet since I want to, in future revisions, be able to trust the order of the elements of the collection.
+	 */
+	TArray<mu::OP::ADDRESS> NavigationOPAddresses;
+
+	/** Provided an Operation type store as navigation addresses all operations of the targeted Op type
+	 * @note The operations found get saved on the collection  NavigationOPAddresses.
+	 * @param TargetedOperationType The type of operation we want to cache to be later able to navigate over it.
+	 */
+	void CacheAddressesOfOperationsOfType(mu::OP_TYPE TargetedOperationType);
+
+	/** Fills an array with all the addresses of operations that do have in common the same targeted operation type
+	 * @param TargetOperationType The operation type used to discriminate what operation addresses we want to retrieve.
+	 * @param InParentAddress The address of the parent object. Used to get the children and then process them too.
+	 * @param OutAddressesOfType Output with all the addresses that are of the same type as TargetOperationType.
+	 * @param  AlreadyProcessedAddresses Used during recursive calls. Cache the already processed operations to avoid
+	 * processing them more than once.
+	 */
+	void GetOperationsOfType ( const mu::OP_TYPE& TargetOperationType, const mu::OP::ADDRESS& InParentAddress, const mu::PROGRAM& InProgram,
+		TSet<mu::OP::ADDRESS>& OutAddressesOfType,
+		TSet<mu::OP::ADDRESS>& AlreadyProcessedAddresses);
+
+
+	/*
+	 * Navigation : Caching of operation types being used on the model
+	 */
+
+	/** Operation types present on the currently set mutable model. */
+	TArray<mu::OP_TYPE> ModelOperationTypes;
+
+	
+	/** Fills ModelOperationTypes with all the types present on the current model.
+	 * It also makes sure we have a NONE operation and that the operations are sorted alphabetically.
+	 * @note Method designed to be called once per model.*/
+	void CacheOperationTypesPresentOnModel();
+
+	
+	/** Method to scan over all the operations performed on the current model to produce a set of operation types present on it.
+	 * @param InParentAddress Address of the parent object. Required to be able to perform recursive calls to this method (to get the data from the children)
+	 * @param  InProgram Mutable program.
+	 * @param OutLocatedOperations Output of the method : Set with all the operation types present on the model.
+	 * @param AlreadyProcessedAddresses Addresses already processed. Used to avoid processing the same operation twice.
+	 */
+	void GetOperationTypesPresentOnModel( const mu::OP::ADDRESS& InParentAddress, const mu::PROGRAM& InProgram,
+		TSet<mu::OP_TYPE>& OutLocatedOperations,
+		TSet<mu::OP::ADDRESS>& AlreadyProcessedAddresses);
+
+	/*
+	 * Navigation : Cache operation instances based on constant resource
+	 */
+
+public:
+	
+	/**
+	 * Provided a constant resource data type and index on its constant array this method sets all the operations that
+	 * make usage of said constant to be navigable by the navigation system.
+	 * @note The output of this operation will be cached onto NavigationOPAddresses so we can navigate over them.
+	 * @param ConstantDataType The constant data type to locate operations referencing it: It is required for the
+	 * system to know what operations should be checked for the usage of the provided constant resource. For example, a
+	 * type of mu::DATATYPE::DT_Curve will try to locate operations that we know operate with constant curves.
+	 * @param IndexOnConstantsArray The index of the constant on its constants array. If it is, for example, a constant mesh
+	 * this value represents the index of the constant inside constantMeshes array on mu::program.
+	 */
+	void CacheAddressesRelatedWithConstantResource(const mu::DATATYPE ConstantDataType, const int32 IndexOnConstantsArray);
+
+private:
+
+	/** The addresses of the root operations. Cached when this object gets loaded on the Construct method of this slate */
+	TArray<mu::OP::ADDRESS> RootNodeAddresses;
+	
+	/** Utility method that provides you with the addresses of the root nodes of the program */
+	void CacheRootNodeAddresses();
+	
+	/**
+	 * Provided a data type and the index of the constant object using said data type locate all operations that do
+	 * directly make use of said constant resource.
+	 * @param ConstantDataType The type of the resource we are providing an index of. 
+	 * @param IndexOnConstantsArray The index of the constant resource we want to know what operations are using it.
+	 * @param InProgram The mutable program containing all the operations of the graph.
+	 * @param InParentAddress The address of the root object whose child operations we are searching the provided constant on.
+	 * @param OutAddressesWithPresence An array with all the addresses to mutable operations making use of the provided constant resource.
+	 * @param AlreadyProcessedAddresses Cache with the addresses already processed. It is only intended to be used on recursive calls.
+	 */
+	void GetOperationsReferencingConstantResource(const mu::DATATYPE ConstantDataType,
+	                                              const int32 IndexOnConstantsArray,
+	                                              const mu::PROGRAM& InProgram,
+	                                              const mu::OP::ADDRESS& InParentAddress,
+	                                              TSet<mu::OP::ADDRESS>& OutAddressesWithPresence,
+	                                              TSet<mu::OP::ADDRESS>& AlreadyProcessedAddresses);
+
+	/**
+	 * Provided an operation address this method tells us if that address is making use of our constant resource or not.
+	 * It will take into account the datatype of the provided resource to discard or acknowledge the operation type of the
+	 * provided operation address.
+	 * @param IndexOnConstantsArray The index onto its respective constants array of the constant resource provided. We do
+	 * not give a pointer to the constant resource but the index of it onto its hosting array.
+	 * @param ConstantDataType The type of data we know the constant uses. It is used to check for ones or other operations since
+	 * not all operations do access the same constant resources.
+	 * @param OperationAddress The address of the operation currently being scanned for constant resources.
+	 * @param InProgram The mutable program to use to check the type of operation that corresponds to that operation address.
+	 * @return True if the operation at OperationAddress does make use of the provided constant resource. False otherwise
+	 */
+	bool IsConstantResourceUsedByOperation(const int32 IndexOnConstantsArray,
+	                                       const mu::DATATYPE ConstantDataType, const mu::OP::ADDRESS OperationAddress,
+	                                       const mu::PROGRAM& InProgram) const;
+	
+	
+	/*
+	 * Navigation : operation methods and variables
+	 */
+	
+	/** The element on FoundElementsOfType that we are inspecting and using as origin for our search.
+	 * @note Do not update it's value manually */
+	TSharedPtr< FMutableCodeTreeElement> CurrentNavigationElement;
+	
+	/** Array with all the elements on the tree of the same type as the one defined on "TargetedOperationType" */
+	TArray<TSharedPtr<FMutableCodeTreeElement>> NavigationFoundElements;
+	
+	/** Copy of the array FoundElementsOfType done before moving on to the next element. Required to check if the
+	 * expansion of the tree view did add or remove elements of the type we are looking for.
+	 */
+	TArray<TSharedPtr<FMutableCodeTreeElement>> NavigationPreviouslyFoundElements;
+	
+	
+	/** Amount of elements the scroll will jump up and down while searching for new elements
+	 * It is required to be able to locate elements whose rows have not been yet generated
+	 * Do not change.
+	 */
+	const int32 MinTreeScrollStep = 1;
+
+	/** Variable holding the scroll step to be used. It will vary depending on the amount of elements the tree can display
+	 * witch may change from resolution to resolution. This is then the value to be used during the automatic scrolling
+	 * operation of the tree navigation system. It will start at the min possible value and only increase when posible.
+	 */
+	int32 ComputedMaxViewScrollStep = MinTreeScrollStep;
+
+	/** Tree view manipulation method designed to set the element to be selected on the UI side of the tree. It is also
+	 *required to take note of what the current found element is to be able to navigate the tree of elements
+	 *@param NewlyFoundElement - The element we want to now be using as the NavigationFoundElement */
+	void SetCurrentNavigationElement(TSharedPtr<FMutableCodeTreeElement> NewlyFoundElement);
+
+	
+	/** Consistency focused method designed to sort the list of navigation elements based on the position of them
+	* on the tree view. It is required since new elements may be added or removed due to user activity.
+	*/
+	void SortNavigationElementsArray();
+	
+	/** Used by the UI and internally does what is necessary to get to the previous element of the targeted search type.
+	 * It is designed to work alongside GoToPreviousOperationAfterRefresh(...)
+	 * @return A reply to tell if the UI action has been handled or not.
+	 */
+	FReply OnGoToPreviousOperationButtonPressed();
+
+	/** Used by the UI and internally does what is necessary to get to the next element of the targeted search type.
+	 * It is designed to work alongside GoToNextOperationAfterRefresh(...)
+	 * @return A reply to tell if the UI action has been handled or not.
+	 */
+	FReply OnGoToNextOperationButtonPressed();
+	
+	/** Makes the currently selected operation be the previous one able to be located on the tree of operations.
+	 * @note It must only be called after the tree has been fully refreshed. If not inconsistent behaviour may arise.
+	 */
+	void GoToPreviousOperationAfterRefresh();
+	
+	/** Makes the currently selected operation of the targeted type to be the next one. Called after the tree has been
+	 * refreshed with new contents (fully refreshed).
+	 * @note It must only be called after the tree has been fully refreshed. If not inconsistent behaviour may arise.
+	 */
+	void GoToNextOperationAfterRefresh();
+
+	/** Method that compares the contents of "PreviousTargetedElementsOfType" against the ones on "CurrentTargetedElementsOfType" to
+	 *check if there is any new element found on the current elements array not found on the Previous one
+	 * @return True if there is any new element defined on the "CurrentTargetedElementsOfType" not present on "PreviousTargetedElementsOfType",
+	 * false if not.
+	 */
+	bool HaveNewElementsBeenFound() const;
+	
+	/*
+	* Main callbacks from the tree widget standard operation. 
+	*/
+
+	/** Callback that generates a new Row object based on a Tree Node*/
+ 	TSharedRef<ITableRow> GenerateRowForNodeTree(TSharedPtr<FMutableCodeTreeElement> InInfo, const TSharedRef<STableViewBase>& OwnerTable);
+	
+	/**
+	* Provided an element it is returned a list with all the immediate children it has 
+	* @param InInfo - The MutableCodeElement whose children objects we are searching
+	* @param OutChildren - The children of the provided MutableCodeTreeElement
+	*/
+	void GetChildrenForInfo(TSharedPtr<FMutableCodeTreeElement> InInfo, TArray< TSharedPtr<FMutableCodeTreeElement> >& OutChildren);
+
+	/** 
+	* Callback invoked each time an element of the tree gets expanded or contracted 
+	* @param InInfo - The MutableCodeElement to be contracted or expanded
+	* @param bInExpanded - Determines if the action is of contraction (false) or expansion (true).
+	*/
+	void OnExpansionChanged(TSharedPtr<FMutableCodeTreeElement> InInfo, bool bInExpanded);
+
+	/** 
+	* Callback invoked each time the selected element changes used to generate the previews depending of the 
+	* operation selected
+	* @param InInfo - The MutableCodeElement that has been selected. Can be null if unselected
+	* @param SelectInfo - Way in witch the selection changed
+	*/
+	void OnSelectionChanged(TSharedPtr<FMutableCodeTreeElement> InInfo, ESelectInfo::Type SelectInfo);
+
+	/** 
+	* Callback invoked when a tree row is getting removed 
+	* @param InTreeRow - The tree row that is being removed from the tree view
+	*/
+	void OnRowReleased(const TSharedRef<ITableRow>& InTreeRow);
+
+	/** Callback invoked when whe press the right mouse button. Useful for adding context menu objects */
+	TSharedPtr<SWidget> OnTreeContextMenuOpening();
+
+	void TreeExpandRecursive(TSharedPtr<FMutableCodeTreeElement> Info, bool bExpand );
+	
+	/** Temporal object designed to be used during the recursive operation of TreeExpandElements() and group
+	 * strongly related data
+	 */
+	struct FProcessedOperationsBuffer
+	{
+		/** Array with the operation addresses of all original (not duplicates) expanded operations. */
+		TArray<mu::OP::ADDRESS> ExpandedOriginalOperations;
+
+		/** Array with the operation addresses of all duplicated expanded operations. */
+		TArray<mu::OP::ADDRESS> ExpandedDuplicatedOperations;
+	};
+	
+	/** 
+	* Expands the provided elements and all the children they have. 
+	* @note It performs the operation by going to the deepest children and expanding from then upwards to the origin before
+	* proceeding to the next parent. This way we make sure that the expansion of the elements is following the order
+	* of left to right and top to bottom in a "Z" like pattern like the way the tree gets read by humans.
+	* @param InElementsToExpand - The root info objects where to start the expansion
+	* @param bForceExpandDuplicates - (Optional) Determines if the duplicated objects must be expanded. 
+	*	Used for certain situations where we want to expand them. By default they do not get expanded
+	* @param FilteringDataType - (Optional) determines what kind of operations should be expanded. By default no
+	* filtering is applied
+	* @param InExpandedOperationsBuffer (Optional) Object containing all the duplicated and original elements already processed
+	* during the subsequent recursive calls to this method.
+	*/
+	void TreeExpandElements(TArray<TSharedPtr<FMutableCodeTreeElement>>& InElementsToExpand,
+		bool bForceExpandDuplicates = false,
+		mu::DATATYPE FilteringDataType = mu::DATATYPE::DT_NONE,
+		TSharedPtr<FProcessedOperationsBuffer> InExpandedOperationsBuffer = nullptr);
+	
+
+	/** Expand only the unique elements (mot duplicates) */
+	void TreeExpandUnique();
+
+	/** Expand only the elements using as operation type DT_INSTANCE */
+	void TreeExpandInstance();
+	
+	/** Grabs the selected element and expands all the elements inside the selected branch. Duplicates are ignored */
+	void TreeExpandSelected();
+
+	/** 
+	* Highlights all tree elements that share the same operation as the element provided
+	* @param InTargetElement - The info object to be used as blueprint to search for similar objects and highlight them
+	*/
+	void HighlightDuplicatesOfEntry(const TSharedPtr<FMutableCodeTreeElement>& InTargetElement);
+
+	/** Clear all the highlighted elements and set them to their default visual state */
+	void ClearHighlightedItems();
+
+	/** Called when any parameter value has changed, with the parameter index as argument.  */
+	void OnPreviewParameterValueChanged( int32 ParamIndex );
+
+	
+
+	/*
+	 * Control of "Skip Mips" for image operations.
+	 */
+
+	/** Operation type we are using to search for tree nodes. Driven by the UI */
+	int32 MipsToSkip = 0;
+
+	/** */
+	bool bSelectedOperationIsImage = false;
+
+	/** Stores a list of strings based on the possible types of operations mutable define to be used by the UI */
+	EVisibility IsMipSkipVisible() const;
+
+	/** Stores a list of strings based on the possible types of operations mutable define to be used by the UI */
+	TOptional<int32> GetCurrentMipSkip() const;
+
+	/** */
+	void OnCurrentMipSkipChanged(int32 NewValue);
+
+	/*
+	 * Remote previewer invocation methods
+	 */
+
+public:
+	void PreviewMutableImage (mu::ImagePtrConst InImagePtr);
+	void PreviewMutableMesh (mu::MeshPtrConst InMeshPtr);
+	void PreviewMutableLayout(mu::LayoutPtrConst Layout);
+	void PreviewMutableSkeleton(mu::SkeletonPtrConst Skeleton);
+	void PreviewMutableString(const mu::string* InStringPtr);
+	void PreviewMutableProjector(const mu::PROJECTOR* Projector);
+	void PreviewMutableMatrix(const mu::mat4f* Mat);
+	void PreviewMutableShape(const mu::SHAPE* Shape);
+	void PreviewMutableCurve(const mu::Curve* Curve);
+	
+private:
+	void PrepareStringViewer();
+	void PrepareImageViewer();
+	void PrepareMeshViewer();
+	void PrepareLayoutViewer();
+	void PrepareProjectorViewer();
+};
+
+
+
+/** An row of the code tree in the SMutableCodeViewer. */
+class FMutableCodeTreeElement : public TSharedFromThis<FMutableCodeTreeElement>
+{
+public:
+	FMutableCodeTreeElement(const mu::ModelPtr& InModel, mu::OP::ADDRESS InOperation, const FString& InCaption, const TSharedPtr<FMutableCodeTreeElement>* InDuplicatedOf = nullptr)
+	{
+		MutableModel = InModel;
+		MutableOperation = InOperation;
+		Caption = InCaption;
+		if (InDuplicatedOf)
+		{
+			DuplicatedOf = *InDuplicatedOf;
+		}
+
+		// Check what type of operation is (state constant or dynamic resource)
+		{
+			// If duplicated then grab the already processed data on the original operation
+			if (InDuplicatedOf)
+			{
+				bIsDynamicResource = DuplicatedOf->bIsDynamicResource;
+				bIsStateConstant = DuplicatedOf->bIsStateConstant;
+
+				// All required data has been processed so an early exit is required
+				return;
+			}
+			
+			// Iterate over all states and try to locate the operation
+			const mu::PROGRAM& MutableProgram = InModel->GetPrivate()->m_program;
+			for (const mu::PROGRAM::STATE& CurrentState : MutableProgram.m_states)
+			{
+				// Check if it is a dynamic resource
+				for (auto& DynamicResource : CurrentState.m_dynamicResources)
+				{
+					// If the operation gets located then mark it as dynamic resource
+					if (DynamicResource.first == MutableOperation)
+					{
+						bIsDynamicResource = true;
+						break;
+					}
+				}
+				
+				// Early exit: A dynamic resource can not be at the same time a state constant
+				if (bIsDynamicResource)
+				{
+					return;
+				}
+				
+				// Check if it is a state constant
+				for (auto& CachedOperation : CurrentState.m_updateCache)
+				{
+					// If the operation gets located then mark it as dynamic resource
+					if (CachedOperation == MutableOperation)
+					{
+						bIsStateConstant = true;
+						break;
+					}
+				}
+			}
+		}
+		
+	}
+
+public:
+
+	/** */
+	mu::ModelPtr MutableModel;
+
+	/** Mutable Graph Node represented in this tree row. */
+	mu::OP::ADDRESS MutableOperation;
+
+	/** If true means that it will not update when a runtime parameter on the state gets updated */
+	bool bIsStateConstant = false;
+
+	/** If true then the mesh or image of this operation may change during the state update */
+	bool bIsDynamicResource = false;
+
+	/** Label representing this operation. */
+	FString Caption;
+
+	/** If this tree element is a duplicated of another op, this is the op. */
+	TSharedPtr<FMutableCodeTreeElement> DuplicatedOf;
+};
diff --git a/Engine/Plugins/Experimental/Mutable/Source/CustomizableObjectEditor/Private/SMutableColorPreviewBox.cpp b/Engine/Plugins/Experimental/Mutable/Source/CustomizableObjectEditor/Private/SMutableColorPreviewBox.cpp
new file mode 100644
--- /dev/null
+++ b/Engine/Plugins/Experimental/Mutable/Source/CustomizableObjectEditor/Private/SMutableColorPreviewBox.cpp
@@ -0,0 +1,38 @@
+// Copyright Epic Games, Inc. All Rights Reserved.
+
+#include "SMutableColorPreviewBox.h"
+
+#include "Widgets/SBoxPanel.h"
+#include "Widgets/Layout/SBorder.h"
+#include "Styling/SlateColor.h"
+#include "Styling/AppStyle.h"
+
+#define LOCTEXT_NAMESPACE "CustomizableObjectEditor"
+
+void SMutableColorPreviewBox::Construct(const FArguments& InArgs)
+{
+	this->Color = InArgs._BoxColor.Get();
+
+	ChildSlot
+	[
+		SNew(SHorizontalBox)
+
+		+ SHorizontalBox::Slot()
+		[
+			SNew(SBorder)
+			.BorderImage(FAppStyle::GetBrush("WhiteBrush"))
+			.BorderBackgroundColor(this,&SMutableColorPreviewBox::GetColor)
+		]
+	];
+}
+
+void SMutableColorPreviewBox::SetColor(const FSlateColor& InColor)
+{
+	this->Color = InColor;
+}
+
+FSlateColor SMutableColorPreviewBox::GetColor() const
+{
+	return this->Color;
+}
+#undef LOCTEXT_NAMESPACES
\ No newline at end of file
diff --git a/Engine/Plugins/Experimental/Mutable/Source/CustomizableObjectEditor/Private/SMutableColorPreviewBox.h b/Engine/Plugins/Experimental/Mutable/Source/CustomizableObjectEditor/Private/SMutableColorPreviewBox.h
new file mode 100644
--- /dev/null
+++ b/Engine/Plugins/Experimental/Mutable/Source/CustomizableObjectEditor/Private/SMutableColorPreviewBox.h
@@ -0,0 +1,35 @@
+// Copyright Epic Games, Inc. All Rights Reserved.
+
+#pragma once
+
+#include "Widgets/SCompoundWidget.h"
+#include "Widgets/DeclarativeSyntaxSupport.h"
+#include "Styling/SlateColor.h"
+
+/** Custom widget designed to display a color on a box */
+class SMutableColorPreviewBox : public SCompoundWidget
+{
+public:
+	/** Slate Constructor that does allow the caller to provide the default value for the preview */
+	SLATE_BEGIN_ARGS(SMutableColorPreviewBox) {}
+		SLATE_ATTRIBUTE(FSlateColor, BoxColor)
+	SLATE_END_ARGS()
+
+	/** Builds the widget */
+	void Construct(const FArguments& InArgs);
+
+	/** Sets the color to be displayed by this widget */
+	void SetColor(const FSlateColor& InColor);
+
+private:
+
+	/** Defines the color used by this widget*/
+	FSlateColor Color;
+
+	/*
+	* Method that returns the FSLateColor being currently used by the preview 
+	* @return the FSLateColor that is being used by the Preview
+	*/
+	FSlateColor GetColor() const;
+};
+
diff --git a/Engine/Plugins/Experimental/Mutable/Source/CustomizableObjectEditor/Private/SMutableColorViewer.cpp b/Engine/Plugins/Experimental/Mutable/Source/CustomizableObjectEditor/Private/SMutableColorViewer.cpp
new file mode 100644
--- /dev/null
+++ b/Engine/Plugins/Experimental/Mutable/Source/CustomizableObjectEditor/Private/SMutableColorViewer.cpp
@@ -0,0 +1,163 @@
+// Copyright Epic Games, Inc. All Rights Reserved.
+
+#include "SMutableColorViewer.h"
+#include "SMutableColorPreviewBox.h"
+
+#include "GenericPlatform/GenericPlatform.h"
+#include "Widgets/SBoxPanel.h"
+#include "Internationalization/Text.h"
+#include "Templates/SharedPointer.h"
+#include "Math/Color.h"
+#include "Styling/SlateColor.h"
+#include "Widgets/Text/STextBlock.h"
+
+#define LOCTEXT_NAMESPACE "CustomizableObjectEditor"
+
+
+void SMutableColorViewer::Construct(const FArguments& InArgs)
+{
+	const FText ColorValueTitle = LOCTEXT("ColorValuesTitle", "Color Values : ");
+
+	const FText RedColorValueTitle = LOCTEXT("RedColorValueTitle", "Red : ");
+	const FText GreenColorValueTitle = LOCTEXT("GreenColorValueTitle", "Green : ");
+	const FText BlueColorValueTitle = LOCTEXT("BlueColorValueTitle", "Blue : ");
+
+	const int32 IndentationSpace = 16;	
+	const int32 AfterTitleSpacing = 4;
+	const int32 ColorPreviewHorizontalPadding = 18;
+
+	ChildSlot
+	[
+		SNew(SVerticalBox)
+
+		+ SVerticalBox::Slot()
+		.AutoHeight()
+		[	
+			SNew(STextBlock).
+			Text(ColorValueTitle)
+		]
+		
+		+ SVerticalBox::Slot()
+		.AutoHeight()
+		.Padding(IndentationSpace, AfterTitleSpacing)
+		[
+			SNew(SHorizontalBox)
+
+			+ SHorizontalBox::Slot()
+			.AutoWidth()
+			[
+				SNew(SVerticalBox)
+
+				// Red channel -----------------------------------------
+				+ SVerticalBox::Slot()
+				.AutoHeight()
+				[
+					SNew(SHorizontalBox)
+
+					+ SHorizontalBox::Slot()
+					.AutoWidth()
+					[
+						SNew(STextBlock).
+						Text(RedColorValueTitle)
+					]
+
+					+ SHorizontalBox::Slot()
+					.AutoWidth()
+					[
+						SNew(STextBlock).
+						Text(this, &SMutableColorViewer::GetRedValue)
+					]
+				]
+				// -----------------------------------------------------
+
+				// Green channel ---------------------------------------
+				+SVerticalBox::Slot()
+				.Padding(0, 1)
+				.AutoHeight()
+				[
+					SNew(SHorizontalBox)
+
+					+ SHorizontalBox::Slot()
+					.AutoWidth()
+					[
+						SNew(STextBlock).
+						Text(GreenColorValueTitle)
+					]
+
+					+ SHorizontalBox::Slot()
+					.AutoWidth()
+					[
+						SNew(STextBlock).
+						Text(this, &SMutableColorViewer::GetGreenValue)
+					]
+				]
+				// -----------------------------------------------------
+
+				// Blue channel ----------------------------------------
+				+SVerticalBox::Slot()
+				.Padding(0, 1)
+				.AutoHeight()
+				[
+					SNew(SHorizontalBox)
+
+					+ SHorizontalBox::Slot()
+					.AutoWidth()
+					[
+						SNew(STextBlock).
+						Text(BlueColorValueTitle)
+					]
+
+					+ SHorizontalBox::Slot()
+					.AutoWidth()
+					[
+						SNew(STextBlock).
+						Text(this, &SMutableColorViewer::GetBlueValue)
+					]
+				]
+				// -----------------------------------------------------
+			]
+		
+			// Color preview widget --------------------------------
+			+ SHorizontalBox::Slot()
+			.Padding(ColorPreviewHorizontalPadding,0)
+			.MaxWidth(120)
+			[
+				SAssignNew(this->ColorPreview, SMutableColorPreviewBox)
+				// .BoxColor(this, &SMutableColorViewer::GetPreviewColor)
+			]
+			// -----------------------------------------------------
+		]
+	];
+}
+
+void SMutableColorViewer::SetColor(const float& InRed, const float& InGreen, const float& InBlue)
+{
+	this->RedValue = InRed;
+	this->GreenValue = InGreen;
+	this->BlueValue = InBlue;
+
+	this->ColorPreview->SetColor( this->GetPreviewColor() );
+}
+
+
+FSlateColor SMutableColorViewer::GetPreviewColor() const
+{
+	return FSlateColor { FLinearColor{this->RedValue, this->GreenValue, this->BlueValue, this->AlphaValue} };
+}
+
+FText SMutableColorViewer::GetRedValue() const
+{
+	return FText::AsNumber(this->RedValue);
+}
+
+FText SMutableColorViewer::GetGreenValue() const
+{
+	return FText::AsNumber(this->GreenValue);
+}
+
+FText SMutableColorViewer::GetBlueValue() const
+{
+	return FText::AsNumber(this->BlueValue);
+}
+
+#undef LOCTEXT_NAMESPACE
\ No newline at end of file
diff --git a/Engine/Plugins/Experimental/Mutable/Source/CustomizableObjectEditor/Private/SMutableColorViewer.h b/Engine/Plugins/Experimental/Mutable/Source/CustomizableObjectEditor/Private/SMutableColorViewer.h
new file mode 100644
--- /dev/null
+++ b/Engine/Plugins/Experimental/Mutable/Source/CustomizableObjectEditor/Private/SMutableColorViewer.h
@@ -0,0 +1,46 @@
+// Copyright Epic Games, Inc. All Rights Reserved.
+
+#pragma once
+
+#include "SMutableColorPreviewBox.h"
+#include "Widgets/SCompoundWidget.h"
+#include "Widgets/DeclarativeSyntaxSupport.h"
+#include "Internationalization/Text.h"
+#include "Templates/SharedPointer.h"
+#include "Styling/SlateColor.h"
+
+
+class SMutableColorViewer : public SCompoundWidget
+{
+public:
+	SLATE_BEGIN_ARGS(SMutableColorViewer) {}
+	SLATE_END_ARGS()
+
+	/** Builds the widget */
+	void Construct(const FArguments& InArgs);
+
+	/** Set the Mutable Bool to be used for this widget */
+	void SetColor(const float& InRed, const float& InGreen, const float& InBlue);
+
+private:
+
+	/** Color being displayed */
+	float RedValue = 0.0f;
+	float GreenValue = 0.0f;
+	float BlueValue = 0.0f;
+	const float AlphaValue{1};
+
+	/** Color box widget designed to serve as a preview of the color reported by mutable */
+	TSharedPtr<SMutableColorPreviewBox> ColorPreview;
+
+	/** Retrieve the values of each color component for the UI Texts to be updated */
+	FText GetRedValue() const;
+	FText GetGreenValue() const;
+	FText GetBlueValue() const;
+
+	/*
+	* Get a color object that the SMutableColorPreviewBox is able to display 
+	* @return a new FSlateColor generated from the float values held by this widget
+	*/
+	FSlateColor GetPreviewColor() const;
+};
diff --git a/Engine/Plugins/Experimental/Mutable/Source/CustomizableObjectEditor/Private/SMutableConstantsWidget.cpp b/Engine/Plugins/Experimental/Mutable/Source/CustomizableObjectEditor/Private/SMutableConstantsWidget.cpp
new file mode 100644
--- /dev/null
+++ b/Engine/Plugins/Experimental/Mutable/Source/CustomizableObjectEditor/Private/SMutableConstantsWidget.cpp
@@ -0,0 +1,1495 @@
+// Copyright Epic Games, Inc. All Rights Reserved.
+
+#include "SMutableConstantsWidget.h"
+
+#include "MutableString.h"
+#include "MutableRuntime/Private/ModelPrivate.h"
+
+#include "MutableRuntime/Public/Image.h"
+#include "MutableRuntime/Public/Layout.h"
+#include "MutableRuntime/Public/Mesh.h"
+
+#include "SMutableImageViewer.h"
+#include "SMutableMeshViewer.h"
+#include "SMutableStringViewer.h"
+
+#include "SMutableCodeViewer.h"
+#include "Algo/StableSort.h"
+#include "Components/TextBlock.h"			
+#include "Internationalization/Text.h"		
+#include "MutableTools/Public/TypeInfo.h"
+#include "Widgets/Layout/SExpandableArea.h"	
+#include "Widgets/Layout/SScrollBox.h"		
+#include "Widgets/Layout/SSeparator.h"
+#include "Widgets/Views/SListView.h"
+#include "Widgets/Views/STileView.h"
+#include "Widgets/Input/SButton.h"
+#include "Math/UnrealMathUtility.h"
+
+
+#define LOCTEXT_NAMESPACE "CustomizableObjectEditor"
+
+// Private namespace with utility functionality used by this slate object
+namespace
+{
+	/** Provided a byte count this function proceeds to output that byte value as bytes, kilobytes, megabytes and gigabytes.
+	 * @param SizeInBytes The amount of bytes to convert to a formatted text that represents it.
+	 * @return A text representing the value provided as Bytes, Kilobytes, Megabytes and Gigabytes. It will not return
+	 * size types with value 0.
+	 */
+	FText GetSizeInBytesAsBiggerUnits(const uint32 SizeInBytes)
+	{
+		FString OutputString = "";
+		
+		uint32 Bytes = SizeInBytes;
+		uint32 KiloBytes = 0;
+		uint32 MegaBytes = 0;
+		uint32 GigaBytes = 0;
+		
+		// B to KB
+		if (Bytes >= 1024)
+		{
+			KiloBytes = Bytes / 1024;
+			Bytes = Bytes % 1024;
+			
+			// KB to MB
+			if (KiloBytes >= 1024)
+			{
+				MegaBytes = KiloBytes / 1024;
+				KiloBytes = KiloBytes % 1024;
+				
+				// MB to GB
+				if (MegaBytes >= 1024)
+				{
+					GigaBytes = MegaBytes / 1024;
+					MegaBytes = MegaBytes % 1024;
+				}
+			}
+		}
+
+		// Compose the string that will be the output of this function 
+		if (GigaBytes)
+		{
+			OutputString.Append(FString::FromInt(GigaBytes) + " GB, ");
+		}
+
+		if (MegaBytes)
+		{
+			OutputString.Append(FString::FromInt(MegaBytes) + " MB, ");
+		}
+
+		if (KiloBytes)
+		{
+			OutputString.Append(FString::FromInt(KiloBytes) + " KB, ");
+		}
+
+		if (Bytes)
+		{
+			OutputString.Append( FString::FromInt(Bytes) + " B ");
+		}
+
+		return FText::FromString(OutputString);
+	}
+}
+
+
+#pragma region SUPPORT_CLASSES
+
+class SMutableConstantMeshRow final : public STableRow<TSharedPtr<FMutableConstantMeshElement>>
+{
+public:
+	void Construct(const FArguments& Args, const TSharedRef<STableViewBase>& InOwnerTableView, const TSharedPtr<FMutableConstantMeshElement>& InRowElement)
+	{
+		check(InRowElement);
+		
+		const FText MeshProxyText =
+			FText::Format( LOCTEXT("MeshConstantProxyLabel", "{0}_MESH "),InRowElement->IndexOnSourceVector);
+
+		this->ChildSlot
+		[
+			SNew(STextBlock)
+			.Text(MeshProxyText)
+		];
+		
+		STableRow< TSharedPtr<FMutableConstantMeshElement> >::ConstructInternal(
+			STableRow::FArguments()
+			.ShowSelection(true)
+			, InOwnerTableView);
+
+	}
+};
+
+class SMutableConstantStringRow final : public STableRow<TSharedPtr<FMutableConstantStringElement>>
+{
+public:
+	void Construct(const FArguments& Args, const TSharedRef<STableViewBase>& InOwnerTableView, const TSharedPtr<FMutableConstantStringElement>& InRowElement)
+	{
+		check(InRowElement);
+		
+		// Generate the text to be displayed taking in mind the string value held by the constant to be able to "preview"
+		// it for easier navigation
+		const FString MainString = FString::FromInt(InRowElement->IndexOnSourceVector) + FString(TEXT("_STR "));
+		const FString ConstantStringText = FString(InRowElement->MutableString->c_str());
+		const FString GlimpseConstantText = ConstantStringText.Left(GlimpseCharacterCount);
+		
+		// Compose the FStrings to produce the UI text to be displayed
+		FString UiString = MainString + "\"" + GlimpseConstantText;
+		if (ConstantStringText.Len() > GlimpseConstantText.Len())
+		{
+			// Shortening occured, adding "..." to show it to the user
+			UiString.Append("...");
+		}
+		UiString.Append("\"");
+		
+		this->ChildSlot
+		[
+			SNew(STextBlock)
+			.Text(FText::FromString(UiString))
+		];
+		
+		STableRow< TSharedPtr<FMutableConstantStringElement> >::ConstructInternal(
+			STableRow::FArguments()
+			.ShowSelection(true)
+			, InOwnerTableView);
+	}
+
+
+private:
+
+	/** Determines the amount of characters to be displayed on the string constant UI text as a preview of the value
+	 * of the actual string constant.
+	 */
+	const uint32 GlimpseCharacterCount = 8;
+};
+
+namespace ImageConstantTitles
+{
+	static const FName ImageID("Id");
+	static const FName ImageSize("Resolution");
+	static const FName ImageMipMaps("MipMapCount");
+	static const FName ImageFormat("Format");
+	static const FName ImageTotalMemory("MemorySize");
+}
+
+class SMutableConstantImageRow final : public SMultiColumnTableRow<TSharedPtr<FMutableConstantImageElement>>
+{
+public:
+	void Construct(const FArguments& Args, const TSharedRef<STableViewBase>& InOwnerTableView, const TSharedPtr<FMutableConstantImageElement>& InRowElement)
+	{
+		check(InRowElement);
+		RowElement = InRowElement;
+		
+		SMultiColumnTableRow< TSharedPtr<FMutableConstantImageElement> >::Construct(
+	STableRow::FArguments()
+			.ShowSelection(true)
+			, InOwnerTableView
+		);
+	}
+
+	virtual TSharedRef<SWidget> GenerateWidgetForColumn(const FName& ColumnName) override
+	{
+		check(RowElement);
+
+		const mu::Image* Image = RowElement->ImagePtr.get();
+		
+		// Index
+		if (ColumnName == ImageConstantTitles::ImageID)
+		{
+			const FText IndexAsText = FText::AsNumber(RowElement->IndexOnSourceVector);
+			return SNew(STextBlock).Text(IndexAsText);
+		}
+
+		// ImageSize (Resolution size)
+		if (ColumnName == ImageConstantTitles::ImageSize)
+		{
+			const FString XSizeString = FString::FromInt( Image->GetSizeX());
+			const FString YSizeString = FString::FromInt( Image->GetSizeY());
+			const FString XSign = FString("x");
+			const FText ImageResolution = FText::FromString( XSizeString + XSign + YSizeString);
+			
+			return SNew(STextBlock).Text(ImageResolution);
+		}
+
+		// Image Mip maps (LODs)
+		if (ColumnName == ImageConstantTitles::ImageMipMaps)
+		{
+			const FText LodCount = FText::AsNumber(Image->GetLODCount());
+			return SNew(STextBlock).Text(LodCount);
+		}
+
+		// Image format
+		if (ColumnName == ImageConstantTitles::ImageFormat)
+		{
+			const mu::EImageFormat ImageFormat = Image->GetFormat();
+			const uint8 ImageFormatValue = static_cast<uint8>(ImageFormat);
+			const FText FormatAsText = FText::FromString(
+				FString(mu::TypeInfo::s_imageFormatName[ImageFormatValue]));
+			
+			return SNew(STextBlock).Text(FormatAsText);
+		}
+
+		// Memory
+		if (ColumnName == ImageConstantTitles::ImageTotalMemory)
+		{
+			// Todo: Not sure if I do prefer this way of showing the sizes or the one used for the expandable area titles.
+			
+			uint32 SizeInBytes = Image->GetDataSize();
+			FString SizeUnit = " B";
+			
+			// If convertible to KB then display it as KB (Kilo Bytes)
+			if (SizeInBytes >= 1024 && SizeInBytes % 1024 == 0)
+			{
+				SizeUnit = " KB";
+				SizeInBytes /= 1024;
+			
+				// If convertible to MB then display it as MB (Mega Bytes)
+				if (SizeInBytes >= 1024 && SizeInBytes % 1024 == 0)
+				{
+					SizeUnit = " MB";
+					SizeInBytes /= 1024;
+				
+					// If convertible to GB then display it as GB (Giga Bytes)
+					if (SizeInBytes >= 1024 && SizeInBytes % 1024 == 0)
+					{
+						SizeUnit = " GB";
+						SizeInBytes /= 1024;
+					}
+				}
+			}
+			
+			FText SizeAsText = FText::GetEmpty();
+			if (SizeInBytes > 0)
+			{
+				// Return the text formatted as number
+				SizeAsText = FText::FromString( FString::FormatAsNumber(SizeInBytes) + SizeUnit);
+			}
+
+			// Return the text with the size
+			return SNew(STextBlock).Text(SizeAsText);
+		}
+
+		checkNoEntry();
+		return SNullWidget::NullWidget;
+	}
+
+private:
+	TSharedPtr<FMutableConstantImageElement> RowElement;
+};
+
+
+class SMutableConstantLayoutRow final : public STableRow<TSharedPtr<FMutableConstantLayoutElement>>
+{
+public:
+	void Construct(const FArguments& Args, const TSharedRef<STableViewBase>& InOwnerTableView, const TSharedPtr<FMutableConstantLayoutElement>& InRowElement)
+	{
+		check(InRowElement);
+		const FText LayoutProxyText = FText::Format( LOCTEXT("LayoutConstantProxyLabel", "{0}_LAYOUT "),InRowElement->IndexOnSourceVector);
+
+		this->ChildSlot
+		[
+			SNew(STextBlock)
+			.Text(LayoutProxyText)
+		];
+
+		STableRow< TSharedPtr<FMutableConstantLayoutElement> >::ConstructInternal(
+			STableRow::FArguments()
+			.ShowSelection(true)
+			, InOwnerTableView);
+	}
+	
+};
+
+class SMutableConstantProjectorRow final : public STableRow<TSharedPtr<FMutableConstantProjectorElement>>
+{
+public:
+	void Construct(const FArguments& Args, const TSharedRef<STableViewBase>& InOwnerTableView, const TSharedPtr<FMutableConstantProjectorElement>& InRowElement)
+	{
+		check(InRowElement);
+		const FText ProjectorProxyText = FText::Format( LOCTEXT("ProjectorConstantProxyLabel", "{0}_PROJECTOR "),InRowElement->IndexOnSourceVector);
+		
+		this->ChildSlot
+		[
+			SNew(STextBlock)
+			.Text(ProjectorProxyText)
+		];
+
+		STableRow< TSharedPtr<FMutableConstantProjectorElement> >::ConstructInternal(
+			STableRow::FArguments()
+			.ShowSelection(true)
+			, InOwnerTableView);
+	}
+	
+};
+
+class SMutableConstantMatrixRow final : public STableRow<TSharedPtr<FMutableConstantMatrixElement>>
+{
+public:
+	void Construct(const FArguments& Args, const TSharedRef<STableViewBase>& InOwnerTableView, const TSharedPtr<FMutableConstantMatrixElement>& InRowElement)
+	{
+		check(InRowElement);
+		const FText MatrixProxyText = FText::Format( LOCTEXT("MatrixConstantProxyLabel", "{0}_MATRIX "),InRowElement->IndexOnSourceVector);
+
+		this->ChildSlot
+		[
+			SNew(STextBlock)
+			.Text(MatrixProxyText)
+		];
+
+		STableRow< TSharedPtr<FMutableConstantMatrixElement> >::ConstructInternal(
+			STableRow::FArguments()
+			.ShowSelection(true)
+			, InOwnerTableView);
+	}
+};
+
+class SMutableConstantShapeRow final : public STableRow<TSharedPtr<FMutableConstantShapeElement>>
+{
+public:
+	void Construct(const FArguments& Args, const TSharedRef<STableViewBase>& InOwnerTableView, const TSharedPtr<FMutableConstantShapeElement>& InRowElement)
+	{
+		check(InRowElement);
+		const FText ShapeProxyText = FText::Format( LOCTEXT("ShapeConstantProxyLabel", "{0}_SHAPE "),InRowElement->IndexOnSourceVector);
+
+		this->ChildSlot
+		[
+			SNew(STextBlock)
+			.Text(ShapeProxyText)
+		];
+
+		STableRow< TSharedPtr<FMutableConstantShapeElement> >::ConstructInternal(
+			STableRow::FArguments()
+			.ShowSelection(true)
+			, InOwnerTableView);
+	}
+};
+
+class SMutableConstantCurveRow final : public STableRow<TSharedPtr<FMutableConstantCurveElement>>
+{
+public:
+	void Construct(const FArguments& Args, const TSharedRef<STableViewBase>& InOwnerTableView, const TSharedPtr<FMutableConstantCurveElement>& InRowElement)
+	{
+		check(InRowElement);
+		const FText CurveProxyText = FText::Format( LOCTEXT("CurveConstantProxyLabel", "{0}_CURVE "),InRowElement->IndexOnSourceVector);
+
+		this->ChildSlot
+		[
+			SNew(STextBlock)
+			.Text(CurveProxyText)
+		];
+
+		STableRow< TSharedPtr<FMutableConstantCurveElement> >::ConstructInternal(
+			STableRow::FArguments()
+			.ShowSelection(true)
+			, InOwnerTableView);
+	}
+};
+
+
+class SMutableConstantSkeletonRow final : public STableRow<TSharedPtr<FMutableConstantSkeletonElement>>
+{
+public:
+	void Construct(const FArguments& Args, const TSharedRef<STableViewBase>& InOwnerTableView, const TSharedPtr<FMutableConstantSkeletonElement>& InRowElement)
+	{
+		check(InRowElement);
+		const FText SkeletonProxyText = FText::Format( LOCTEXT("SkeletonConstantProxyLabel", "{0}_SKELETON "),InRowElement->IndexOnSourceVector);
+
+		this->ChildSlot
+		[
+			SNew(STextBlock)
+			.Text(SkeletonProxyText)
+		];
+
+		STableRow< TSharedPtr<FMutableConstantSkeletonElement> >::ConstructInternal(
+			STableRow::FArguments()
+			.ShowSelection(true)
+			, InOwnerTableView);
+	}
+};
+
+#pragma endregion SUPPORT_CLASSES
+
+
+
+void SMutableConstantsWidget::Construct(const FArguments& InArgs,const mu::PROGRAM* InMutableProgramPtr,  TSharedPtr<SMutableCodeViewer> InMutableCodeViewerPtr)
+{
+	// A pointer to MutableCodeViewerPtr is required in order to be able to invoke the preview of our constants
+	check (InMutableCodeViewerPtr);
+	MutableCodeViewerPtr = InMutableCodeViewerPtr;
+	
+	// A pointer to the mutable program object is required to get the constants data
+	check(InMutableProgramPtr);
+	SetProgram(InMutableProgramPtr);
+	
+	// Formatting constants
+	constexpr uint32 InBetweenListsVerticalPadding = 4;
+	
+	// Vertical size for each entry 
+	constexpr float ProxyEntryHeight = 20;
+	
+	// Panel title
+	const FText ConstantsPanelTitle = LOCTEXT("ConstantsPannelName", "Constants : ");
+	
+	// Child structure
+	this->ChildSlot
+	[
+		SNew(SScrollBox)
+		.Orientation(EOrientation::Orient_Vertical)
+		
+		// String constants
+		+ SScrollBox::Slot()
+		.Padding(0,InBetweenListsVerticalPadding)
+		[
+			SAssignNew(StringsExpandableArea,SExpandableArea)
+			.OnAreaExpansionChanged(this,&SMutableConstantsWidget::OnStringsRegionExpansionChanged)
+			.InitiallyCollapsed(true)
+			.AreaTitle(this,&SMutableConstantsWidget::OnDrawStringsAreaTitle)
+			.BodyContent()
+			[
+				SNew(STileView<TSharedPtr<FMutableConstantStringElement>>)
+				.OnSelectionChanged(this,&SMutableConstantsWidget::OnSelectedStringChanged)
+				.ListItemsSource(&ConstantStringElements)
+				.ItemHeight(ProxyEntryHeight)
+				.OnGenerateTile(this, &SMutableConstantsWidget::OnGenerateStringRow)
+				.SelectionMode(ESelectionMode::Single)
+			]
+		]
+
+		+ SScrollBox::Slot()
+		[
+			SNew(SSeparator)
+		]
+
+		// Image constants
+		+ SScrollBox::Slot()
+		.Padding(0,InBetweenListsVerticalPadding)
+		[
+			SAssignNew(ImagesExpandableArea,SExpandableArea)
+			.OnAreaExpansionChanged(this,&SMutableConstantsWidget::OnImagesRegionExpansionChanged)
+			.InitiallyCollapsed(true)
+			.AreaTitle(this,&SMutableConstantsWidget::OnDrawImagesAreaTitle)
+			.BodyContent()
+			[
+				SNew(SVerticalBox)
+
+				// Buttons being used to change the segment to be displayed
+				+SVerticalBox::Slot()
+				.AutoHeight()
+				[
+					GenerateImagesSegmentSelectionWidget()
+				]
+
+				// List view showing the selected segment of images
+				+ SVerticalBox::Slot()
+				[
+					SAssignNew(ConstantImagesListView,SListView<TSharedPtr<FMutableConstantImageElement>>)
+					.OnGenerateRow(this,&SMutableConstantsWidget::OnGenerateImageRow)
+					.ListItemsSource(&ConstantImageElementsSegment)
+					.OnSelectionChanged(this,&SMutableConstantsWidget::OnSelectedImageChanged)
+					.SelectionMode(ESelectionMode::Single)
+					.HeaderRow
+					(
+						SNew(SHeaderRow)
+						+ SHeaderRow::Column(ImageConstantTitles::ImageID)
+						.DefaultLabel(FText(LOCTEXT("ImageId", "ID")))
+						.OnSort(this, &SMutableConstantsWidget::OnImageTableSortRequested)
+						.FillWidth(0.28f)
+					
+						+ SHeaderRow::Column(ImageConstantTitles::ImageSize)
+						.DefaultLabel(FText(LOCTEXT("ImageResolution","Resolution")))
+						.OnSort(this, &SMutableConstantsWidget::OnImageTableSortRequested)
+						
+						+ SHeaderRow::Column(ImageConstantTitles::ImageMipMaps)
+						.DefaultLabel(FText(LOCTEXT("ImageMipMaps","Mip Maps")))
+						.OnSort(this, &SMutableConstantsWidget::OnImageTableSortRequested)
+						
+						+ SHeaderRow::Column(ImageConstantTitles::ImageFormat)
+						.DefaultLabel(FText(LOCTEXT("ImageFormat","Format")))
+						.OnSort(this, &SMutableConstantsWidget::OnImageTableSortRequested)
+						
+						+ SHeaderRow::Column(ImageConstantTitles::ImageTotalMemory)
+						.DefaultLabel(FText(LOCTEXT("ImageMemorySize","Size")))
+						.OnSort(this, &SMutableConstantsWidget::OnImageTableSortRequested)
+					)
+				]
+			]
+		]
+
+		+ SScrollBox::Slot()
+		[
+			SNew(SSeparator)
+		]
+		
+		// Mesh Constants
+		+ SScrollBox::Slot()
+		.Padding(0,InBetweenListsVerticalPadding)
+		[
+			SAssignNew(MeshesExpandableArea,SExpandableArea)
+			.OnAreaExpansionChanged(this,&SMutableConstantsWidget::OnMeshesRegionExpansionChanged)
+			.InitiallyCollapsed(true)
+			.AreaTitle(this,&SMutableConstantsWidget::OnDrawMeshesAreaTitle)
+			.BodyContent()
+			[
+				SNew(STileView<TSharedPtr<FMutableConstantMeshElement>>)
+				.OnSelectionChanged(this,&SMutableConstantsWidget::OnSelectedMeshChanged)
+				.ListItemsSource(&ConstantMeshElements)
+				.ItemHeight(ProxyEntryHeight)
+				.OnGenerateTile(this,&SMutableConstantsWidget::OnGenerateMeshRow)
+				.SelectionMode(ESelectionMode::Single)
+			]
+		]
+
+		+ SScrollBox::Slot()
+		[
+			SNew(SSeparator)
+		]
+		
+		// Layout Constants		
+		+ SScrollBox::Slot()
+		.Padding(0,InBetweenListsVerticalPadding)
+		[
+			SAssignNew(LayoutsExpandableArea,SExpandableArea)
+			.OnAreaExpansionChanged(this,&SMutableConstantsWidget::OnLayoutsRegionExpansionChanged)
+			.InitiallyCollapsed(true)
+			.AreaTitle(this,&SMutableConstantsWidget::OnDrawLayoutsAreaTitle)
+			.BodyContent()
+			[
+				SNew(STileView<TSharedPtr<FMutableConstantLayoutElement>>)
+				.OnSelectionChanged(this,&SMutableConstantsWidget::OnSelectedLayoutChanged)
+				.ListItemsSource(&ConstantLayoutElements)
+				.ItemHeight(ProxyEntryHeight)
+				.OnGenerateTile(this,&SMutableConstantsWidget::OnGenerateLayoutRow)
+				.SelectionMode(ESelectionMode::Single)
+			]
+		]
+
+		+ SScrollBox::Slot()
+		[
+			SNew(SSeparator)
+		]
+		
+		// Projector Constants
+		+ SScrollBox::Slot()
+		.Padding(0,InBetweenListsVerticalPadding)
+		[
+			SAssignNew(ProjectorsExpandableArea,SExpandableArea)
+			.OnAreaExpansionChanged(this,&SMutableConstantsWidget::OnProjectorsRegionExpansionChanged)
+			.InitiallyCollapsed(true)
+			.AreaTitle(this,&SMutableConstantsWidget::OnDrawProjectorsAreaTitle)
+			.BodyContent()
+			[
+				SNew(STileView<TSharedPtr<FMutableConstantProjectorElement>>)
+				.OnSelectionChanged(this,&SMutableConstantsWidget::OnSelectedProjectorChanged)
+				.ListItemsSource(&ConstantProjectorElements)
+				.ItemHeight(ProxyEntryHeight)
+				.OnGenerateTile(this,&SMutableConstantsWidget::OnGenerateProjectorRow)
+				.SelectionMode(ESelectionMode::Single)
+			]
+		]
+
+		+ SScrollBox::Slot()
+		[
+			SNew(SSeparator)
+		]
+		
+		// Matrix Constants
+		+ SScrollBox::Slot()
+		.Padding(0,InBetweenListsVerticalPadding)
+		[
+			SAssignNew(MatricesExpandableArea,SExpandableArea)
+			.OnAreaExpansionChanged(this,&SMutableConstantsWidget::OnMatricesRegionExpansionChanged)
+			.InitiallyCollapsed(true)
+			.AreaTitle(this,&SMutableConstantsWidget::OnDrawMatricesAreaTitle)
+			.BodyContent()
+			[
+				SNew(STileView<TSharedPtr<FMutableConstantMatrixElement>>)
+				.OnSelectionChanged(this,&SMutableConstantsWidget::OnSelectedMatrixChanged)
+				.ListItemsSource(&ConstantMatrixElements)
+				.ItemHeight(ProxyEntryHeight)
+				.OnGenerateTile(this,&SMutableConstantsWidget::OnGenerateMatrixRow)
+				.SelectionMode(ESelectionMode::Single)
+			]
+		]
+
+		+ SScrollBox::Slot()
+		[
+			SNew(SSeparator)
+		]
+		
+		// Shape Constants
+		+ SScrollBox::Slot()
+		.Padding(0,InBetweenListsVerticalPadding)
+		[
+			SAssignNew(ShapesExpandableArea,SExpandableArea)
+			.OnAreaExpansionChanged(this,&SMutableConstantsWidget::OnShapesRegionExpansionChanged)
+			.InitiallyCollapsed(true)
+			.AreaTitle(this,&SMutableConstantsWidget::OnDrawShapesAreaTitle)
+			.BodyContent()
+			[
+				SNew(STileView<TSharedPtr<FMutableConstantShapeElement>>)
+				.OnSelectionChanged(this,&SMutableConstantsWidget::OnSelectedShapeChanged)
+				.ListItemsSource(&ConstantShapeElements)
+				.ItemHeight(ProxyEntryHeight)
+				.OnGenerateTile(this,&SMutableConstantsWidget::OnGenerateShapeRow)
+				.SelectionMode(ESelectionMode::Single)
+			]
+		]
+
+		+ SScrollBox::Slot()
+		[
+			SNew(SSeparator)
+		]
+		
+		// Curve Constants
+		+ SScrollBox::Slot()
+		.Padding(0,InBetweenListsVerticalPadding)
+		[
+			SAssignNew(CurvesExpandableArea,SExpandableArea)
+			.OnAreaExpansionChanged(this,&SMutableConstantsWidget::OnCurvesRegionExpansionChanged)
+			.InitiallyCollapsed(true)
+			.AreaTitle(this,&SMutableConstantsWidget::OnDrawCurvesAreaTitle)
+			.BodyContent()
+			[
+				SNew(STileView<TSharedPtr<FMutableConstantCurveElement>>)
+				.OnSelectionChanged(this,&SMutableConstantsWidget::OnSelectedCurveChanged)
+				.ListItemsSource(&ConstantCurveElements)
+				.ItemHeight(ProxyEntryHeight)
+				.OnGenerateTile(this,&SMutableConstantsWidget::OnGenerateCurveRow)
+				.SelectionMode(ESelectionMode::Single)
+			]
+		]
+
+		+ SScrollBox::Slot()
+		[
+			SNew(SSeparator)
+		]
+		
+		// Skeleton Constants
+		+ SScrollBox::Slot()
+		.Padding(0,InBetweenListsVerticalPadding)
+		[
+			SAssignNew(SkeletonsExpandableArea,SExpandableArea)
+			.OnAreaExpansionChanged(this,&SMutableConstantsWidget::OnSkeletonsRegionExpansionChanged)
+			.InitiallyCollapsed(true)
+			.AreaTitle(this,&SMutableConstantsWidget::OnDrawSkeletonsAreaTitle)
+			.BodyContent()
+			[
+				SNew(STileView<TSharedPtr<FMutableConstantSkeletonElement>>)
+				.OnSelectionChanged(this,&SMutableConstantsWidget::OnSelectedSkeletonChanged)
+				.ListItemsSource(&ConstantSkeletonElements)
+				.ItemHeight(ProxyEntryHeight)
+				.OnGenerateTile(this,&SMutableConstantsWidget::OnGenerateSkeletonRow)
+				.SelectionMode(ESelectionMode::Single)
+			]
+		]
+	];
+
+	// Store all the expandable areas so they are later reachable using loops
+	ExpandableAreas.Add(StringsExpandableArea);
+	ExpandableAreas.Add(ImagesExpandableArea);
+	ExpandableAreas.Add(MeshesExpandableArea);
+	ExpandableAreas.Add(LayoutsExpandableArea);
+	ExpandableAreas.Add(ProjectorsExpandableArea);
+	ExpandableAreas.Add(ShapesExpandableArea);
+	ExpandableAreas.Add(CurvesExpandableArea);
+	ExpandableAreas.Add(MatricesExpandableArea);
+	ExpandableAreas.Add(SkeletonsExpandableArea);
+}
+
+
+
+void SMutableConstantsWidget::SetProgram(const mu::PROGRAM* InProgram)
+{
+	// Make sure we do not process the setting of the same program object as the one already set
+	if (this->MutableProgramPtr == InProgram)
+	{
+		return;
+	}
+	
+	// Set only once the program that is being used. No further updates should be required
+	this->MutableProgramPtr = InProgram;
+	
+	// Generate the backend for the lists used in this object
+	if (MutableProgramPtr != nullptr)
+	{
+		LoadConstantElements();
+	}
+}
+
+#pragma region Row Generation
+
+TSharedRef<ITableRow> SMutableConstantsWidget::OnGenerateStringRow(
+	TSharedPtr<FMutableConstantStringElement> MutableConstantStringElement, const TSharedRef<STableViewBase>& OwnerTable) const
+{
+	TSharedRef<SMutableConstantStringRow> Row = SNew(SMutableConstantStringRow, OwnerTable, MutableConstantStringElement);
+	return Row;
+}
+
+TSharedRef<ITableRow> SMutableConstantsWidget::OnGenerateImageRow(
+	TSharedPtr<FMutableConstantImageElement> MutableConstantImageElement,
+	const TSharedRef<STableViewBase>& OwnerTable) const
+{
+	TSharedRef<SMutableConstantImageRow> Row = SNew(SMutableConstantImageRow,OwnerTable,MutableConstantImageElement);
+	return Row;
+}
+
+TSharedRef<ITableRow> SMutableConstantsWidget::OnGenerateMeshRow(
+	TSharedPtr<FMutableConstantMeshElement> MutableConstantMeshElement, const TSharedRef<STableViewBase>& OwnerTable) const
+{
+	TSharedRef<SMutableConstantMeshRow> Row = SNew(SMutableConstantMeshRow, OwnerTable, MutableConstantMeshElement);
+	return Row;
+}
+
+TSharedRef<ITableRow> SMutableConstantsWidget::OnGenerateLayoutRow(
+	TSharedPtr<FMutableConstantLayoutElement> MutableConstantLayoutElement,
+	const TSharedRef<STableViewBase>& OwnerTable) const
+{
+	TSharedRef<SMutableConstantLayoutRow> Row = SNew(SMutableConstantLayoutRow,OwnerTable,MutableConstantLayoutElement);
+	return Row;
+}
+
+TSharedRef<ITableRow> SMutableConstantsWidget::OnGenerateProjectorRow(
+	TSharedPtr<FMutableConstantProjectorElement> MutableConstantProjectorElement,
+	const TSharedRef<STableViewBase>& OwnerTable) const
+{
+	TSharedRef<SMutableConstantProjectorRow> Row = SNew(SMutableConstantProjectorRow,OwnerTable,MutableConstantProjectorElement);
+	return Row;
+}
+
+TSharedRef<ITableRow> SMutableConstantsWidget::OnGenerateMatrixRow(
+	TSharedPtr<FMutableConstantMatrixElement> MutableConstantMatrixElement,
+	const TSharedRef<STableViewBase>& OwnerTable) const
+{
+	TSharedRef<SMutableConstantMatrixRow> Row = SNew(SMutableConstantMatrixRow,OwnerTable,MutableConstantMatrixElement);
+	return Row;
+}
+
+TSharedRef<ITableRow> SMutableConstantsWidget::OnGenerateShapeRow(
+	TSharedPtr<FMutableConstantShapeElement> MutableConstantShapeElement,
+	const TSharedRef<STableViewBase>& OwnerTable) const
+{
+	TSharedRef<SMutableConstantShapeRow> Row = SNew(SMutableConstantShapeRow,OwnerTable,MutableConstantShapeElement);
+	return Row;
+}
+
+TSharedRef<ITableRow> SMutableConstantsWidget::OnGenerateCurveRow(
+	TSharedPtr<FMutableConstantCurveElement> MutableConstantCurveElement,
+	const TSharedRef<STableViewBase>& OwnerTable) const
+{
+	TSharedRef<SMutableConstantCurveRow> Row = SNew(SMutableConstantCurveRow,OwnerTable,MutableConstantCurveElement);
+	return Row;
+}
+
+TSharedRef<ITableRow> SMutableConstantsWidget::OnGenerateSkeletonRow(
+	TSharedPtr<FMutableConstantSkeletonElement> MutableConstantSkeletonElement,
+	const TSharedRef<STableViewBase>& OwnerTable) const
+{
+	TSharedRef<SMutableConstantSkeletonRow> Row = SNew(SMutableConstantSkeletonRow,OwnerTable,MutableConstantSkeletonElement);
+	return Row;
+}
+
+bool SMutableConstantsWidget::ShouldBackButtonBeEnabled() const
+{
+	return  CurrentArraySegment > 0;
+}
+
+bool SMutableConstantsWidget::ShouldNextButtonBeEnabled() const
+{
+	return TotalAmountOfSegments >= 1 && CurrentArraySegment < TotalAmountOfSegments - 1;
+}
+
+void SMutableConstantsWidget::RegenerateProxyImageArray()
+{
+	// Prepare the segment of elements to be filled up
+	ConstantImageElementsSegment.SetNum(0);
+	ConstantImageElementsSegment.Reserve(ElementsPerSegment);
+	
+	const uint32 StartingIndex = CurrentArraySegment * ElementsPerSegment;
+	const uint32 FinishIndex =
+		StartingIndex + FMath::Min(ElementsPerSegment, (static_cast<uint32>(ConstantImageElements.Num()) - StartingIndex));
+	// Iterate over the main array from the point where the segment starts to the point where it ends. Stop before if needed
+	for	(uint32 MainArrayIndex = StartingIndex ; MainArrayIndex < FinishIndex; ++ MainArrayIndex)
+	{
+		ConstantImageElementsSegment.Add(ConstantImageElements[MainArrayIndex]);
+	}
+	
+	if (ConstantImagesListView.IsValid())
+	{
+		ConstantImagesListView->RequestListRefresh();
+	}
+}
+
+TSharedRef<SWidget> SMutableConstantsWidget::GenerateImagesSegmentSelectionWidget()
+{
+	// If there are segments to navigate then generate the navigation area
+	if (TotalAmountOfSegments > 1)
+	{
+		TSharedRef<SHorizontalBox> Container =	SNew(SHorizontalBox)
+
+		+ SHorizontalBox::Slot()
+		.HAlign(EHorizontalAlignment::HAlign_Left)
+		.AutoWidth()
+		[
+			SNew(SButton)
+			.OnClicked(this,&SMutableConstantsWidget::OnImageFullBackButtonClicked)
+			.Text(FText(INVTEXT("|<")))
+			.IsEnabled(this,&SMutableConstantsWidget::ShouldBackButtonBeEnabled)
+		]
+		
+		+ SHorizontalBox::Slot()
+		.HAlign(EHorizontalAlignment::HAlign_Left)
+		.AutoWidth()
+		[
+			SNew(SButton)
+			.OnClicked(this,&SMutableConstantsWidget::OnImageBackButtonClicked)
+			.Text(FText(INVTEXT("<")))
+			.IsEnabled(this,&SMutableConstantsWidget::ShouldBackButtonBeEnabled)
+		]
+
+		+ SHorizontalBox::Slot()
+		.HAlign(EHorizontalAlignment::HAlign_Center)
+		[
+			SNew(STextBlock)
+			.Justification(ETextJustify::Center)
+			.Text(this,&SMutableConstantsWidget::OnDrawCurrentImageSegmentText)
+		]
+			
+		+ SHorizontalBox::Slot()
+		.HAlign(EHorizontalAlignment::HAlign_Right)
+		.AutoWidth()
+		[
+			SNew(SButton)
+			.OnClicked(this,&SMutableConstantsWidget::OnImageForwardButtonClicked)
+			.Text(FText(INVTEXT(">")))
+			.IsEnabled(this,&SMutableConstantsWidget::ShouldNextButtonBeEnabled)
+		]
+
+		+ SHorizontalBox::Slot()
+		.HAlign(EHorizontalAlignment::HAlign_Right)
+		.AutoWidth()
+		[
+			SNew(SButton)
+			.OnClicked(this,&SMutableConstantsWidget::OnImageFullForwardButtonClicked)
+			.Text(FText(INVTEXT(">|")))
+			.IsEnabled(this,&SMutableConstantsWidget::ShouldNextButtonBeEnabled)
+		];
+
+		return Container;
+	}
+
+	// If no more than 1 section is required
+	return SNullWidget::NullWidget;
+}
+
+FText SMutableConstantsWidget::OnDrawCurrentImageSegmentText() const
+{
+	FString BuiltText = FString::FromInt(  CurrentArraySegment + 1);
+	BuiltText += FString(" / ");
+	BuiltText += FString::FromInt(TotalAmountOfSegments);
+	
+	return FText::Format(LOCTEXT("PageNumber","Page : {0}"), FText::FromString(BuiltText));
+}
+
+void SMutableConstantsWidget::OnImageTableSortRequested(EColumnSortPriority::Type ColumnSortPriority, const FName& ColumnID,
+                                                        EColumnSortMode::Type ColumnSortMode)
+{
+	// If the colum has been sorted on one way now do the inverse way of sorting
+	if (LastSortedColumnID == ColumnID)
+	{
+		bSortAscending = !bSortAscending;
+	}
+
+	Algo::StableSort(ConstantImageElements, [&](const TSharedPtr<FMutableConstantImageElement>& A, const TSharedPtr<FMutableConstantImageElement>& B)
+	{
+		// Sort by image id
+		if (ColumnID == ImageConstantTitles::ImageID)
+		{
+			if (bSortAscending)
+			{
+				return A->IndexOnSourceVector < B->IndexOnSourceVector;
+			}
+			return A->IndexOnSourceVector > B->IndexOnSourceVector;
+		}
+		
+		// Sort Image mip maps
+		if (ColumnID == ImageConstantTitles::ImageMipMaps)
+		{
+			if (bSortAscending)
+			{
+				return A->ImagePtr->GetLODCount() < B->ImagePtr->GetLODCount();
+			}
+			return A->ImagePtr->GetLODCount() > B->ImagePtr->GetLODCount();
+		}
+
+		// Sort by image format
+		if (ColumnID == ImageConstantTitles::ImageFormat)
+		{
+			const uint8 AImageFormatValue = static_cast<uint8>(A->ImagePtr->GetFormat());
+			const FString AImageFormatString = FString(mu::TypeInfo::s_imageFormatName[AImageFormatValue]);
+			
+			const uint8 BImageFormatValue = static_cast<uint8>(B->ImagePtr->GetFormat());
+			const FString BImageFormatString = FString(mu::TypeInfo::s_imageFormatName[BImageFormatValue]);
+			
+			if (bSortAscending)
+			{
+				return AImageFormatString.Compare(BImageFormatString) < 0;
+			}
+			return AImageFormatString.Compare(BImageFormatString) > 0;
+		}
+		
+		
+		// Sort by image size
+		if (ColumnID == ImageConstantTitles::ImageSize)
+		{
+			if (bSortAscending)
+			{
+				return A->ImagePtr->GetSizeX() *  A->ImagePtr->GetSizeY() <
+					B->ImagePtr->GetSizeX() *  B->ImagePtr->GetSizeY();
+			}
+			return A->ImagePtr->GetSizeX() *  A->ImagePtr->GetSizeY() >
+					B->ImagePtr->GetSizeX() *  B->ImagePtr->GetSizeY();
+		}
+		
+		
+		// Sort by image used memory
+		if (ColumnID == ImageConstantTitles::ImageTotalMemory)
+		{
+			if (bSortAscending)
+			{
+				return A->ImagePtr->GetDataSize() < B->ImagePtr->GetDataSize();
+			}
+			return A->ImagePtr->GetDataSize() > B->ImagePtr->GetDataSize();
+		}
+		
+		return false;
+	});
+
+	// Store the last column to be able to do the inverse sorting order the next time it gets pressed
+	LastSortedColumnID = ColumnID;
+
+	RegenerateProxyImageArray();
+	ConstantImagesListView->RequestListRefresh();
+}
+
+
+FReply SMutableConstantsWidget::OnImageBackButtonClicked()
+{
+	CurrentArraySegment--;
+	
+	// Regenerate the proxy list with the elements on the currently selected list segment
+	RegenerateProxyImageArray();
+
+	return FReply::Handled();
+}
+
+FReply SMutableConstantsWidget::OnImageForwardButtonClicked()
+{
+	CurrentArraySegment++;
+	
+	// Regenerate the proxy list with the elements on the currently selected list segment
+	RegenerateProxyImageArray();
+
+	return FReply::Handled();
+}
+
+FReply SMutableConstantsWidget::OnImageFullBackButtonClicked()
+{ 
+	CurrentArraySegment = 0;
+	
+	// Regenerate the proxy list with the elements on the currently selected list segment
+	RegenerateProxyImageArray();
+
+	return FReply::Handled();
+}
+
+FReply SMutableConstantsWidget::OnImageFullForwardButtonClicked()
+{
+	CurrentArraySegment = TotalAmountOfSegments - 1;
+	
+	// Regenerate the proxy list with the elements on the currently selected list segment
+	RegenerateProxyImageArray();
+
+	return FReply::Handled();
+}
+
+#pragma endregion 
+
+#pragma  region  Expansions Handling
+
+void SMutableConstantsWidget::OnStringsRegionExpansionChanged(bool bExpanded)
+{
+	if (bExpanded)
+	{
+		ContractExpandableAreas(StringsExpandableArea);
+	}
+}
+
+void SMutableConstantsWidget::OnImagesRegionExpansionChanged(bool bExpanded)
+{
+	if (bExpanded)
+	{
+		ContractExpandableAreas(ImagesExpandableArea);
+	}
+}
+
+void SMutableConstantsWidget::OnMeshesRegionExpansionChanged(bool bExpanded)
+{
+	if (bExpanded)
+	{
+		ContractExpandableAreas(MeshesExpandableArea);
+	}
+}
+
+void SMutableConstantsWidget::OnLayoutsRegionExpansionChanged(bool bExpanded)
+{
+	if (bExpanded)
+	{
+		ContractExpandableAreas(LayoutsExpandableArea);
+	}
+}
+
+void SMutableConstantsWidget::OnProjectorsRegionExpansionChanged(bool bExpanded)
+{
+	if (bExpanded)
+	{
+		ContractExpandableAreas(ProjectorsExpandableArea);
+	}
+}
+
+void SMutableConstantsWidget::OnMatricesRegionExpansionChanged(bool bExpanded)
+{
+	if (bExpanded)
+	{
+		ContractExpandableAreas(MatricesExpandableArea);
+	}
+}
+
+void SMutableConstantsWidget::OnShapesRegionExpansionChanged(bool bExpanded)
+{
+	if (bExpanded)
+	{
+		ContractExpandableAreas(ShapesExpandableArea);
+	}
+}
+
+void SMutableConstantsWidget::OnCurvesRegionExpansionChanged(bool bExpanded)
+{
+	if (bExpanded)
+	{
+		ContractExpandableAreas(CurvesExpandableArea);
+	}
+}
+
+void SMutableConstantsWidget::OnSkeletonsRegionExpansionChanged(bool bExpanded)
+{
+	if (bExpanded)
+	{
+		ContractExpandableAreas(SkeletonsExpandableArea);
+	}
+}
+
+void SMutableConstantsWidget::ContractExpandableAreas(const TSharedPtr<SExpandableArea>& InException)
+{
+#if UE_BUILD_DEBUG || UE_BUILD_DEVELOPMENT
+	if(!InException)
+	{
+		UE_LOG(LogTemp,Warning,TEXT("No valid expandable area has been provided as exception : All expandable areas will therefore get contracted"));
+	}
+#endif
+	
+	for (TSharedPtr<SExpandableArea>& CurrentExpandableArea : ExpandableAreas)
+	{
+		if (CurrentExpandableArea == InException)
+		{
+			continue;
+		}
+
+		CurrentExpandableArea->SetExpanded(false);
+	}
+}
+
+#pragma endregion 
+
+#pragma  region Element caches loading
+
+void SMutableConstantsWidget::LoadConstantElements()
+{
+	LoadConstantStrings();
+	LoadConstantImages();
+	LoadConstantMeshes();
+	LoadConstantLayouts();
+	LoadConstantProjectors();
+	LoadConstantMatrices();
+	LoadConstantShapes();
+	LoadConstantCurves();
+	LoadConstantSkeletons();
+}
+
+void SMutableConstantsWidget::LoadConstantStrings()
+{
+	check (MutableProgramPtr);
+
+	ConstantStringElements.Empty();
+	uint32 ConstantStringsAccumulatedSize = 0;
+	
+	for (uint32 StringAddressIndex = 0; StringAddressIndex < MutableProgramPtr->m_constantStrings.size(); StringAddressIndex++)
+	{
+		TSharedPtr<FMutableConstantStringElement> ConstantStringElement = MakeShared<FMutableConstantStringElement>();
+		ConstantStringElement->MutableString = &(MutableProgramPtr->m_constantStrings[StringAddressIndex]);
+		ConstantStringElement->IndexOnSourceVector = StringAddressIndex;
+		
+		ConstantStringElements.Add(ConstantStringElement);
+		ConstantStringsAccumulatedSize += sizeof(*ConstantStringElement->MutableString);
+	}
+
+	// Cache the size in memory of the constants as a formatted text so it is able to be be later used by the UI
+	ConstantStringsFormattedSize = GetSizeInBytesAsBiggerUnits(ConstantStringsAccumulatedSize);
+}
+
+void SMutableConstantsWidget::LoadConstantImages()
+{
+	check (MutableProgramPtr);
+	
+	ConstantImageElements.Empty();
+	uint32 ConstantImagesAccumulatedSize = 0;
+
+	for (int32 ImageIndex = 0; ImageIndex < MutableProgramPtr->m_constantImages.Num(); ImageIndex++)
+	{
+		TSharedPtr<FMutableConstantImageElement> ConstantImageElement =MakeShared<FMutableConstantImageElement>();
+		MutableProgramPtr->GetConstant(ImageIndex, ConstantImageElement->ImagePtr, 0);
+		ConstantImageElement->IndexOnSourceVector = ImageIndex;
+		
+		ConstantImageElements.Add(ConstantImageElement);
+		ConstantImagesAccumulatedSize += ConstantImageElement->ImagePtr->GetDataSize();
+	}
+
+	// Cache the size in memory of the constants as a formatted text so it is able to be be later used by the UI
+	ConstantImagesFormattedSize = GetSizeInBytesAsBiggerUnits(ConstantImagesAccumulatedSize);
+	
+	// Regenerate the array with only a part of the the full images array
+	RegenerateProxyImageArray();
+
+	// Compute the total amount of segments
+	{
+		// Compute the max amount of segments in case it needs to be updated
+		TotalAmountOfSegments = ConstantImageElements.Num() / ElementsPerSegment;
+		
+		// Add one more segment if there are elements not filling a segment
+		TotalAmountOfSegments = ConstantImageElements.Num() % ElementsPerSegment > 0 ? TotalAmountOfSegments + 1 : TotalAmountOfSegments;
+	}
+
+}
+
+void SMutableConstantsWidget::LoadConstantMeshes()
+{
+	check (MutableProgramPtr);
+	
+	ConstantMeshElements.Empty();
+	uint32 ConstantMeshesAccumulatedSize = 0;
+	
+	for (int32 MeshIndex = 0; MeshIndex < MutableProgramPtr->m_constantMeshes.Num(); MeshIndex++)
+	{
+		TSharedPtr< FMutableConstantMeshElement> ConstantMeshElement = MakeShared<FMutableConstantMeshElement>();
+		ConstantMeshElement->MeshPtr = MutableProgramPtr->m_constantMeshes[MeshIndex].Value;
+		ConstantMeshElement->IndexOnSourceVector = MeshIndex;
+		
+		ConstantMeshElements.Add(ConstantMeshElement);
+		ConstantMeshesAccumulatedSize += ConstantMeshElement->MeshPtr->GetDataSize();
+	}
+
+	// Cache the size in memory of the constants as a formatted text so it is able to be be later used by the UI
+	ConstantMeshesFormattedSize = GetSizeInBytesAsBiggerUnits(ConstantMeshesAccumulatedSize);
+}
+
+void SMutableConstantsWidget::LoadConstantLayouts()
+{
+	check (MutableProgramPtr);
+	ConstantLayoutElements.Empty();
+	uint32 ConstantLayoutsAccumulatedSize = 0;
+	
+	for (uint32 LayoutIndex = 0; LayoutIndex < MutableProgramPtr->m_constantLayouts.size(); LayoutIndex++)
+	{		
+		TSharedPtr<FMutableConstantLayoutElement> ConstantLayoutElement = MakeShared<FMutableConstantLayoutElement>();
+		ConstantLayoutElement->Layout = MutableProgramPtr->m_constantLayouts[LayoutIndex];
+		ConstantLayoutElement->IndexOnSourceVector = LayoutIndex;
+		
+		ConstantLayoutElements.Add(ConstantLayoutElement);
+		ConstantLayoutsAccumulatedSize += sizeof(*ConstantLayoutElement->Layout);
+	}
+
+	// Cache the size in memory of the constants as a formatted text so it is able to be be later used by the UI
+	ConstantLayoutsFormattedSize = GetSizeInBytesAsBiggerUnits(ConstantLayoutsAccumulatedSize);
+
+}
+
+void SMutableConstantsWidget::LoadConstantSkeletons()
+{
+	check (MutableProgramPtr);
+	ConstantSkeletonElements.Empty();
+	uint32 ConstantSkeletonsAccumulatedSize = 0;
+	
+	for (uint32 SkeletonIndex = 0; SkeletonIndex < MutableProgramPtr->m_constantSkeletons.size(); SkeletonIndex++)
+	{
+		TSharedPtr<FMutableConstantSkeletonElement> ConstantSkeletonElement = MakeShared<FMutableConstantSkeletonElement>();
+		ConstantSkeletonElement->Skeleton = MutableProgramPtr->m_constantSkeletons[SkeletonIndex];
+		ConstantSkeletonElement->IndexOnSourceVector = SkeletonIndex;
+		
+		ConstantSkeletonElements.Add(ConstantSkeletonElement);
+		ConstantSkeletonsAccumulatedSize += sizeof(*ConstantSkeletonElement->Skeleton);
+	}
+
+	// Cache the size in memory of the constants as a formatted text so it is able to be be later used by the UI
+	ConstantSkeletonsFormattedSize = GetSizeInBytesAsBiggerUnits(ConstantSkeletonsAccumulatedSize);
+}
+
+void SMutableConstantsWidget::LoadConstantProjectors()
+{
+	check (MutableProgramPtr);
+	ConstantProjectorElements.Empty();
+	uint32 ConstantProjectorsAccumulatedSize = 0;
+	
+	for (uint32 ProjectorIndex = 0; ProjectorIndex < MutableProgramPtr->m_constantProjectors.size(); ProjectorIndex++)
+	{
+		TSharedPtr<FMutableConstantProjectorElement> ConstantProjectorElement = MakeShared<FMutableConstantProjectorElement>();
+		ConstantProjectorElement->Projector = &(MutableProgramPtr->m_constantProjectors[ProjectorIndex]);
+		ConstantProjectorElement->IndexOnSourceVector = ProjectorIndex;
+		
+		ConstantProjectorElements.Add(ConstantProjectorElement);
+		ConstantProjectorsAccumulatedSize += sizeof(*ConstantProjectorElement->Projector);
+	}
+
+	// Cache the size in memory of the constants as a formatted text so it is able to be be later used by the UI
+	ConstantProjectorsFormattedSize = GetSizeInBytesAsBiggerUnits(ConstantProjectorsAccumulatedSize);
+
+}
+
+void SMutableConstantsWidget::LoadConstantMatrices()
+{
+	check (MutableProgramPtr);
+	ConstantMatrixElements.Empty();
+	uint32 ConstantMatricesAccumulatedSize = 0;
+	
+	for (uint32 MatrixIndex = 0; MatrixIndex < MutableProgramPtr->m_constantMatrices.size(); MatrixIndex++)
+	{
+		TSharedPtr<FMutableConstantMatrixElement> ConstantMatrixElement = MakeShared<FMutableConstantMatrixElement>();
+		ConstantMatrixElement->Matrix = &(MutableProgramPtr->m_constantMatrices[MatrixIndex]);
+		ConstantMatrixElement->IndexOnSourceVector = MatrixIndex;
+		
+		ConstantMatrixElements.Add(ConstantMatrixElement);
+		ConstantMatricesAccumulatedSize += sizeof(*ConstantMatrixElement->Matrix);
+	}
+
+	// Cache the size in memory of the constants as a formatted text so it is able to be be later used by the UI
+	ConstantMatricesFormattedSize = GetSizeInBytesAsBiggerUnits(ConstantMatricesAccumulatedSize);
+
+}
+
+void SMutableConstantsWidget::LoadConstantShapes()
+{
+	check (MutableProgramPtr);
+	ConstantShapeElements.Empty();
+	uint32 ConstantShapesAccumulatedSize = 0;
+	
+	for (uint32 ShapeIndex = 0; ShapeIndex < MutableProgramPtr->m_constantShapes.size(); ShapeIndex++)
+	{
+		TSharedPtr<FMutableConstantShapeElement> ConstantShapeElement = MakeShared<FMutableConstantShapeElement>();
+		ConstantShapeElement->Shape = &(MutableProgramPtr->m_constantShapes[ShapeIndex]);
+		ConstantShapeElement->IndexOnSourceVector = ShapeIndex;
+		
+		ConstantShapeElements.Add(ConstantShapeElement);
+		ConstantShapesAccumulatedSize += sizeof(*ConstantShapeElement->Shape);
+	}
+	
+	// Cache the size in memory of the constants as a formatted text so it is able to be be later used by the UI
+	ConstantShapesFormattedSize = GetSizeInBytesAsBiggerUnits(ConstantShapesAccumulatedSize);
+
+}
+
+void SMutableConstantsWidget::LoadConstantCurves()
+{
+	check (MutableProgramPtr);
+	ConstantCurveElements.Empty();
+	uint32 ConstantCurvesAccumulatedSize = 0;
+	
+	for (uint32 CurveIndex = 0; CurveIndex < MutableProgramPtr->m_constantCurves.size(); CurveIndex++)
+	{
+		TSharedPtr<FMutableConstantCurveElement> ConstantCurveElement = MakeShared<FMutableConstantCurveElement>();
+		ConstantCurveElement->Curve = &(MutableProgramPtr->m_constantCurves[CurveIndex]);
+		ConstantCurveElement->IndexOnSourceVector = CurveIndex;
+		
+		ConstantCurveElements.Add(ConstantCurveElement);
+		ConstantCurvesAccumulatedSize += sizeof(*ConstantCurveElement->Curve);
+	}
+
+	// Cache the size in memory of the constants as a formatted text so it is able to be be later used by the UI
+	ConstantCurvesFormattedSize = GetSizeInBytesAsBiggerUnits(ConstantCurvesAccumulatedSize);
+
+}
+
+#pragma  endregion
+
+#pragma region Previewer invocation methods
+
+
+void SMutableConstantsWidget::OnSelectedStringChanged(
+	TSharedPtr<FMutableConstantStringElement> MutableConstantStringElement,
+	ESelectInfo::Type SelectionType) const
+{
+	if (MutableConstantStringElement)
+	{
+		// Ask the Code viewer to present the element held on your element on the previewer window
+		MutableCodeViewerPtr->PreviewMutableString(MutableConstantStringElement->MutableString);
+
+		MutableCodeViewerPtr->CacheAddressesRelatedWithConstantResource(mu::DATATYPE::DT_STRING,MutableConstantStringElement->IndexOnSourceVector);
+	}
+}
+
+void SMutableConstantsWidget::OnSelectedImageChanged(
+	TSharedPtr<FMutableConstantImageElement> MutableConstantImageElement,
+	ESelectInfo::Type SelectionType) const
+{
+	if (MutableConstantImageElement)
+	{
+		MutableCodeViewerPtr->PreviewMutableImage(MutableConstantImageElement->ImagePtr);
+
+		MutableCodeViewerPtr->CacheAddressesRelatedWithConstantResource(mu::DATATYPE::DT_IMAGE,MutableConstantImageElement->IndexOnSourceVector);
+	}
+}
+
+void SMutableConstantsWidget::OnSelectedMeshChanged(
+	TSharedPtr<FMutableConstantMeshElement> MutableConstantMeshElement,
+	ESelectInfo::Type SelectionType) const
+{
+	if (MutableConstantMeshElement)
+	{
+		MutableCodeViewerPtr->PreviewMutableMesh(MutableConstantMeshElement->MeshPtr);
+
+		MutableCodeViewerPtr->CacheAddressesRelatedWithConstantResource(mu::DATATYPE::DT_MESH,MutableConstantMeshElement->IndexOnSourceVector);
+	}
+}
+
+void SMutableConstantsWidget::OnSelectedLayoutChanged(
+	TSharedPtr<FMutableConstantLayoutElement> MutableConstantLayoutElement,
+	ESelectInfo::Type SelectionType) const
+{
+	if (MutableConstantLayoutElement)
+	{
+		MutableCodeViewerPtr->PreviewMutableLayout(MutableConstantLayoutElement->Layout);
+
+		MutableCodeViewerPtr->CacheAddressesRelatedWithConstantResource(mu::DATATYPE::DT_LAYOUT,MutableConstantLayoutElement->IndexOnSourceVector);
+	}
+}
+
+void SMutableConstantsWidget::OnSelectedProjectorChanged(
+	TSharedPtr<FMutableConstantProjectorElement> MutableConstantProjectorElement,
+	ESelectInfo::Type SelectionType) const
+{
+	if (MutableConstantProjectorElement)
+	{
+		MutableCodeViewerPtr->PreviewMutableProjector(MutableConstantProjectorElement->Projector);
+
+		MutableCodeViewerPtr->CacheAddressesRelatedWithConstantResource(mu::DATATYPE::DT_PROJECTOR,MutableConstantProjectorElement->IndexOnSourceVector);
+	}
+}
+
+void SMutableConstantsWidget::OnSelectedMatrixChanged(
+	TSharedPtr<FMutableConstantMatrixElement> MutableConstantMatrixElement,
+	ESelectInfo::Type SelectionType) const
+{
+	if (MutableConstantMatrixElement)
+	{
+		MutableCodeViewerPtr->PreviewMutableMatrix(MutableConstantMatrixElement->Matrix);
+
+		MutableCodeViewerPtr->CacheAddressesRelatedWithConstantResource(mu::DATATYPE::DT_MATRIX,MutableConstantMatrixElement->IndexOnSourceVector);
+	}
+}
+
+void SMutableConstantsWidget::OnSelectedShapeChanged(
+	TSharedPtr<FMutableConstantShapeElement> MutableConstantShapeElement,
+	ESelectInfo::Type SelectionType) const
+{
+	if (MutableConstantShapeElement)
+	{
+		MutableCodeViewerPtr->PreviewMutableShape(MutableConstantShapeElement->Shape);
+
+		MutableCodeViewerPtr->CacheAddressesRelatedWithConstantResource(mu::DATATYPE::DT_SHAPE,MutableConstantShapeElement->IndexOnSourceVector);
+	}
+}
+
+void SMutableConstantsWidget::OnSelectedCurveChanged(
+	TSharedPtr<FMutableConstantCurveElement> MutableConstantCurveElement,
+	ESelectInfo::Type SelectionType) const
+{
+	if (MutableConstantCurveElement)
+	{
+		MutableCodeViewerPtr->PreviewMutableCurve(MutableConstantCurveElement->Curve);
+
+		MutableCodeViewerPtr->CacheAddressesRelatedWithConstantResource(mu::DATATYPE::DT_CURVE,MutableConstantCurveElement->IndexOnSourceVector);
+	}
+}
+
+void SMutableConstantsWidget::OnSelectedSkeletonChanged(
+	TSharedPtr<FMutableConstantSkeletonElement> MutableConstantSkeletonElement,
+	ESelectInfo::Type SelectionType) const
+{
+	if (MutableConstantSkeletonElement)
+	{
+		MutableCodeViewerPtr->PreviewMutableSkeleton(MutableConstantSkeletonElement->Skeleton);
+
+		MutableCodeViewerPtr->CacheAddressesRelatedWithConstantResource(mu::DATATYPE::DT_SKELETON,MutableConstantSkeletonElement->IndexOnSourceVector);
+	}
+}
+
+#pragma endregion 
+
+#pragma region Region title drawing callback methods
+
+
+FText SMutableConstantsWidget::OnDrawStringsAreaTitle() const
+{
+	return FText::Format( LOCTEXT("StringConstantsTitle", "String Constants ({0}) : {1} "),ConstantStringElements.Num(), ConstantStringsFormattedSize);
+}
+
+FText SMutableConstantsWidget::OnDrawImagesAreaTitle() const
+{
+	return FText::Format( LOCTEXT("ImageConstantsTitle", "Image Constants ({0}) : {1} "), ConstantImageElements.Num(), ConstantImagesFormattedSize);
+}
+
+FText SMutableConstantsWidget::OnDrawMeshesAreaTitle() const
+{
+	return FText::Format(LOCTEXT("MeshConstantsTitle", "Mesh Constants ({0}) : {1} "), ConstantMeshElements.Num(), ConstantMeshesFormattedSize) ;
+}
+
+FText SMutableConstantsWidget::OnDrawLayoutsAreaTitle() const
+{
+	return FText::Format(LOCTEXT("LayoutConstantsTitle", "Layout Constants ({0}) : {1} "), ConstantLayoutElements.Num(), ConstantLayoutsFormattedSize) ;
+}
+
+FText SMutableConstantsWidget::OnDrawProjectorsAreaTitle() const
+{
+	return  FText::Format( LOCTEXT("ProjectorConstantsTitle", "Projector Constants ({0}) : {1} "), ConstantProjectorElements.Num(),ConstantProjectorsFormattedSize);
+}
+
+FText SMutableConstantsWidget::OnDrawMatricesAreaTitle() const
+{
+	return FText::Format(LOCTEXT("MatrixConstantsTitle", "Matrix Constants ({0}) : {1} "), ConstantMatrixElements.Num(),ConstantMatricesFormattedSize) ;
+}
+
+FText SMutableConstantsWidget::OnDrawShapesAreaTitle() const
+{
+	return FText::Format(LOCTEXT("ShapeConstantsTitle", "Shape Constants ({0}) : {1} "),ConstantShapeElements.Num(), ConstantShapesFormattedSize) ;
+}
+
+FText SMutableConstantsWidget::OnDrawCurvesAreaTitle() const
+{
+	return FText::Format(LOCTEXT("CurveConstantsTitle", "Curve Constants ({0}) : {1} "), ConstantCurveElements.Num(), ConstantCurvesFormattedSize) ;
+}
+
+FText SMutableConstantsWidget::OnDrawSkeletonsAreaTitle() const
+{
+	return FText::Format( LOCTEXT("SkeletonConstantsTitle", "Skeleton Constants ({0}) : {1} "), ConstantSkeletonElements.Num(), ConstantSkeletonsFormattedSize);
+}
+
+#pragma endregion
+
+
+#undef LOCTEXT_NAMESPACE
