[mutable] Move Mutable plugin from the restricted folder to the experimental plugins folder.

diff --git a/Engine/Plugins/Experimental/Mutable/Source/CustomizableObjectEditor/Private/SMutableConstantsWidget.h b/Engine/Plugins/Experimental/Mutable/Source/CustomizableObjectEditor/Private/SMutableConstantsWidget.h
new file mode 100644
--- /dev/null
+++ b/Engine/Plugins/Experimental/Mutable/Source/CustomizableObjectEditor/Private/SMutableConstantsWidget.h
@@ -0,0 +1,327 @@
+// Copyright Epic Games, Inc. All Rights Reserved.
+
+#pragma once
+
+#include "MutableString.h"
+#include "Table.h"
+#include "MutableRuntime/Private/ModelPrivate.h"
+#include "MutableRuntime/Private/MutableMath.h"
+#include "MutableRuntime/Private/ParametersPrivate.h"
+#include "MutableRuntime/Private/SystemPrivate.h"
+
+#include "MutableRuntime/Public/Image.h"
+#include "MutableRuntime/Public/Layout.h"
+#include "MutableRuntime/Public/Mesh.h"
+
+#include "Templates/SharedPointer.h"
+#include "Widgets/SCompoundWidget.h"
+#include "Widgets/Layout/SExpandableArea.h"
+#include "Widgets/Views/SListView.h"
+
+
+// Forward declarations
+class SMutableCodeViewer;
+
+/**
+ * Base Structure to define the different elements used by the lists on this object
+ */
+struct FMutableConstantElement
+{
+	/** The index of this element on the host vector. */
+	uint32 IndexOnSourceVector;
+};
+
+/**
+* Cache object used for the generation of the ui elements related to the constant images found on the model
+*/
+struct FMutableConstantImageElement : public FMutableConstantElement
+{
+	mu::ImagePtrConst ImagePtr = nullptr;
+};
+
+/**
+* Cache object used for the generation of the ui elements related to the constant meshes found on the model
+*/
+struct FMutableConstantMeshElement : public FMutableConstantElement
+{
+	mu::MeshPtrConst MeshPtr = nullptr;
+};
+
+/**
+* Cache object used for the generation of the ui elements related to the constant strings found on the model
+*/
+struct FMutableConstantStringElement : public FMutableConstantElement
+{
+	const mu::string* MutableString = nullptr;
+};
+
+/**
+* Cache object used for the generation of the ui elements related to the constant layouts found on the model
+*/
+struct FMutableConstantLayoutElement : public FMutableConstantElement
+{
+	mu::LayoutPtrConst Layout = nullptr;
+};
+
+/**
+* Cache object used for the generation of the ui elements related to the constant skeletons found on the model
+*/
+struct FMutableConstantSkeletonElement : public FMutableConstantElement
+{
+	mu::SkeletonPtrConst Skeleton = nullptr;
+};
+
+/**
+* Cache object used for the generation of the ui elements related to the constant projectors found on the model
+*/
+struct FMutableConstantProjectorElement : public FMutableConstantElement
+{
+	const mu::PROJECTOR* Projector = nullptr;
+};
+
+/**
+* Cache object used for the generation of the ui elements related to the constant matrices found on the model
+*/
+struct FMutableConstantMatrixElement : public FMutableConstantElement
+{ 
+	const mu::mat4f* Matrix = nullptr;
+};
+
+/**
+* Cache object used for the generation of the ui elements related to the constant shapes found on the model
+*/
+struct FMutableConstantShapeElement : public FMutableConstantElement
+{
+	const mu::SHAPE* Shape = nullptr;
+};
+
+/**
+* Cache object used for the generation of the ui elements related to the constant curves found on the model
+*/ 
+struct FMutableConstantCurveElement : public FMutableConstantElement
+{
+	const mu::Curve* Curve = nullptr;
+};
+
+
+/**
+* Slate panel object designed to hold all the model constants
+*/
+class SMutableConstantsWidget : public SCompoundWidget
+{
+	SLATE_BEGIN_ARGS(SMutableConstantsWidget) {}
+	SLATE_END_ARGS()
+
+public:
+
+	/** Builds the widget
+	 * @param InArgs - Arguments provided when generating this slate object
+	 * @param InMutableProgramPtr - Pointer to the mu::PROGRAM object that holds the constants data.
+	 * @param  InMutableCodeViewerPtr - Pointer to the MutableCodeViewer tasked with the previewing of the constant values
+	 */
+	void Construct(const FArguments& InArgs, const mu::PROGRAM*  InMutableProgramPtr ,  TSharedPtr<SMutableCodeViewer> InMutableCodeViewerPtr);
+	
+private:
+
+	/** Mutable object containing the constants data */
+	const mu::PROGRAM* MutableProgramPtr = nullptr;
+
+	/** Slate capable of accessing the previewer object */
+	 TSharedPtr<SMutableCodeViewer> MutableCodeViewerPtr = nullptr;
+	
+	/**
+	 *Sets the back end for the operation of this widget. Each time this is done the ui backend gets updated
+	 * @param InProgram - Mutable program object holding all the constants data
+	 */
+	void SetProgram(const mu::PROGRAM* InProgram);
+	
+	/*
+	* Data backend for the lists of constants
+	*/
+	
+	TArray<TSharedPtr<FMutableConstantImageElement>> ConstantImageElements;
+	TArray<TSharedPtr<FMutableConstantMeshElement>> ConstantMeshElements;
+	TArray<TSharedPtr<FMutableConstantStringElement>> ConstantStringElements;
+	TArray<TSharedPtr<FMutableConstantLayoutElement>> ConstantLayoutElements;
+	TArray<TSharedPtr<FMutableConstantProjectorElement>> ConstantProjectorElements;
+	TArray<TSharedPtr<FMutableConstantMatrixElement>> ConstantMatrixElements;
+	TArray<TSharedPtr<FMutableConstantShapeElement>> ConstantShapeElements;
+	TArray<TSharedPtr<FMutableConstantCurveElement>> ConstantCurveElements;
+	TArray<TSharedPtr<FMutableConstantSkeletonElement>> ConstantSkeletonElements;
+
+	/**
+	 * Load up all the elements with the data found on the mu::PROGRAM object onto TArrays after parsing the data found.
+	 */
+	void LoadConstantElements();
+	
+	void LoadConstantStrings();
+	void LoadConstantMeshes();
+	void LoadConstantImages();
+	void LoadConstantLayouts();
+	void LoadConstantProjectors();
+	void LoadConstantMatrices();
+	void LoadConstantShapes();
+	void LoadConstantCurves();
+	void LoadConstantSkeletons();
+
+	/*
+	 * Proxy slates operation objects 
+	 */
+	
+	void OnSelectedStringChanged(TSharedPtr<FMutableConstantStringElement> MutableConstantStringElement, ESelectInfo::Type SelectionType) const;
+	void OnSelectedImageChanged(TSharedPtr<FMutableConstantImageElement> MutableConstantImageElement, ESelectInfo::Type SelectionType)const;
+	void OnSelectedMeshChanged(TSharedPtr<FMutableConstantMeshElement> MutableConstantMeshElement, ESelectInfo::Type SelectionType) const;
+	void OnSelectedLayoutChanged(TSharedPtr<FMutableConstantLayoutElement> MutableConstantLayoutElement, ESelectInfo::Type SelectionType) const;
+	void OnSelectedProjectorChanged(TSharedPtr<FMutableConstantProjectorElement> MutableConstantProjectorElement, ESelectInfo::Type SelectionType) const;
+	void OnSelectedMatrixChanged(TSharedPtr<FMutableConstantMatrixElement> MutableConstantMatrixElement, ESelectInfo::Type SelectionType) const;
+	void OnSelectedShapeChanged(TSharedPtr<FMutableConstantShapeElement> MutableConstantShapeElement, ESelectInfo::Type SelectionType) const;
+	void OnSelectedCurveChanged(TSharedPtr<FMutableConstantCurveElement> MutableConstantCurveElement, ESelectInfo::Type SelectionType) const;
+	void OnSelectedSkeletonChanged(TSharedPtr<FMutableConstantSkeletonElement> MutableConstantSkeletonElement, ESelectInfo::Type SelectionType) const;
+	
+	/*
+	* UI updating methods
+	*/
+	
+	TSharedRef<ITableRow> OnGenerateStringRow(TSharedPtr<FMutableConstantStringElement> MutableConstantStringElement, const TSharedRef<STableViewBase>& OwnerTable) const;
+	TSharedRef<ITableRow> OnGenerateImageRow(TSharedPtr<FMutableConstantImageElement> MutableConstantImageElement, const TSharedRef<STableViewBase>& OwnerTable) const;
+	TSharedRef<ITableRow> OnGenerateMeshRow(TSharedPtr<FMutableConstantMeshElement> MutableConstantMeshElement, const TSharedRef<STableViewBase>& OwnerTable) const;
+	TSharedRef<ITableRow> OnGenerateLayoutRow(TSharedPtr<FMutableConstantLayoutElement> MutableConstantLayoutElement, const TSharedRef<STableViewBase>& OwnerTable) const;
+	TSharedRef<ITableRow> OnGenerateProjectorRow(TSharedPtr<FMutableConstantProjectorElement> MutableConstantProjectorElement, const TSharedRef<STableViewBase>& OwnerTable)const;
+	TSharedRef<ITableRow> OnGenerateMatrixRow(TSharedPtr<FMutableConstantMatrixElement> MutableConstantMatrixElement, const TSharedRef<STableViewBase>& OwnerTable) const;
+	TSharedRef<ITableRow> OnGenerateShapeRow(TSharedPtr<FMutableConstantShapeElement> MutableConstantShapeElement, const TSharedRef<STableViewBase>& OwnerTable) const;
+	TSharedRef<ITableRow> OnGenerateCurveRow(TSharedPtr<FMutableConstantCurveElement> MutableConstantCurveElement, const TSharedRef<STableViewBase>& OwnerTable) const;
+	TSharedRef<ITableRow> OnGenerateSkeletonRow(TSharedPtr<FMutableConstantSkeletonElement> MutableConstantSkeletonElement, const TSharedRef<STableViewBase>& OwnerTable) const;
+	
+	/*
+	 * Image segment operation 
+	 */
+
+	// Part of the original array used do display only a part of the data
+	TArray<TSharedPtr<FMutableConstantImageElement>> ConstantImageElementsSegment;
+
+	/** Amount of elements an array segment can hold. The lowest the more stable the UI drawing gets
+	*(avoid crash due to oversize slate Y size) */
+	const uint32 ElementsPerSegment = 24;
+
+	// Amount of elements a segment can hold. The lowest the more stable the UI drawing gets (avoid crash due to oversize slate Y coordinate)
+	uint32 CurrentArraySegment = 0;
+	uint32 TotalAmountOfSegments = 0;
+
+	/** Provided the array on ConstantImageElements a new sub array gets generated for the currently selected segment */
+	void RegenerateProxyImageArray();
+	
+	/** Method designed to generate the UI controls to let the user change the image array segment being previewed */
+	TSharedRef<SWidget> GenerateImagesSegmentSelectionWidget();
+
+	/** Returns the text containing the page (segment) being drawn at the moment. Designed for the UI to show data to
+	 * the user
+	 */
+	FText OnDrawCurrentImageSegmentText() const;
+
+	/** Determines if the button that loads the previous image array segment should be interactable by the user
+	 * It only relates to UI, there are other internal checks to avoid the processing of an existent segment
+	 */
+	bool ShouldBackButtonBeEnabled() const;
+	/** Determines if the button that loads the next image array segment should be interactable by the user
+	 * It only relates to UI, there are other internal checks to avoid the processing of an existent segment
+	 */
+	bool ShouldNextButtonBeEnabled() const;
+
+	/** Action performed when the back button is pressed on the UI. Sets the current segment to the previous one*/
+	FReply OnImageBackButtonClicked();
+	/** Action performed when the forward button is pressed on the UI. Sets the current segment to the next one*/
+	FReply OnImageForwardButtonClicked();
+
+	/** Action performed when the full back button is pressed on the UI. Sets the current segment to the first one*/
+	FReply OnImageFullBackButtonClicked();
+	/** Action performed when the full forward button is pressed on the UI. Sets the current segment to the last one*/
+	FReply OnImageFullForwardButtonClicked();
+
+	
+	/*
+	 * Image Table sorting methods
+	 */
+
+	/** List view showing the constant images data*/
+	TSharedPtr<SListView<TSharedPtr<FMutableConstantImageElement>>> ConstantImagesListView;
+
+	/** Id of the last column the user decided to Sort. Usefully in order to interpolate ascending and descending sorting
+	 * order 
+	 */
+	FName LastSortedColumnID = "";
+	/** Variable holding what kind of sorting has been used on the last sorting operation*/
+	bool bSortAscending = false;
+
+	/** Callback method designed to sort the list of images. It sorts ConstantImageElements */
+	void OnImageTableSortRequested(EColumnSortPriority::Type ColumnSortPriority, const FName& ColumnID, EColumnSortMode::Type ColumnSortMode);
+
+	/*
+	 * Expandable areas objects and behaviours
+	 */
+
+	// Pointers to all expandable areas part of this slate 
+	TSharedPtr<SExpandableArea> StringsExpandableArea;
+	TSharedPtr<SExpandableArea> ImagesExpandableArea;
+	TSharedPtr<SExpandableArea> MeshesExpandableArea;
+	TSharedPtr<SExpandableArea> LayoutsExpandableArea;
+	TSharedPtr<SExpandableArea> ProjectorsExpandableArea;
+	TSharedPtr<SExpandableArea> MatricesExpandableArea;
+	TSharedPtr<SExpandableArea> ShapesExpandableArea;
+	TSharedPtr<SExpandableArea> CurvesExpandableArea;
+	TSharedPtr<SExpandableArea> SkeletonsExpandableArea;
+
+	/** Array with all expandable areas set on this object. Used for dynamic expansion/contraction */
+	TArray<TSharedPtr<SExpandableArea>> ExpandableAreas;
+
+	/*
+	 * Callback methods called each time one expandable area gets expanded or contracted
+	 */
+	
+	void OnStringsRegionExpansionChanged(bool bExpanded);
+	void OnImagesRegionExpansionChanged(bool bExpanded);
+	void OnMeshesRegionExpansionChanged(bool bExpanded);
+	void OnLayoutsRegionExpansionChanged(bool bExpanded);
+	void OnProjectorsRegionExpansionChanged(bool bExpanded);
+	void OnMatricesRegionExpansionChanged(bool bExpanded);
+	void OnShapesRegionExpansionChanged(bool bExpanded);
+	void OnCurvesRegionExpansionChanged(bool bExpanded);
+	void OnSkeletonsRegionExpansionChanged(bool bExpanded);
+	
+	/** Method called each time one expandable area changes expansion state
+	 * @param InException - Pointer to the expandable area that will not be contracted while all the others will
+	 */
+	void ContractExpandableAreas(const TSharedPtr<SExpandableArea>& InException);
+
+	/*
+	 * Methods used to get the size of each constant type collection on memory.  
+	 */
+
+	// Sizes decomposed on GB, MB, KB and B
+	FText ConstantStringsFormattedSize;
+	FText ConstantImagesFormattedSize;
+	FText ConstantMeshesFormattedSize;
+	FText ConstantLayoutsFormattedSize;
+	FText ConstantProjectorsFormattedSize;
+	FText ConstantMatricesFormattedSize;
+	FText ConstantShapesFormattedSize;
+	FText ConstantCurvesFormattedSize;
+	FText ConstantSkeletonsFormattedSize;
+	
+	/*
+	 * Callback methods used for drawing the titles of each of the constant expandable areas
+	 */
+	
+	FText OnDrawStringsAreaTitle() const;
+	FText OnDrawImagesAreaTitle() const;
+	FText OnDrawMeshesAreaTitle() const;
+	FText OnDrawLayoutsAreaTitle() const;
+	FText OnDrawProjectorsAreaTitle() const;
+	FText OnDrawMatricesAreaTitle() const;
+	FText OnDrawShapesAreaTitle() const;
+	FText OnDrawCurvesAreaTitle() const;
+	FText OnDrawSkeletonsAreaTitle() const;
+	
+};
+
+
+
diff --git a/Engine/Plugins/Experimental/Mutable/Source/CustomizableObjectEditor/Private/SMutableGraphViewer.cpp b/Engine/Plugins/Experimental/Mutable/Source/CustomizableObjectEditor/Private/SMutableGraphViewer.cpp
new file mode 100644
--- /dev/null
+++ b/Engine/Plugins/Experimental/Mutable/Source/CustomizableObjectEditor/Private/SMutableGraphViewer.cpp
@@ -0,0 +1,520 @@
+// Copyright Epic Games, Inc. All Rights Reserved.
+
+#include "SMutableGraphViewer.h"
+
+#include "CustomizableObjectEditorStyle.h"
+#include "CustomizableObjectCompiler.h"
+#include "SMutableCodeViewer.h"
+#include "MutableTools/Public/Streams.h"
+
+#include "Misc/App.h"
+#include "Framework/MultiBox/MultiBoxBuilder.h"
+#include "Widgets/Input/SButton.h"
+#include "Widgets/Images/SImage.h"
+#include "Widgets/Views/STreeView.h"
+#include "Widgets/Input/STextComboBox.h"
+#include "Widgets/Docking/SDockTab.h"
+#include "EditorDirectories.h"
+#include "DesktopPlatformModule.h"
+#include "HAL/PlatformApplicationMisc.h"
+
+#define LOCTEXT_NAMESPACE "SMutableDebugger"
+
+
+// \todo: multi-column tree
+namespace MutableGraphTreeViewColumns
+{
+	static const FName Name("Name");
+};
+
+
+class SMutableGraphTreeRow : public STableRow<TSharedPtr<FMutableGraphTreeElement>>
+{
+public:
+
+	void Construct(const FArguments& Args, const TSharedRef<STableViewBase>& InOwnerTableView, const TSharedPtr<FMutableGraphTreeElement>& InRowItem)
+	{
+		RowItem = InRowItem;
+
+		const char* TypeName = RowItem->MutableNode->GetType()->m_strName;
+
+		FText MainLabel = FText::FromString(ANSI_TO_TCHAR(TypeName));
+		if (RowItem->DuplicatedOf)
+		{
+			MainLabel = FText::FromString( FString::Printf(TEXT("%s (Duplicated)"), ANSI_TO_TCHAR(TypeName)));
+		}
+
+		// TODO
+		const FSlateBrush* IconBrush = nullptr; 
+
+		this->ChildSlot
+		[
+			SNew(SHorizontalBox)
+			+ SHorizontalBox::Slot()
+			.VAlign(VAlign_Center)
+			.AutoWidth()
+			[
+				SNew(SExpanderArrow, SharedThis(this))
+			]
+
+			+ SHorizontalBox::Slot()
+			.VAlign(VAlign_Center)
+			.Padding(FMargin(5.f, 0.f, 5.f, 0.f))
+			.AutoWidth()
+			[
+				SNew(SOverlay)
+				+ SOverlay::Slot()
+				[
+					SNew(SImage)
+					.Image(IconBrush ? IconBrush : FCoreStyle::Get().GetDefaultBrush())
+					.ColorAndOpacity(IconBrush ? FLinearColor::White : FLinearColor::Transparent)
+				]
+			]
+
+			+ SHorizontalBox::Slot()
+			[
+				SNew(STextBlock)
+				.Text(MainLabel)
+			]
+		];
+
+		STableRow< TSharedPtr<FMutableGraphTreeElement> >::ConstructInternal(
+			STableRow::FArguments()
+			//.Style(FAppStyle::Get(), "DetailsView.TreeView.TableRow")
+			.ShowSelection(true)
+			, InOwnerTableView
+		);
+
+	}
+
+
+private:
+
+	TSharedPtr<FMutableGraphTreeElement> RowItem;
+};
+
+
+void SMutableGraphViewer::AddReferencedObjects(FReferenceCollector& Collector)
+{
+	// Add UObjects here if we own any at some point
+	//Collector.AddReferencedObject(CustomizableObject);
+}
+
+
+FString SMutableGraphViewer::GetReferencerName() const
+{
+	return TEXT("SMutableGraphViewer");
+}
+
+
+void SMutableGraphViewer::Construct(const FArguments& InArgs, const mu::NodePtr& InRootNode, 
+	const FCompilationOptions& InCompileOptions,
+	TWeakPtr<FTabManager> InParentTabManager, const FName& InParentNewTabId)
+{
+	DataTag = InArgs._DataTag;
+	RootNode = InRootNode;
+	CompileOptions = InCompileOptions;	
+	ParentTabManager = InParentTabManager;
+	ParentNewTabId = InParentNewTabId;
+
+	FToolBarBuilder ToolbarBuilder(TSharedPtr<const FUICommandList>(), FMultiBoxCustomization::None, TSharedPtr<FExtender>(), true);
+	ToolbarBuilder.SetLabelVisibility(EVisibility::Visible);
+	ToolbarBuilder.SetStyle(&FAppStyle::Get(), "SlimToolBar");
+
+	// Export
+	ToolbarBuilder.AddToolBarButton(
+		FUIAction(
+			FExecuteAction::CreateLambda([InRootNode]()
+				{
+					TArray<FString> SaveFilenames;
+					IDesktopPlatform* DesktopPlatform = FDesktopPlatformModule::Get();
+					bool bSave = false;
+					if (!DesktopPlatform) return;
+
+					FString LastExportPath = FEditorDirectories::Get().GetLastDirectory(ELastDirectory::GENERIC_EXPORT);
+					FString FileTypes = TEXT("Mutable source data files|*.mutable_source|All files|*.*");
+					bSave = DesktopPlatform->SaveFileDialog(
+						FSlateApplication::Get().FindBestParentWindowHandleForDialogs(nullptr),
+						TEXT("Export Mutable object"),
+						*LastExportPath,
+						TEXT("exported.mutable_source"),
+						*FileTypes,
+						EFileDialogFlags::None,
+						SaveFilenames
+					);
+
+					if (!bSave) return;
+
+					// Dump source model to a file.
+					FString SaveFileName = FString(SaveFilenames[0]);
+					mu::OutputFileStream stream(TCHAR_TO_ANSI(*SaveFileName));
+					stream.Write(MUTABLE_SOURCE_MODEL_FILETAG, 4);
+					mu::OutputArchive arch(&stream);
+					mu::Node::Serialise(InRootNode.get(), arch);
+					stream.Flush();
+
+					FEditorDirectories::Get().SetLastDirectory(ELastDirectory::GENERIC_EXPORT, SaveFileName);
+				})
+			),
+			NAME_None,
+			LOCTEXT("ExportMutableGraph", "Export"),
+			LOCTEXT("ExportMutableGraphTooltip", "Export a debug mutable graph file."),
+			FSlateIcon(),
+			EUserInterfaceActionType::Button
+		);
+		
+	ToolbarBuilder.EndSection();
+
+
+	ToolbarBuilder.BeginSection("Compilation");
+
+	ToolbarBuilder.AddToolBarButton(
+		FUIAction(FExecuteAction::CreateSP(this, &SMutableGraphViewer::CompileMutableCodePressed)),
+		NAME_None,
+		LOCTEXT("GenerateMutableCode", "Unreal to Mutable Code"),
+		LOCTEXT("GenerateMutableCodeTooltip", "Generate a mutable code from the mutable graph."),
+		FSlateIcon(FCustomizableObjectEditorStyle::Get().GetStyleSetName(), "CustomizableObjectDebugger.CompileMutableCode", "CustomizableObjectDebugger.CompileMutableCode.Small"),
+		EUserInterfaceActionType::Button
+	);
+
+	ToolbarBuilder.AddComboButton(
+		FUIAction(),
+		FOnGetContent::CreateSP(this, &SMutableGraphViewer::GenerateCompileOptionsMenuContent),
+		LOCTEXT("Compile_Options_Label", "Compile Options"),
+		LOCTEXT("Compile_Options_Tooltip", "Change Compile Options"),
+		TAttribute<FSlateIcon>(),
+		true);
+
+	ToolbarBuilder.EndSection();
+
+	ToolbarBuilder.AddWidget(SNew(STextBlock).Text(MakeAttributeLambda([this]() { return FText::FromString(DataTag); })));
+
+	ChildSlot
+	[
+		SNew(SVerticalBox)
+		+ SVerticalBox::Slot()
+		.AutoHeight()
+		.VAlign(VAlign_Center)
+		[
+			ToolbarBuilder.MakeWidget()
+		]
+		+ SVerticalBox::Slot()
+		.VAlign(VAlign_Fill)
+		[
+			SNew(SSplitter)
+			.Orientation(EOrientation::Orient_Horizontal)
+			+ SSplitter::Slot()
+			.Value(0.25f)
+			[
+				SNew(SBorder)
+				.BorderImage(FAppStyle::GetBrush("ToolPanel.GroupBorder"))
+				.Padding(FMargin(4.0f, 4.0f))
+				[
+					SAssignNew(TreeView, STreeView<TSharedPtr<FMutableGraphTreeElement>>)
+					.TreeItemsSource(&RootNodes)
+					.OnGenerateRow(this,&SMutableGraphViewer::GenerateRowForNodeTree)
+					.OnGetChildren(this, &SMutableGraphViewer::GetChildrenForInfo)
+					.OnSetExpansionRecursive(this, &SMutableGraphViewer::TreeExpandRecursive)
+					.OnContextMenuOpening(this, &SMutableGraphViewer::OnTreeContextMenuOpening)
+					.SelectionMode(ESelectionMode::Single)
+					.HeaderRow
+					(
+						SNew(SHeaderRow)
+						+ SHeaderRow::Column(MutableGraphTreeViewColumns::Name)
+						.FillWidth(25.f)
+						.DefaultLabel(LOCTEXT("Node Name", "Node Name"))
+					)
+				]
+			]
+			+ SSplitter::Slot()
+			.Value(0.75f)
+			[
+				SNew(SBorder)
+				.BorderImage(FAppStyle::GetBrush("ToolPanel.GroupBorder"))
+				.Padding(FMargin(4.0f, 4.0f))
+				//[
+				//	SubjectsTreeView->AsShared()
+				//]
+			]
+		]
+	];
+	
+	RebuildTree();
+}
+
+
+void SMutableGraphViewer::CompileMutableCodePressed()
+{
+	// Do the compilation to Mutable Code synchronously.
+	TSharedPtr<FCustomizableObjectCompileRunnable> CompileTask = MakeShareable(new FCustomizableObjectCompileRunnable(RootNode, false));
+	CompileTask->Options = CompileOptions;
+	CompileTask->Init();
+	CompileTask->Run();
+
+	FString NewDataTag = FString::Printf(TEXT("%s graph, opt %d "), *DataTag, CompileOptions.OptimizationLevel);
+
+	TSharedPtr<SDockTab> NewMutableCodeTab = SNew(SDockTab)
+		.Label(LOCTEXT("MutableCode", "Mutable Code"))
+		[
+			SNew(SMutableCodeViewer, CompileTask->Model)
+			.DataTag(NewDataTag)
+		];
+
+	TSharedPtr<FTabManager> TabManager = ParentTabManager.Pin();
+	check(TabManager);
+	TabManager->InsertNewDocumentTab(ParentNewTabId, FTabManager::ESearchPreference::PreferLiveTab, NewMutableCodeTab.ToSharedRef());
+}
+
+
+TSharedRef<SWidget> SMutableGraphViewer::GenerateCompileOptionsMenuContent()
+{
+	const bool bShouldCloseWindowAfterMenuSelection = false;
+	FMenuBuilder MenuBuilder(bShouldCloseWindowAfterMenuSelection, nullptr);
+
+	// settings
+	MenuBuilder.BeginSection("Optimization", LOCTEXT("MutableCompileOptimizationHeading", "Optimization"));
+	{
+		// Compilation options
+		//-----------------------------------
+
+		// Optimisation level
+		CompileOptimizationStrings.Empty();
+		CompileOptimizationStrings.Add(MakeShareable(new FString(NSLOCTEXT("UnrealEd", "OptimizationNone", "None").ToString())));
+		CompileOptimizationStrings.Add(MakeShareable(new FString(NSLOCTEXT("UnrealEd", "OptimizationMin", "Minimal").ToString())));
+		CompileOptimizationStrings.Add(MakeShareable(new FString(NSLOCTEXT("UnrealEd", "OptimizationMed", "Medium").ToString())));
+		CompileOptimizationStrings.Add(MakeShareable(new FString(NSLOCTEXT("UnrealEd", "OptimizationMax", "Maximum").ToString())));
+
+		CompileOptions.OptimizationLevel = FMath::Min(CompileOptions.OptimizationLevel, CompileOptimizationStrings.Num() - 1);
+
+		CompileOptimizationCombo =
+			SNew(STextComboBox)
+			.OptionsSource(&CompileOptimizationStrings)
+			.InitiallySelectedItem(CompileOptimizationStrings[CompileOptions.OptimizationLevel])
+			.OnSelectionChanged(this, &SMutableGraphViewer::OnChangeCompileOptimizationLevel)
+			;
+		MenuBuilder.AddWidget(CompileOptimizationCombo.ToSharedRef(), LOCTEXT("MutableCompileOptimizationLevel", "Optimization Level"));
+
+		// Parallel compilation
+		MenuBuilder.AddMenuEntry(
+			LOCTEXT("MutableEnableParallelCompilation", "Enable compiling in multiple threads."),
+			LOCTEXT("MutableEnableParallelCompilationTooltip", "This is faster but use more memory."),
+			FSlateIcon(),
+			FUIAction(
+				FExecuteAction::CreateLambda([this]() { CompileOptions.bUseParallelCompilation = !CompileOptions.bUseParallelCompilation; }),
+				FCanExecuteAction(),
+				FIsActionChecked::CreateLambda([this]() { return CompileOptions.bUseParallelCompilation; })),
+			NAME_None,
+			EUserInterfaceActionType::ToggleButton
+		);
+
+		// Disk as cache
+		MenuBuilder.AddMenuEntry(
+			LOCTEXT("MutableEnableTextureCompression", "Enable compiling using the disk as memory."),
+			LOCTEXT("MutableEnableTextureCompressionTooltip", "This is very slow but supports compiling huge objects. It requires a lot of free space in the OS disk."),
+			FSlateIcon(),
+			FUIAction(
+				FExecuteAction::CreateLambda([this]() { CompileOptions.bUseDiskCompilation = !CompileOptions.bUseDiskCompilation; }),
+				FCanExecuteAction(),
+				FIsActionChecked::CreateLambda([this]() { return CompileOptions.bUseDiskCompilation; })),
+			NAME_None,
+			EUserInterfaceActionType::ToggleButton
+		);
+	}
+	MenuBuilder.EndSection();
+
+	return MenuBuilder.MakeWidget();
+}
+
+
+void SMutableGraphViewer::OnChangeCompileOptimizationLevel(TSharedPtr<FString> NewSelection, ESelectInfo::Type SelectInfo)
+{
+	CompileOptions.OptimizationLevel = CompileOptimizationStrings.Find(NewSelection);
+}
+
+
+void SMutableGraphViewer::RebuildTree()
+{
+	RootNodes.Reset();
+	ItemCache.Reset();
+	MainItemPerNode.Reset();
+
+	RootNodes.Add(MakeShareable(new FMutableGraphTreeElement(RootNode)));
+	TreeView->RequestTreeRefresh();
+	TreeExpandUnique();
+}
+
+
+TSharedRef<ITableRow> SMutableGraphViewer::GenerateRowForNodeTree(TSharedPtr<FMutableGraphTreeElement> InTreeNode, const TSharedRef<STableViewBase>& InOwnerTable)
+{
+	TSharedRef<SMutableGraphTreeRow> Row = SNew(SMutableGraphTreeRow, InOwnerTable, InTreeNode);
+	return Row;
+}
+
+void SMutableGraphViewer::GetChildrenForInfo(TSharedPtr<FMutableGraphTreeElement> InInfo, TArray<TSharedPtr<FMutableGraphTreeElement>>& OutChildren)
+{
+	if (!InInfo->MutableNode)
+	{
+		return;
+	}
+
+	// If this is a duplicated of another row, don't provide its children.
+	if (InInfo->DuplicatedOf)
+	{
+		return;
+	}
+
+	// Generic node case
+	int32 InputCount = InInfo->MutableNode->GetInputCount();
+	for (int32 InputIndex=0; InputIndex<InputCount; ++InputIndex)
+	{
+		mu::NodePtr ChildNode = InInfo->MutableNode->GetInputNode( InputIndex );
+		if (ChildNode)
+		{
+			FItemCacheKey Key = { InInfo->MutableNode.get(),ChildNode.get(), uint32(InputIndex) };
+			TSharedPtr<FMutableGraphTreeElement>* CachedItem = ItemCache.Find(Key);
+
+			if (CachedItem)
+			{
+				OutChildren.Add(*CachedItem);
+			}
+			else
+			{
+				TSharedPtr<FMutableGraphTreeElement>* MainItemPtr = MainItemPerNode.Find(ChildNode.get());
+				TSharedPtr<FMutableGraphTreeElement> Item = MakeShareable(new FMutableGraphTreeElement(ChildNode, MainItemPtr));
+				OutChildren.Add(Item);
+				ItemCache.Add(Key, Item);
+
+				if (!MainItemPtr)
+				{
+					MainItemPerNode.Add(ChildNode.get(),Item);
+				}
+			}
+		}
+	}
+}
+
+
+TSharedPtr<SWidget> SMutableGraphViewer::OnTreeContextMenuOpening()
+{
+	FMenuBuilder MenuBuilder(true, nullptr);
+
+	MenuBuilder.AddMenuEntry(
+		LOCTEXT("Expand_Instance", "Expand Instance-Level Operations"),
+		LOCTEXT("Expand_Instance", "Expands all the operations in the tree that are instance operations (not images, meshes, booleans, etc.)."),
+		FSlateIcon(),
+		FUIAction(FExecuteAction::CreateSP(this, &SMutableGraphViewer::TreeExpandUnique)
+			//, FCanExecuteAction::CreateSP(this, &SMutableCodeViewer::HasAnyItemInPalette)
+		)
+	);
+
+	return MenuBuilder.MakeWidget();
+}
+
+
+void SMutableGraphViewer::TreeExpandRecursive(TSharedPtr<FMutableGraphTreeElement> InInfo, bool bExpand)
+{
+	if (bExpand)
+	{
+		TreeExpandUnique();
+	}
+}
+
+
+void SMutableGraphViewer::TreeExpandUnique()
+{
+	TArray<TSharedPtr<FMutableGraphTreeElement>> Pending = RootNodes;
+
+	TSet<TSharedPtr<FMutableGraphTreeElement>> Processed;
+
+	TArray<TSharedPtr<FMutableGraphTreeElement>> Children;
+
+	while (!Pending.IsEmpty())
+	{
+		TSharedPtr<FMutableGraphTreeElement> Item = Pending.Pop();
+		TreeView->SetItemExpansion(Item, true);
+
+		Children.SetNum(0);
+		GetChildrenForInfo(Item, Children);
+		Pending.Append(Children);
+	}
+}
+
+
+
+FReply SMutableGraphViewer::OnDragOver(const FGeometry& MyGeometry, const FDragDropEvent& DragDropEvent)
+{
+	if (TSharedPtr<FExternalDragOperation> DragDropOp = DragDropEvent.GetOperationAs<FExternalDragOperation>())
+	{
+		if (DragDropOp->HasFiles())
+		{
+			// For now, only allow a single file.
+			const TArray<FString>& Files = DragDropOp->GetFiles();
+			if (Files.Num() == 1)
+			{
+				const FString DraggedFileExtension = FPaths::GetExtension(Files[0], true);
+				if (DraggedFileExtension == TEXT(".mutable_source"))
+				{
+					// Dump source model to a file.
+					mu::InputFileStream stream(TCHAR_TO_ANSI(*Files[0]));
+
+					char MutableSourceTag[4] = {};
+					stream.Read(MutableSourceTag, 4);
+
+					if (!FMemory::Memcmp(MutableSourceTag, MUTABLE_SOURCE_MODEL_FILETAG, 4))
+					{
+						return FReply::Handled();
+					}
+
+					return FReply::Unhandled();
+				}
+			}
+		}
+	}
+
+	return FReply::Unhandled();
+}
+
+
+FReply SMutableGraphViewer::OnDrop(const FGeometry& MyGeometry, const FDragDropEvent& DragDropEvent)
+{
+	SCompoundWidget::OnDrop(MyGeometry, DragDropEvent);
+
+	if (TSharedPtr<FExternalDragOperation> DragDropOp = DragDropEvent.GetOperationAs<FExternalDragOperation>())
+	{
+		if (DragDropOp->HasFiles())
+		{
+			// For now, only allow a single file.
+			const TArray<FString>& Files = DragDropOp->GetFiles();
+			if (Files.Num() == 1)
+			{
+				const FString DraggedFileExtension = FPaths::GetExtension(Files[0], true);
+				if (DraggedFileExtension == TEXT(".mutable_source"))
+				{
+					// Dump source model to a file.
+					mu::InputFileStream stream(TCHAR_TO_ANSI(*Files[0]));
+
+					char MutableSourceTag[4] = {};
+					stream.Read(MutableSourceTag, 4);
+
+					if (!FMemory::Memcmp(MutableSourceTag, MUTABLE_SOURCE_MODEL_FILETAG, 4))
+					{
+						mu::InputArchive arch(&stream);
+						RootNode = mu::Node::StaticUnserialise( arch );
+						DataTag = FString("dropped-file ") + FPaths::GetCleanFilename(Files[0]);
+						RebuildTree();
+
+						return FReply::Handled();
+					}
+
+					return FReply::Unhandled();
+				}
+			}
+		}
+
+		return FReply::Unhandled();
+	}
+
+	return FReply::Unhandled();
+}
+
+#undef LOCTEXT_NAMESPACE 
diff --git a/Engine/Plugins/Experimental/Mutable/Source/CustomizableObjectEditor/Private/SMutableGraphViewer.h b/Engine/Plugins/Experimental/Mutable/Source/CustomizableObjectEditor/Private/SMutableGraphViewer.h
new file mode 100644
--- /dev/null
+++ b/Engine/Plugins/Experimental/Mutable/Source/CustomizableObjectEditor/Private/SMutableGraphViewer.h
@@ -0,0 +1,135 @@
+// Copyright Epic Games, Inc. All Rights Reserved.
+
+#pragma once
+
+#include "Widgets/SCompoundWidget.h"
+#include "CustomizableObject.h"
+#include "MutableTools/Public/Node.h"
+#include "Widgets/Views/STreeView.h"
+#include "Widgets/Input/STextComboBox.h"
+#include "CustomizableObjectCompiler.h"
+#include "Framework/Docking/TabManager.h"
+
+
+// Forward declarations
+class FMutableGraphTreeElement;
+
+
+/** This widget shows the internal Mutable Graph for debugging purposes. 
+ * This is not the Unreal source graph in the UCustomizableObject, but an intermediate step of the compilation process.
+ */
+class SMutableGraphViewer :
+	public SCompoundWidget,
+	public FGCObject
+{
+public:
+
+	// SWidget interface
+	SLATE_BEGIN_ARGS(SMutableGraphViewer) {}
+
+	/** User-visible tag to indentify the source of the data shwon. */
+	SLATE_ARGUMENT(FString, DataTag)
+
+	SLATE_END_ARGS()
+
+	void Construct(const FArguments& InArgs, const mu::NodePtr& InRootNode, const FCompilationOptions& InCompileOptions,
+		TWeakPtr<FTabManager> InParentTabManager, const FName& InParentNewTabId);
+
+	// SWidget interface
+	FReply OnDragOver(const FGeometry& MyGeometry, const FDragDropEvent& DragDropEvent) override;
+	FReply OnDrop(const FGeometry& MyGeometry, const FDragDropEvent& DragDropEvent) override;
+
+	// FGCObject interface
+	virtual void AddReferencedObjects(FReferenceCollector& Collector) override;
+	virtual FString GetReferencerName() const override;
+
+private:
+
+	/** Tag for the user to identify the data being shown. */
+	FString DataTag;
+
+	/** The Mutable Graph to show, represented by its root. */
+	mu::NodePtr RootNode;
+
+	/** Compilation options to use in the debugger operations. */
+	FCompilationOptions CompileOptions;
+
+	/** Object compiler. */
+	FCustomizableObjectCompiler Compiler;
+
+	/** UI references used to create new tabs. */
+	TWeakPtr<FTabManager> ParentTabManager;
+	FName ParentNewTabId;
+
+	/** Root nodes of the tree widget. */
+	TArray<TSharedPtr<FMutableGraphTreeElement>> RootNodes;
+
+	/** Tree showing the graph. */
+	TSharedPtr<STreeView<TSharedPtr<FMutableGraphTreeElement>>> TreeView;
+
+	/** Cache of tree elements matching the graph nodes that have been generated so far. 
+	* We store both the parent and the node in the key, because a single node may appear multiple times if it has different parents.
+	*/
+	struct FItemCacheKey
+	{
+		const mu::Node* Parent = nullptr;
+		const mu::Node* Child = nullptr;
+		uint32 ChildIndexInParent = 0;
+
+		friend FORCEINLINE bool operator == (const FItemCacheKey& A, const FItemCacheKey& B)
+		{
+			return A.Parent == B.Parent && A.Child == B.Child && A.ChildIndexInParent == B.ChildIndexInParent;
+		}
+
+		friend FORCEINLINE uint32 GetTypeHash(const FItemCacheKey& Key)
+		{
+			return HashCombine( GetTypeHash(Key.Parent), HashCombine(GetTypeHash(Key.Child), Key.ChildIndexInParent));
+		}
+	};
+	TMap< FItemCacheKey, TSharedPtr<FMutableGraphTreeElement>> ItemCache;
+
+	/** Main tree item for each node. A graph node can be representes with multiple tree ndoes if it is reachable from different paths. */
+	TMap< const mu::Node*, TSharedPtr<FMutableGraphTreeElement>> MainItemPerNode;
+
+	/** */
+	void RebuildTree();
+
+	/** UI callbacks */
+	void CompileMutableCodePressed();
+	TSharedRef<SWidget> GenerateCompileOptionsMenuContent();
+	TSharedPtr<STextComboBox> CompileOptimizationCombo;
+	TArray< TSharedPtr<FString> > CompileOptimizationStrings;
+	void OnChangeCompileOptimizationLevel(TSharedPtr<FString> NewSelection, ESelectInfo::Type SelectInfo);
+
+	/** Callbacks from the tree widget. */
+ 	TSharedRef<ITableRow> GenerateRowForNodeTree(TSharedPtr<FMutableGraphTreeElement> InTreeNode, const TSharedRef<STableViewBase>& InOwnerTable);
+	void GetChildrenForInfo(TSharedPtr<FMutableGraphTreeElement> InInfo, TArray< TSharedPtr<FMutableGraphTreeElement> >& OutChildren);
+	TSharedPtr<SWidget> OnTreeContextMenuOpening();
+	void TreeExpandRecursive(TSharedPtr<FMutableGraphTreeElement> Info, bool bExpand);
+	void TreeExpandUnique();
+
+};
+
+
+/** An row of the code tree in the SMutableGraphViewer. */
+class FMutableGraphTreeElement : public TSharedFromThis<FMutableGraphTreeElement>
+{
+public:
+	FMutableGraphTreeElement(const mu::NodePtr& InNode, TSharedPtr<FMutableGraphTreeElement>* InDuplicatedOf=nullptr )
+	{
+		MutableNode = InNode;
+		if (InDuplicatedOf)
+		{
+			DuplicatedOf = *InDuplicatedOf;
+		}
+	}
+
+public:
+
+	/** Mutable Graph Node represented in this tree row. */
+	mu::NodePtr MutableNode;
+
+	/** If this tree element is a duplicated of another node, this is the node. */
+	TSharedPtr<FMutableGraphTreeElement> DuplicatedOf;
+
+};
diff --git a/Engine/Plugins/Experimental/Mutable/Source/CustomizableObjectEditor/Private/SMutableMeshViewer.h b/Engine/Plugins/Experimental/Mutable/Source/CustomizableObjectEditor/Private/SMutableMeshViewer.h
new file mode 100644
--- /dev/null
+++ b/Engine/Plugins/Experimental/Mutable/Source/CustomizableObjectEditor/Private/SMutableMeshViewer.h
@@ -0,0 +1,209 @@
+// Copyright Epic Games, Inc. All Rights Reserved.
+
+#pragma once
+
+#include "Containers/Array.h"
+#include "MutableRuntime/Public/Mesh.h"
+#include "Widgets/DeclarativeSyntaxSupport.h"
+#include "Widgets/SCompoundWidget.h"
+#include "Widgets/Views/SListView.h"
+#include "Widgets/Views/STableRow.h"
+#include "Widgets/Layout/SBorder.h"
+#include "AssetThumbnail.h"
+#include "UObject/GCObject.h"
+
+#include "Templates/SharedPointer.h"
+
+#include "SMutableMeshViewport.h"
+#include "Widgets/Input/SSearchBox.h"
+#include "Widgets/Layout/SScrollBox.h"
+
+#include "PropertyCustomizationHelpers.h"
+#include "SMutableSkeletonViewer.h"
+
+
+/** Container designed to hold the buffer channel data of a mutable mesh buffer to be later used by the UI */
+struct FBufferChannelElement
+{
+	FText SemanticIndex = FText(INVTEXT(""));
+	FText BufferSemantic = FText(INVTEXT(""));
+	FText BufferFormat = FText(INVTEXT(""));
+	FText BufferComponentCount = FText(INVTEXT(""));
+};
+
+
+/** Element representing a mutable buffer. It contains an array with all the elements representing the channels the mutable buffer is made of.*/
+struct FBufferElement
+{
+	/** The index of the buffer on the origin mutable buffer set */
+	FText BufferIndex = FText(INVTEXT(""));
+
+	/** An array of BufferChannels that represent the relative mutable channels on the mutable buffer*/
+	TSharedPtr<TArray<TSharedPtr<FBufferChannelElement>>> BufferChannels;
+};
+
+/** Widget designed to show the statistical data from a Mutable Mesh*/
+class SMutableMeshViewer : public SCompoundWidget, public FGCObject
+{
+public:
+	SLATE_BEGIN_ARGS(SMutableMeshViewer) {}
+		SLATE_ARGUMENT_DEFAULT(mu::MeshPtrConst, Mesh){nullptr};
+	SLATE_END_ARGS()
+
+	/** Builds the widget */
+	void Construct(const FArguments& InArgs);
+
+	// SWidget interface
+	// void Tick(const FGeometry& AllottedGeometry, const double InCurrentTime, const float InDeltaTime) override;
+
+	// FGCObject interface
+	virtual void AddReferencedObjects(FReferenceCollector& Collector) override;
+
+	virtual FString GetReferencerName() const override
+	{
+		return TEXT("SMutableMeshViewer");
+	}
+
+	/** Set the Mutable Mesh to be used for this widget */
+	void SetMesh(const mu::MeshPtrConst& InMesh);
+
+	
+	/** Generate a new SListView for all the ChannelElements provided.
+	 * @param InBufferChannelElements - List with all the elements that will be ued to fill InHostListView
+	 * @return A SListView containing all the channel elements as UI elements. 
+	 */
+	TSharedRef<SWidget> GenerateBufferChannelsListView(
+		const TSharedPtr<TArray<TSharedPtr<FBufferChannelElement>>>& InBufferChannelElements);
+
+private:
+	
+	/** Slate whose task is to display the skeleton found on this mesh as a slate tree view */
+	TSharedPtr<SMutableSkeletonViewer> MutableSkeletonViewer;
+
+	/** Thumbnail objects used to render a view of the selected skeletal mesh for the mesh viewport*/
+	TSharedPtr<FAssetThumbnailPool> AssetThumbnailPool;
+	TSharedPtr<FAssetThumbnail> AssetThumbnail;
+
+	/** Data backend for the widget. It represents the mesh that is being "displayed" */
+	mu::MeshPtrConst MutableMesh = nullptr;
+
+	/** Is true, mesh has changed and we need to update. */
+	bool bIsPendingUpdate = false;
+
+	/** Splitter used to separate the two sides of the slate (tables and viewport) */
+	TSharedPtr<SSplitter> SpaceSplitter;
+
+	/** Slate object containing all the buffer tables alongside with the bone tree */
+	TSharedPtr<SScrollBox> DataSpaceSlate;
+
+	/** Viewport object to preview the current mesh inside an actual Unreal scene */
+	TSharedPtr<SMutableMeshViewport> MeshViewport;
+	
+	/*
+	 * Slate views for the main types of mesh buffers (vertex , index and face)
+	 * Each buffer element also contains the channels it uses
+	 */
+	
+	TSharedPtr<SListView<TSharedPtr<FBufferElement>>> VertexBuffersSlateView;
+	TSharedPtr<SListView<TSharedPtr<FBufferElement>>> IndexBuffersSlateView;
+	TSharedPtr<SListView<TSharedPtr<FBufferElement>>> FaceBuffersSlateView;
+
+private:
+	
+	/** The border where the object responsible for providing the user of a way to select objects is */
+	TSharedPtr<SBorder> ReferenceMeshSelectionSpace;
+
+	/** Generates all slate objects related with the Mesh Viewport Slate */
+	TSharedRef<SWidget> GenerateViewportSlates();
+
+	/** Generates the tables showing the buffer data on the mesh alongside with the bone tree found on the mutable mesh */
+	TSharedRef<SWidget> GenerateDataTableSlates();
+	
+	/** Width and height for the skeletal mesh thumbnail */
+	FVector2f ThumbnailSize = FVector2f();
+	FVector2f WidgetSize = FVector2f();
+
+	// Single property that only draws the combo box widget of the skeletal mesh
+	TSharedPtr<class ISinglePropertyView> SkeletalMeshSelector;
+	
+	/** Method called each time the mesh selected changes so the UI gets updated reliably */
+	void OnMeshChanged();
+
+private:
+	
+	/* Generic UI callbacks used by the Widget */
+	FText GetVertexCount() const;
+	FText GetFaceCount() const;
+	FText GetBoneCount() const;
+
+	/*
+	 * Elements used to feed the buffers list (index and buffer channels as an internal list)
+	 */
+	
+	TArray<TSharedPtr<FBufferElement>> VertexBuffers;
+	TArray<TSharedPtr<FBufferElement>> IndexBuffers;
+	TArray<TSharedPtr<FBufferElement>> FaceBuffers;
+	
+	/** 
+	* Fills the provided TArray with the buffer definitions generated from the mutable buffers found on the provided mutable buffer set
+	* @param InMuMeshBufferSetPtr - Mutable buffer set pointer pointing to the buffer set witch buffer data is wanted to be saved on the provided TArray
+	* @param OutBuffersDataArray - TArray witch will contain the parsed data from the provided mutable buffer set converted onto FBufferElements
+	* @param InHostListView - The slate list using the previously referenced array as data backend. 
+	*/
+	void FillTargetBufferSetDataArray(const mu::FMeshBufferSet& InMuMeshBufferSetPtr,
+	                                  TArray<TSharedPtr<FBufferElement>>& OutBuffersDataArray,
+	                                  TSharedPtr<SListView<TSharedPtr<FBufferElement>>>& InHostListView);
+
+	
+	/**
+	 * Provided an array of buffer elements and the type of buffer it generates a new SListView for said buffer elements
+	 * @param InBufferElements - The FBufferElement objects to be used as backend for the SListView
+	 * @param InBufferSetTypeName - The name of the buffer set you are generating. Some examples are "Vertex", "index" or "Face"
+	 * @param OutHostListView - THe SListView generated.
+	 * @return - A shared reference to the SListView generated. Required by this slate operation.
+	 */
+	TSharedRef<SWidget> GenerateBuffersListView(TSharedPtr<SListView<TSharedPtr<FBufferElement>>>& OutHostListView,
+	                                          const TArray<TSharedPtr<FBufferElement>>& InBufferElements,
+	                                          const FText& InBufferSetTypeName);
+	
+	/** Callback method invoked each time a new row of the SListView containing the buffer elements needs to be built
+	 * @param InBuffer - The buffer element being used to build the new row.
+	 * @param OwnerTable - The table that is going to get the new row added into it.
+	 * @return The new row object as the base interface used for all the table rows
+	 */
+	TSharedRef<ITableRow> OnGenerateBufferRow(TSharedPtr<FBufferElement, ESPMode::ThreadSafe> InBuffer,
+											  const TSharedRef<STableViewBase, ESPMode::ThreadSafe>& OwnerTable);
+	
+	/** 
+	* Callback Method responsible of generating each row of the buffer channel lists based on the channel definition provided.
+	* @param InBufferChannel - The mesh Buffer channel definition that is being used as blueprint for the actual UI row
+	* @param OwnerTable - The parent UI element that is going to be expanded with the new row object
+	* @return The new row object as the base interface used for all the table rows
+	*/
+	TSharedRef<ITableRow> OnGenerateBufferChannelRow(TSharedPtr<FBufferChannelElement> InBufferChannel,
+	                                                 const TSharedRef<STableViewBase>& OwnerTable);
+	
+	
+	/*
+	 * Reference skeleton mesh selection
+	 */
+
+	/** Currently selected reference mesh */
+	const USkeletalMesh* SelectedReferenceMesh = nullptr;
+
+	/** Callback invoked when the selected asset on the reference skeletal mesh selection space has changed
+	 * @param AssetData - The data of the selected asset on the selection space
+	 */
+	void OnSelectedReferenceSkeletalMeshChanged(const FAssetData& AssetData);
+
+	/** Filters the assets found in order to know if the asset should be exposed on the UI for usage
+	 * @param AssetData - The data of the asset being checked for validation
+	 * @return - True means that the asset WILL NOT appear on the list of selectable resources
+	 */
+	bool OnShouldFilterAsset(const FAssetData& AssetData) const;
+
+	/**
+	 * Provides the SObjectPropertyEntryBox a way to know what to display in the selected object space
+	 */
+	FString OnObjectPath() const;
+};
diff --git a/Engine/Plugins/Experimental/Mutable/Source/CustomizableObjectEditor/Private/SMutableMeshViewport.h b/Engine/Plugins/Experimental/Mutable/Source/CustomizableObjectEditor/Private/SMutableMeshViewport.h
new file mode 100644
--- /dev/null
+++ b/Engine/Plugins/Experimental/Mutable/Source/CustomizableObjectEditor/Private/SMutableMeshViewport.h
@@ -0,0 +1,96 @@
+// Copyright Epic Games, Inc. All Rights Reserved.
+
+#pragma once
+
+#include "MutableRuntime/Public/Mesh.h"
+#include "FMutableViewportClient.h"
+#include "Widgets/SCompoundWidget.h"
+#include "AdvancedPreviewScene.h"
+#include "Components/SkeletalMeshComponent.h"
+#include "SEditorViewport.h"
+#include "UnrealWidget.h"
+#include "UObject/GCObject.h"
+
+#include "Engine/SkeletalMesh.h"
+
+
+/**
+* Object with the objective of showing a preview of a selected mutable mesh. It is designed to have few dependencies
+* as possible
+*/
+class SMutableMeshViewport : public SEditorViewport , public FGCObject
+{
+	SLATE_BEGIN_ARGS(SMutableMeshViewport) {}
+    	SLATE_ARGUMENT_DEFAULT(mu::MeshPtrConst, Mesh) { nullptr };
+    SLATE_END_ARGS()
+    
+public:
+    /** Builds the widget */
+    void Construct(const FArguments& InArgs);
+	
+    /** Set the Mutable Mesh to be used for this widget
+     * @param InMesh - The mutable mesh to be displayed
+     */
+    void SetMesh(const mu::MeshPtrConst& InMesh);
+
+	/** Set the reference mesh used by the mutable to skeletal mesh conversion process
+	 * @param InReferenceMesh - The Skeletal mesh to be used to aid on the conversion of the MutableMesh to a SkeletalMesh
+	 */
+	void SetReferenceMesh(const USkeletalMesh* InReferenceMesh);
+
+	// FGCObject interface
+	virtual void AddReferencedObjects(FReferenceCollector& Collector) override;
+	virtual FString GetReferencerName() const override
+	{
+		return TEXT("SMutableMeshViewport");
+	}
+	
+    // SWidget interface
+	virtual void Tick(const FGeometry& AllottedGeometry, const double InCurrentTime, const float InDeltaTime) override;
+
+protected:
+
+	// SEditorViewport interface (called by SEditorViewport::Construct)
+	virtual TSharedRef<FEditorViewportClient> MakeEditorViewportClient() override;
+
+private:
+
+	/** The mutable mesh being displayed */
+	mu::MeshPtrConst MutableMesh = nullptr;
+
+	/*
+	 *	UI  
+	 */
+	
+	/** Skeleton provided by the user */
+	const USkeletalMesh* ReferenceSkeletalMesh = nullptr;
+	
+	/** The preview scene that we are viewing */
+	TSharedPtr<FAdvancedPreviewScene> PreviewScene = nullptr;
+
+	/** Editor viewport client */
+	TSharedPtr<FMutableMeshViewportClient> ViewportClient = nullptr;
+
+	/** Restores the default state for the viewport */
+	void ClearViewport() const;
+
+	/** Adds the required content to the viewport*/
+	void SendMeshToViewport();
+
+	/** Clears or displays a mesh depending on the presence of a mutable mesh on MutableMesh**/
+	void RefreshViewportContents();
+
+	/*
+	 * Utility
+	 */
+
+	/** Mutable mesh converted into a Unreal mesh object */
+	USkeletalMeshComponent* SkeletalMeshComponent = nullptr;
+	
+	/** Generates a new USkeletalMesh from the MutableMesh set on this object
+	 * @return - If the operation was able to be processed or not
+	 */
+	bool GenerateUnrealMesh();
+};
+
+
diff --git a/Engine/Plugins/Experimental/Mutable/Source/CustomizableObjectEditor/Private/SMutableObjectViewer.cpp b/Engine/Plugins/Experimental/Mutable/Source/CustomizableObjectEditor/Private/SMutableObjectViewer.cpp
new file mode 100644
--- /dev/null
+++ b/Engine/Plugins/Experimental/Mutable/Source/CustomizableObjectEditor/Private/SMutableObjectViewer.cpp
@@ -0,0 +1,449 @@
+// Copyright Epic Games, Inc. All Rights Reserved.
+
+#include "SMutableObjectViewer.h"
+
+#include "CustomizableObjectEditorStyle.h"
+#include "CustomizableObjectCompiler.h"
+#include "SMutableGraphViewer.h"
+#include "SMutableCodeViewer.h"
+#include "MutableTools/Public/Streams.h"
+
+#include "Misc/App.h"
+#include "Framework/MultiBox/MultiBoxBuilder.h"
+#include "Widgets/Input/SButton.h"
+#include "Widgets/Images/SImage.h"
+#include "Widgets/Input/STextComboBox.h"
+#include "Widgets/Docking/SDockTab.h"
+#include "EditorDirectories.h"
+#include "DesktopPlatformModule.h"
+#include "HAL/PlatformApplicationMisc.h"
+#include "Interfaces/ITargetPlatform.h"
+#include "Interfaces/ITargetPlatformManagerModule.h"
+
+#define LOCTEXT_NAMESPACE "SMutableDebugger"
+
+// \todo: multi-column tree
+namespace MutableObjectTreeViewColumns
+{
+	static const FName Name("Property");
+	static const FName Value("Value");
+};
+
+
+class SMutableObjectTreeRow : public STableRow<TSharedPtr<FMutableObjectTreeElement>>
+{
+public:
+
+	void Construct(const FArguments& Args, const TSharedRef<STableViewBase>& InOwnerTableView, const TSharedPtr<FMutableObjectTreeElement>& InRowItem, const UCustomizableObject* InObject )
+	{
+		RowItem = InRowItem;
+
+		FText MainLabel;;
+		switch (RowItem->Type)
+		{
+		case FMutableObjectTreeElement::EType::SectionCaption:
+		{
+			switch (RowItem->Section)
+			{
+			case FMutableObjectTreeElement::ESection::General:
+				MainLabel = FText::FromString(TEXT("General"));
+				break;
+			case FMutableObjectTreeElement::ESection::ChildObjects:
+				MainLabel = FText::FromString(TEXT("Children"));
+				break;
+			default:
+				MainLabel = FText::FromString(TEXT("Unknown"));
+				break;
+			}
+			break;
+		}
+
+		case FMutableObjectTreeElement::EType::Name:
+			MainLabel = FText::FromString(InObject->GetName() );
+
+		default:
+			MainLabel = FText::FromString(TEXT("Unknown"));
+			break;
+		}
+
+
+		// TODO
+		const FSlateBrush* IconBrush = nullptr;
+
+		this->ChildSlot
+		[
+			SNew(SHorizontalBox)
+			+ SHorizontalBox::Slot()
+			.VAlign(VAlign_Center)
+			.AutoWidth()
+			[
+				SNew(SExpanderArrow, SharedThis(this))
+			]
+
+			+ SHorizontalBox::Slot()
+				.VAlign(VAlign_Center)
+				.Padding(FMargin(5.f, 0.f, 5.f, 0.f))
+				.AutoWidth()
+				[
+					SNew(SOverlay)
+					+ SOverlay::Slot()
+					[
+						SNew(SImage)
+						.Image(IconBrush ? IconBrush : FCoreStyle::Get().GetDefaultBrush())
+						.ColorAndOpacity(IconBrush ? FLinearColor::White : FLinearColor::Transparent)
+					]
+				]
+
+			+ SHorizontalBox::Slot()
+				[
+					SNew(STextBlock)
+					.Text(MainLabel)
+				]
+		];
+
+		STableRow< TSharedPtr<FMutableObjectTreeElement> >::ConstructInternal(
+			STableRow::FArguments()
+			//.Style(FAppStyle::Get(), "DetailsView.TreeView.TableRow")
+			.ShowSelection(true)
+			, InOwnerTableView
+		);
+
+	}
+
+
+private:
+
+	TSharedPtr<FMutableObjectTreeElement> RowItem;
+};
+
+
+void SMutableObjectViewer::AddReferencedObjects(FReferenceCollector& Collector)
+{
+	Collector.AddReferencedObject(CustomizableObject);
+}
+
+
+FString SMutableObjectViewer::GetReferencerName() const
+{
+	return TEXT("SMutableObjectViewer");
+}
+
+
+void SMutableObjectViewer::Construct(const FArguments& InArgs, UCustomizableObject* InObject, 
+	TWeakPtr<FTabManager> InParentTabManager, const FName& InParentNewTabId )
+{
+	CustomizableObject = InObject;
+	ParentTabManager = InParentTabManager;
+	ParentNewTabId = InParentNewTabId;
+
+	// Initialize the debugger compile options
+	CompileOptions.bTextureCompression = true;
+	CompileOptions.OptimizationLevel = 3;
+	{
+		ITargetPlatformManagerModule* TPM = GetTargetPlatformManager();
+		const TArray<ITargetPlatform*>& Platforms = TPM->GetActiveTargetPlatforms();
+		CompileOptions.TargetPlatform = Platforms.IsEmpty() ? nullptr : Platforms[0];
+	}
+
+	FToolBarBuilder ToolbarBuilder(TSharedPtr<const FUICommandList>(), FMultiBoxCustomization::None, TSharedPtr<FExtender>(), true);
+	ToolbarBuilder.SetLabelVisibility(EVisibility::Visible);
+	ToolbarBuilder.SetStyle(&FAppStyle::Get(), "SlimToolBar");
+
+	ToolbarBuilder.BeginSection("Compilation");
+
+	ToolbarBuilder.AddToolBarButton(
+		FUIAction(FExecuteAction::CreateSP(this, &SMutableObjectViewer::GenerateMutableGraphPressed)),
+		NAME_None,
+		LOCTEXT("GenerateMutableGraph", "Unreal to Mutable Graph"),
+		LOCTEXT("GenerateMutableGraphTooltip", "Generate a mutable graph from the customizable object source graph."),
+		FSlateIcon(FCustomizableObjectEditorStyle::Get().GetStyleSetName(), "CustomizableObjectDebugger.GenerateMutableGraph", "CustomizableObjectDebugger.GenerateMutableGraph.Small"),
+		EUserInterfaceActionType::Button
+	);
+
+	ToolbarBuilder.AddToolBarButton(
+		FUIAction(FExecuteAction::CreateSP(this, &SMutableObjectViewer::CompileMutableCodePressed)),
+		NAME_None,
+		LOCTEXT("GenerateMutableCode", "Unreal to Mutable Code"),
+		LOCTEXT("GenerateMutableCodeTooltip", "Generate a mutable code from the customizable object source graph."),
+		FSlateIcon(FCustomizableObjectEditorStyle::Get().GetStyleSetName(), "CustomizableObjectDebugger.CompileMutableCode", "CustomizableObjectDebugger.CompileMutableCode.Small"),
+		EUserInterfaceActionType::Button
+	);
+
+	ToolbarBuilder.AddComboButton(
+		FUIAction(),
+		FOnGetContent::CreateSP(this, &SMutableObjectViewer::GenerateCompileOptionsMenuContent),
+		LOCTEXT("Compile_Options_Label", "Compile Options"),
+		LOCTEXT("Compile_Options_Tooltip", "Change Compile Options"),
+		TAttribute<FSlateIcon>(),
+		true);
+
+	ToolbarBuilder.EndSection();
+
+	// Initialize the root tree nodes
+	RootTreeNodes.Add(MakeShareable(new FMutableObjectTreeElement(FMutableObjectTreeElement::EType::SectionCaption, FMutableObjectTreeElement::ESection::General, InObject)));
+	RootTreeNodes.Add(MakeShareable(new FMutableObjectTreeElement(FMutableObjectTreeElement::EType::SectionCaption, FMutableObjectTreeElement::ESection::ChildObjects, InObject)));
+
+	ChildSlot
+	[
+		SNew(SVerticalBox)
+		+ SVerticalBox::Slot()
+		.AutoHeight()
+		.VAlign(VAlign_Center)
+		[
+			ToolbarBuilder.MakeWidget()
+		]
+		+ SVerticalBox::Slot()
+		.VAlign(VAlign_Fill)
+		[
+			SNew(SBorder)
+			.BorderImage(FAppStyle::GetBrush("ToolPanel.GroupBorder"))
+			.Padding(FMargin(4.0f, 4.0f))
+			[
+				SAssignNew(TreeView, STreeView<TSharedPtr<FMutableObjectTreeElement>>)
+				.TreeItemsSource(&RootTreeNodes)
+				.OnGenerateRow(this,&SMutableObjectViewer::GenerateRowForNodeTree)
+				.OnGetChildren(this, &SMutableObjectViewer::GetChildrenForInfo)
+				.SelectionMode(ESelectionMode::None)					
+			]
+		]
+	];	
+}
+
+
+void SMutableObjectViewer::GenerateMutableGraphPressed()
+{
+	// Convert from Unreal graph to Mutable graph.
+	mu::NodePtr RootNode = Compiler.Export(CustomizableObject, CompileOptions);
+	if (!RootNode)
+	{
+		// TODO: Show errors
+		return;
+	}
+
+	FString DataTag = FString::Printf(TEXT("%s for %s"), *CustomizableObject->GetName(), *CompileOptions.TargetPlatform->PlatformName());
+
+	TSharedPtr<SDockTab> NewMutableGraphTab = SNew(SDockTab)
+		.Label(LOCTEXT("MutableGraph", "Mutable Graph"))
+		[
+			SNew(SMutableGraphViewer, RootNode, CompileOptions, ParentTabManager, ParentNewTabId)
+			.DataTag(DataTag)
+		];
+
+	TSharedPtr<FTabManager> TabManager = ParentTabManager.Pin();
+	check(TabManager);
+	TabManager->InsertNewDocumentTab(ParentNewTabId, FTabManager::ESearchPreference::PreferLiveTab, NewMutableGraphTab.ToSharedRef());
+}
+
+
+void SMutableObjectViewer::CompileMutableCodePressed()
+{
+	// Convert from Unreal graph to Mutable graph.
+	mu::NodePtr RootNode = Compiler.Export(CustomizableObject, CompileOptions);
+	if (!RootNode)
+	{
+		// TODO: Show errors
+		return;
+	}
+
+	// Do the compilation to Mutable Code synchronously.
+	TSharedPtr<FCustomizableObjectCompileRunnable> CompileTask = MakeShareable(new FCustomizableObjectCompileRunnable(RootNode, false));
+	CompileTask->Options = CompileOptions;
+	CompileTask->Init();
+	CompileTask->Run();
+
+	FString DataTag = FString::Printf( TEXT("%s for %s opt %d "), *CustomizableObject->GetName(), *CompileOptions.TargetPlatform->PlatformName(), CompileOptions.OptimizationLevel );
+
+	TSharedPtr<SDockTab> NewMutableCodeTab = SNew(SDockTab)
+		.Label(LOCTEXT("MutableCode", "Mutable Code"))
+		[
+			SNew(SMutableCodeViewer, CompileTask->Model)
+			.DataTag(DataTag)
+		];
+
+	TSharedPtr<FTabManager> TabManager = ParentTabManager.Pin();
+	check(TabManager);
+	TabManager->InsertNewDocumentTab(ParentNewTabId, FTabManager::ESearchPreference::PreferLiveTab, NewMutableCodeTab.ToSharedRef());
+}
+
+
+TSharedRef<SWidget> SMutableObjectViewer::GenerateCompileOptionsMenuContent()
+{
+	const bool bShouldCloseWindowAfterMenuSelection = false;
+	FMenuBuilder MenuBuilder(bShouldCloseWindowAfterMenuSelection, nullptr);
+
+	// settings
+	MenuBuilder.BeginSection("Optimization", LOCTEXT("MutableCompileOptimizationHeading", "Optimization"));
+	{
+		// Unreal Graph to Mutable Graph options
+		//-----------------------------------
+
+		// Platform
+		DebugPlatformStrings.Empty();
+
+		TSharedPtr<FString> SelectedPlatform;
+		{
+			ITargetPlatformManagerModule* TPM = GetTargetPlatformManager();
+			check(TPM);
+
+			ITargetPlatform* CurrentPlatform = NULL;
+			const TArray<ITargetPlatform*>& Platforms = TPM->GetTargetPlatforms();
+			for (const ITargetPlatform* Platform : Platforms)
+			{
+				TSharedPtr<FString> ThisPlatform = MakeShareable(new FString(Platform->PlatformName()));
+				DebugPlatformStrings.Add(ThisPlatform);
+				if (Platform == CompileOptions.TargetPlatform)
+				{
+					SelectedPlatform = ThisPlatform;
+				}
+			}
+		}
+
+		if (!SelectedPlatform.IsValid() && DebugPlatformStrings.Num())
+		{
+			SelectedPlatform = DebugPlatformStrings[0];
+		}
+
+		DebugPlatformCombo =
+			SNew(STextComboBox)
+			.OptionsSource(&DebugPlatformStrings)
+			.InitiallySelectedItem(SelectedPlatform)
+			.OnSelectionChanged(this, &SMutableObjectViewer::OnChangeDebugPlatform)
+			;
+
+		MenuBuilder.AddWidget(DebugPlatformCombo.ToSharedRef(), LOCTEXT("MutableDebugPlatform", "Target Platform"));
+
+		// Texture compresssion
+		MenuBuilder.AddMenuEntry(
+			LOCTEXT("MutableEnableTextureCompression", "Enable texture compression."),
+			LOCTEXT("MutableEnableTextureCompressionTooltip", "Enable or disable the compression of textures in the mutable code compilation."),
+			FSlateIcon(),
+			FUIAction(
+				FExecuteAction::CreateLambda([this]() { CompileOptions.bTextureCompression = !CompileOptions.bTextureCompression; }),
+				FCanExecuteAction(),
+				FIsActionChecked::CreateLambda([this]() { return CompileOptions.bTextureCompression; })),
+			NAME_None,
+			EUserInterfaceActionType::ToggleButton
+		);
+
+		// Compilation options
+		//-----------------------------------
+
+		// Optimisation level
+		CompileOptimizationStrings.Empty();
+		CompileOptimizationStrings.Add(MakeShareable(new FString(NSLOCTEXT("UnrealEd", "OptimizationNone", "None").ToString())));
+		CompileOptimizationStrings.Add(MakeShareable(new FString(NSLOCTEXT("UnrealEd", "OptimizationMin", "Minimal").ToString())));
+		CompileOptimizationStrings.Add(MakeShareable(new FString(NSLOCTEXT("UnrealEd", "OptimizationMed", "Medium").ToString())));
+		CompileOptimizationStrings.Add(MakeShareable(new FString(NSLOCTEXT("UnrealEd", "OptimizationMax", "Maximum").ToString())));
+
+		CompileOptions.OptimizationLevel = FMath::Min(CompileOptions.OptimizationLevel, CompileOptimizationStrings.Num() - 1);
+
+		CompileOptimizationCombo =
+			SNew(STextComboBox)
+			.OptionsSource(&CompileOptimizationStrings)
+			.InitiallySelectedItem(CompileOptimizationStrings[CompileOptions.OptimizationLevel])
+			.OnSelectionChanged(this, &SMutableObjectViewer::OnChangeCompileOptimizationLevel)
+			;
+		MenuBuilder.AddWidget(CompileOptimizationCombo.ToSharedRef(), LOCTEXT("MutableCompileOptimizationLevel", "Optimization Level"));
+
+		// Parallel compilation
+		MenuBuilder.AddMenuEntry(
+			LOCTEXT("MutableEnableParallelCompilation", "Enable compiling in multiple threads."),
+			LOCTEXT("MutableEnableParallelCompilationTooltip", "This is faster but use more memory."),
+			FSlateIcon(),
+			FUIAction(
+				FExecuteAction::CreateLambda([this]() { CompileOptions.bUseParallelCompilation = !CompileOptions.bUseParallelCompilation; }),
+				FCanExecuteAction(),
+				FIsActionChecked::CreateLambda([this]() { return CompileOptions.bUseParallelCompilation; })),
+			NAME_None,
+			EUserInterfaceActionType::ToggleButton
+		);
+
+		// Disk as cache
+		MenuBuilder.AddMenuEntry(
+			LOCTEXT("MutableEnableTextureCompression", "Enable compiling using the disk as memory."),
+			LOCTEXT("MutableEnableTextureCompressionTooltip", "This is very slow but supports compiling huge objects. It requires a lot of free space in the OS disk."),
+			FSlateIcon(),
+			FUIAction(
+				FExecuteAction::CreateLambda([this]() { CompileOptions.bUseDiskCompilation = !CompileOptions.bUseDiskCompilation; }),
+				FCanExecuteAction(),
+				FIsActionChecked::CreateLambda([this]() { return CompileOptions.bUseDiskCompilation; })),
+			NAME_None,
+			EUserInterfaceActionType::ToggleButton
+		);
+	}
+	MenuBuilder.EndSection();
+
+	return MenuBuilder.MakeWidget();
+}
+
+
+void SMutableObjectViewer::OnChangeCompileOptimizationLevel(TSharedPtr<FString> NewSelection, ESelectInfo::Type SelectInfo)
+{
+	CompileOptions.OptimizationLevel = CompileOptimizationStrings.Find(NewSelection);
+}
+
+
+void SMutableObjectViewer::OnChangeDebugPlatform(TSharedPtr<FString> NewSelection, ESelectInfo::Type SelectInfo)
+{
+	if (!CustomizableObject) return;
+
+	CompileOptions.TargetPlatform = nullptr;
+
+	ITargetPlatformManagerModule* TPM = GetTargetPlatformManager();
+	check(TPM);
+	const TArray<ITargetPlatform*>& Platforms = TPM->GetTargetPlatforms();
+	check(Platforms.Num());
+
+	CompileOptions.TargetPlatform = Platforms[0];
+
+	for (int32 Index = 1; Index < Platforms.Num(); Index++)
+	{
+		if (Platforms[Index]->PlatformName() == *NewSelection)
+		{
+			CompileOptions.TargetPlatform = Platforms[Index];
+			break;
+		}
+	}
+}
+
+
+TSharedRef<ITableRow> SMutableObjectViewer::GenerateRowForNodeTree(TSharedPtr<FMutableObjectTreeElement> InTreeNode, const TSharedRef<STableViewBase>& InOwnerTable)
+{
+	TSharedRef<SMutableObjectTreeRow> Row = SNew(SMutableObjectTreeRow, InOwnerTable, InTreeNode, CustomizableObject);
+	return Row;
+}
+
+
+void SMutableObjectViewer::GetChildrenForInfo(TSharedPtr<FMutableObjectTreeElement> InInfo, TArray<TSharedPtr<FMutableObjectTreeElement>>& OutChildren)
+{
+	switch (InInfo->Type)
+	{
+	case FMutableObjectTreeElement::EType::SectionCaption:
+		switch (InInfo->Section)
+		{
+
+		case FMutableObjectTreeElement::ESection::General:
+		{
+			OutChildren.Add(MakeShareable(new FMutableObjectTreeElement(FMutableObjectTreeElement::EType::Name, InInfo->Section, CustomizableObject )));
+			break;
+		}
+
+		case FMutableObjectTreeElement::ESection::ChildObjects:
+		{
+			// TODO: based on what is done in FCustomizableObjectCompiler::ProcessChildObjectsRecursively
+			//OutChildren.Add(MakeShareable(new FMutableObjectTreeElement(FMutableObjectTreeElement::EType::Name, InInfo->Section)));
+			break;
+		}
+
+		default:
+			break;
+		}
+		break;
+
+	default:
+		break;
+	}
+}
+
+
+#undef LOCTEXT_NAMESPACE 
diff --git a/Engine/Plugins/Experimental/Mutable/Source/CustomizableObjectEditor/Private/SMutableObjectViewer.h b/Engine/Plugins/Experimental/Mutable/Source/CustomizableObjectEditor/Private/SMutableObjectViewer.h
new file mode 100644
--- /dev/null
+++ b/Engine/Plugins/Experimental/Mutable/Source/CustomizableObjectEditor/Private/SMutableObjectViewer.h
@@ -0,0 +1,122 @@
+// Copyright Epic Games, Inc. All Rights Reserved.
+
+#pragma once
+
+#include "Widgets/SCompoundWidget.h"
+#include "Framework/Docking/TabManager.h"
+#include "Widgets/Views/STreeView.h"
+#include "CustomizableObject.h"
+#include "CustomizableObjectCompiler.h"
+
+
+// Forward declarations
+class FMutableObjectTreeElement;
+
+
+/** This widget shows the internals of a CustomizableObject for debugging purposes. 
+ * This is not the Unreal source graph in the UCustomizableObject, but an intermediate step of the compilation process.
+ */
+class SMutableObjectViewer :
+	public SCompoundWidget,
+	public FGCObject
+{
+public:
+
+	// SWidget interface
+	SLATE_BEGIN_ARGS(SMutableObjectViewer) {}
+	SLATE_END_ARGS()
+
+	void Construct(const FArguments& InArgs, UCustomizableObject*,
+		TWeakPtr<FTabManager> InParentTabManager, const FName& InParentNewTabId );
+
+	// FGCObject interface
+	virtual void AddReferencedObjects(FReferenceCollector& Collector) override;
+	virtual FString GetReferencerName() const override;
+
+private:
+
+	/** The Mutable Graph to show, represented by its root. */
+	UCustomizableObject* CustomizableObject;
+
+	/** Compilation options to use in the debugger operations. */
+	FCompilationOptions CompileOptions;
+	TArray< TSharedPtr<FString> > DebugPlatformStrings;
+	TSharedPtr<STextComboBox> DebugPlatformCombo;
+
+	/** Object compiler. */
+	FCustomizableObjectCompiler Compiler;
+
+	/** UI references used to create new tabs. */
+	TWeakPtr<FTabManager> ParentTabManager;
+	FName ParentNewTabId;
+
+	/** Tree showing the object properties. */
+	TSharedPtr<STreeView<TSharedPtr<FMutableObjectTreeElement>>> TreeView;
+
+	/** Root nodes of the tree widget. */
+	TArray<TSharedPtr<FMutableObjectTreeElement>> RootTreeNodes;
+
+	/** UI callbacks */
+	void GenerateMutableGraphPressed();
+	void CompileMutableCodePressed();
+	TSharedRef<SWidget> GenerateCompileOptionsMenuContent();
+	TSharedPtr<STextComboBox> CompileOptimizationCombo;
+	TArray< TSharedPtr<FString> > CompileOptimizationStrings;
+	void OnChangeCompileOptimizationLevel(TSharedPtr<FString> NewSelection, ESelectInfo::Type SelectInfo);
+	void OnChangeDebugPlatform(TSharedPtr<FString> NewSelection, ESelectInfo::Type SelectInfo);
+
+	/** Tree callbacks. */
+	TSharedRef<ITableRow> GenerateRowForNodeTree(TSharedPtr<FMutableObjectTreeElement> InTreeNode, const TSharedRef<STableViewBase>& InOwnerTable);
+	void GetChildrenForInfo(TSharedPtr<FMutableObjectTreeElement> InInfo, TArray< TSharedPtr<FMutableObjectTreeElement> >& OutChildren);
+
+};
+
+
+/** An row of the code tree in . */
+class FMutableObjectTreeElement : public TSharedFromThis<FMutableObjectTreeElement>
+{
+public:
+
+	enum class EType
+	{
+		None,
+		SectionCaption,
+		Name,
+		ChildObject
+	};
+
+	enum class ESection
+	{
+		None,
+		General,
+		ChildObjects
+	};
+
+	/** Constructor for section rows */
+	FMutableObjectTreeElement(ESection InSection, UCustomizableObject* InObject)
+	{
+		Type = EType::SectionCaption;
+		Section = InSection;
+		CustomizableObject = InObject;
+	}
+
+	/** Constructor for generic property rows. */
+	FMutableObjectTreeElement(EType InType, ESection InSection, UCustomizableObject* InObject)
+	{
+		Type = InType;
+		Section = InSection;
+		CustomizableObject = InObject;
+	}
+
+public:
+
+	/** Row type. */
+	EType Type = EType::None;
+
+	/** If the row is a section row. */
+	ESection Section = ESection::None;
+
+	/** */
+	UCustomizableObject* CustomizableObject = nullptr;
+
+};
diff --git a/Engine/Plugins/Experimental/Mutable/Source/CustomizableObjectEditor/Private/SMutableParametersWidget.cpp b/Engine/Plugins/Experimental/Mutable/Source/CustomizableObjectEditor/Private/SMutableParametersWidget.cpp
new file mode 100644
--- /dev/null
+++ b/Engine/Plugins/Experimental/Mutable/Source/CustomizableObjectEditor/Private/SMutableParametersWidget.cpp
@@ -0,0 +1,589 @@
+// Copyright Epic Games, Inc. All Rights Reserved.
+
+#include "SMutableParametersWidget.h"
+#include "CustomizableObjectEditorModule.h"
+#include "Styling/AppStyle.h"
+#include "Async/TaskGraphInterfaces.h"
+#include "Brushes/SlateImageBrush.h"
+#include "Widgets/Input/SNumericEntryBox.h"
+#include "Widgets/Input/SCheckBox.h"
+#include "Widgets/Input/STextComboBox.h"
+#include "Widgets/Input/SVectorInputBox.h"
+#include "Widgets/Colors/SColorPicker.h"
+#include "Widgets/Layout/SGridPanel.h"
+#include "MutableTools/Public/TypeInfo.h"
+
+
+#define LOCTEXT_NAMESPACE "CustomizableObjectEditor"
+
+
+void SMutableParametersWidget::Construct(const FArguments& InArgs)
+{
+	OnParametersValueChanged = InArgs._OnParametersValueChanged;
+
+	ChildSlot
+	[
+		SNew(SVerticalBox)
+		+ SVerticalBox::Slot()
+		.AutoHeight()
+		.Padding(0, 4.0f)
+		.HAlign(HAlign_Fill)
+		[
+			SAssignNew(ParamBox, SVerticalBox)
+		]
+	];
+}
+
+
+void SMutableParametersWidget::SetParameters(const mu::ParametersPtr& InParameters)
+{
+	MutableParameters = InParameters;
+	bIsPendingUpdate = true;
+}
+
+
+void SMutableParametersWidget::Tick(const FGeometry& AllottedGeometry, const double InCurrentTime, const float InDeltaTime)
+{
+	SCompoundWidget::Tick(AllottedGeometry, InCurrentTime, InDeltaTime);
+
+	if (!bIsPendingUpdate)
+	{
+		return;
+	}
+
+	bIsPendingUpdate = false;
+
+	ParamBox->ClearChildren();
+
+	if (!MutableParameters)
+	{
+		return;
+	}
+
+	for ( int32 ParamIndex=0; ParamIndex< MutableParameters->GetCount(); ++ParamIndex )
+	{
+		FString ParamName = ANSI_TO_TCHAR(MutableParameters->GetName(ParamIndex));
+		TSharedPtr<SHorizontalBox> ParameterBox;
+
+		ParamBox->AddSlot()
+			.Padding(1.0f, 6.0f)
+			.HAlign(HAlign_Fill)
+			.VAlign(VAlign_Fill)
+			.AutoHeight()
+			[
+				SAssignNew(ParameterBox, SHorizontalBox)
+
+				+ SHorizontalBox::Slot()
+				.FillWidth(0.66f)
+				.VAlign(VAlign_Center)
+				[
+					SNew(STextBlock)
+					.Text(FText::FromString(ParamName))
+					.Justification(ETextJustify::Left)
+					.AutoWrapText(true)
+				]
+			];
+
+		switch (MutableParameters->GetType(ParamIndex) )
+		{
+		case mu::PARAMETER_TYPE::T_BOOL:
+			ParameterBox->AddSlot()
+				.Padding(4.0f)
+				.VAlign(VAlign_Center)
+				[
+					SNew(SCheckBox)
+					.Visibility(this, &SMutableParametersWidget::GetParameterVisibility, ParamIndex)
+					.IsChecked(this, &SMutableParametersWidget::GetBoolParameterValue, ParamIndex)
+					.OnCheckStateChanged(this, &SMutableParametersWidget::OnBoolParameterChanged, ParamIndex)
+				];
+			break;
+
+		case mu::PARAMETER_TYPE::T_FLOAT:
+			ParameterBox->AddSlot()
+				.Padding(4.0f)
+				.VAlign(VAlign_Center)
+				[
+					SNew(SNumericEntryBox<float>)
+					.Visibility(this, &SMutableParametersWidget::GetParameterVisibility, ParamIndex)
+					.AllowSpin(true)
+					.MinSliderValue(0.0f)
+					.MaxSliderValue(1.0f)
+					.Value(this, &SMutableParametersWidget::GetFloatParameterValue, ParamIndex)
+					.OnValueChanged(this, &SMutableParametersWidget::OnFloatParameterChanged, ParamIndex)
+				];
+			break;
+
+		case mu::PARAMETER_TYPE::T_COLOUR:
+			ParameterBox->AddSlot()
+				.Padding(4.0f)
+				.VAlign(VAlign_Center)
+				[
+					SNew(SColorBlock)
+					.Color(this, &SMutableParametersWidget::GetColorParameterValue, ParamIndex)
+					.ShowBackgroundForAlpha(false)
+					.AlphaDisplayMode(EColorBlockAlphaDisplayMode::Ignore)
+					.UseSRGB(false)
+					.OnMouseButtonDown(this, &SMutableParametersWidget::OnColorBlockMouseButtonDown, ParamIndex)
+					.Size(FVector2D(10.0f, 10.0f))
+				];
+			break;
+
+		case mu::PARAMETER_TYPE::T_INT:
+		{
+			// If we have a list of options, add a combo box too
+			TSharedPtr<SMutableTextSearchBox> ParamComboBox;
+			int32 ValueCount = MutableParameters->GetIntPossibleValueCount(ParamIndex);
+			if (ValueCount > 0)
+			{
+				FString ToolTipText = FString("None");
+				TArray<FString> OptionNamesAttribute;
+				int Value = MutableParameters->GetIntValue( ParamIndex );
+				int ValueIndex = MutableParameters->GetIntValueIndex( ParamIndex, Value );
+				for (int i = 0; i < ValueCount; ++i)
+				{
+					const char* ValueText = MutableParameters->GetIntPossibleValueName( ParamIndex, i );
+					OptionNamesAttribute.Add( FString(ANSI_TO_TCHAR(ValueText)) );				
+				}
+
+				ParameterBox->AddSlot()
+					.Padding(4.0f)
+					.VAlign(VAlign_Center)
+					[
+						SAssignNew(ParamComboBox, SMutableTextSearchBox)
+						.Visibility(this, &SMutableParametersWidget::GetParameterVisibility, ParamIndex)
+						.ToolTipText(FText::FromString(ToolTipText))
+						.PossibleSuggestions(OptionNamesAttribute)
+						.InitialText(FText::FromString(OptionNamesAttribute[ValueIndex]))
+						.MustMatchPossibleSuggestions(TAttribute<bool>(true))
+						.SuggestionListPlacement(EMenuPlacement::MenuPlacement_ComboBox)
+						.OnTextCommitted(this, &SMutableParametersWidget::OnIntParameterTextChanged, ParamIndex)
+					];
+			}
+
+			ParameterBox->AddSlot()
+				.Padding(4.0f)
+				.VAlign(VAlign_Center)
+				[
+					SNew(SNumericEntryBox<int32>)
+					.Visibility(this, &SMutableParametersWidget::GetParameterVisibility, ParamIndex)
+					.AllowSpin(true)
+					.MinSliderValue(0)
+					.MaxSliderValue(this, &SMutableParametersWidget::GetIntParameterValueMax, ParamIndex)
+					.Value(this, &SMutableParametersWidget::GetIntParameterValue, ParamIndex)
+					.OnValueChanged(this, &SMutableParametersWidget::OnIntParameterChanged, ParamIndex, ParamComboBox)
+				];
+			break;
+		}
+
+		case mu::PARAMETER_TYPE::T_PROJECTOR:
+			ParameterBox->AddSlot()
+				.Padding(4.0f)
+				.VAlign(VAlign_Center)
+				[
+					SNew(SGridPanel)
+					.FillColumn(0, 1.0f)
+
+					+SGridPanel::Slot(0,0)
+					.VAlign(VAlign_Center)
+					[
+						SNew(STextBlock)
+						.Text(FText::FromString(TEXT("Location")))
+						.Font(FAppStyle::GetFontStyle(TEXT("PropertyWindow.NormalFont")))
+					]
+
+					+SGridPanel::Slot(1, 0)
+					[
+						SNew(SNumericVectorInputBox<FVector::FReal>)
+						.X(this, &SMutableParametersWidget::GetProjectorLocation, EAxis::X, ParamIndex)
+						.Y(this, &SMutableParametersWidget::GetProjectorLocation, EAxis::Y, ParamIndex)
+						.Z(this, &SMutableParametersWidget::GetProjectorLocation, EAxis::Z, ParamIndex)
+						.bColorAxisLabels(true)
+						.OnXChanged(this, &SMutableParametersWidget::SetProjectorLocation, ETextCommit::Default, EAxis::X, false, ParamIndex)
+						.OnYChanged(this, &SMutableParametersWidget::SetProjectorLocation, ETextCommit::Default, EAxis::Y, false, ParamIndex)
+						.OnZChanged(this, &SMutableParametersWidget::SetProjectorLocation, ETextCommit::Default, EAxis::Z, false, ParamIndex)
+						.OnXCommitted(this, &SMutableParametersWidget::SetProjectorLocation, EAxis::X, true, ParamIndex)
+						.OnYCommitted(this, &SMutableParametersWidget::SetProjectorLocation, EAxis::Y, true, ParamIndex)
+						.OnZCommitted(this, &SMutableParametersWidget::SetProjectorLocation, EAxis::Z, true, ParamIndex)
+						.AllowSpin(true)
+						.SpinDelta(1)
+					]
+
+					+ SGridPanel::Slot(0, 1)
+					.VAlign(VAlign_Center)
+					[
+						SNew(STextBlock)
+						.Text(FText::FromString(TEXT("Scale")))
+						.Font(FAppStyle::GetFontStyle("PropertyWindow.NormalFont"))
+					]
+
+					+ SGridPanel::Slot(1, 1)
+					[
+						SNew(SNumericVectorInputBox<FVector::FReal>)
+						.X(this, &SMutableParametersWidget::GetProjectorScale, EAxis::X, ParamIndex)
+						.Y(this, &SMutableParametersWidget::GetProjectorScale, EAxis::Y, ParamIndex)
+						.Z(this, &SMutableParametersWidget::GetProjectorScale, EAxis::Z, ParamIndex)
+						.bColorAxisLabels(true)
+						.OnXChanged(this, &SMutableParametersWidget::SetProjectorScale, ETextCommit::Default, EAxis::X, false, ParamIndex)
+						.OnYChanged(this, &SMutableParametersWidget::SetProjectorScale, ETextCommit::Default, EAxis::Y, false, ParamIndex)
+						.OnZChanged(this, &SMutableParametersWidget::SetProjectorScale, ETextCommit::Default, EAxis::Z, false, ParamIndex)
+						.OnXCommitted(this, &SMutableParametersWidget::SetProjectorScale, EAxis::X, true, ParamIndex)
+						.OnYCommitted(this, &SMutableParametersWidget::SetProjectorScale, EAxis::Y, true, ParamIndex)
+						.OnZCommitted(this, &SMutableParametersWidget::SetProjectorScale, EAxis::Z, true, ParamIndex)
+						.AllowSpin(true)
+						.SpinDelta(1)
+					]
+
+				];
+			break;
+
+		default:
+			// Unsupported parameter type
+			ParameterBox->AddSlot()
+				.Padding(4.0f)
+				.VAlign(VAlign_Center)
+				[
+					SNew(STextBlock)
+					.Text(FText::FromString(TEXT("Unsupported parameter type.")))
+					.Justification(ETextJustify::Left)
+					.AutoWrapText(true)
+				];
+		}
+
+	}
+}
+
+
+EVisibility SMutableParametersWidget::GetParameterVisibility(int32 ParamIndex) const
+{
+	// \TODO: Parameter relevancy?
+	if (!MutableParameters
+		|| ParamIndex >= MutableParameters->GetCount() )
+	{
+		return EVisibility::HitTestInvisible;
+	}
+
+	return EVisibility::Visible;
+}
+
+
+ECheckBoxState SMutableParametersWidget::GetBoolParameterValue(int32 ParamIndex) const
+{
+	if ( !MutableParameters 
+		|| ParamIndex>=MutableParameters->GetCount()
+		|| MutableParameters->GetType(ParamIndex)!=mu::PARAMETER_TYPE::T_BOOL )
+	{
+		return ECheckBoxState::Undetermined;
+	}	
+
+	return MutableParameters->GetBoolValue(ParamIndex) ? ECheckBoxState::Checked : ECheckBoxState::Unchecked;
+}
+
+
+void SMutableParametersWidget::OnBoolParameterChanged(ECheckBoxState InCheckboxState, int32 ParamIndex)
+{
+	if (!MutableParameters
+		|| ParamIndex >= MutableParameters->GetCount()
+		|| MutableParameters->GetType(ParamIndex) != mu::PARAMETER_TYPE::T_BOOL)
+	{
+		return;
+	}
+
+	MutableParameters->SetBoolValue( ParamIndex, InCheckboxState==ECheckBoxState::Checked );
+
+	OnParametersValueChanged.ExecuteIfBound(ParamIndex);
+}
+
+
+TOptional<float> SMutableParametersWidget::GetFloatParameterValue(int32 ParamIndex) const
+{
+	if (!MutableParameters
+		|| ParamIndex >= MutableParameters->GetCount()
+		|| MutableParameters->GetType(ParamIndex) != mu::PARAMETER_TYPE::T_FLOAT)
+	{
+		return TOptional<float>();
+	}
+
+	return MutableParameters->GetFloatValue(ParamIndex);
+}
+
+
+void SMutableParametersWidget::OnFloatParameterChanged(float InValue, int32 ParamIndex)
+{
+	if (!MutableParameters
+		|| ParamIndex >= MutableParameters->GetCount()
+		|| MutableParameters->GetType(ParamIndex) != mu::PARAMETER_TYPE::T_FLOAT)
+	{
+		return;
+	}
+
+	MutableParameters->SetFloatValue(ParamIndex, InValue);
+
+	OnParametersValueChanged.ExecuteIfBound(ParamIndex);
+}
+
+
+TOptional<int32> SMutableParametersWidget::GetIntParameterValue(int32 ParamIndex) const
+{
+	if (!MutableParameters
+		|| ParamIndex >= MutableParameters->GetCount()
+		|| MutableParameters->GetType(ParamIndex) != mu::PARAMETER_TYPE::T_INT)
+	{
+		return TOptional<int32>();
+	}
+
+	int32 Value = MutableParameters->GetIntValue(ParamIndex);
+	int32 ValueIndex = MutableParameters->GetIntValueIndex(ParamIndex, Value);
+
+	return ValueIndex;
+}
+
+
+TOptional<int32> SMutableParametersWidget::GetIntParameterValueMax(int32 ParamIndex) const
+{
+	if (!MutableParameters
+		|| ParamIndex >= MutableParameters->GetCount()
+		|| MutableParameters->GetType(ParamIndex) != mu::PARAMETER_TYPE::T_INT)
+	{
+		return TOptional<int32>();
+	}
+
+	int32 ValueCount = MutableParameters->GetIntPossibleValueCount(ParamIndex);
+
+	if (!ValueCount)
+	{
+		ValueCount = 16;
+	}
+
+	return ValueCount -1;
+}
+
+
+void SMutableParametersWidget::OnIntParameterChanged(int32 InValue, int32 ParamIndex, TSharedPtr<SMutableTextSearchBox> Combo )
+{
+	if (!MutableParameters
+		|| ParamIndex >= MutableParameters->GetCount()
+		|| MutableParameters->GetType(ParamIndex) != mu::PARAMETER_TYPE::T_INT)
+	{
+		return;
+	}
+
+	// The slider value is actually the value index for the parameter
+	if (InValue >= 0 && InValue < MutableParameters->GetIntPossibleValueCount(ParamIndex))
+	{
+		int32 RealValue = MutableParameters->GetIntPossibleValue(ParamIndex, InValue);
+		MutableParameters->SetIntValue(ParamIndex, RealValue);
+
+		// Update the text combo if any
+		if (Combo)
+		{
+			FString Text = TEXT("Invalid");
+			Text = ANSI_TO_TCHAR(MutableParameters->GetIntPossibleValueName(ParamIndex, InValue));
+			Combo->SetText(FText::FromString(Text));
+		}
+
+		OnParametersValueChanged.ExecuteIfBound(ParamIndex);
+	}
+}
+
+
+void SMutableParametersWidget::OnIntParameterTextChanged(TSharedPtr<FString> Selection, ESelectInfo::Type, int32 ParamIndex)
+{
+	if (!Selection
+		|| !MutableParameters
+		|| ParamIndex >= MutableParameters->GetCount()
+		|| MutableParameters->GetType(ParamIndex) != mu::PARAMETER_TYPE::T_INT)
+	{
+		return;
+	}
+
+	int32 ValueIndex = MutableParameters->GetIntValueIndex(ParamIndex, TCHAR_TO_ANSI(**Selection));
+	if (ValueIndex >= 0 && ValueIndex < MutableParameters->GetIntPossibleValueCount(ParamIndex))
+	{
+		int32 RealValue = MutableParameters->GetIntPossibleValue(ParamIndex, ValueIndex);
+		MutableParameters->SetIntValue(ParamIndex, RealValue);
+
+		OnParametersValueChanged.ExecuteIfBound(ParamIndex);
+	}
+}
+
+
+FLinearColor SMutableParametersWidget::GetColorParameterValue(int32 ParamIndex) const
+{
+	if (!MutableParameters
+		|| ParamIndex >= MutableParameters->GetCount()
+		|| MutableParameters->GetType(ParamIndex) != mu::PARAMETER_TYPE::T_COLOUR)
+	{
+		return FLinearColor();
+	}
+
+	FLinearColor Result;
+	MutableParameters->GetColourValue(ParamIndex, &Result.R, &Result.G, &Result.B);
+	Result.A = 1.0f;
+
+	return Result;
+}
+
+
+FReply SMutableParametersWidget::OnColorBlockMouseButtonDown(const FGeometry& MyGeometry, const FPointerEvent& MouseEvent, int32 ParamIndex)
+{
+	if (MouseEvent.GetEffectingButton() != EKeys::LeftMouseButton)
+	{
+		return FReply::Unhandled();
+	}
+
+	FLinearColor col = GetColorParameterValue(ParamIndex);
+
+	TArray<FLinearColor*> LinearColorArray;
+	LinearColorArray.Add(&col);
+
+	FColorPickerArgs args;
+	args.bIsModal = true;
+	args.bUseAlpha = false;
+	args.bOnlyRefreshOnMouseUp = false;
+	args.InitialColorOverride = col;
+	args.OnColorCommitted = FOnLinearColorValueChanged::CreateSP(this, &SMutableParametersWidget::OnSetColorFromColorPicker, ParamIndex);
+	OpenColorPicker(args);
+
+	return FReply::Handled();
+}
+
+
+void SMutableParametersWidget::OnSetColorFromColorPicker(FLinearColor NewColor, int32 ParamIndex)
+{
+	if (!MutableParameters
+		|| ParamIndex >= MutableParameters->GetCount()
+		|| MutableParameters->GetType(ParamIndex) != mu::PARAMETER_TYPE::T_COLOUR)
+	{
+		return;
+	}
+
+	MutableParameters->SetColourValue(ParamIndex, NewColor.R, NewColor.G, NewColor.B);
+
+	OnParametersValueChanged.ExecuteIfBound(ParamIndex);
+}
+
+
+TOptional<FVector::FReal> SMutableParametersWidget::GetProjectorLocation(EAxis::Type Axis, int32 ParamIndex) const
+{ 
+	if (!MutableParameters
+		|| ParamIndex >= MutableParameters->GetCount()
+		|| MutableParameters->GetType(ParamIndex) != mu::PARAMETER_TYPE::T_PROJECTOR)
+	{
+		return 0.0f;
+	}
+
+	mu::PROJECTOR_TYPE Type;
+	float Pos[3];
+	float Dir[3];
+	float Up[3];
+	float Scale[3];
+	float Angle;
+	MutableParameters->GetProjectorValue(ParamIndex, &Type, 
+		&Pos[0], &Pos[1], &Pos[2],
+		&Dir[0], &Dir[1], &Dir[2],
+		&Up[0], &Up[1], &Up[2],
+		&Scale[0], &Scale[1], &Scale[2],
+		&Angle );
+
+	return Pos[ Axis-EAxis::X ];
+}
+
+
+void SMutableParametersWidget::SetProjectorLocation(FVector::FReal NewValue, ETextCommit::Type, EAxis::Type Axis, bool bCommitted, int32 ParamIndex)
+{
+	if (!MutableParameters
+		|| ParamIndex >= MutableParameters->GetCount()
+		|| MutableParameters->GetType(ParamIndex) != mu::PARAMETER_TYPE::T_PROJECTOR)
+	{
+		return;
+	}
+
+	mu::PROJECTOR_TYPE Type;
+	float Pos[3];
+	float Dir[3];
+	float Up[3];
+	float Scale[3];
+	float Angle;
+	MutableParameters->GetProjectorValue(ParamIndex, &Type,
+		&Pos[0], &Pos[1], &Pos[2],
+		&Dir[0], &Dir[1], &Dir[2],
+		&Up[0], &Up[1], &Up[2],
+		&Scale[0], &Scale[1], &Scale[2],
+		&Angle);
+
+	Pos[Axis - EAxis::X] = NewValue;
+
+	MutableParameters->SetProjectorValue(ParamIndex,
+		Pos[0], Pos[1], Pos[2],
+		Dir[0], Dir[1], Dir[2],
+		Up[0], Up[1], Up[2],
+		Scale[0], Scale[1], Scale[2],
+		Angle);
+
+	OnParametersValueChanged.ExecuteIfBound(ParamIndex);
+}
+
+
+TOptional<FVector::FReal> SMutableParametersWidget::GetProjectorScale(EAxis::Type Axis, int32 ParamIndex) const
+{
+	if (!MutableParameters
+		|| ParamIndex >= MutableParameters->GetCount()
+		|| MutableParameters->GetType(ParamIndex) != mu::PARAMETER_TYPE::T_PROJECTOR)
+	{
+		return 0.0f;
+	}
+
+	mu::PROJECTOR_TYPE Type;
+	float Pos[3];
+	float Dir[3];
+	float Up[3];
+	float Scale[3];
+	float Angle;
+	MutableParameters->GetProjectorValue(ParamIndex, &Type,
+		&Pos[0], &Pos[1], &Pos[2],
+		&Dir[0], &Dir[1], &Dir[2],
+		&Up[0], &Up[1], &Up[2],
+		&Scale[0], &Scale[1], &Scale[2],
+		&Angle);
+
+	return Scale[Axis - EAxis::X];
+}
+
+
+void SMutableParametersWidget::SetProjectorScale(FVector::FReal NewValue, ETextCommit::Type, EAxis::Type Axis, bool bCommitted, int32 ParamIndex)
+{
+	if (!MutableParameters
+		|| ParamIndex >= MutableParameters->GetCount()
+		|| MutableParameters->GetType(ParamIndex) != mu::PARAMETER_TYPE::T_PROJECTOR)
+	{
+		return;
+	}
+
+	mu::PROJECTOR_TYPE Type;
+	float Pos[3];
+	float Dir[3];
+	float Up[3];
+	float Scale[3];
+	float Angle;
+	MutableParameters->GetProjectorValue(ParamIndex, &Type,
+		&Pos[0], &Pos[1], &Pos[2],
+		&Dir[0], &Dir[1], &Dir[2],
+		&Up[0], &Up[1], &Up[2],
+		&Scale[0], &Scale[1], &Scale[2],
+		&Angle);
+
+	Scale[Axis - EAxis::X] = NewValue;
+
+	MutableParameters->SetProjectorValue(ParamIndex,
+		Pos[0], Pos[1], Pos[2],
+		Dir[0], Dir[1], Dir[2],
+		Up[0], Up[1], Up[2],
+		Scale[0], Scale[1], Scale[2],
+		Angle);
+
+	OnParametersValueChanged.ExecuteIfBound(ParamIndex);
+}
+
+
+#undef LOCTEXT_NAMESPACE
+
diff --git a/Engine/Plugins/Experimental/Mutable/Source/CustomizableObjectEditor/Private/SMutableParametersWidget.h b/Engine/Plugins/Experimental/Mutable/Source/CustomizableObjectEditor/Private/SMutableParametersWidget.h
new file mode 100644
--- /dev/null
+++ b/Engine/Plugins/Experimental/Mutable/Source/CustomizableObjectEditor/Private/SMutableParametersWidget.h
@@ -0,0 +1,82 @@
+// Copyright Epic Games, Inc. All Rights Reserved.
+
+#pragma once
+
+#include "SMutableTextSearchBox.h"
+#include "Widgets/SCompoundWidget.h"
+#include "Widgets/SBoxPanel.h"
+#include "Widgets/Input/SCheckBox.h"
+#include "Widgets/DeclarativeSyntaxSupport.h"
+#include "MutableRuntime/Public/Parameters.h"
+#include "Math/Axis.h"
+
+
+/** Delegate that is executed when a parameter value changes */
+DECLARE_DELEGATE_OneParam(FOnMutableParameterValueChanged, int32);
+
+
+/**  This widget displays and allows edition of a set of parameter values for the lower level Mutable Parameters 
+* object. It is meant to be used for tools and debugging and it doesn't have any knowledge of Customizable Objects
+* and the additional parameter details stored there.
+*/
+class SMutableParametersWidget : public SCompoundWidget
+{
+public:
+
+	SLATE_BEGIN_ARGS(SMutableParametersWidget) {}
+
+		/** Parameters to show and edit. */
+		SLATE_ARGUMENT_DEFAULT(mu::ParametersPtr, Parameters) { nullptr };
+
+		/** Called when any parameter value has changed, with the parameter index as argument.  */
+		SLATE_EVENT(FOnMutableParameterValueChanged, OnParametersValueChanged)
+
+	SLATE_END_ARGS()
+
+	void Construct(const FArguments& InArgs);
+
+	// SWidget interface
+	void Tick(const FGeometry& AllottedGeometry, const double InCurrentTime, const float InDeltaTime) override;
+
+	// Own interface
+
+	/** Set the image to show in the widget. */
+	void SetParameters(const mu::ParametersPtr& InParameters);
+
+private:
+
+	/** Reference to the mutable parameters. */
+	mu::ParametersPtr MutableParameters;
+
+	/** Is true, the parameter object has changed and we need to update. */
+	bool bIsPendingUpdate = false;
+
+	/** Delegate called when the check box changes state */
+	FOnMutableParameterValueChanged OnParametersValueChanged;
+
+	/** Dynamically filled box with per-parameter widgets */
+	TSharedPtr<SVerticalBox> ParamBox;
+
+	/** Internal UI callbacks. */
+	EVisibility GetParameterVisibility(int32 ParamIndex) const;
+	ECheckBoxState GetBoolParameterValue(int32 ParamIndex) const;
+	void OnBoolParameterChanged(ECheckBoxState CheckboxState, int32 ParamIndex);
+	TOptional<float> GetFloatParameterValue(int32 ParamIndex) const;
+	void OnFloatParameterChanged(float Value, int32 ParamIndex);
+	FLinearColor GetColorParameterValue(int32 ParamIndex) const;
+	FReply OnColorBlockMouseButtonDown(const FGeometry& MyGeometry, const FPointerEvent& MouseEvent, int32 ParamIndex);
+	void OnSetColorFromColorPicker(FLinearColor NewColor, int32 ParamIndex);
+	TOptional<int32> GetIntParameterValue(int32 ParamIndex) const;
+	TOptional<int32> GetIntParameterValueMax(int32 ParamIndex) const;
+	void OnIntParameterChanged(int32 Value, int32 ParamIndex, TSharedPtr<SMutableTextSearchBox> );
+	void OnIntParameterTextChanged(TSharedPtr<FString> Selection, ESelectInfo::Type SelectInfo, int32 ParamIndex);
+
+	/** Projector UI callbacks */
+	TOptional<FVector::FReal> GetProjectorLocation(EAxis::Type Axis, int32 ParamIndex) const;
+	void SetProjectorLocation(FVector::FReal NewValue, ETextCommit::Type, EAxis::Type, bool bCommitted, int32 ParamIndex);
+	TOptional<FVector::FReal> GetProjectorScale(EAxis::Type Axis, int32 ParamIndex) const;
+	void SetProjectorScale(FVector::FReal NewValue, ETextCommit::Type, EAxis::Type, bool bCommitted, int32 ParamIndex);
+
+
+};
+
