[mutable] Move Mutable plugin from the restricted folder to the experimental plugins folder.

diff --git a/Engine/Plugins/Experimental/Mutable/Source/CustomizableObjectEditor/Private/SMutableTextSearchBox.cpp b/Engine/Plugins/Experimental/Mutable/Source/CustomizableObjectEditor/Private/SMutableTextSearchBox.cpp
new file mode 100644
--- /dev/null
+++ b/Engine/Plugins/Experimental/Mutable/Source/CustomizableObjectEditor/Private/SMutableTextSearchBox.cpp
@@ -0,0 +1,374 @@
+// Copyright Epic Games, Inc. All Rights Reserved.
+
+#include "SMutableTextSearchBox.h"
+
+#include "CustomizableObject.h" // Required for LogMutable
+
+#include "Layout/WidgetPath.h"
+#include "Framework/Application/SlateApplication.h"
+#include "Widgets/Text/STextBlock.h"
+#include "Widgets/Input/SMenuAnchor.h"
+#include "Widgets/Views/SListView.h"
+#include "Widgets/Input/SSearchBox.h"
+#include "Widgets/Input/SButton.h"
+#include "Widgets/Input/SComboButton.h"
+#include "Widgets/Images/SImage.h"
+
+#include "Styling/AppStyle.h"
+
+void SMutableTextSearchBox::Construct( const FArguments& InArgs )
+{
+	OnTextChanged = InArgs._OnTextChanged;
+	OnTextCommitted = InArgs._OnTextCommitted;
+	OnKeyDownHandler = InArgs._OnKeyDownHandler;
+	PossibleSuggestions = InArgs._PossibleSuggestions;
+	PreCommittedText = InArgs._InitialText.Get();
+	bMustMatchPossibleSuggestions = InArgs._MustMatchPossibleSuggestions.Get();
+
+	ChildSlot
+		[
+			//SNew( SHorizontalBox )
+			//	+ SHorizontalBox::Slot()
+			//	.FillWidth( 1 )
+			//	[
+			SNew( SButton )
+				//.ButtonStyle( OurButtonStyle )
+				.ClickMethod( EButtonClickMethod::MouseDown )
+				.OnClicked( this, &SMutableTextSearchBox::OnButtonClicked )
+				//.ContentPadding( InArgs._ContentPadding )
+				//.ForegroundColor( InArgs._ForegroundColor )
+				.ButtonColorAndOpacity( FSlateColor(FLinearColor(1.f, 1.f, 1.f, 1.f)) )
+				.IsFocusable( true )
+				.Visibility(EVisibility::Visible)
+				.Content()
+				[
+					SAssignNew(SuggestionBox, SMenuAnchor)
+					.Placement( InArgs._SuggestionListPlacement )
+					.Padding(-2.0f)
+					[
+						SAssignNew(InputText, SSearchBox)
+						.InitialText(InArgs._InitialText)
+						.HintText(InArgs._HintText)
+						.OnTextChanged(this, &SMutableTextSearchBox::HandleTextChanged)
+						.OnTextCommitted(this, &SMutableTextSearchBox::HandleTextCommitted)
+						.SelectAllTextWhenFocused( true )
+						.DelayChangeNotificationsWhileTyping( InArgs._DelayChangeNotificationsWhileTyping )
+						.OnKeyDownHandler(this, &SMutableTextSearchBox::HandleKeyDown)
+						.Visibility(EVisibility::SelfHitTestInvisible)
+					]
+					.MenuContent
+						(
+							SNew(SBorder)
+							.BorderImage(FAppStyle::GetBrush("Menu.Background"))
+							[
+								SNew(SBox)
+								.WidthOverride(175)	// to enforce some minimum width, ideally we define the minimum, not a fixed width
+								//.HeightOverride(175) // avoids flickering, ideally this would be adaptive to the content without flickering
+								[
+									SAssignNew(SuggestionListView, SListView< TSharedPtr<FString> >)
+									.ListItemsSource(&Suggestions)
+									.SelectionMode( ESelectionMode::Single )							// Ideally the mouse over would not highlight while keyboard controls the UI
+									.OnGenerateRow(this, &SMutableTextSearchBox::MakeSuggestionListItemWidget)
+									.OnSelectionChanged( this, &SMutableTextSearchBox::OnSelectionChanged)
+									.OnListViewScrolled(this, &SMutableTextSearchBox::OnListViewScrolled)
+								]
+							]
+						)
+				]
+		];
+
+	InputText->SetVisibility(EVisibility::HitTestInvisible);
+}
+
+void SMutableTextSearchBox::OnListViewScrolled(double InScrollOffset)
+{	
+	if (!HasKeyboardFocus() && !SuggestionListView->IsUserScrolling())
+	{
+		FocusEditBox();
+	}
+}
+
+FReply SMutableTextSearchBox::OnButtonClicked()
+{
+	InputText->SetText(FText());
+	FReply Reply = FReply::Handled().SetUserFocus(InputText.ToSharedRef());
+	UpdateSuggestionList();
+
+	// Forward keyboard focus to our editable text widget
+	return Reply;
+}
+
+void SMutableTextSearchBox::SetText(const TAttribute< FText >& InNewText)
+{
+	InputText->SetText(InNewText);
+	PreCommittedText = InNewText.Get();
+}
+
+void SMutableTextSearchBox::SetError( const FText& InError )
+{
+	InputText->SetError(InError);
+}
+
+void SMutableTextSearchBox::SetError( const FString& InError )
+{
+	InputText->SetError(InError);
+}
+
+FReply SMutableTextSearchBox::OnPreviewKeyDown( const FGeometry& MyGeometry, const FKeyEvent& InKeyEvent )
+{
+	if ( SuggestionBox->IsOpen() && InKeyEvent.GetKey() == EKeys::Escape )
+	{
+		// Clear any selection first to prevent the currently selection being set in the text box
+		SuggestionListView->ClearSelection();
+		SuggestionBox->SetIsOpen(false, false);
+		return FReply::Handled();
+	}
+
+	return FReply::Unhandled();
+}
+
+FReply SMutableTextSearchBox::HandleKeyDown( const FGeometry& MyGeometry, const FKeyEvent& InKeyEvent )
+{
+	if ( SuggestionBox->IsOpen() && (InKeyEvent.GetKey() == EKeys::Up || InKeyEvent.GetKey() == EKeys::Down) )
+	{
+		const bool bSelectingUp = InKeyEvent.GetKey() == EKeys::Up;
+		TSharedPtr<FString> SelectedSuggestion = GetSelectedSuggestion();
+
+		if ( SelectedSuggestion.IsValid() )
+		{
+			// Find the selection index and select the previous or next one
+			int32 TargetIdx = INDEX_NONE;
+			for ( int32 SuggestionIdx = 0; SuggestionIdx < Suggestions.Num(); ++SuggestionIdx )
+			{
+				if ( Suggestions[SuggestionIdx] == SelectedSuggestion )
+				{
+					if ( bSelectingUp )
+					{
+						TargetIdx = SuggestionIdx - 1;
+					}
+					else
+					{
+						TargetIdx = SuggestionIdx + 1;
+					}
+					break;
+				}
+			}
+
+			if ( Suggestions.IsValidIndex(TargetIdx) )
+			{
+				SuggestionListView->SetSelection( Suggestions[TargetIdx] );
+				SuggestionListView->RequestScrollIntoView( Suggestions[TargetIdx] );
+			}
+		}
+		else if ( !bSelectingUp && Suggestions.Num() > 0 )
+		{
+			// Nothing selected and pressed down, select the first item
+			SuggestionListView->SetSelection( Suggestions[0] );
+		}
+
+		return FReply::Handled();
+	}
+
+	if (OnKeyDownHandler.IsBound())
+	{
+		return OnKeyDownHandler.Execute(MyGeometry, InKeyEvent);
+	}
+
+	return FReply::Unhandled();
+}
+
+bool SMutableTextSearchBox::SupportsKeyboardFocus() const
+{
+	return InputText->SupportsKeyboardFocus();
+}
+
+bool SMutableTextSearchBox::HasKeyboardFocus() const
+{
+	// Since keyboard focus is forwarded to our editable text, we will test it instead
+	return InputText->HasKeyboardFocus();
+}
+
+FReply SMutableTextSearchBox::OnFocusReceived( const FGeometry& MyGeometry, const FFocusEvent& InFocusEvent )
+{
+	UpdateSuggestionList();
+
+	// Forward keyboard focus to our editable text widget
+	return FReply::Handled().SetUserFocus(InputText.ToSharedRef(), InFocusEvent.GetCause());
+}
+
+void SMutableTextSearchBox::HandleTextChanged(const FText& NewText)
+{
+	OnTextChanged.ExecuteIfBound(NewText);
+	UpdateSuggestionList();
+}
+
+void SMutableTextSearchBox::HandleTextCommitted(const FText& NewText, ETextCommit::Type CommitType)
+{
+	if (SuggestionListView->IsUserScrolling())
+	{
+		return;
+	}
+
+	TSharedPtr<FString> SelectedSuggestion = GetSelectedSuggestion();
+
+	bool bCommitText = true;
+	FText CommittedText;
+	if ( SelectedSuggestion.IsValid() && CommitType != ETextCommit::OnCleared )
+	{
+		// Pressed selected a suggestion, set the text
+		CommittedText = FText::FromString( *SelectedSuggestion.Get() );
+	}
+	else
+	{
+		if ( CommitType == ETextCommit::OnCleared )
+		{
+			// Clear text when escape is pressed then commit an empty string
+			//CommittedText = FText::GetEmpty();
+			CommittedText = PreCommittedText;
+		}
+		else if( PossibleSuggestions.Get().Contains(NewText.ToString()) )
+		{
+			// If the text is a suggestion, set the text.
+			CommittedText = NewText;
+		}
+		else if( bMustMatchPossibleSuggestions )
+		{
+			// commit the original text if we have to match a suggestion
+			CommittedText = PreCommittedText;
+		}
+		else
+		{
+			// otherwise, set the typed text
+			CommittedText = NewText;
+		}	
+	}
+
+	// Set the text and execute the delegate
+	SetText(CommittedText);
+
+	TSharedPtr<FString> CommittedStringPtr = MakeShareable<FString>(new FString(CommittedText.ToString()));
+	OnTextCommitted.ExecuteIfBound(CommittedStringPtr, ESelectInfo::Direct);
+
+	if(CommitType != ETextCommit::Default)
+	{
+		// Clear the suggestion box if the user has navigated away or set their own text.
+		SuggestionListView->ClearSelection();
+		SuggestionBox->SetIsOpen(false, false);
+		//SuggestionListView->SetVisibility(EVisibility::Hidden);
+	}
+}
+
+void SMutableTextSearchBox::OnSelectionChanged( TSharedPtr<FString> NewValue, ESelectInfo::Type SelectInfo )
+{
+	// If the user clicked directly on an item to select it, then accept the choice and close the window
+	if( SelectInfo == ESelectInfo::OnMouseClick )
+	{
+		//SetText( FText::FromString( *NewValue.Get() ));
+		//SuggestionBox->SetIsOpen(false, false);
+		//FocusEditBox();
+		HandleTextCommitted(FText::FromString(*NewValue.Get()), ETextCommit::OnEnter);
+	}
+}
+
+
+TSharedRef<ITableRow> SMutableTextSearchBox::MakeSuggestionListItemWidget(TSharedPtr<FString> Text, const TSharedRef<STableViewBase>& OwnerTable)
+{
+	check(Text.IsValid());
+
+	return
+		SNew(STableRow< TSharedPtr<FString> >, OwnerTable)
+		[
+			SNew(STextBlock)
+			.Text(FText::FromString(*Text.Get()))
+			.HighlightText(this, &SMutableTextSearchBox::GetHighlightText)
+		];
+}
+
+FText SMutableTextSearchBox::GetHighlightText() const
+{
+	return InputText->GetText();
+}
+
+void SMutableTextSearchBox::UpdateSuggestionList()
+{
+	const FString TypedText = InputText->GetText().ToString();
+
+	Suggestions.Empty();
+
+	if ( TypedText.Len() > 0 )
+	{
+		TArray<FString> AllSuggestions = PossibleSuggestions.Get();
+
+		for ( auto SuggestionIt = AllSuggestions.CreateConstIterator(); SuggestionIt; ++SuggestionIt )
+		{
+			const FString& Suggestion = *SuggestionIt;
+			if ( Suggestion.Contains(TypedText))
+			{
+				Suggestions.Add( MakeShareable(new FString(Suggestion)) );
+			}
+		}
+
+		if ( Suggestions.Num() > 0 )
+		{
+			// At least one suggestion was found, open the menu
+			SuggestionBox->SetIsOpen(true, false);
+		}
+		else
+		{
+			// No suggestions were found, close the menu
+			SuggestionBox->SetIsOpen(false, false);
+		}
+	}
+	else //if (HasKeyboardFocus())
+	{
+		TArray<FString> AllSuggestions = PossibleSuggestions.Get();
+
+		for ( auto SuggestionIt = AllSuggestions.CreateConstIterator(); SuggestionIt; ++SuggestionIt )
+		{
+			const FString& Suggestion = *SuggestionIt;
+			Suggestions.Add( MakeShareable(new FString(Suggestion)) );
+		}
+
+		if ( Suggestions.Num() > 0 )
+		{
+			// At least one suggestion was found, open the menu
+			SuggestionBox->SetIsOpen(true, false);
+		}
+		else
+		{
+			// No suggestions were found, close the menu
+			SuggestionBox->SetIsOpen(false, false);
+		}
+	}
+	//else
+	//{
+	//	// No text was typed, close the menu
+	//	SuggestionBox->SetIsOpen(false, false);
+	//}
+
+	SuggestionListView->RequestListRefresh();
+}
+
+void SMutableTextSearchBox::FocusEditBox()
+{
+	FWidgetPath WidgetToFocusPath;
+	FSlateApplication::Get().GeneratePathToWidgetUnchecked( InputText.ToSharedRef(), WidgetToFocusPath );
+	FSlateApplication::Get().SetKeyboardFocus( WidgetToFocusPath, EFocusCause::SetDirectly );
+}
+
+TSharedPtr<FString> SMutableTextSearchBox::GetSelectedSuggestion()
+{
+	TSharedPtr<FString> SelectedSuggestion;
+	if ( SuggestionBox->IsOpen() )
+	{
+		const TArray< TSharedPtr<FString> >& SelectedSuggestionList = SuggestionListView->GetSelectedItems();
+		if ( SelectedSuggestionList.Num() > 0 )
+		{
+			// Selection mode is Single, so there should only be one suggestion at the most
+			SelectedSuggestion = SelectedSuggestionList[0];
+		}
+	}
+
+	return SelectedSuggestion;
+}
+
diff --git a/Engine/Plugins/Experimental/Mutable/Source/CustomizableObjectEditor/Private/SMutableTextSearchBox.h b/Engine/Plugins/Experimental/Mutable/Source/CustomizableObjectEditor/Private/SMutableTextSearchBox.h
new file mode 100644
--- /dev/null
+++ b/Engine/Plugins/Experimental/Mutable/Source/CustomizableObjectEditor/Private/SMutableTextSearchBox.h
@@ -0,0 +1,150 @@
+// Copyright Epic Games, Inc. All Rights Reserved.
+
+#pragma once
+
+
+#include "CoreMinimal.h"
+#include "SlateFwd.h"
+#include "Misc/Attribute.h"
+#include "Input/Reply.h"
+#include "Widgets/DeclarativeSyntaxSupport.h"
+#include "Widgets/SCompoundWidget.h"
+#include "Widgets/Views/STableViewBase.h"
+#include "Widgets/Views/STableRow.h"
+#include "Widgets/Input/SSearchBox.h"
+
+class SMenuAnchor;
+
+
+/**
+ * A widget to provide a search box with a filtered dropdown menu.
+ */
+
+typedef TSlateDelegates< TSharedPtr<FString> >::FOnSelectionChanged FOnTextSelectionCommitted;
+//DECLARE_DELEGATE_ThreeParams( FOnTextSelectionCommitted, const FText&, ETextCommit::Type, const FString& );
+
+class SMutableTextSearchBox : public SCompoundWidget
+{
+public:
+	SLATE_BEGIN_ARGS(SMutableTextSearchBox)
+		: _SuggestionListPlacement( MenuPlacement_BelowAnchor )
+		, _OnTextChanged()
+		, _OnTextCommitted()
+		, _InitialText()
+		, _HintText()
+		, _PossibleSuggestions(TArray<FString>())
+		, _DelayChangeNotificationsWhileTyping( false )
+		, _MustMatchPossibleSuggestions( false )
+	{}
+
+	/** Where to place the suggestion list */
+	SLATE_ARGUMENT( EMenuPlacement, SuggestionListPlacement )
+
+	/** Invoked whenever the text changes */
+	SLATE_EVENT( FOnTextChanged, OnTextChanged )
+
+	/** Invoked whenever the text is committed (e.g. user presses enter) */
+	SLATE_EVENT( FOnTextSelectionCommitted, OnTextCommitted )
+
+	/** Initial text to display for the search text */
+	SLATE_ATTRIBUTE( FText, InitialText )
+
+	/** Hint text to display for the search text when there is no value */
+	SLATE_ATTRIBUTE( FText, HintText )
+
+	/** All possible suggestions for the search text */
+	SLATE_ATTRIBUTE( TArray<FString>, PossibleSuggestions )
+
+	/** Whether the SearchBox should delay notifying listeners of text changed events until the user is done typing */
+	SLATE_ATTRIBUTE( bool, DelayChangeNotificationsWhileTyping )
+
+	/** Whether the SearchBox allows entries that don't match the possible suggestions */
+	SLATE_ATTRIBUTE( bool, MustMatchPossibleSuggestions )
+
+	/** Callback delegate to have first chance handling of the OnKeyDown event */
+	SLATE_EVENT( FOnKeyDown, OnKeyDownHandler )
+
+	SLATE_END_ARGS()
+
+	/** Constructs this widget with InArgs */
+	void Construct( const FArguments& InArgs );
+
+	/** Sets the text string currently being edited */
+	void SetText(const TAttribute< FText >& InNewText);
+
+	/** Set or clear the current error reporting information for this search box */
+	void SetError( const FText& InError );
+	void SetError( const FString& InError );
+
+	// SWidget implementation
+	virtual FReply OnPreviewKeyDown( const FGeometry& MyGeometry, const FKeyEvent& InKeyEvent ) override;
+	virtual bool SupportsKeyboardFocus() const override;
+	virtual bool HasKeyboardFocus() const override;
+	virtual FReply OnFocusReceived( const FGeometry& MyGeometry, const FFocusEvent& InFocusEvent ) override;
+
+private:
+	FReply OnButtonClicked();
+
+	/** First chance handler for key down events to the editable text widget */
+	FReply HandleKeyDown(const FGeometry& MyGeometry, const FKeyEvent& InKeyEvent);
+
+	/** Handler for when text in the editable text box changed */
+	void HandleTextChanged(const FText& NewText);
+
+	/** Handler for when text in the editable text box changed */
+	void HandleTextCommitted(const FText& NewText, ETextCommit::Type CommitType);
+
+	/** Called by SListView when the selection changes in the suggestion list */
+	void OnSelectionChanged( TSharedPtr<FString> NewValue, ESelectInfo::Type SelectInfo );
+
+	void OnListViewScrolled(double InScrollOffset);
+
+	/** Makes the widget for a suggestion message in the list view */
+	TSharedRef<ITableRow> MakeSuggestionListItemWidget(TSharedPtr<FString> Text, const TSharedRef<STableViewBase>& OwnerTable);
+
+	/** Gets the text to highlight in the suggestion list */
+	FText GetHighlightText() const;
+
+	/** Updates and shows or hides the suggestion list */
+	void UpdateSuggestionList();
+
+	/** Sets the focus to the InputText box */
+	void FocusEditBox();
+
+	/** Returns the currently selected suggestion */
+	TSharedPtr<FString> GetSelectedSuggestion();
+
+private:
+	/** The editable text field */
+	TSharedPtr< SSearchBox > InputText;
+
+	/** The the state of the text prior to being committed */
+	FText PreCommittedText;
+
+	/** Auto completion elements */
+	TSharedPtr< SMenuAnchor > SuggestionBox;
+
+	/** All suggestions stored in this widget for the list view */
+	TArray< TSharedPtr<FString> > Suggestions;
+
+	/** The list view for showing all suggestions */
+	TSharedPtr< SListView< TSharedPtr<FString> > > SuggestionListView;
+
+	/** Delegate for when text is changed in the edit box */
+	FOnTextChanged OnTextChanged;
+
+	/** Delegate for when text is changed in the edit box */
+	FOnTextSelectionCommitted OnTextCommitted;
+
+	/** Delegate for first chance handling for key down events */
+	FOnKeyDown OnKeyDownHandler;
+
+	/** All possible suggestions for the search text */
+	TAttribute< TArray<FString> > PossibleSuggestions;
+
+	/** Determines whether or not the committed text should match a suggestion */
+	bool bMustMatchPossibleSuggestions;
+
+	FString ParamName;
+};
+
diff --git a/Engine/Plugins/Experimental/Mutable/Source/CustomizableObjectEditor/Private/StatePerformanceTesting.cpp b/Engine/Plugins/Experimental/Mutable/Source/CustomizableObjectEditor/Private/StatePerformanceTesting.cpp
new file mode 100644
--- /dev/null
+++ b/Engine/Plugins/Experimental/Mutable/Source/CustomizableObjectEditor/Private/StatePerformanceTesting.cpp
@@ -0,0 +1,650 @@
+// Copyright Epic Games, Inc. All Rights Reserved.
+
+
+#include "StatePerformanceTesting.h"
+#include "CustomizableObjectEditor.h"
+#include "CustomizableObjectEditorModule.h"
+#include "CustomizableObjectClasses.h"
+#include "CustomizableObjectInstance.h"
+#include "Components/SkeletalMeshComponent.h"
+#include "Logging/MessageLog.h"
+
+
+#define LOCTEXT_NAMESPACE "CustomizableObjectEditor"
+
+
+FTestingCustomizableObject::FTestingCustomizableObject(TestType TestParameter)
+	: Instance(nullptr)
+	, BatchTestInCourse(false)
+	, BatchTestFinished(false)
+	, ActiveBatchIndex(-1)
+	, ActiveBatch(nullptr)
+	, TestResourcesPresent(false)
+	, Test(TestParameter)
+{
+	HelperCallback = NewObject<UUpdateClassHelper>(GetTransientPackage());
+	HelperCallback->StateData = this;
+}
+
+
+FTestingCustomizableObject::~FTestingCustomizableObject()
+{
+	if (BatchTestInCourse)
+	{
+		FinishTest();
+	}
+
+	HelperCallback = nullptr;
+	Instance = nullptr;	
+}
+
+
+void FTestingCustomizableObject::InitTest()
+{
+
+}
+
+
+void FTestingCustomizableObject::AddReferencedObjects(FReferenceCollector& Collector)
+{
+	Collector.AddReferencedObject(Instance);
+	Collector.AddReferencedObject(HelperCallback);
+
+	for (UCustomizableSkeletalComponent* Elem : CustomizableSkeletalComponents)
+	{
+		Collector.AddReferencedObject(Elem);
+	}
+	for (USkeletalMeshComponent* Elem : SkeletalMeshComponents)
+	{
+		Collector.AddReferencedObject(Elem);
+	}
+}
+
+
+FString FTestingCustomizableObject::GetReferencerName() const
+{
+	return TEXT("FTestingCustomizableObject");
+}
+
+
+void FTestingCustomizableObject::StartTest(UCustomizableObject* CustomizableObject)
+{
+	if (BatchTestInCourse)
+	{
+		FinishTest();
+	}
+	
+	if (!CustomizableObject->IsCompiled())
+	{
+		FText Msg(LOCTEXT("StateTestCompileRequired","Please compile the Customizable Object before starting the test"));
+		FMessageLog MessageLog("Mutable");
+		MessageLog.Notify(Msg, EMessageSeverity::Info, true);
+		return;
+	}
+
+	TestResourcesPresent = true;
+
+	Instance = CustomizableObject->CreateInstance();
+
+	const int32 NumComponents = CustomizableObject->ReferenceSkeletalMeshes.Num();
+	SkeletalMeshComponents.Reset(NumComponents);
+	CustomizableSkeletalComponents.Reset(NumComponents);
+
+	for (int32 ComponentIndex = 0; ComponentIndex < NumComponents; ++ComponentIndex)
+	{
+		UCustomizableSkeletalComponent* CustomizableSkeletalComponent = NewObject<UCustomizableSkeletalComponent>(UCustomizableSkeletalComponent::StaticClass());
+		if (CustomizableSkeletalComponent)
+		{
+			CustomizableSkeletalComponents.Add(CustomizableSkeletalComponent);
+			
+			CustomizableSkeletalComponent->CustomizableObjectInstance = Instance;
+			CustomizableSkeletalComponent->ComponentIndex = ComponentIndex;
+
+			USkeletalMeshComponent* SkeletalMeshComponent = NewObject<USkeletalMeshComponent>(GetTransientPackage(), NAME_None, RF_Transient);
+			if (SkeletalMeshComponent)
+			{
+				SkeletalMeshComponents.Add(SkeletalMeshComponent);
+
+				CustomizableSkeletalComponent->AttachToComponent(SkeletalMeshComponent, FAttachmentTransformRules::KeepRelativeTransform);
+			}
+		}
+	}
+
+	if (SkeletalMeshComponents.Num() && SkeletalMeshComponents.Num() == CustomizableSkeletalComponents.Num())
+	{
+		Instance->UpdateSkeletalMeshAsync(true);
+		Instance->UpdateBeginDelegate.BindRaw(this, &FTestingCustomizableObject::InstanceUpdateBegin);
+
+		Instance->UpdatedDelegate.AddDynamic(HelperCallback, &UUpdateClassHelper::DelegatedCallback);
+		Instance->UpdateSkeletalMeshAsync(true);
+		InitStateAndParameterSequence();
+	}
+}
+
+
+void FTestingCustomizableObject::InitStateAndParameterSequence()
+{
+	BatchTestFinished = false;
+	BatchTestInCourse = true;
+
+	ArrayBatch.Empty();
+
+	InitTest();
+	
+	TickDelegate = FTickerDelegate::CreateRaw(this, &FTestingCustomizableObject::RunStateTestTick);
+	TickDelegateHandle = FTSTicker::GetCoreTicker().AddTicker(TickDelegate, 0.f);
+
+	if (ArrayBatch.Num() > 0)
+	{
+		ActiveBatchIndex = 0;
+		ActiveBatch = &ArrayBatch[0];
+	}
+}
+
+
+void FTestingCustomizableObject::FinishTest()
+{
+	WorstTimeDelegate.Unbind();
+	StateTestEnded.ExecuteIfBound();
+	if (Instance)
+	{
+		Instance->UpdateBeginDelegate.Unbind();
+		Instance->UpdatedDelegate.RemoveDynamic(HelperCallback, &UUpdateClassHelper::DelegatedCallback);
+	}
+
+	for(UCustomizableSkeletalComponent* CustomizableSkeletalComponent : CustomizableSkeletalComponents)
+	{
+		if (CustomizableSkeletalComponent)
+		{
+			CustomizableSkeletalComponent->DestroyComponent();
+		}
+	}
+
+	for (USkeletalMeshComponent* SkeletalMeshComponent : SkeletalMeshComponents)
+	{
+		if (SkeletalMeshComponent)
+		{
+			SkeletalMeshComponent->DestroyComponent();
+		}
+	}
+
+	CustomizableSkeletalComponents.Empty();
+	SkeletalMeshComponents.Empty();
+
+	Instance = nullptr; // NOTE: is this the right thing? doing like in the CO editor destructor
+
+	TestResourcesPresent = false;
+
+	TickDelegate = nullptr;
+	FTSTicker::GetCoreTicker().RemoveTicker(TickDelegateHandle);
+
+	GEditor->ForceGarbageCollection();
+
+	BatchTestInCourse = false;
+	BatchTestFinished = true;
+}
+
+
+void FTestingCustomizableObject::InstanceUpdateBegin()
+{
+	if (ActiveBatch != nullptr)
+	{
+		ActiveBatch->TestStarted = true;
+	}
+}
+
+
+void FTestingCustomizableObject::InstanceUpdateEnd()
+{
+	const int32 MutableRuntimeCycles = Instance->LastUpdateMutableRuntimeCycles;
+
+	if ((ActiveBatch != nullptr) && !ActiveBatch->TestFinished)
+	{
+		ActiveBatch->TestFinished = true;
+		ActiveBatch->ElapsedTime = FPlatformTime::ToMilliseconds(MutableRuntimeCycles);
+	}
+}
+
+
+bool FTestingCustomizableObject::RunStateTestTick(float DeltaTime)
+{
+	if (!ActiveBatch->TestStarted)
+	{
+		ActiveBatch->TestStarted = true;
+		LaunchNextBatch();
+		Instance->UpdateSkeletalMeshAsync(true);
+	}
+	
+	if ((ActiveBatch->TestStarted && ActiveBatch->TestFinished && ((ActiveBatchIndex + 1) < ArrayBatch.Num())))
+	{
+		InstanceTestUpdated(ActiveBatch);
+		ActiveBatchIndex++;
+		ActiveBatch = &ArrayBatch[ActiveBatchIndex];
+	}
+	
+	if (((ActiveBatchIndex + 1) == ArrayBatch.Num()) && ActiveBatch->TestStarted && ActiveBatch->TestFinished)
+	{
+		FinishTest();
+		return false;
+	}
+
+	return true;
+}
+
+
+bool FTestingCustomizableObject::GetTestResourcesPresent()
+{
+	return TestResourcesPresent;
+}
+
+
+int32 FTestingCustomizableObject::GetCurrentActiveBatchIndex()
+{
+	return ActiveBatchIndex;
+}
+
+
+int32 FTestingCustomizableObject::GetNumBatches()
+{
+	return ArrayBatch.Num();
+}
+
+bool FTestingCustomizableObject::GetBatchTestInCourse()
+{
+	return BatchTestInCourse;
+}
+
+void FTestingCustomizableObject::LaunchNextBatch()
+{
+	if (ActiveBatch == nullptr)
+	{
+		return;
+	}
+
+	if (ActiveBatch->IsStateChange)
+	{
+		Instance->SetState(ActiveBatch->State);
+	}
+	else
+	{
+		switch (ActiveBatch->ParameterType)
+		{
+			case EMutableParameterType::Bool:
+			{
+				Instance->SetBoolParameterSelectedOption(ActiveBatch->Name, ActiveBatch->ValueBool);
+				break;
+			}
+			case EMutableParameterType::Int:
+			{
+				Instance->SetIntParameterSelectedOption(ActiveBatch->Name, ActiveBatch->OptionName);
+				break;
+			}
+			case EMutableParameterType::Float:
+			{
+				Instance->SetFloatParameterSelectedOption(ActiveBatch->Name, ActiveBatch->ValueFloat);
+				break;
+			}
+			case EMutableParameterType::Color:
+			{
+				Instance->SetColorParameterSelectedOption(ActiveBatch->Name, ActiveBatch->ColorValue);
+				break;
+			}
+			case EMutableParameterType::Projector:
+			{
+				// For projectors we have small random offsets to apply to the position
+				FVector Position = Instance->GetProjectorPosition(ActiveBatch->Name);
+				Position *= ActiveBatch->VectorValue;
+				Instance->SetProjectorPosition(ActiveBatch->Name, (FVector3f)Position);
+				break;
+			}
+			case EMutableParameterType::Texture:
+			{
+				// What to do in this case?
+				break;
+			}
+			default:
+			{
+				break;
+			}
+		}
+	}
+}
+
+
+void FTestingCustomizableObject::InstanceTestUpdated(BatchParameterChange* BatchChange)
+{
+
+}
+
+
+void UUpdateClassHelper::DelegatedCallback(UCustomizableObjectInstance* Instance)
+{
+	if (StateData != nullptr)
+	{
+		StateData->InstanceUpdateEnd();
+	}
+}
+
+
+FRuntimeTest::FRuntimeTest() : FTestingCustomizableObject(TestType::RuntimeParameter)
+{
+}
+
+
+void FRuntimeTest::LoadStateData(int32 StateIndex)
+{
+	UCustomizableObject* CustomizableObject = Instance->GetCustomizableObject(); 
+
+	StateWholeInformation StateData;
+	StateData.State = StateIndex;
+	StateData.StateName = CustomizableObject->GetStateName(StateIndex);
+
+	ArrayBatch.Add(BatchParameterChange(StateIndex, StateData.StateName));
+
+	uint32 MaxParamChanges = 16;
+	
+	// TODO: test no-state available cases
+	uint32 StateParameterCount = CustomizableObject->GetStateParameterCount(StateIndex);
+	for (uint32 i = 0; i < StateParameterCount; ++i)
+	{
+		uint32 ParamIndexInObject = CustomizableObject->GetStateParameterIndex(StateIndex, i);
+		FString RunTimeParameter = CustomizableObject->GetParameterName(ParamIndexInObject);
+		EMutableParameterType Temp = CustomizableObject->GetParameterType(ParamIndexInObject);
+		switch (CustomizableObject->GetParameterType(ParamIndexInObject))
+		{
+		case EMutableParameterType::None:
+		{
+			break;
+		}
+		case EMutableParameterType::Bool:
+		{
+			TArray<FCustomizableObjectBoolParameterValue>& BoolParameters = Instance->GetBoolParameters();
+			FCustomizableObjectBoolParameterValue Parameter;
+			Parameter.ParameterName = RunTimeParameter;
+			const uint32 MaxIndex = BoolParameters.Num();
+			for (uint32 j = 0; j < MaxIndex; ++j)
+			{
+				if (BoolParameters[j].ParameterName == Parameter.ParameterName)
+				{
+					Parameter.ParameterValue = BoolParameters[j].ParameterValue;
+					StateData.BoolParametersIndex.Add(j);
+					for (uint32 k = 0; k < (MaxParamChanges / 2); ++k)
+					{
+						ArrayBatch.Add(BatchParameterChange(StateIndex, StateData.StateName, RunTimeParameter, Parameter.ParameterValue));
+						ArrayBatch.Add(BatchParameterChange(StateIndex, StateData.StateName, RunTimeParameter, !Parameter.ParameterValue));
+					}
+				}
+			}
+			StateData.BoolParameters.Add(Parameter);
+			break;
+		}
+		case EMutableParameterType::Int:
+		{
+			TArray<FCustomizableObjectIntParameterValue>& IntParameters = Instance->GetIntParameters();
+			FCustomizableObjectIntParameterValue Parameter;
+			Parameter.ParameterName = RunTimeParameter;
+			const uint32 MaxIndex = IntParameters.Num();
+			for (uint32 j = 0; j < MaxIndex; ++j)
+			{
+				if (IntParameters[j].ParameterName == Parameter.ParameterName)
+				{
+					//Parameter.ParameterValue = Descriptor.IntParameters[j].ParameterValue;
+					Parameter.ParameterValueName = IntParameters[j].ParameterValueName;
+					StateData.IntParametersIndex.Add(j);
+
+					int32 ParameterIndex = CustomizableObject->FindParameter(IntParameters[j].ParameterName);
+
+					uint32 NumOptions = CustomizableObject->GetIntParameterNumOptions(ParameterIndex);
+					for (uint32 l = 0; l < NumOptions && l < MaxParamChanges; ++l)
+					{
+						FString PossibleValue = CustomizableObject->GetIntParameterAvailableOption(ParameterIndex, l);
+						ArrayBatch.Add(BatchParameterChange(StateIndex, StateData.StateName, RunTimeParameter, PossibleValue));
+					}
+				}
+			}
+			StateData.IntParameters.Add(Parameter);
+			break; 
+		}
+		case EMutableParameterType::Float:
+		{
+			TArray<FCustomizableObjectFloatParameterValue>& FloatParameters = Instance->GetFloatParameters();
+			FCustomizableObjectFloatParameterValue Parameter;
+			Parameter.ParameterName = RunTimeParameter;
+			const uint32 MaxIndex = FloatParameters.Num();
+			for (uint32 j = 0; j < MaxIndex; ++j)
+			{
+				if (FloatParameters[j].ParameterName == Parameter.ParameterName)
+				{
+					Parameter.ParameterValue = FloatParameters[j].ParameterValue;
+					StateData.FloatParametersIndex.Add(j);
+
+					for (uint32 k = 0; k < MaxParamChanges; ++k)
+					{
+						ArrayBatch.Add(BatchParameterChange(StateIndex, StateData.StateName, RunTimeParameter, FMath::RandRange(0.0f, 1.0f)));
+					}
+					break;
+				}
+			}
+			StateData.FloatParameters.Add(Parameter);
+			break;
+		}
+		case EMutableParameterType::Color:
+		{
+			TArray<FCustomizableObjectVectorParameterValue>& VectorParameters = Instance->GetVectorParameters();
+			FCustomizableObjectVectorParameterValue Parameter;
+			Parameter.ParameterName = RunTimeParameter;
+			const uint32 MaxIndex = VectorParameters.Num();
+			for (uint32 j = 0; j < MaxIndex; ++j)
+			{
+				if (VectorParameters[j].ParameterName == Parameter.ParameterName)
+				{
+					Parameter.ParameterValue = VectorParameters[j].ParameterValue;
+					StateData.VectorParametersIndex.Add(j);
+					for (uint32 k = 0; k < MaxParamChanges; ++k)
+					{
+						ArrayBatch.Add(BatchParameterChange(StateIndex, StateData.StateName, RunTimeParameter, FLinearColor(FMath::RandRange(0.0f, 1.0f), FMath::RandRange(0.0f, 1.0f), FMath::RandRange(0.0f, 1.0f), 1.0f)));
+					}
+					break;
+				}
+			}
+			StateData.VectorParameters.Add(Parameter);
+			break;
+		}
+		case EMutableParameterType::Projector:
+		{
+			TArray<FCustomizableObjectVectorParameterValue>& VectorParameters = Instance->GetVectorParameters();
+			FCustomizableObjectVectorParameterValue Parameter;
+			Parameter.ParameterName = RunTimeParameter;
+			const uint32 MaxIndex = VectorParameters.Num();
+			for (uint32 j = 0; j < MaxIndex; ++j)
+			{
+				if (VectorParameters[j].ParameterName == Parameter.ParameterName)
+				{
+					Parameter.ParameterValue = VectorParameters[j].ParameterValue;
+					StateData.VectorParametersIndex.Add(j);
+					for (uint32 k = 0; k < MaxParamChanges; ++k)
+					{
+						// For projector we multiply the position with a small factor around 1
+						ArrayBatch.Add(BatchParameterChange(StateIndex, StateData.StateName, RunTimeParameter, FVector(FMath::RandRange(0.9f, 1.1f), FMath::RandRange(0.9f, 1.1f), FMath::RandRange(0.9f, 1.1f))));
+					}
+					break;
+				}
+			}
+			StateData.VectorParameters.Add(Parameter);
+			break;
+		}
+		case EMutableParameterType::Texture:
+		{
+			TArray<FCustomizableObjectTextureParameterValue>& TextureParameters = Instance->GetTextureParameters();
+			FCustomizableObjectTextureParameterValue Parameter;
+			Parameter.ParameterName = RunTimeParameter;
+			const uint32 MaxIndex = TextureParameters.Num();
+			for (uint32 j = 0; j < MaxIndex; ++j)
+			{
+				if (TextureParameters[j].ParameterName == Parameter.ParameterName)
+				{
+					Parameter.ParameterValue = TextureParameters[j].ParameterValue;
+					StateData.TextureParametersIndex.Add(j);
+					// How would be the testing here done?
+					break;
+				}
+			}
+			StateData.TextureParameters.Add(Parameter);
+			break;
+		}
+		default:
+		{
+			break;
+		}
+		}
+	}
+
+	ArrayState.Add(StateData);
+}
+
+
+void FRuntimeTest::InitTest()
+{
+	MapStateEnterData.Empty();
+	MapStateUpdateData.Empty();
+	MapStateRuntimeData.Empty();
+
+	// Build a batch of changes for each state
+	const int32 NumStates = Instance->GetCustomizableObject()->GetStateCount();
+
+	// Enter each state and, for each runtime parameter, set random values (in the case
+	// of integer parameter, set all possible values)
+	for (int32 i = 0; i < NumStates; ++i)
+	{
+		LoadStateData(i);
+	}
+
+	const int32 NumStateChanges = NumStates * 50;
+	uint32 IndexState;
+	for (int32 i = 0; i < NumStateChanges; ++i)
+	{
+		IndexState = FMath::Rand() % NumStates;
+		ArrayBatch.Add(BatchParameterChange(IndexState, ArrayState[IndexState].StateName));
+	}
+}
+
+
+TMap<uint32, MeasuredData> FRuntimeTest::GetMapStateEnterData()
+{
+	return MapStateEnterData;
+}
+
+
+TMap<uint32, MeasuredData> FRuntimeTest::GetMapStateUpdateData()
+{
+	return MapStateUpdateData;
+}
+
+
+TMap<uint32, TArray<MeasuredData>> FRuntimeTest::GetMapStateRuntimeData()
+{
+	return MapStateRuntimeData;
+}
+
+
+void FRuntimeTest::InstanceTestUpdated(BatchParameterChange* BatchChange)
+{
+	if (BatchChange->IsStateChange)
+	{
+		MeasuredData* DataPtr = MapStateEnterData.Find(BatchChange->State);
+
+		if (DataPtr == nullptr)
+		{
+			MeasuredData& Data = MapStateEnterData.Add(BatchChange->State, MeasuredData(BatchChange->StateName));
+			AddMeasuredDataAndReportTime(&Data, BatchChange->ElapsedTime, BatchChange->State);
+			TArray<MeasuredData> ArrayMeasuredData;
+			MapStateRuntimeData.Add(BatchChange->State, ArrayMeasuredData);
+		}
+		else
+		{
+			AddMeasuredDataAndReportTime(DataPtr, BatchChange->ElapsedTime, BatchChange->State);
+		}
+	}
+	else
+	{
+		// Add current change in course to statistics about the instance being previewed
+		MeasuredData* DataPtr = FindStateRuntimeParameter(BatchChange->State, BatchChange->Name);
+
+		if (DataPtr == nullptr)
+		{
+			MeasuredData Data = MeasuredData(BatchChange->Name);
+			AddMeasuredDataAndReportTime(&Data, BatchChange->ElapsedTime, BatchChange->State, BatchChange->Name);
+			TArray<MeasuredData>* ArrayMeasuredData = MapStateRuntimeData.Find(BatchChange->State);
+			ArrayMeasuredData->Add(Data);
+		}
+		else
+		{
+			AddMeasuredDataAndReportTime(DataPtr, BatchChange->ElapsedTime, BatchChange->State, BatchChange->Name);
+		}
+
+		DataPtr = MapStateUpdateData.Find(BatchChange->State);
+		if (DataPtr == nullptr)
+		{
+			MeasuredData Data = MeasuredData(BatchChange->StateName);
+			AddMeasuredDataAndReportTime(&Data, BatchChange->ElapsedTime, BatchChange->State, BatchChange->Name);
+			MapStateUpdateData.Add(BatchChange->State, Data);
+		}
+		else
+		{
+			AddMeasuredDataAndReportTime(DataPtr, BatchChange->ElapsedTime, BatchChange->State, BatchChange->Name);
+		}
+	}
+}
+
+
+MeasuredData* FRuntimeTest::FindStateRuntimeParameter(uint32 State, FString& Name)
+{
+	TArray<MeasuredData>* ArrayData = MapStateRuntimeData.Find(State);
+
+	const uint32 MaxIndex = ArrayData->Num();
+	for (uint32 i = 0; i < MaxIndex; ++i)
+	{
+		if ((*ArrayData)[i].Name == Name)
+		{
+			return &(*ArrayData)[i];
+		}
+	}
+
+	return nullptr;
+}
+
+
+void FRuntimeTest::AddMeasuredDataAndReportTime(MeasuredData* DataPtr, float ElapsedTime, int32 StateIndex, FString ParameterName)
+{
+	if (DataPtr)
+	{
+		if (ElapsedTime > DataPtr->MaxTime)
+		{
+			if (ParameterName.IsEmpty())
+			{
+				WorstTimeDelegate.ExecuteIfBound(ElapsedTime, StateIndex, INDEX_NONE, 0/* TODO: Get lod!*/, Instance, DataPtr);
+			}
+			else
+			{
+				int32 ParameterIndexInCO = INDEX_NONE;
+
+				if (Instance)
+				{
+					if (const UCustomizableObject* Object = Instance->GetCustomizableObject())
+					{
+						ParameterIndexInCO = Object->FindParameter(ParameterName);
+					}
+				}
+
+				if (ParameterIndexInCO != INDEX_NONE)
+				{
+					WorstTimeDelegate.ExecuteIfBound(ElapsedTime, StateIndex, ParameterIndexInCO, 0/* TODO: Get lod!*/, Instance, DataPtr);
+				}
+			}
+		}
+		DataPtr->AddMeasureData(ElapsedTime);
+	}
+}
+
+
+#undef LOCTEXT_NAMESPACE
diff --git a/Engine/Plugins/Experimental/Mutable/Source/CustomizableObjectEditor/Private/StatePerformanceTesting.h b/Engine/Plugins/Experimental/Mutable/Source/CustomizableObjectEditor/Private/StatePerformanceTesting.h
new file mode 100644
--- /dev/null
+++ b/Engine/Plugins/Experimental/Mutable/Source/CustomizableObjectEditor/Private/StatePerformanceTesting.h
@@ -0,0 +1,478 @@
+// Copyright Epic Games, Inc. All Rights Reserved.
+
+#pragma once
+
+#include "ICustomizableObjectEditor.h"
+#include "Misc/NotifyHook.h"
+#include "CustomizableObjectCompiler.h"
+#include "CustomizableObjectInstance.h"
+#include "StatePerformanceTesting.generated.h"
+
+
+class UInstance;
+class UCustomizableSkeletalComponent;
+class UStaticMeshComponent;
+class USkeletalMeshComponent;
+class UUpdateClassHelper;
+
+
+enum class TestType
+{
+	RuntimeParameter,
+	CustomizableObjectMemory
+};
+
+
+struct MeasuredData
+{
+	MeasuredData() :
+		Name("")
+		, MeanTime(0.0f)
+		, MinTime(FLT_MAX)
+		, MaxTime(FLT_MIN)
+		, Accumulated(0.0f)
+		, NumMeasures(0.0f)
+	{
+
+	}
+
+	MeasuredData(FString& NameParam) :
+		Name(NameParam)
+		, MeanTime(0.0f)
+		, MinTime(FLT_MAX)
+		, MaxTime(FLT_MIN)
+		, Accumulated(0.0f)
+		, NumMeasures(0.0f)
+	{
+
+	}
+
+	void AddMeasureData(float ElapsedTime)
+	{
+		Accumulated += ElapsedTime;
+		MinTime = FMath::Min(MinTime, ElapsedTime);
+		MaxTime = FMath::Max(MaxTime, ElapsedTime);
+		NumMeasures += 1.0f;
+		MeanTime = Accumulated / NumMeasures;
+	}
+
+
+	MeasuredData operator+=(const MeasuredData& Other)
+	{
+		if (this->Name != Other.Name)
+		{
+			return *this;
+		}
+
+		if ((HasInitialValues(*this) && HasInitialValues(Other)) ||
+			(!HasInitialValues(*this) && HasInitialValues(Other)))
+		{
+			return *this;
+		}
+		else if (HasInitialValues(*this) && !HasInitialValues(Other))
+		{
+			*this = Other;
+			return *this;
+		}
+
+		this->Accumulated += Other.Accumulated;
+		this->MaxTime = FMath::Max(this->MaxTime, Other.MaxTime);
+		this->MinTime = FMath::Min(this->MinTime, Other.MinTime);
+		this->NumMeasures += Other.NumMeasures;
+		this->MeanTime = this->Accumulated / this->NumMeasures;
+		return *this;
+	}
+
+	bool HasInitialValues(const MeasuredData& Data)
+	{
+		if ((Data.MeanTime == 0.0f) &&
+			(Data.MinTime == FLT_MAX) &&
+			(Data.MaxTime == FLT_MIN) &&
+			(Data.Accumulated == 0.0f) &&
+			(Data.NumMeasures == 0.0f))
+			{
+				return true;
+			}
+
+		return false;
+	}
+
+
+	FString Name;
+	float MeanTime;
+	float MinTime;
+	float MaxTime;
+	float Accumulated;
+	float NumMeasures;
+};
+
+
+struct BatchParameterChange
+{
+	/** Float change constructor */
+	BatchParameterChange(int32 StateParam, FString& StateName, FString& NameParam, float ValueFloatParam) :
+		  State(StateParam)
+		, StateName(StateName)
+		, ParameterType(EMutableParameterType::Float)
+		, Name(NameParam)
+		, ValueFloat(ValueFloatParam)
+		, ValueBool(false)
+		, IsStateChange(false)
+		, ElapsedTime(0.0f)
+		, TestStarted(false)
+		, TestFinished(false)
+	{
+
+	}
+
+	/** Integer parameter change constructor */
+	BatchParameterChange(int32 StateParam, FString& StateName, FString& NameParam, FString ParameterOptionParam) :
+		  State(StateParam)
+		, StateName(StateName)
+		, ParameterType(EMutableParameterType::Int)
+		, Name(NameParam)
+		, ValueFloat(-1.0f)
+		, OptionName(ParameterOptionParam)
+		, ValueBool(false)
+		, IsStateChange(false)
+		, ElapsedTime(0.0f)
+		, TestStarted(false)
+		, TestFinished(false)
+	{
+
+	}
+
+	/** Bool change constructor */
+	BatchParameterChange(int32 StateParam, FString& StateName, FString& NameParam, bool ValueBoolParam) :
+		  State(StateParam)
+		, StateName(StateName)
+		, ParameterType(EMutableParameterType::Bool)
+		, Name(NameParam)
+		, ValueFloat(-1.0f)
+		, ValueBool(ValueBoolParam)
+		, IsStateChange(false)
+		, ElapsedTime(0.0f)
+		, TestStarted(false)
+		, TestFinished(false)
+	{
+
+	}
+
+	/** Color change constructor */
+	BatchParameterChange(int32 StateParam, FString& StateName, FString& NameParam, FLinearColor ColorValueParam) :
+		  State(StateParam)
+		, StateName(StateName)
+		, ParameterType(EMutableParameterType::Color)
+		, Name(NameParam)
+		, ValueFloat(-1.0f)
+		, ValueBool(false)
+		, ColorValue(ColorValueParam)
+		, IsStateChange(false)
+		, ElapsedTime(0.0f)
+		, TestStarted(false)
+		, TestFinished(false)
+	{
+
+	}
+
+	/** Projector position change constructor */
+	BatchParameterChange(int32 StateParam, FString& StateName, FString& NameParam, FVector VectorValueParam) :
+		  State(StateParam)
+		, StateName(StateName)
+		, ParameterType(EMutableParameterType::Projector)
+		, Name(NameParam)
+		, ValueFloat(-1.0f)
+		, ValueBool(false)
+		, VectorValue(VectorValueParam)
+		, IsStateChange(false)
+		, ElapsedTime(0.0f)
+		, TestStarted(false)
+		, TestFinished(false)
+	{
+
+	}
+
+	/** State change constructor */
+	BatchParameterChange(int32 StateParam, FString& NameParam) :
+		  State(StateParam)
+		, StateName(NameParam)
+		, ParameterType(EMutableParameterType::None)
+		, Name(NameParam)
+		, ValueFloat(-1.0f)
+		, ValueBool(false)
+		, IsStateChange(true)
+		, ElapsedTime(0.0f)
+		, TestStarted(false)
+		, TestFinished(false)
+	{
+
+	}
+
+	/** Default constructor */
+	BatchParameterChange() :
+		  State(-1)
+		, ParameterType(EMutableParameterType::None)
+		, ValueFloat(-1.0f)
+		, ValueBool(false)
+		, IsStateChange(false)
+		, ElapsedTime(0.0f)
+		, TestStarted(false)
+		, TestFinished(false)
+	{
+
+	}
+
+	/** State of the parameter, in case this change in the batch is for a parameter. Otherwise, is the state to set (as a batch change) */
+	int32 State = -1;
+
+	/** State name (parameter name in the case of a parameter, or state name in the case of a state) */
+	FString StateName;
+
+	/** Parameter type, in case this change in the batch is for a parameter */
+	EMutableParameterType ParameterType = EMutableParameterType::None;
+
+	/** Name (parameter name in the case of a parameter, or state name in the case of a state) */
+	FString Name;
+
+	/** Value to change the parameter to, in case it's a float parameter */
+	float ValueFloat = -1.0f;
+
+	/** Value to change the parameter to, in case it's an integer parameter */
+	FString OptionName;
+
+	/** Value to change the parameter to, in case it's a bool parameter */
+	bool ValueBool = false;
+
+	/** Value to change the parameter to, in case it's a color parameter */
+	FLinearColor ColorValue;
+
+	/** Value to change the parameter to, in case it's a projector parameter: offset the position */
+	FVector VectorValue;
+
+	/** Flag to distinguish between state and parameter changes */
+	bool IsStateChange = false;
+
+	/** Th elpased time to perform this batch change */
+	float ElapsedTime = 0.0f;
+
+	/** Flag to know if the test for this batch test started */
+	bool TestStarted = false;
+
+	/** Flag to know if the test for this batch test finished */
+	bool TestFinished = false;
+};
+
+
+struct StateWholeInformation
+{
+	/** State index */
+	int32 State;
+
+	/** State name */
+	FString StateName;
+
+	/** Copy of run time boolean parameters */
+	TArray<struct FCustomizableObjectBoolParameterValue> BoolParameters;
+
+	/** Index in the Customizable Object of the parameters in BoolParameters */
+	TArray<int32> BoolParametersIndex;
+
+	/** Copy of run time integer parameters */
+	TArray<struct FCustomizableObjectIntParameterValue> IntParameters;
+
+	/** Index in the Customizable Object of the parameters in IntParameters */
+	TArray<int32> IntParametersIndex;
+
+	/** Copy of run time float parameters */
+	TArray<struct FCustomizableObjectFloatParameterValue> FloatParameters;
+
+	/** Index in the Customizable Object of the parameters in FloatParameters */
+	TArray<int32> FloatParametersIndex;
+
+	/** Copy of run time texture parameters */
+	TArray<struct FCustomizableObjectTextureParameterValue> TextureParameters;
+
+	/** Index in the Customizable Object of the parameters in TextureParameters */
+	TArray<int32> TextureParametersIndex;
+
+	/** Copy of run time vector parameters */
+	TArray<struct FCustomizableObjectVectorParameterValue> VectorParameters;
+
+	/** Index in the Customizable Object of the parameters in VectorParameters */
+	TArray<int32> VectorParametersIndex;
+
+	/** Copy of run time projector parameters */
+	TArray<struct FCustomizableObjectProjectorParameterValue> ProjectorParameters;
+
+	/** Index in the Customizable Object of the parameters in ProjectorParameters */
+	TArray<int32> ProjectorParametersIndex;
+};
+
+
+// Callbacks to notify of worst cases found when running the test. It can be called multiple times per test, as worse cases pop up.
+DECLARE_DELEGATE_SixParams(FStateTestWorstTimeFound, float, int32, int32, uint32, UCustomizableObjectInstance*, MeasuredData*);
+// Callback to notify that the test has finished, before data is cleared
+DECLARE_DELEGATE(FStateTestEnded);
+
+
+class FTestingCustomizableObject : public FGCObject
+{
+public:
+
+	// FGCObject interface
+	virtual void AddReferencedObjects(FReferenceCollector& Collector) override;
+	virtual FString GetReferencerName() const override;
+
+public:
+	/** Default constructor */
+	FTestingCustomizableObject(TestType TestParameter);
+
+	/** Default destructor */
+	virtual ~FTestingCustomizableObject();
+
+	/** Init test */
+	virtual void InitTest();
+
+	/** Start the test */
+	void StartTest(UCustomizableObject* CustomizableObject);
+
+	/** Inits the loop through all states and runtime parameters */
+	void InitStateAndParameterSequence();
+
+	/** Finish the test */
+	void FinishTest();
+
+	/** Callback for instance update begin delegate */
+	void InstanceUpdateBegin();
+
+	/** Callback for instance update end delegate */
+	void InstanceUpdateEnd();
+
+	/** Getter of ActiveBatchIndex */
+	int32 GetCurrentActiveBatchIndex();
+
+	/** Getter of ArrayBatch.Num() */
+	int32 GetNumBatches();
+
+	/** Getter of BatchTestInCourse */
+	bool GetBatchTestInCourse();
+
+	/** Coordinates the advaces of each test of the batch test */
+	bool RunStateTestTick(float DeltaTime);
+
+	/** Instance to build to perform the test */
+	UCustomizableObjectInstance* Instance;
+
+	/** Instance skeletal component */
+	TArray<UCustomizableSkeletalComponent*> CustomizableSkeletalComponents;
+
+	/** Instance skeletal mesh component */
+	TArray<USkeletalMeshComponent*> SkeletalMeshComponents;
+
+	/** UObject class to be able to use the update callback */
+	UPROPERTY()
+	UUpdateClassHelper* HelperCallback;
+
+	/** Getter of TestResourcesPresent */
+	bool GetTestResourcesPresent();
+
+	/** Delegate called when the worst generation time is found UP UNTIL NOW, with its info and a pointer to the instance. It can be called multiple times on a single test run. */
+	FStateTestWorstTimeFound WorstTimeDelegate;
+
+	/** Delegate called when the stress test finishes, before any data is erased. */
+	FStateTestEnded StateTestEnded;
+
+protected:
+	/** Performs the change in state or runtime parameter according to the information in ActiveBatch */
+	void LaunchNextBatch();
+
+	/** To be called when a running test updates its instance parameters after a BatchParameterChange execution */
+	virtual void InstanceTestUpdated(BatchParameterChange* BatchChange);
+
+	/** Flag to know if there's a batch test running */
+	bool BatchTestInCourse;
+
+	/** Flag to know if the last ran batch test finished */
+	bool BatchTestFinished;
+
+	/** All Customizable Object state information to generate a proper batch of changes */
+	TArray<StateWholeInformation> ArrayState;
+
+	/** Array with the batch of changes to apply */
+	TArray<BatchParameterChange> ArrayBatch;
+
+	/** Active batch change index */
+	int32 ActiveBatchIndex;
+
+	/** Active batch change */
+	BatchParameterChange* ActiveBatch;
+
+	/** Handle for the delegate to run the Tick method in charge of coordinating the test */
+	FTSTicker::FDelegateHandle TickDelegateHandle;
+
+	/** Delegate to run the Tick method in charge of coordinating the test */
+	FTickerDelegate TickDelegate;
+
+	/** Flag to control when test resources are built (for AddReferencedObjects method) */
+	bool TestResourcesPresent;
+
+	/** Enum to know the type of test is being carried out by this instance */
+	TestType Test;
+};
+
+
+UCLASS()
+class UUpdateClassHelper : public UObject
+{
+public:
+	GENERATED_BODY()
+
+	/** Method to assign for the callback */
+	UFUNCTION()
+	void DelegatedCallback(UCustomizableObjectInstance* Instance);
+
+	/** Pointer to existing class to notify */
+	FTestingCustomizableObject* StateData = nullptr;
+};
+
+
+/** Class used to make the runtime parameter and state enter tests */
+class FRuntimeTest : public FTestingCustomizableObject
+{
+public:
+	/** Default constructor */
+	FRuntimeTest();
+
+	/** Loads the state and runtime parameter data of the instance */
+	void LoadStateData(int32 StateIndex);
+
+	/** Init test */
+	void InitTest() final override;
+
+	/** Getter of MapStateEnterData (copy to avoid any possible thread issues) */
+	TMap<uint32, MeasuredData> GetMapStateEnterData();
+
+	/** Getter of MapStateUpdateData (copy to avoid any possible thread issues) */
+	TMap<uint32, MeasuredData> GetMapStateUpdateData();
+
+	/** Getter of MapStateRuntimeData (copy to avoid any possible thread issues) */
+	TMap<uint32, TArray<MeasuredData>> GetMapStateRuntimeData();
+
+protected:
+	/** To be called when a running test updates its instance parameters after a BatchParameterChange execution */
+	void InstanceTestUpdated(BatchParameterChange* BatchChange) final override;
+
+	/** Retrieve pointer to runtime parameter of name given as parameter, for state given as parameter */
+	MeasuredData* FindStateRuntimeParameter(uint32 State, FString& Name);
+
+	void AddMeasuredDataAndReportTime(MeasuredData* DataPtr, float ElapsedTime, int32 StateIndex, FString ParameterName = FString());
+
+	/** Map storing state id -> state entering time data */
+	TMap<uint32, MeasuredData> MapStateEnterData;
+
+	/** Map storing state id -> state update time data */
+	TMap<uint32, MeasuredData> MapStateUpdateData;
+
+	/** Map storing state id -> state runtime parameters measure data */
+	TMap<uint32, TArray<MeasuredData>> MapStateRuntimeData;
+};
\ No newline at end of file
diff --git a/Engine/Plugins/Experimental/Mutable/Source/CustomizableObjectEditor/Private/StressTest.cpp b/Engine/Plugins/Experimental/Mutable/Source/CustomizableObjectEditor/Private/StressTest.cpp
new file mode 100644
--- /dev/null
+++ b/Engine/Plugins/Experimental/Mutable/Source/CustomizableObjectEditor/Private/StressTest.cpp
@@ -0,0 +1,1079 @@
+// Copyright Epic Games, Inc. All Rights Reserved.
+
+#include "StressTest.h"
+#include "UObject/Package.h"
+#include "CustomizableObject.h"
+#include "CustomizableObjectInstance.h"
+#include "Widgets/Notifications/SNotificationList.h"
+#include "Engine/SkeletalMesh.h"
+#include "Materials/Material.h"
+#include "Materials/MaterialInstance.h"
+#include "EditorFramework/AssetImportData.h"
+#include "UnrealPortabilityHelpers.h"
+
+
+void ULiveInstance::DelegatedCallback(UCustomizableObjectInstance* UpdatedInstance)
+{
+	AddInstanceInformation();
+	Updated = true;
+}
+
+
+void ULiveInstance::AddInstanceInformation()
+{
+	if(!Instance->HasAnySkeletalMesh())
+	{
+		UE_LOG(LogMutable, Warning, TEXT("ERROR: no skeletal mesh in ULiveInstance::AddInstanceInformataion"));
+		return;
+	}
+
+	FSkeletalMeshRenderData* RenderData = Helper_GetResourceForRendering(Instance->GetSkeletalMesh());
+	TIndirectArray<FSkeletalMeshLODRenderData>* ArrayLODRenderData = Helper_GetLODDataPtr(RenderData);
+	int32 MaxIndexLOD = Helper_GetLODDataPtr(RenderData)->Num();
+
+	for (int32 i = 0; i < MaxIndexLOD; ++i)
+	{
+		TMap<FString, MaterialBriefInfo> MapMaterialInfo;
+		uint32 CurrentLODTextureSize = 0;
+		uint32 CurrentLODNumFaces = 0;
+
+		TMap<FString, MeasuredData> MapMeasuredData = ULiveInstance::InitInstanceInformation();
+
+		for (int32 c = 0; c < Instance->SkeletalMeshes.Num(); ++c)
+		{
+			RenderData = Helper_GetResourceForRendering(Instance->GetSkeletalMesh(c));
+			ArrayLODRenderData = Helper_GetLODDataPtr(RenderData);
+			MaxIndexLOD = Helper_GetLODDataPtr(RenderData)->Num();
+			CurrentLODNumFaces = Helper_LODGetTotalFaces(&((*ArrayLODRenderData)[i]));
+			const TArray<struct FSkeletalMeshLODInfo>& LODInfoArray = Helper_GetLODInfoArray(Instance->GetSkeletalMesh(c));
+			if (LODInfoArray.IsValidIndex(i))
+			{
+				const FSkeletalMeshLODInfo& LODInfo = LODInfoArray[i];
+
+				MapMeasuredData["Number of faces"].AddMeasureData(Helper_LODGetTotalFaces(&((*ArrayLODRenderData)[i])));
+				MapMeasuredData["Number of texture coordinates"].AddMeasureData(Helper_LODGetNumTexCoords(&((*ArrayLODRenderData)[i])));
+				const TArray<FSkeletalMaterial>& Materials = Instance->GetSkeletalMesh(c)->GetMaterials();
+				MapMeasuredData["Number of materials"].AddMeasureData(LODInfo.LODMaterialMap.Num());
+
+				for (int32 j = 0; j < LODInfo.LODMaterialMap.Num(); ++j)
+				{
+					if (Materials.IsValidIndex(LODInfo.LODMaterialMap[j]))
+					{
+						UMaterialInterface* MaterialInterface = Materials[LODInfo.LODMaterialMap[j]].MaterialInterface;
+						UMaterialInstance* BaseMaterial = Cast<UMaterialInstance>(MaterialInterface);
+						FString MaterialPath = MaterialInterface->GetPathName();
+						MaterialBriefInfo Info;
+
+						Info.Name = BaseMaterial ? BaseMaterial->GetName() : MaterialInterface->GetName();
+
+						if (BaseMaterial && MaterialPath.Contains("Transient"))
+						{
+							for (int32 k = 0; k < BaseMaterial->TextureParameterValues.Num(); ++k)
+							{
+								if (UTexture* Texture = BaseMaterial->TextureParameterValues[k].ParameterValue)
+								{
+									FString TextPath = Texture->GetPathName();
+									if (Texture->GetPathName().Contains("Transient"))
+									{
+										FIntVector Size = FIntVector(0, 0, 0);
+										Info.ArrayTextureName.Add(Texture->GetName());
+										Size.X = Texture->GetSurfaceWidth();
+										Size.Y = Texture->GetSurfaceHeight();
+										uint32 CurrentTextureSize = Texture->CalcTextureMemorySizeEnum(ETextureMipCount::TMC_AllMips);
+										CurrentLODTextureSize += CurrentTextureSize;
+										MapMeasuredData["Texture size"].AddMeasureData(float(CurrentTextureSize));
+										Info.ArrayTextureSize.Add(Size);
+										Info.ArrayTextureFormat.Add(((UTexture2D*)Texture)->GetPixelFormat());
+									}
+								}
+							}
+							MapMaterialInfo.Add(Info.Name, Info);
+						}
+					}
+				}
+				MapMeasuredData["Instance texture size"].AddMeasureData(CurrentLODTextureSize);
+			}
+		}
+
+		TestDelegate.ExecuteIfBound(MapMeasuredData, MapMaterialInfo, CurrentLODTextureSize, CurrentLODNumFaces, uint32(i), Instance);
+
+		NumInstanceUpdated++;
+	}
+}
+
+
+TMap<FString, MeasuredData> ULiveInstance::InitInstanceInformation()
+{
+	TMap<FString, MeasuredData> MapResult;
+
+	TArray<FString> ArrayParam;
+
+	ArrayParam.Add("Number of faces");
+	ArrayParam.Add("Number of materials");
+	ArrayParam.Add("Number of texture coordinates");
+	ArrayParam.Add("Texture size");
+	ArrayParam.Add("Instance texture size");
+
+	int32 MaxIndex = ArrayParam.Num();
+
+	for (int32 i = 0; i < MaxIndex; ++i)
+	{
+		MapResult.Add(ArrayParam[i], MeasuredData(ArrayParam[i]));
+	}
+
+	return MapResult;
+}
+
+
+FRunningStressTest::FRunningStressTest()
+{
+	NumUpdateRequested = 0;
+}
+
+
+bool FRunningStressTest::RunStressTestTick(float DeltaTime)
+{
+	if (PendingInstanceCount <= 0)
+	{
+		VerifyInstancesThreshold -= DeltaTime;
+
+		if (VerifyInstancesThreshold > 0.0f)
+		{
+			return true;
+		}
+
+		// Verify all instances in the test made their last update before finising the test properly
+		for (TPair<UCustomizableObjectInstance*, ULiveInstance*>& it : LiveInstances)
+		{
+			if (!it.Value->Updated)
+			{
+				VerifyInstancesThreshold = 5.0f;
+				return true;
+			}
+		}
+
+		if (NumUpdatedDone == NumUpdateRequested)
+		{
+			UE_LOG(LogMutable, Warning, TEXT("\t\t\t Final Stress Test results for model %s"), *CustomizableObject->GetName());
+		}
+
+		WriteTestResults();
+
+		PrintMostExpensivePerLOD();
+
+		UE_LOG(LogMutable, Warning, TEXT("///////////////////////////////////////////////////////////////////////////////////"));
+
+		UE_LOG(LogMutable, Warning, TEXT("Finished stress test for model [%s]."), *CustomizableObject->GetName());
+		UE_LOG(LogMutable, Warning, TEXT("\n"));
+
+		FinishTest();
+
+		return false;
+	}
+	else
+	{
+		int thisFrameTime = int(DeltaTime*1000.0f);
+		TotalTimeMs += thisFrameTime;
+
+		// Create a new instance if necessary
+		int remainingTime = thisFrameTime;
+		NextInstanceTimeMs -= remainingTime;
+		while (NextInstanceTimeMs < 0)
+		{
+			--PendingInstanceCount;
+
+			// Remaining time
+			remainingTime = -NextInstanceTimeMs;
+
+			UCustomizableObjectInstance* Instance = CustomizableObject->CreateInstance();
+
+			if (!ArrayLODResized)
+			{
+				NumLOD = CustomizableObject->GetNumLODs();
+				ArrayLODMostExpensiveGeometry.SetNum(NumLOD);
+				ArrayLODMostExpensiveTexture.SetNum(NumLOD);
+				ArrayMeasuredDataPerLOD.SetNum(NumLOD);
+				ArrayMaterialInfoPerLOD.SetNum(NumLOD);
+
+				for (uint32 i = 0; i < NumLOD; ++i)
+				{
+					ArrayMeasuredDataPerLOD[i] = ULiveInstance::InitInstanceInformation();
+				}
+
+				ArrayLODResized = true;
+			}
+
+			Instance->SetRandomValues();
+
+			UCustomizableSkeletalComponent* Component = NewObject<UCustomizableSkeletalComponent>(GetTransientPackage(), NAME_None, RF_Transient);
+			Component->CustomizableObjectInstance = Instance;
+
+			ULiveInstance* Data = NewObject<ULiveInstance>(GetTransientPackage());
+			Instance->UpdatedDelegate.AddDynamic(Data, &ULiveInstance::DelegatedCallback);
+			Instance->UpdateSkeletalMeshAsync(true);
+			Data->Updated = false;
+
+			Data->TimeToDie = FMath::Max(InstanceLifeTimeMs, InstanceLifeTimeMs + (FMath::Rand() % (InstanceLifeTimeMsVar * 2)) - InstanceLifeTimeMsVar);
+			Data->TimeToUpdate = InstanceUpdateTimeMs
+				+ (FMath::Rand() % (InstanceUpdateTimeMsVar * 2))
+				- InstanceUpdateTimeMsVar;
+
+			Data->Instance = Instance;
+			Data->NumInstanceUpdated = 0;
+			Data->TestDelegate.BindRaw(this, &FRunningStressTest::AddInstanceData);
+
+			LiveInstances.Add(Instance, Data);
+			NumUpdateRequested++;
+
+			NextInstanceTimeMs = CreateInstanceTimeMs
+				+ (FMath::Rand() % (CreateInstanceTimeMsVar * 2))
+				- CreateInstanceTimeMsVar;
+
+			if (remainingTime > 0)
+			{
+				NextInstanceTimeMs -= remainingTime;
+			}
+
+			UE_LOG(LogMutable, Warning, TEXT("Stress test created instance. [%d active]"), LiveInstances.Num());
+		}
+
+		// Update instances that need it
+		for (TPair<UCustomizableObjectInstance*, ULiveInstance*>& it : LiveInstances)
+		{
+			if (!it.Value->Updated)
+			{
+				continue;
+			}
+
+			it.Value->TimeToUpdate -= thisFrameTime;
+			if (it.Value->TimeToUpdate < 0)
+			{
+				it.Key->SetRandomValues();
+				it.Key->UpdateSkeletalMeshAsync(true);
+
+				it.Value->Updated = false;
+				it.Value->TimeToUpdate = InstanceUpdateTimeMs
+					+ (FMath::Rand() % (InstanceUpdateTimeMsVar * 2))
+					- InstanceUpdateTimeMsVar;
+			}
+		}
+
+		// Destroy instances that expired
+		TArray<UCustomizableObjectInstance*> ToRemove;
+		for (TPair<UCustomizableObjectInstance*, ULiveInstance*>& it : LiveInstances)
+		{
+			if (!it.Value) continue;
+
+			it.Value->TimeToDie -= thisFrameTime;
+			if (it.Value->TimeToDie < 0 && it.Value->Updated)
+			{
+				it.Key->UpdatedDelegate.RemoveDynamic(it.Value, &ULiveInstance::DelegatedCallback);
+				it.Value = nullptr;
+				ToRemove.Add(it.Key);
+			}
+		}
+
+		for (UCustomizableObjectInstance* i : ToRemove)
+		{
+			LiveInstances.Remove(i);
+		}
+
+		if (ToRemove.Num())
+		{
+			UE_LOG(LogMutable, Warning, TEXT("Stress test removed instance. [%d active]"), LiveInstances.Num());
+		}
+	}
+
+	return true;
+}
+
+
+void FRunningStressTest::InitMeasureTest()
+{
+	StressTestReadyToReset = false;
+	NumUpdateRequested = 0;
+	NumUpdatedDone = 0;
+	ArrayMeasuredDataPerLOD.Empty();
+	ArrayMaterialInfoPerLOD.Empty();
+
+	int32 MaxValue = TNumericLimits<int32>::Max();
+	int32 MinValue = TNumericLimits<int32>::Min();
+
+	TextureMinSize = FIntVector(MaxValue, MaxValue, 0);
+	TextureMaxSize = FIntVector(MinValue, MinValue, 0);
+
+	TestInCourse = true;
+
+	MinTextureSize = TNumericLimits<uint32>::Max();
+	MeanTextureSize = 0;
+	MaxTextureSize = 0;
+
+	StressTestTickDelegate = FTickerDelegate::CreateRaw(this, &FRunningStressTest::RunStressTestTick);
+	StressTestTickDelegateHandle = FTSTicker::GetCoreTicker().AddTicker(StressTestTickDelegate, 0.f);
+
+	ArrayLODMostExpensiveGeometry.Empty();
+	ArrayLODMostExpensiveTexture.Empty();
+
+	ArrayLODResized = false;
+	NumLOD = 0;
+}
+
+
+void FRunningStressTest::WriteTestResults()
+{
+	FString Line;
+	Line = "///////////////////////////////////////////////////////////////////////////////////";
+	UE_LOG(LogMutable, Warning, TEXT("%s"), *Line);
+
+	for (uint32 i = 0; i < NumLOD; ++i)
+	{
+		int32 MaxOffsetName = 0;
+		for (TPair<FString, MeasuredData>& it : ArrayMeasuredDataPerLOD[i])
+		{
+			MaxOffsetName = FMath::Max(MaxOffsetName, it.Key.Len());
+		}
+		MaxOffsetName += 5;
+
+		MinTextureSize = uint32(ArrayMeasuredDataPerLOD[i]["Texture size"].MinTime);
+		MeanTextureSize = uint32(ArrayMeasuredDataPerLOD[i]["Texture size"].MeanTime);
+		MaxTextureSize = uint32(ArrayMeasuredDataPerLOD[i]["Texture size"].MaxTime);
+
+		Line = "Results for LOD level ";
+		Line += FString::Printf(TEXT("%d"), i);
+		UE_LOG(LogMutable, Warning, TEXT("%s:"), *Line);
+
+		uint32 MaxLength1 = MaxOffsetName + 25;
+		uint32 MaxLength2 = MaxLength1 + 25;
+
+		Line = "\t Parameter";
+		AppendWhitespace(Line, MaxOffsetName - Line.Len());
+		Line += "Min";
+		AppendWhitespace(Line, MaxLength1 - Line.Len());
+		Line += "Mean";
+		AppendWhitespace(Line, MaxLength2 - Line.Len());
+		Line += "Max";
+		UE_LOG(LogMutable, Warning, TEXT("%s"), *Line);
+
+		Line = "\t ---------------------------------------------------------------------------------------------";
+		UE_LOG(LogMutable, Warning, TEXT("%s"), *Line);
+
+		for (TPair<FString, MeasuredData>& it : ArrayMeasuredDataPerLOD[i])
+		{
+			if (it.Key == "Texture size")
+			{
+				continue;
+			}
+
+			Line = "\t ";
+			Line += FString::Printf(TEXT("%s: "), *(it.Key));
+			AppendWhitespace(Line, MaxOffsetName - Line.Len());
+			Line += FString::Printf(TEXT("%d"), int32(it.Value.MinTime));
+			AppendWhitespace(Line, MaxLength1 - Line.Len());
+			Line += FString::Printf(TEXT("%d"), int32(it.Value.MeanTime));
+			AppendWhitespace(Line, MaxLength2 - Line.Len());
+			Line += FString::Printf(TEXT("%d"), int32(it.Value.MaxTime));
+			UE_LOG(LogMutable, Warning, TEXT("%s"), *Line);
+		}
+
+		// Texture size information
+		Line = "\t Texture size ";
+		AppendWhitespace(Line, MaxOffsetName - Line.Len());
+		float Value = (MinTextureSize > 1048576) ? float(MinTextureSize) / 1048576.0f : float(MinTextureSize) / 1024.0f;
+		Line += FString::Printf(TEXT("%.2f"), Value);
+		Line += (MinTextureSize > 1048576) ? "(MB)" : "(KB)";
+		AppendWhitespace(Line, MaxLength1 - Line.Len());
+		Value = (MeanTextureSize > 1048576) ? float(MeanTextureSize) / 1048576.0f : float(MeanTextureSize) / 1024.0f;
+		Line += FString::Printf(TEXT("%.2f"), Value);
+		Line += (MeanTextureSize > 1048576) ? "(MB)" : "(KB)";
+		AppendWhitespace(Line, MaxLength2 - Line.Len());
+		Value = (MaxTextureSize > 1048576) ? float(MaxTextureSize) / 1048576.0f : float(MaxTextureSize) / 1024.0f;
+		Line += FString::Printf(TEXT("%.2f"), Value);
+		Line += (MaxTextureSize > 1048576) ? "(MB)" : "(KB)";
+		UE_LOG(LogMutable, Warning, TEXT("%s\n"), *Line);
+	}
+}
+
+
+void FRunningStressTest::AddInstanceData(MeasureDataMap MapMeasuredParam, MaterialInfoMap MapMaterialParam, uint32 CurrentLODTextureSize, uint32 CurrentLODNumFaces, uint32 LOD, UCustomizableObjectInstance* Instance)
+{
+	for (TPair<FString, MeasuredData>& it : MapMeasuredParam)
+	{
+		MeasuredData* Data = ArrayMeasuredDataPerLOD[LOD].Find(it.Key);
+
+		if (Data != nullptr)
+		{
+			*Data += it.Value;
+		}
+	}
+
+	if (ArrayLODMostExpensiveGeometry[LOD].CurrentLODNumFaces <= CurrentLODNumFaces)
+	{
+		CaptureInstanceParameters(Instance, &ArrayLODMostExpensiveGeometry[LOD]);
+		ArrayLODMostExpensiveGeometry[LOD].CurrentLODNumFaces = CurrentLODNumFaces;
+		ArrayLODMostExpensiveGeometry[LOD].CurrentLODTextureSize = CurrentLODTextureSize;
+		ArrayLODMostExpensiveGeometry[LOD].MapMaterial = MapMaterialParam;
+	}
+
+	WorstFacesDelegate.ExecuteIfBound(CurrentLODNumFaces, LOD, Instance, &ArrayMeasuredDataPerLOD[LOD]["Number of faces"]);
+
+	if (ArrayLODMostExpensiveTexture[LOD].CurrentLODTextureSize <= CurrentLODTextureSize)
+	{
+		CaptureInstanceParameters(Instance, &ArrayLODMostExpensiveTexture[LOD]);
+		ArrayLODMostExpensiveTexture[LOD].CurrentLODNumFaces = CurrentLODNumFaces;
+		ArrayLODMostExpensiveTexture[LOD].CurrentLODTextureSize = CurrentLODTextureSize;
+		ArrayLODMostExpensiveTexture[LOD].MapMaterial = MapMaterialParam;
+	}
+
+	WorstTextureDelegate.ExecuteIfBound(MapMaterialParam, CurrentLODTextureSize, LOD, Instance, &ArrayMeasuredDataPerLOD[LOD]["Instance texture size"]);
+
+	for (TPair<FString, MaterialBriefInfo>& it : MapMaterialParam)
+	{
+		MaterialBriefInfo* Data = ArrayMaterialInfoPerLOD[LOD].Find(it.Key);
+
+		if (Data == nullptr)
+		{
+			ArrayMaterialInfoPerLOD[LOD].Add(it.Key, it.Value);
+
+			const int32 MaxIndex = it.Value.ArrayTextureSize.Num();
+			for (int32 i = 0; i < MaxIndex; ++i)
+			{
+				FIntVector& Value = it.Value.ArrayTextureSize[i];
+
+				if ((Value.X <= TextureMinSize.X) || (Value.Y <= TextureMinSize.Y))
+				{
+					TextureMinSize = Value;
+				}
+
+				if ((Value.X >= TextureMaxSize.X) || (Value.Y >= TextureMaxSize.Y))
+				{
+					TextureMaxSize = Value;
+				}
+			}
+		}
+	}
+
+	NumUpdatedDone++;
+
+	FString Line = "Test update instance number ";
+	Line += FString::Printf(TEXT("%d"), NumUpdatedDone);
+	UE_LOG(LogMutable, Warning, TEXT("%s"), *Line);
+
+	WriteTestResults();
+}
+
+
+void FRunningStressTest::FinishTest()
+{
+	WorstFacesDelegate.Unbind();
+	WorstTextureDelegate.Unbind();
+	StressTestEnded.ExecuteIfBound();
+	FTSTicker::GetCoreTicker().RemoveTicker(StressTestTickDelegateHandle);
+	StressTestTickDelegate = nullptr;
+
+	StressTestReadyToReset = true;
+
+	for (TPair<UCustomizableObjectInstance*, ULiveInstance*>& it : LiveInstances)
+	{
+		it.Value->TestDelegate.Unbind();
+	}
+
+	LiveInstances.Empty();
+	TestInCourse = false;
+}
+
+
+void FRunningStressTest::CaptureInstanceParameters(UCustomizableObjectInstance* Instance, MostExpensiveInstanceData* Destination)
+{
+	Destination->BoolParameters = Instance->GetBoolParameters();
+	Destination->IntParameters = Instance->GetIntParameters();
+	Destination->FloatParameters = Instance->GetFloatParameters();
+	Destination->TextureParameters = Instance->GetTextureParameters();
+	Destination->VectorParameters = Instance->GetVectorParameters();
+	Destination->ProjectorParameters = Instance->GetProjectorParameters();
+}
+
+
+TMap<class UCustomizableObjectInstance*, ULiveInstance*>& FRunningStressTest::GetLiveInstances()
+{
+	return LiveInstances;
+}
+
+
+bool FRunningStressTest::GetStressTestReadyToReset()
+{
+	return StressTestReadyToReset;
+}
+
+
+bool FRunningStressTest::GetTestInCourse()
+{
+	return TestInCourse;
+}
+
+
+void FRunningStressTest::SetCreateInstanceTimeMs(int32 CreateInstanceTimeMsParam)
+{
+	CreateInstanceTimeMs = CreateInstanceTimeMsParam;
+}
+
+
+void FRunningStressTest::SetCreateInstanceTimeMsVar(int32 CreateInstanceTimeMsVarParam)
+{
+	CreateInstanceTimeMsVar = CreateInstanceTimeMsVarParam;
+}
+
+
+void FRunningStressTest::SetInstanceUpdateTimeMs(int32 InstanceUpdateTimeMsParam)
+{
+	InstanceUpdateTimeMs = InstanceUpdateTimeMsParam;
+}
+
+
+void FRunningStressTest::SetInstanceUpdateTimeMsVar(int32 InstanceUpdateTimeMsVarParam)
+{
+	InstanceUpdateTimeMsVar = InstanceUpdateTimeMsVarParam;
+}
+
+
+void FRunningStressTest::SetInstanceLifeTimeMs(int32 InstanceLifeTimeMsParam)
+{
+	InstanceLifeTimeMs = InstanceLifeTimeMsParam;
+}
+
+
+void FRunningStressTest::SetInstanceLifeTimeMsVar(int32 InstanceLifeTimeMsVarParam)
+{
+	InstanceLifeTimeMsVar = InstanceLifeTimeMsVarParam;
+}
+
+
+void FRunningStressTest::SetPendingInstanceCount(int32 PendingInstanceCountParam)
+{
+	PendingInstanceCount = PendingInstanceCountParam;
+}
+
+
+void FRunningStressTest::SetCustomizableObject(UCustomizableObject* CustomizableObjectParam)
+{
+	CustomizableObject = CustomizableObjectParam;
+}
+
+
+void FRunningStressTest::SetStressTestReadyToReset(bool StressTestReadyToResetParam)
+{
+	StressTestReadyToReset = StressTestReadyToResetParam;
+}
+
+
+void FRunningStressTest::SetVerifyInstancesThreshold(float VerifyInstancesThresholdParam)
+{
+	VerifyInstancesThreshold = VerifyInstancesThresholdParam;
+}
+
+
+void FRunningStressTest::SetNextInstanceTimeMs(float NextInstanceTimeMsParam)
+{
+	NextInstanceTimeMs = NextInstanceTimeMsParam;
+}
+
+
+void FRunningStressTest::AppendWhitespace(FString& Target, int32 Amount)
+{
+	for (int32 i = 0; i < Amount; ++i)
+	{
+		Target += " ";
+	}
+}
+
+
+void FRunningStressTest::PrintMostExpensivePerLOD()
+{
+	const uint32 MaxIndexLOD = ArrayLODMostExpensiveGeometry.Num();
+
+	// For each LOD level
+	for (uint32 i = 0; i < MaxIndexLOD; ++i)
+	{
+		FString Line = "///////////////////////////////////////////////////////////////////////////////////";
+		UE_LOG(LogMutable, Warning, TEXT("%s:"), *Line);
+
+		Line = "Instance detailed information for LOD level ";
+		Line += FString::Printf(TEXT("%d:"), i);
+		UE_LOG(LogMutable, Warning, TEXT("%s\n"), *Line);
+
+		Line = "\t Information for most expensive instance in terms of geometry size";
+		UE_LOG(LogMutable, Warning, TEXT("%s:"), *Line);
+		MostExpensiveInstanceData& DataGeometry = ArrayLODMostExpensiveGeometry[i];
+		PrintMostExpensiveInstanceData(DataGeometry, i);
+
+		Line = "\t Information for most expensive instance in terms of material texture size";
+		UE_LOG(LogMutable, Warning, TEXT("%s:"), *Line);
+		MostExpensiveInstanceData& DataTexture = ArrayLODMostExpensiveTexture[i];
+		PrintMostExpensiveInstanceData(DataTexture, i);
+	}
+}
+
+
+void FRunningStressTest::PrintParameterData(const MostExpensiveInstanceData& Data)
+{
+	uint32 MaxLength = FindLongestParameterName(Data);
+	MaxLength += 5;
+
+	FString Line;
+
+	if (Data.BoolParameters.Num() > 0)
+	{
+		PrintParameterHeader(FString("\t\t Bool Parameters"), MaxLength, false);
+		for (const FCustomizableObjectBoolParameterValue& it : Data.BoolParameters)
+		{
+			Line = FString::Printf(TEXT("\t\t\t %s"), *it.ParameterName);
+			AppendWhitespace(Line, MaxLength - Line.Len());
+			Line += FString::Printf(TEXT("%d"), it.ParameterValue);
+			UE_LOG(LogMutable, Warning, TEXT("%s"), *Line);
+		}
+	}
+
+	if (Data.IntParameters.Num() > 0)
+	{
+		PrintParameterHeader(FString("\t\t Integer Parameters"), MaxLength, true);
+		for (const FCustomizableObjectIntParameterValue& it : Data.IntParameters)
+		{
+			Line = FString::Printf(TEXT("\t\t\t %s"), *it.ParameterName);
+			AppendWhitespace(Line, MaxLength - Line.Len());
+			Line += FString::Printf(TEXT("%s"), *it.ParameterValueName);
+			UE_LOG(LogMutable, Warning, TEXT("%s"), *Line);
+		}
+	}
+
+	if (Data.FloatParameters.Num() > 0)
+	{
+		PrintParameterHeader(FString("\t\t Float Parameters"), MaxLength, true);
+		for (const FCustomizableObjectFloatParameterValue& it : Data.FloatParameters)
+		{
+			Line = FString::Printf(TEXT("\t\t\t %s"), *it.ParameterName);
+			AppendWhitespace(Line, MaxLength - Line.Len());
+			Line += FString::Printf(TEXT("%.2f"), it.ParameterValue);
+			UE_LOG(LogMutable, Warning, TEXT("%s"), *Line);
+		}
+	}
+
+	if (Data.TextureParameters.Num() > 0)
+	{
+		PrintParameterHeader(FString("\t\t Texture Parameters"), MaxLength, true);
+		for (const FCustomizableObjectTextureParameterValue& it : Data.TextureParameters)
+		{
+			Line = FString::Printf(TEXT("\t\t\t %s"), *it.ParameterName);
+			AppendWhitespace(Line, MaxLength - Line.Len());
+			Line += FString::Printf(TEXT("%d"), it.ParameterValue);
+			UE_LOG(LogMutable, Warning, TEXT("%s"), *Line);
+		}
+	}
+
+	if (Data.VectorParameters.Num() > 0)
+	{
+		PrintParameterHeader(FString("\t\t Vector Parameters"), MaxLength, true);
+		for (const FCustomizableObjectVectorParameterValue& it : Data.VectorParameters)
+		{
+			Line = FString::Printf(TEXT("\t\t\t %s"), *it.ParameterName);
+			AppendWhitespace(Line, MaxLength - Line.Len());
+			Line += FString::Printf(TEXT("(%.2f,"), it.ParameterValue.R);
+			Line += FString::Printf(TEXT("%.2f,"), it.ParameterValue.G);
+			Line += FString::Printf(TEXT("%.2f,"), it.ParameterValue.B);
+			Line += FString::Printf(TEXT("%.2f)"), it.ParameterValue.A);
+			UE_LOG(LogMutable, Warning, TEXT("%s"), *Line);
+		}
+	}
+
+	if (Data.ProjectorParameters.Num() > 0)
+	{
+		PrintParameterHeader(FString("\t\t Projector Parameters"), MaxLength, true);
+		for (const FCustomizableObjectProjectorParameterValue& it : Data.ProjectorParameters)
+		{
+			Line = FString::Printf(TEXT("\t\t\t %s"), *it.ParameterName);
+			AppendWhitespace(Line, MaxLength - Line.Len());
+			Line += FString::Printf(TEXT("(%.2f,"), it.Value.Position.X);
+			Line += FString::Printf(TEXT("%.2f,"), it.Value.Position.Y);
+			Line += FString::Printf(TEXT("%.2f)"), it.Value.Position.Z);
+			UE_LOG(LogMutable, Warning, TEXT("%s"), *Line);
+		}
+	}
+}
+
+uint32 FRunningStressTest::FindLongestParameterName(const MostExpensiveInstanceData& Data)
+{
+	int32 Result = 0;
+	for (const FCustomizableObjectBoolParameterValue& it : Data.BoolParameters)
+	{
+		Result = FMath::Max(Result, it.ParameterName.Len());
+	}
+
+	for (const FCustomizableObjectIntParameterValue& it : Data.IntParameters)
+	{
+		Result = FMath::Max(Result, it.ParameterName.Len());
+	}
+
+	for (const FCustomizableObjectFloatParameterValue& it : Data.FloatParameters)
+	{
+		Result = FMath::Max(Result, it.ParameterName.Len());
+	}
+
+	for (const FCustomizableObjectTextureParameterValue& it : Data.TextureParameters)
+	{
+		Result = FMath::Max(Result, it.ParameterName.Len());
+	}
+
+	for (const FCustomizableObjectVectorParameterValue& it : Data.VectorParameters)
+	{
+		Result = FMath::Max(Result, it.ParameterName.Len());
+	}
+
+	for (const FCustomizableObjectProjectorParameterValue& it : Data.ProjectorParameters)
+	{
+		Result = FMath::Max(Result, it.ParameterName.Len());
+	}
+
+	return Result;
+}
+
+
+void FRunningStressTest::PrintParameterHeader(FString Param, int32 Value, bool AddNewLine)
+{
+	if (AddNewLine)
+	{
+		UE_LOG(LogMutable, Warning, TEXT("\n"));
+	}
+	FString Line = Param;
+	UE_LOG(LogMutable, Warning, TEXT("\t\t\t %s"), *Line);
+	Line = "\t\t\t Name";
+	AppendWhitespace(Line, Value - Line.Len());
+	Line += "Value";
+	UE_LOG(LogMutable, Warning, TEXT("%s"), *Line);
+	UE_LOG(LogMutable, Warning, TEXT("\t\t\t ------------------------------------------"));
+}
+
+
+void FRunningStressTest::PrintMostExpensiveInstanceData(const MostExpensiveInstanceData& Data, uint32 LOD)
+{
+	FString Line;
+	Line = FString::Printf(TEXT("\t\t Number of faces: "));
+	Line += FString::Printf(TEXT("%d"), Data.CurrentLODNumFaces);
+	UE_LOG(LogMutable, Warning, TEXT("%s"), *Line);
+
+	// Texture size information
+	Line = "\t\t Texture size: ";
+	float Value = (MinTextureSize > 1048576) ? float(MinTextureSize) / 1048576.0f : float(MinTextureSize) / 1024.0f;
+	Line += FString::Printf(TEXT("min=%.2f"), Value);
+	Line += (MinTextureSize > 1048576) ? "(MB)" : "(KB)";
+	Value = (MaxTextureSize > 1048576) ? float(MaxTextureSize) / 1048576.0f : float(MaxTextureSize) / 1024.0f;
+	Line += FString::Printf(TEXT(", max=%.2f"), Value);
+	Line += (MaxTextureSize > 1048576) ? "(MB)" : "(KB)";
+	UE_LOG(LogMutable, Warning, TEXT("%s\n"), *Line);
+
+	// Material information
+	const TMap<FString, MaterialBriefInfo>& MapMaterialLOD = Data.MapMaterial;
+
+	Line = FString::Printf(TEXT("\t\t Material information"));
+
+	if (ArrayMaterialInfoPerLOD[LOD].Num() == 1)
+	{
+		Line += FString::Printf(TEXT(" (1 element):"));
+	}
+	else
+	{
+		Line += FString::Printf(TEXT(" (%d elements):"), MapMaterialLOD.Num());
+	}
+
+	UE_LOG(LogMutable, Warning, TEXT("%s"), *Line);
+
+	int32 MaxOffsetName = 0;
+	for (const TPair<FString, MaterialBriefInfo>& itMaterial : MapMaterialLOD)
+	{
+		MaxOffsetName = FMath::Max(MaxOffsetName, itMaterial.Key.Len());
+	}
+	MaxOffsetName += 20;
+	int32 SecondOffsetTexture = MaxOffsetName + 20;
+
+	FString TextureFormatString;
+
+	for (const TPair<FString, MaterialBriefInfo>& itMaterial : MapMaterialLOD)
+	{
+		Line = FString::Printf(TEXT("\t\t\t Material %s: "), *(itMaterial.Key));
+		UE_LOG(LogMutable, Warning, TEXT("%s"), *Line);
+
+		int32 MaxIndex = itMaterial.Value.ArrayTextureName.Num();
+		for (int32 i = 0; i < MaxIndex; ++i)
+		{
+			Line = FString::Printf(TEXT("\t\t\t\t Texture %s"), *(itMaterial.Value.ArrayTextureName[i]));
+			AppendWhitespace(Line, MaxOffsetName - Line.Len());
+			Line += FString::Printf(TEXT("size=(%d,"), itMaterial.Value.ArrayTextureSize[i].X);
+			Line += FString::Printf(TEXT("%d)"), itMaterial.Value.ArrayTextureSize[i].Y);
+			AppendWhitespace(Line, SecondOffsetTexture - Line.Len());
+			TextureFormatString = StringyfyEPixelFormat(itMaterial.Value.ArrayTextureFormat[i]);
+			Line += FString::Printf(TEXT("format=%s"), *TextureFormatString);
+			UE_LOG(LogMutable, Warning, TEXT("%s"), *Line);
+		}
+	}
+
+	UE_LOG(LogMutable, Warning, TEXT("\n"));
+
+	if ((Data.BoolParameters.Num() > 0) || (Data.IntParameters.Num() > 0) || (Data.FloatParameters.Num() > 0) ||
+		(Data.TextureParameters.Num() > 0) || (Data.VectorParameters.Num() > 0) || (Data.ProjectorParameters.Num() > 0))
+	{
+		Line = FString::Printf(TEXT("\t\t Parameter information"));
+		UE_LOG(LogMutable, Warning, TEXT("%s"), *Line);
+
+		// Parameter information
+		PrintParameterData(Data);
+	}
+}
+
+
+FString FRunningStressTest::StringyfyEPixelFormat(EPixelFormat Value)
+{
+	switch (Value)
+	{
+		case PF_Unknown:
+		{
+			return FString("PF_Unknown");
+		}
+		case PF_A32B32G32R32F:
+		{
+			return FString("PF_A32B32G32R32F");
+		}
+		case PF_B8G8R8A8:
+		{
+			return FString("PF_B8G8R8A8");
+		}
+		case PF_G8:
+		{
+			return FString("PF_G8");
+		}
+		case PF_G16:
+		{
+			return FString("PF_G16");
+		}
+		case PF_DXT1:
+		{
+			return FString("PF_DXT1");
+		}
+		case PF_DXT3:
+		{
+			return FString("PF_DXT3");
+		}
+		case PF_DXT5:
+		{
+			return FString("PF_DXT5");
+		}
+		case PF_UYVY:
+		{
+			return FString("PF_UYVY");
+		}
+		case PF_FloatRGB:
+		{
+			return FString("PF_FloatRGB");
+		}
+		case PF_FloatRGBA:
+		{
+			return FString("PF_FloatRGBA");
+		}
+		case PF_DepthStencil:
+		{
+			return FString("PF_DepthStencil");
+		}
+		case PF_ShadowDepth:
+		{
+			return FString("PF_ShadowDepth");
+		}
+		case PF_R32_FLOAT:
+		{
+			return FString("PF_R32_FLOAT");
+		}
+		case PF_G16R16:
+		{
+			return FString("PF_G16R16");
+		}
+		case PF_G16R16F:
+		{
+			return FString("PF_G16R16F");
+		}
+		case PF_G16R16F_FILTER:
+		{
+			return FString("PF_G16R16F_FILTER");
+		}
+		case PF_G32R32F:
+		{
+			return FString("PF_G32R32F");
+		}
+		case PF_A2B10G10R10:
+		{
+			return FString("PF_A2B10G10R10");
+		}
+		case PF_A16B16G16R16:
+		{
+			return FString("PF_A16B16G16R16");
+		}
+		case PF_D24:
+		{
+			return FString("PF_D24");
+		}
+		case PF_R16F:
+		{
+			return FString("PF_R16F");
+		}
+		case PF_R16F_FILTER:
+		{
+			return FString("PF_R16F_FILTER");
+		}
+		case PF_BC5:
+		{
+			return FString("PF_BC5");
+		}
+		case PF_V8U8:
+		{
+			return FString("PF_V8U8");
+		}
+		case PF_A1:
+		{
+			return FString("PF_A1");
+		}
+		case PF_FloatR11G11B10:
+		{
+			return FString("PF_FloatR11G11B10");
+		}
+		case PF_A8:
+		{
+			return FString("PF_A8");
+		}
+		case PF_R32_UINT:
+		{
+			return FString("PF_R32_UINT");
+		}
+		case PF_R32_SINT:
+		{
+			return FString("PF_R32_SINT");
+		}
+		case PF_PVRTC2:
+		{
+			return FString("PF_PVRTC2");
+		}
+		case PF_PVRTC4:
+		{
+			return FString("PF_PVRTC4");
+		}
+		case PF_R16_UINT:
+		{
+			return FString("PF_R16_UINT");
+		}
+		case PF_R16_SINT:
+		{
+			return FString("PF_R16_SINT");
+		}
+		case PF_R16G16B16A16_UINT:
+		{
+			return FString("PF_R16G16B16A16_UINT");
+		}
+		case PF_R16G16B16A16_SINT:
+		{
+			return FString("PF_R16G16B16A16_SINT");
+		}
+		case PF_R5G6B5_UNORM:
+		{
+			return FString("PF_R5G6B5_UNORM");
+		}
+		case PF_R8G8B8A8:
+		{
+			return FString("PF_R8G8B8A8");
+		}
+		case PF_A8R8G8B8:
+		{
+			return FString("PF_A8R8G8B8");
+		}
+		case PF_BC4:
+		{
+			return FString("PF_BC4");
+		}
+		case PF_R8G8:
+		{
+			return FString("PF_R8G8");
+		}
+		case PF_ATC_RGB:
+		{
+			return FString("PF_ATC_RGB");
+		}
+		case PF_ATC_RGBA_E:
+		{
+			return FString("PF_ATC_RGBA_E");
+		}
+		case PF_ATC_RGBA_I:
+		{
+			return FString("PF_ATC_RGBA_I");
+		}
+		case PF_X24_G8:
+		{
+			return FString("PF_X24_G8");
+		}
+		case PF_ETC1:
+		{
+			return FString("PF_ETC1");
+		}
+		case PF_ETC2_RGB:
+		{
+			return FString("PF_ETC2_RGB");
+		}
+		case PF_ETC2_RGBA:
+		{
+			return FString("PF_ETC2_RGBA");
+		}
+		case PF_R32G32B32A32_UINT:
+		{
+			return FString("PF_R32G32B32A32_UINT");
+		}
+		case PF_R16G16_UINT:
+		{
+			return FString("PF_R16G16_UINT");
+		}
+		case PF_ASTC_4x4:
+		{
+			return FString("PF_ASTC_4x4");
+		}
+		case PF_ASTC_6x6:
+		{
+			return FString("PF_ASTC_6x6");
+		}
+		case PF_ASTC_8x8:
+		{
+			return FString("PF_ASTC_8x8");
+		}
+		case PF_ASTC_10x10:
+		{
+			return FString("PF_ASTC_10x10");
+		}
+		case PF_ASTC_12x12:
+		{
+			return FString("PF_ASTC_12x12");
+		}
+		case PF_BC6H:
+		{
+			return FString("PF_BC6H");
+		}
+		case PF_BC7:
+		{
+			return FString("PF_BC7");
+		}
+		case PF_R8_UINT:
+		{
+			return FString("PF_R8_UINT");
+		}
+		case PF_L8:
+		{
+			return FString("PF_L8");
+		}
+		case PF_MAX:
+		{
+			return FString("PF_MAX");
+		}
+		default:
+		{
+			return FString("Error");
+		}
+	}
+}
\ No newline at end of file
diff --git a/Engine/Plugins/Experimental/Mutable/Source/CustomizableObjectEditor/Private/StressTest.h b/Engine/Plugins/Experimental/Mutable/Source/CustomizableObjectEditor/Private/StressTest.h
new file mode 100644
--- /dev/null
+++ b/Engine/Plugins/Experimental/Mutable/Source/CustomizableObjectEditor/Private/StressTest.h
@@ -0,0 +1,292 @@
+// Copyright Epic Games, Inc. All Rights Reserved.
+
+#pragma once
+
+#include "Delegates/IDelegateInstance.h"
+#include "StatePerformanceTesting.h"
+#include "CustomizableObject.h"
+#include "Containers/Ticker.h"
+#include "StressTest.generated.h"
+
+
+class UCustomizableObjectInstance;
+
+
+struct MaterialBriefInfo
+{
+	MaterialBriefInfo()
+	{
+
+	}
+
+	MaterialBriefInfo(FString& NameParam, TArray<FString>& ArrayTextureNameParam, TArray<FIntVector> ArrayTextureSizeParam, TArray<EPixelFormat> ArrayTextureFormatParam) :
+		Name(NameParam)
+		, ArrayTextureName(ArrayTextureNameParam)
+		, ArrayTextureSize(ArrayTextureSizeParam)
+		, ArrayTextureFormat(ArrayTextureFormatParam)
+	{
+
+	}
+
+	FString Name;
+	TArray<FString> ArrayTextureName;
+	TArray<FIntVector> ArrayTextureSize;
+	TArray<EPixelFormat> ArrayTextureFormat;
+};
+
+
+struct MostExpensiveInstanceData
+{
+	MostExpensiveInstanceData():
+		  CurrentLODNumFaces(0)
+		, CurrentLODTextureSize(0)
+	{
+
+	}
+
+	FString Name;
+	uint32 CurrentLODNumFaces;
+	uint32 CurrentLODTextureSize;
+	TMap<FString, MaterialBriefInfo> MapMaterial;
+	TArray<struct FCustomizableObjectBoolParameterValue> BoolParameters;
+	TArray<struct FCustomizableObjectIntParameterValue> IntParameters;
+	TArray<struct FCustomizableObjectFloatParameterValue> FloatParameters;
+	TArray<struct FCustomizableObjectTextureParameterValue> TextureParameters;
+	TArray<struct FCustomizableObjectVectorParameterValue> VectorParameters;
+	TArray<struct FCustomizableObjectProjectorParameterValue> ProjectorParameters;
+};
+
+
+typedef TMap<FString, MeasuredData> MeasureDataMap;
+typedef TMap<FString, MaterialBriefInfo> MaterialInfoMap;
+
+
+// Callback to add update information to test
+DECLARE_DELEGATE_SixParams(FStressTestInstanceUpdate, MeasureDataMap, MaterialInfoMap, uint32, uint32, uint32, UCustomizableObjectInstance*);
+// Callbacks to notify of worst cases found when running the test. It can be called multiple times per test, as worse cases pop up.
+DECLARE_DELEGATE_FourParams(FStressTestWorstNumFacesFound, uint32, uint32, UCustomizableObjectInstance*, MeasuredData*);
+DECLARE_DELEGATE_FiveParams(FStressTestWorstTextureSizeFound, MaterialInfoMap, uint32, uint32, UCustomizableObjectInstance*, MeasuredData*);
+// Callback to notify that the test has finished, before data is cleared
+DECLARE_DELEGATE(FStressTestEnded);
+
+UCLASS()
+class ULiveInstance : public UObject
+{
+public:
+	GENERATED_BODY()
+
+	/** Method to assign for the callback */
+	UFUNCTION()
+	void DelegatedCallback(UCustomizableObjectInstance* Instance);
+
+	/** Add to ArrayInformation the relevant information from Instance */
+	void AddInstanceInformation();
+
+	/** Init elements in MapMeasuredData the elevant information from Instance */
+	static TMap<FString, MeasuredData> InitInstanceInformation();
+
+	/** Amount of time an instance is kept before destroying it (the instance needs to complete at least the first skeletal mesh update) */
+	int TimeToDie;
+
+	/** After an instance has been updated, if this value is negative, a new skeleta lmesh update with random values is requested */
+	int TimeToUpdate;
+
+	/** Flag to control whether the instance has been updated */
+	bool Updated;
+
+	/** Instace being updated with random values */
+	UCustomizableObjectInstance* Instance;
+
+	/** Number of instances updated */
+	uint32 NumInstanceUpdated;
+
+	/** Delgate to call the owning test and add results */
+	FStressTestInstanceUpdate TestDelegate;
+};
+
+
+class FRunningStressTest
+{
+public:
+	/** Default constructor */
+	FRunningStressTest();
+
+	/** Ticker used to contorl logic during stress test */
+	bool RunStressTestTick(float DeltaTime);
+
+	/** Setup map to gather data and ticker */
+	void InitMeasureTest();
+
+	/** Show in the console log the test results */
+	void WriteTestResults();
+
+	/** Callback for new instance update information */
+	void AddInstanceData(MeasureDataMap MapMeasuredParam, MaterialInfoMap MapMaterialParam, uint32 CurrentLODTextureSize, uint32 CurrentLODNumFaces, uint32 LOD, UCustomizableObjectInstance* Instance);
+
+	/** Finish current test */
+	void FinishTest();
+
+	/** Will retrieve the parameters present in Instance and set them in Destination */
+	void CaptureInstanceParameters(UCustomizableObjectInstance* Instance, MostExpensiveInstanceData* Destination);
+
+	/** Getter of LiveInstances */
+	TMap<class UCustomizableObjectInstance*, ULiveInstance*>& GetLiveInstances();
+
+	/** Getter of StressTestReadyToReset */
+	bool GetStressTestReadyToReset();
+
+	/** Getter of TestInCourse */
+	bool GetTestInCourse();
+
+	/** Setter of CreateInstanceTimeMs */
+	void SetCreateInstanceTimeMs(int32 CreateInstanceTimeMsParam);
+
+	/** Setter of CreateInstanceTimeMsVar */
+	void SetCreateInstanceTimeMsVar(int32 CreateInstanceTimeMsVarParam);
+
+	/** Setter of InstanceUpdateTimeMs */
+	void SetInstanceUpdateTimeMs(int32 InstanceUpdateTimeMsParam);
+
+	/** Setter of InstanceUpdateTimeMsVar */
+	void SetInstanceUpdateTimeMsVar(int32 InstanceUpdateTimeMsVarParam);
+
+	/** Setter of InstanceLifeTimeMs */
+	void SetInstanceLifeTimeMs(int32 InstanceLifeTimeMsParam);
+
+	/** Setter of InstanceLifeTimeMsVar */
+	void SetInstanceLifeTimeMsVar(int32 InstanceLifeTimeMsVarParam);
+
+	/** Setter of PendingInstanceCount */
+	void SetPendingInstanceCount(int32 PendingInstanceCountParam);
+
+	/** Setter of CustomizableObject */
+	void SetCustomizableObject(UCustomizableObject* CustomizableObjectParam);
+
+	/** Setter of StressTestReadyToReset */
+	void SetStressTestReadyToReset(bool StressTestReadyToResetParam);
+
+	/** Setter of SetVerifyInstancesThreshold */
+	void SetVerifyInstancesThreshold(float VerifyInstancesThresholdParam);
+
+	/** Setter of NextInstanceTimeMs */
+	void SetNextInstanceTimeMs(float NextInstanceTimeMsParam);
+
+	/** Delegate called when the worst case for number of faces is found UP UNTIL NOW, with its info and a pointer to the instance. It can be called multiple times on a single test run. */
+	FStressTestWorstNumFacesFound WorstFacesDelegate;
+
+	/** Delegate called when the worst case for texture size is found UP UNTIL NOW, with its info and a pointer to the instance. It can be called multiple times on a single test run. */
+	FStressTestWorstTextureSizeFound WorstTextureDelegate;
+
+	/** Delegate called when the stress test finishes, before any data is erased. */
+	FStressTestEnded StressTestEnded;
+
+	/** Return the string verison of the pixel format enum given as parameter */
+	static FString StringyfyEPixelFormat(EPixelFormat Value);
+
+protected:
+	/** Append Amount whitespace characters to Target */
+	void AppendWhitespace(FString& Target, int32 Amount);
+
+	/** Output face, texture size and material information in ArrayLODMostExpensive */
+	void PrintMostExpensivePerLOD();
+
+	/** Print parameter data in Data */
+	void PrintParameterData(const MostExpensiveInstanceData& Data);
+
+	/** Utility method to find the length of the longest parameter name in Data */
+	uint32 FindLongestParameterName(const MostExpensiveInstanceData& Data);
+
+	/** Utility method to add the headers used in PrintMostExpensivePerLOD */
+	void PrintParameterHeader(FString Param, int32 Value, bool AddNewLine);
+
+	/** Utility method to print the parameter information */
+	void PrintMostExpensiveInstanceData(const MostExpensiveInstanceData& Data, uint32 LOD);
+
+	/** Value used to determine how frequently a new instance should be made */
+	int32 CreateInstanceTimeMs;
+
+	/** Value used to determine how frequently a new instance should be made */
+	int32 CreateInstanceTimeMsVar;
+
+	/** Value used to determine how frequently an instance should be updated (the ULiveInstance::Updated flag is used as restriction) */
+	int32 InstanceUpdateTimeMs;
+
+	/** Value used to determine how frequently an instance should be updated (the ULiveInstance::Updated flag is used as restriction) */
+	int32 InstanceUpdateTimeMsVar;
+
+	/** Value used to set the maximum amount of time an instance is updated before being destroyed (the ULiveInstance::Updated flag is used as restriction) */
+	int32 InstanceLifeTimeMs;
+
+	/** Value used to set the maximum amount of time an instance is updated before being destroyed (the ULiveInstance::Updated flag is used as restriction) */
+	int32 InstanceLifeTimeMsVar;
+
+	/** Number of instances for the test */
+	int32 PendingInstanceCount;
+
+	/** Elapsed time in the current test */
+	int64 TotalTimeMs = 0;
+
+	/** Threshold time to build a new instance in the ticker method */
+	int NextInstanceTimeMs = 0;
+
+	/** Ticker delegate handle */
+	FTSTicker::FDelegateHandle StressTestTickDelegateHandle;
+
+	/** Ticker handle */
+	FTickerDelegate StressTestTickDelegate;
+
+	/** Map with the customizable object instances and the corresponding data */
+	UPROPERTY()
+	TMap<class UCustomizableObjectInstance*, ULiveInstance*> LiveInstances;
+
+	/** Threshold time to verify all instances in the test have finished updating */
+	float VerifyInstancesThreshold;
+
+	/** Customizable object to build the instances */
+	UCustomizableObject* CustomizableObject;
+
+	/** True when all the instances in the stress test have finished being updated */
+	bool StressTestReadyToReset;
+
+	/** Counter to know the total number of instance update requested */
+	uint32 NumUpdateRequested;
+
+	/** Counter to know the total number of instance update already done */
+	uint32 NumUpdatedDone;
+
+	/** Information about each random values skeleton update carried out by this instance, each array element covers per-LOD data */
+	TArray<TMap<FString, MeasuredData>> ArrayMeasuredDataPerLOD;
+
+	/** To have a per-material breakdown information */
+	TArray<TMap<FString, MaterialBriefInfo>> ArrayMaterialInfoPerLOD;
+
+	/** Smallest texture found in the test */
+	FIntVector TextureMinSize;
+
+	/** Biggest texture found in the test */
+	FIntVector TextureMaxSize;
+
+	/** Flag to know if a test is in course */
+	bool TestInCourse;
+
+	/** Min texture size in bytes found in the test */
+	uint32 MinTextureSize;
+
+	/** Mean texture size in bytes found in the test */
+	uint32 MeanTextureSize;
+
+	/** Max texture size in bytes found in the test */
+	uint32 MaxTextureSize;
+
+	/** Array with per-LOD detailed information of the most expensive instance built regarding geometry data */
+	TArray<MostExpensiveInstanceData> ArrayLODMostExpensiveGeometry;
+
+	/** Array with per-LOD detailed information of the most expensive instance built regarding texture size data */
+	TArray<MostExpensiveInstanceData> ArrayLODMostExpensiveTexture;
+
+	/** Flag to control when to initialize the number of elements of ArrayLODMostExpensive in each test */
+	bool ArrayLODResized;
+
+	/** Number of LOD of the Customizable Object */
+	uint32 NumLOD;
+};
\ No newline at end of file
diff --git a/Engine/Plugins/Experimental/Mutable/Source/CustomizableObjectEditor/Private/Thumbnails/CustomizableObjectInstanceThumbnailRenderer.cpp b/Engine/Plugins/Experimental/Mutable/Source/CustomizableObjectEditor/Private/Thumbnails/CustomizableObjectInstanceThumbnailRenderer.cpp
new file mode 100644
--- /dev/null
+++ b/Engine/Plugins/Experimental/Mutable/Source/CustomizableObjectEditor/Private/Thumbnails/CustomizableObjectInstanceThumbnailRenderer.cpp
@@ -0,0 +1,123 @@
+// Copyright Epic Games, Inc. All Rights Reserved.
+
+#include "Thumbnails/CustomizableObjectInstanceThumbnailRenderer.h"
+#include "CustomizableObjectEditorModule.h"
+#include "CustomizableObjectInstance.h"
+#include "../CustomizableObjectCompiler.h"
+#include "UObject/ConstructorHelpers.h"
+#include "CanvasTypes.h"
+#include "EngineModule.h"
+
+
+#define LOCTEXT_NAMESPACE "CustomizableObjectInstanceThumbnailRenderer"
+
+
+UCustomizableObjectInstanceThumbnailRenderer::UCustomizableObjectInstanceThumbnailRenderer()
+	: Super()
+{
+	ThumbnailScene = nullptr;
+	//CustomizableObjectInstance = ConstructObject<UCustomizableObjectInstance>(UCustomizableObjectInstance::StaticClass(), GetTransientPackage(), NAME_None, RF_Transient);
+
+	// Structure to hold one-time initialization
+	struct FConstructorStatics
+	{
+		ConstructorHelpers::FObjectFinder<UTexture2D> PSysThumbnail_NoImage;
+		ConstructorHelpers::FObjectFinder<UTexture2D> PSysThumbnail_OOD;
+		FConstructorStatics()
+			: PSysThumbnail_NoImage(TEXT("/Engine/EditorMaterials/ParticleSystems/PSysThumbnail_NoImage"))
+			, PSysThumbnail_OOD(TEXT("/Engine/EditorMaterials/ParticleSystems/PSysThumbnail_OOD"))
+		{
+		}
+	};
+	static FConstructorStatics ConstructorStatics;
+
+	NoImage = ConstructorStatics.PSysThumbnail_NoImage.Object;
+}
+
+void UCustomizableObjectInstanceThumbnailRenderer::Draw(UObject* Object, int32 X, int32 Y, uint32 Width, uint32 Height, FRenderTarget* RenderTarget, FCanvas* Canvas, bool bAdditionalViewFamily)
+{
+	UCustomizableObjectInstance* CustomizableObjectInstance = Cast<UCustomizableObjectInstance>(Object);
+	
+	if (CustomizableObjectInstance)
+	{
+		UCustomizableObject* CustomizableObject = CustomizableObjectInstance->GetCustomizableObject(); 
+
+		if ( ThumbnailScene == nullptr )
+		{
+			ThumbnailScene = new FSkeletalMeshThumbnailScene();
+		}
+
+		//CustomizableObjectInstance = ConstructObject<UCustomizableObjectInstance>(UCustomizableObjectInstance::StaticClass(), GetTransientPackage(), NAME_None, RF_Transient);
+		//CustomizableObjectInstance->SetObject(CustomizableObject);
+		//CustomizableObjectInstance->bShowOnlyRuntimeParameters = false;
+
+		if (!CustomizableObject->IsCompiled())
+		{
+			// Use default compilation options?
+			FCompilationOptions Options;
+
+			FCustomizableObjectCompiler Compiler;
+			// Compiler.Compile(CustomizableObjectInstance->CustomizableObject, Options, false); // TODO Can not be sync!
+		}
+
+		// The compilation might fail because the graph is not valid or incomplete, so it has to be checked again
+		if (CustomizableObject->GetRefSkeletalMesh() && CustomizableObject->IsCompiled()) // TODO MultiComp
+		{
+			CustomizableObjectInstance->SetObject(CustomizableObject);
+			
+			// \todo removed
+			//CustomizableObjectInstance->GetPrivate()->ReloadParametersFromObject();
+			
+			UCustomizableSkeletalComponent* PreviewCustomizableSkeletalComponent = NewObject<UCustomizableSkeletalComponent>(UCustomizableSkeletalComponent::StaticClass());
+			PreviewCustomizableSkeletalComponent->CustomizableObjectInstance = CustomizableObjectInstance;
+
+			PreviewCustomizableSkeletalComponent->UpdateSkeletalMeshAsync();
+			USkeletalMesh* SkeletalMesh = PreviewCustomizableSkeletalComponent->GetSkeletalMesh();
+
+			if (SkeletalMesh)
+			{
+				ThumbnailScene->SetSkeletalMesh(SkeletalMesh);
+
+				FSceneViewFamilyContext ViewFamily(FSceneViewFamily::ConstructionValues(RenderTarget, ThumbnailScene->GetScene(), FEngineShowFlags(ESFIM_Game))
+					// \todo This differs from engine version to engine version. While thumbnails are not enable, we can disable this code.
+					//.SetWorldTimes(FApp::GetCurrentTime() - GStartTime, FApp::GetDeltaTime(), FApp::GetCurrentTime() - GStartTime)
+					.SetTime(UThumbnailRenderer::GetTime())
+					.SetAdditionalViewFamily(bAdditionalViewFamily));
+
+				ViewFamily.EngineShowFlags.DisableAdvancedFeatures();
+				ViewFamily.EngineShowFlags.MotionBlur = 0;
+				ViewFamily.EngineShowFlags.LOD = 0;
+
+				RenderViewFamily(Canvas, &ViewFamily, ThumbnailScene->CreateView(&ViewFamily, X, Y, Width, Height));
+				ThumbnailScene->SetSkeletalMesh(nullptr);
+			}
+			else
+			{
+				// Use the texture interface to draw
+				Canvas->DrawTile(X, Y, Width, Height, 0.f, 0.f, 1.f, 1.f, FLinearColor::White, NoImage->GetResource(), false);
+			}
+
+			PreviewCustomizableSkeletalComponent->DestroyComponent();
+		}
+		else
+		{
+			// Use the texture interface to draw
+			Canvas->DrawTile(X, Y, Width, Height, 0.f, 0.f, 1.f, 1.f, FLinearColor::White, NoImage->GetResource(), false);
+		}
+	}
+}
+
+void UCustomizableObjectInstanceThumbnailRenderer::BeginDestroy()
+{
+	if ( ThumbnailScene != nullptr )
+	{
+		delete ThumbnailScene;
+		ThumbnailScene = nullptr;
+	}
+
+	Super::BeginDestroy();
+}
+
+
+#undef LOCTEXT_NAMESPACE
+
