[mutable] Move Mutable plugin from the restricted folder to the experimental plugins folder.

diff --git a/Engine/Plugins/Experimental/Mutable/Source/CustomizableObjectEditor/Private/Thumbnails/CustomizableObjectThumbnailRenderer.cpp b/Engine/Plugins/Experimental/Mutable/Source/CustomizableObjectEditor/Private/Thumbnails/CustomizableObjectThumbnailRenderer.cpp
new file mode 100644
--- /dev/null
+++ b/Engine/Plugins/Experimental/Mutable/Source/CustomizableObjectEditor/Private/Thumbnails/CustomizableObjectThumbnailRenderer.cpp
@@ -0,0 +1,117 @@
+// Copyright Epic Games, Inc. All Rights Reserved.
+
+#include "Thumbnails/CustomizableObjectThumbnailRenderer.h"
+#include "CustomizableObjectEditorModule.h"
+#include "CustomizableObject.h"
+#include "CustomizableObjectInstance.h"
+#include "../CustomizableObjectCompiler.h"
+#include "UObject/ConstructorHelpers.h"
+#include "CanvasTypes.h"
+#include "EngineModule.h"
+
+#include "UnrealPortabilityHelpers.h"
+
+
+
+UCustomizableObjectThumbnailRenderer::UCustomizableObjectThumbnailRenderer()
+	: Super()
+{
+	ThumbnailScene = nullptr;
+	//CustomizableObjectInstance = ConstructObject<UCustomizableObjectInstance>(UCustomizableObjectInstance::StaticClass(), GetTransientPackage(), NAME_None, RF_Transient);
+
+	// Structure to hold one-time initialization
+	struct FConstructorStatics
+	{
+		ConstructorHelpers::FObjectFinder<UTexture2D> PSysThumbnail_NoImage;
+		ConstructorHelpers::FObjectFinder<UTexture2D> PSysThumbnail_OOD;
+		FConstructorStatics()
+			: PSysThumbnail_NoImage(TEXT("/Engine/EditorMaterials/ParticleSystems/PSysThumbnail_NoImage"))
+			, PSysThumbnail_OOD(TEXT("/Engine/EditorMaterials/ParticleSystems/PSysThumbnail_OOD"))
+		{
+		}
+	};
+	static FConstructorStatics ConstructorStatics;
+
+	NoImage = ConstructorStatics.PSysThumbnail_NoImage.Object;
+}
+
+void UCustomizableObjectThumbnailRenderer::Draw(UObject* Object, int32 X, int32 Y, uint32 Width, uint32 Height, FRenderTarget* RenderTarget, FCanvas* Canvas, bool bAdditionalViewFamily)
+{
+	UCustomizableObject* CustomizableObject = Cast<UCustomizableObject>(Object);
+	
+	if (CustomizableObject != nullptr)
+	{
+		if ( ThumbnailScene == nullptr )
+		{
+			ThumbnailScene = new FSkeletalMeshThumbnailScene();
+		}
+
+		if (!CustomizableObject->IsCompiled())
+		{
+			// Use default compilation options?
+			FCompilationOptions Options;
+
+			FCustomizableObjectCompiler Compiler;
+			// Compiler.Compile(CustomizableObject, Options, false); // TODO Can not be sync!
+		}
+
+		// The compilation might fail because the graph is not valid or incomplete, so it has to be checked again
+		if (CustomizableObject->GetRefSkeletalMesh() && CustomizableObject->IsCompiled()) // TODO MultiComp
+		{
+			CustomizableObjectInstance = CustomizableObject->CreateInstance();
+
+			if (CustomizableObjectInstance)
+			{
+				UCustomizableSkeletalComponent* PreviewCustomizableSkeletalComponent = NewObject<UCustomizableSkeletalComponent>(UCustomizableSkeletalComponent::StaticClass());
+				PreviewCustomizableSkeletalComponent->CustomizableObjectInstance = CustomizableObjectInstance;
+
+				PreviewCustomizableSkeletalComponent->UpdateSkeletalMeshAsync();
+				USkeletalMesh* SkeletalMesh = PreviewCustomizableSkeletalComponent->GetSkeletalMesh();
+
+				if (SkeletalMesh && Helper_GetLODInfoArray(SkeletalMesh).Num())
+				{
+					ThumbnailScene->SetSkeletalMesh(SkeletalMesh);
+
+					FSceneViewFamilyContext ViewFamily(FSceneViewFamily::ConstructionValues(RenderTarget, ThumbnailScene->GetScene(), FEngineShowFlags(ESFIM_Game))
+						// \todo This differs from engine version to engine version. While thumbnails are not enable, we can disable this code.
+						//.SetWorldTimes(FApp::GetCurrentTime() - GStartTime, FApp::GetDeltaTime(), FApp::GetCurrentTime() - GStartTime) 
+						.SetTime(UThumbnailRenderer::GetTime())
+						.SetAdditionalViewFamily(bAdditionalViewFamily));
+
+					ViewFamily.EngineShowFlags.DisableAdvancedFeatures();
+					ViewFamily.EngineShowFlags.MotionBlur = 0;
+					ViewFamily.EngineShowFlags.LOD = 0;
+
+					RenderViewFamily(Canvas, &ViewFamily, ThumbnailScene->CreateView(&ViewFamily, X, Y, Width, Height));
+					ThumbnailScene->SetSkeletalMesh(nullptr);
+				}
+				else
+				{
+					// Use the texture interface to draw
+					Canvas->DrawTile(X, Y, Width, Height, 0.f, 0.f, 1.f, 1.f, FLinearColor::White, NoImage->GetResource(), false);
+				}
+
+				SkeletalMesh = nullptr;
+				PreviewCustomizableSkeletalComponent->DestroyComponent();
+				CustomizableObjectInstance = nullptr;
+			}
+		}
+		else
+		{
+			// Use the texture interface to draw
+			Canvas->DrawTile(X, Y, Width, Height, 0.f, 0.f, 1.f, 1.f, FLinearColor::White, NoImage->GetResource(), false);
+		}
+	}
+}
+
+void UCustomizableObjectThumbnailRenderer::BeginDestroy()
+{
+	if ( ThumbnailScene != nullptr )
+	{
+		delete ThumbnailScene;
+		ThumbnailScene = nullptr;
+	}
+
+	Super::BeginDestroy();
+}
+
diff --git a/Engine/Plugins/Experimental/Mutable/Source/CustomizableObjectEditor/Private/UnrealToMutableTextureConversionUtils.cpp b/Engine/Plugins/Experimental/Mutable/Source/CustomizableObjectEditor/Private/UnrealToMutableTextureConversionUtils.cpp
new file mode 100644
--- /dev/null
+++ b/Engine/Plugins/Experimental/Mutable/Source/CustomizableObjectEditor/Private/UnrealToMutableTextureConversionUtils.cpp
@@ -0,0 +1,311 @@
+// Copyright Epic Games, Inc. All Rights Reserved.
+
+#include <UnrealToMutableTextureConversionUtils.h>
+
+#include "Engine/Texture2D.h"
+#include "RenderUtils.h"
+#include "ImageCore.h"
+#include "ImageCoreUtils.h"
+#include "Math/UnrealMathUtility.h"
+//#include "CustomizableObject.h"	// For the LogMutable logging category
+#include "ImageUtils.h"
+
+namespace UnrealToMutableImageConversion_Interanl
+{
+
+FORCEINLINE ERawImageFormat::Type ConvertFormatSourceToRaw(const ETextureSourceFormat SourceFormat) 
+{
+	return FImageCoreUtils::ConvertToRawImageFormat(SourceFormat);
+}
+
+EUnrealToMutableConversionError ApplyCompositeTexture(
+        FImage& Image, 
+        UTexture* CompositeTexture,
+        const ECompositeTextureMode CompositeTextureMode,
+        const float CompositePower)
+{
+    const int32 SizeX = CompositeTexture->Source.GetSizeX();
+    const int32 SizeY = CompositeTexture->Source.GetSizeY();
+    const ETextureSourceFormat SourceFormat = CompositeTexture->Source.GetFormat();
+
+    const ERawImageFormat::Type RawFormat = ConvertFormatSourceToRaw(SourceFormat);
+
+    if (RawFormat == ERawImageFormat::RGBA32F)
+    {
+        return EUnrealToMutableConversionError::CompositeUnsupportedFormat;
+    }   
+
+    FImage CompositeImage(SizeX, SizeY, 1, RawFormat, EGammaSpace::Linear);
+    
+    if(!CompositeTexture->Source.GetMipData(CompositeImage.RawData, 0))
+    {
+        return EUnrealToMutableConversionError::Unknown;
+    }
+
+    // Convert Composite Image to RGBA32F format and resize so both images have
+    // the source image size.
+    const bool bHaveSimilarAspect = FMath::IsNearlyEqual(
+                float(SizeX) / float(SizeY), 
+                float(Image.SizeX) / float(Image.SizeY), 
+                KINDA_SMALL_NUMBER);
+   
+    if (SizeX < Image.SizeX || SizeY < Image.SizeY || !bHaveSimilarAspect)
+    {
+        return EUnrealToMutableConversionError::CompositeImageDimensionMissmatch;
+    }
+
+    {
+        FImage TempImage;
+        CompositeImage.ResizeTo(
+            TempImage, Image.SizeX, Image.SizeY, ERawImageFormat::RGBA32F, EGammaSpace::Linear);
+        
+        Exchange(CompositeImage, TempImage);
+    }
+
+    TArrayView64<FLinearColor> ImageView = Image.AsRGBA32F();
+    TArrayView64<FLinearColor> CompositeImageView = CompositeImage.AsRGBA32F();
+   
+    const size_t OutChannelOffset = [CompositeTextureMode]() -> size_t
+        {
+            switch(CompositeTextureMode)
+            {
+                case CTM_NormalRoughnessToRed:   return offsetof(FLinearColor, R);
+                case CTM_NormalRoughnessToGreen: return offsetof(FLinearColor, G);
+                case CTM_NormalRoughnessToBlue:  return offsetof(FLinearColor, B);
+                case CTM_NormalRoughnessToAlpha: return offsetof(FLinearColor, A);
+            }
+
+            check(false);
+            return 0;
+        }(); 
+
+    const int64 NumPixels = ImageView.Num();
+    for (int64 I = 0; I < NumPixels; ++I)
+    {
+        const FVector Normal = FVector( 
+                CompositeImageView[I].R * 2.0f - 1.0f,
+                CompositeImageView[I].G * 2.0f - 1.0f,
+                CompositeImageView[I].B * 2.0f - 1.0f); 
+
+        // Is that C++ undefined behaviour?
+        float* Value = reinterpret_cast<float*>(
+                reinterpret_cast<uint8*>(&ImageView[I]) + OutChannelOffset);
+
+        // See TextureCompressorModule.cpp:1924 for details. 
+        // Toksvig estimation of variance
+        float LengthN = FMath::Min( Normal.Size(), 1.0f );
+        float Variance = ( 1.0f - LengthN ) / LengthN;
+        Variance = FMath::Max( 0.0f, Variance - 0.00004f );
+
+        Variance *= CompositePower;
+        
+        float Roughness = *Value;
+
+        float a = Roughness * Roughness;
+        float a2 = a * a;
+        float B = 2.0f * Variance * (a2 - 1.0f);
+        a2 = ( B - a2 ) / ( B - 1.0f );
+        Roughness = FMath::Pow( a2, 0.25f );
+        
+        *Value = Roughness;
+    }
+
+    return EUnrealToMutableConversionError::Success;
+}
+
+void FlipGreenChannel(FImage& Image)
+{
+    TArrayView64<FLinearColor> ImageView = Image.AsRGBA32F();
+
+    for (FLinearColor& Color : ImageView)
+    {
+        Color.G = 1.0f - FMath::Clamp(Color.G, 0.0f, 1.0f); 
+    }
+}
+
+void Normalize(FImage& Image)
+{
+	TArrayView64<FLinearColor> ImageView = Image.AsRGBA32F();
+
+	for (FLinearColor& Color : ImageView)
+	{
+		FVector3f Normal = (FVector3f(Color.R, Color.G, Color.B) * 2.0f - 1.0f).GetUnsafeNormal();
+
+		Color.R = Normal.X * 0.5f + 0.5f;
+		Color.G = Normal.Y * 0.5f + 0.5f;
+		Color.B = Normal.Z * 0.5f + 0.5f;
+	}
+}
+
+void BlurNormalForComposite(FImage& Image)
+{
+	TArrayView64<FLinearColor> ImageView = Image.AsRGBA32F();
+
+	const int32 SizeX = Image.SizeX;
+	const int32 SizeY = Image.SizeY;
+
+	for (int32 Y = 0; Y < SizeY - 1; ++Y)
+	{
+		for (int32 X = 0; X < SizeX - 1; ++X)
+		{
+			const int64 Idx0 = Y * SizeX + X;
+			const int64 Idx1 = Y * SizeX + (X + 1);
+			const int64 Idx2 = (Y + 1) * SizeX + X;
+			const int64 Idx3 = (Y + 1) * SizeX + (X + 1);
+
+			// Simple 2x2 box filter in place to gather info about the top mip normals variance.
+			ImageView[Idx0] = (ImageView[Idx0] + ImageView[Idx1] + ImageView[Idx2] + ImageView[Idx3]) * 0.25f;
+		}
+	}
+}
+
+} //namespace UnrealToMutableImageConversion_Interanl
+
+TTuple<mu::ImagePtr, EUnrealToMutableConversionError> ConvertTextureUnrealToMutable(UTexture2D* Texture, bool bIsNormalComposite)
+{
+    using namespace UnrealToMutableImageConversion_Interanl;
+
+    const int32 LODs = 1;
+    const int32 SizeX = Texture->Source.GetSizeX();
+    const int32 SizeY = Texture->Source.GetSizeY();
+    ETextureSourceFormat Format = Texture->Source.GetFormat();
+ 
+    const ERawImageFormat::Type RawFormat = ConvertFormatSourceToRaw(Format);
+
+    if (RawFormat == ERawImageFormat::RGBA32F)
+    {
+        return MakeTuple(nullptr, EUnrealToMutableConversionError::UnsupportedFormat);
+    }
+
+
+    FImage TempImage0(SizeX, SizeY, 1, RawFormat, EGammaSpace::Linear);
+    FImage TempImage1;
+    
+    if (!Texture->Source.GetMipData(TempImage0.RawData, 0))
+    {
+        return MakeTuple(nullptr, EUnrealToMutableConversionError::Unknown);
+    }
+
+    const bool bFlipGreenChannel = Texture->bFlipGreenChannel;
+    
+	//const bool bApplyCompositeTexture =
+    //        static_cast<bool>(Texture->CompositeTexture) &&
+    //        Texture->CompositeTextureMode != ECompositeTextureMode::CTM_Disabled;
+
+    // If any post processes of the image is needed, convert to RGBA32F
+    if (bFlipGreenChannel || bIsNormalComposite)
+    { 
+
+        TempImage0.CopyTo(TempImage1, ERawImageFormat::RGBA32F, EGammaSpace::Linear);
+
+        if (bFlipGreenChannel)
+        {
+            FlipGreenChannel(TempImage1);
+        }
+
+		// Prepare texture for use as normal composite.
+		if (bIsNormalComposite)
+		{
+			Normalize(TempImage1);
+			BlurNormalForComposite(TempImage1);
+		}
+
+		// Don't do this here as the result depends on the generation of the mips  
+        //if (bApplyCompositeTexture)
+        //{
+        //    EUnrealToMutableConversionError CompositeTextureError = 
+        //        ApplyCompositeTexture(
+        //            TempImage1, 
+        //            Texture->CompositeTexture, 
+        //            Texture->CompositeTextureMode,
+        //            Texture->CompositePower);
+
+        //    // This should probably only be a warning indicating the composite could not
+        //    // be applied.
+        //    if (CompositeTextureError != EUnrealToMutableConversionError::Success)
+        //    {
+        //        return MakeTuple(nullptr, CompositeTextureError);
+        //    }
+        //}
+
+        // The result is needed to TempImage0
+        // Swap internals so the memory allocations is potentially reused.
+        Exchange(TempImage1, TempImage0); 
+    }
+   
+    const ERawImageFormat::Type MutableCompatibleFormat = Format == TSF_G8 
+            ? ERawImageFormat::G8 
+            : ERawImageFormat::BGRA8;
+
+    TempImage0.CopyTo(TempImage1, MutableCompatibleFormat, EGammaSpace::Linear);
+
+    // Convert to RGBA8 in place if needed 
+    if (MutableCompatibleFormat == ERawImageFormat::BGRA8)
+    {
+        TArrayView64<FColor> ImageDataView = TempImage1.AsBGRA8();
+
+        for (FColor& Color : ImageDataView)
+        {
+            Color = FColor(Color.ToPackedABGR()); 
+        } 
+    }
+
+	mu::ImagePtr Image;
+	switch (MutableCompatibleFormat)
+	{
+	case ERawImageFormat::G8:
+		Image = new mu::Image(SizeX, SizeY, LODs, mu::EImageFormat::IF_L_UBYTE);
+		FMemory::Memcpy(Image->GetData(), TempImage1.RawData.GetData(), Image->GetDataSize());
+		break;
+
+	case ERawImageFormat::BGRA8:
+	{
+		// Try to find out if the texture has and actually makes use of the alpha channel
+		bool bHasAlphaChannel = true;
+		FImage SourceImage;
+		if (FImageUtils::GetTexture2DSourceImage(Texture, SourceImage))
+		{
+			bHasAlphaChannel = Texture->AdjustMinAlpha != Texture->AdjustMaxAlpha
+				&& Texture->CompressionSettings != TextureCompressionSettings::TC_Normalmap
+				&& !Texture->CompressionNoAlpha
+				&& (Texture->CompressionForceAlpha
+					|| FImageCore::DetectAlphaChannel(SourceImage));
+		}
+
+		// TODO: If we ever manage to get Pixel Format data on cook compilation time, remove the code that sets bHasAlphaChannel and just use Texture->HasAlphaChannel() here. Currently unreliable, it always returns EPixelFormat::PF_Unknown when cooking, which returns always "false" to HasAlphaChannel().
+		if (bHasAlphaChannel)
+		{
+			Image = new mu::Image(SizeX, SizeY, LODs, mu::EImageFormat::IF_RGBA_UBYTE);
+			FMemory::Memcpy(Image->GetData(), TempImage1.RawData.GetData(), Image->GetDataSize());
+		}
+		else
+		{
+			Image = new mu::Image(SizeX, SizeY, LODs, mu::EImageFormat::IF_RGB_UBYTE);
+			// Manual copy
+			const uint8* DataSource = TempImage1.RawData.GetData();
+			uint8* DataDest = Image->GetData();
+			uint8* DataDestEnd = DataDest + Image->GetDataSize();
+			for (int32 Pixel=0; DataDest!= DataDestEnd; )
+			{
+				DataDest[0] = DataSource[0];
+				DataDest[1] = DataSource[1];
+				DataDest[2] = DataSource[2];
+				DataDest += 3;
+				DataSource += 4;
+			}
+		}
+
+		//FString Msg = FString::Printf(TEXT("Alpha channel is %s for %s"), bHasAlphaChannel ? TEXT("enabled") : TEXT("disabled"), *Texture->GetName());
+		//UE_LOG(LogMutable, VeryVerbose, TEXT("%s"), *Msg);
+
+		break;
+	}
+
+	default:
+		// Format not supported yet?
+		check(false);
+		break;
+	}
+
+    return MakeTuple(Image, EUnrealToMutableConversionError::Success);
+}
diff --git a/Engine/Plugins/Experimental/Mutable/Source/CustomizableObjectEditor/Private/UnrealToMutableTextureConversionUtils.h b/Engine/Plugins/Experimental/Mutable/Source/CustomizableObjectEditor/Private/UnrealToMutableTextureConversionUtils.h
new file mode 100644
--- /dev/null
+++ b/Engine/Plugins/Experimental/Mutable/Source/CustomizableObjectEditor/Private/UnrealToMutableTextureConversionUtils.h
@@ -0,0 +1,21 @@
+// Copyright Epic Games, Inc. All Rights Reserved.
+
+#pragma once
+
+#include "Templates/Tuple.h"
+
+#include "MutableRuntime/Public/Image.h"
+
+// Forward declarations
+class UTexture2D;
+
+enum class EUnrealToMutableConversionError 
+{
+    Success,
+    UnsupportedFormat,
+    CompositeImageDimensionMissmatch,
+    CompositeUnsupportedFormat,
+    Unknown
+};
+
+TTuple<mu::ImagePtr, EUnrealToMutableConversionError> ConvertTextureUnrealToMutable(UTexture2D* Texture, bool bIsNormalComposite = false );
diff --git a/Engine/Plugins/Experimental/Mutable/Source/CustomizableObjectEditor/Public/CustomizableObjectDebuggerActions.h b/Engine/Plugins/Experimental/Mutable/Source/CustomizableObjectEditor/Public/CustomizableObjectDebuggerActions.h
new file mode 100644
--- /dev/null
+++ b/Engine/Plugins/Experimental/Mutable/Source/CustomizableObjectEditor/Public/CustomizableObjectDebuggerActions.h
@@ -0,0 +1,29 @@
+// Copyright Epic Games, Inc. All Rights Reserved.
+
+#pragma once
+
+#include "Framework/Commands/UICommandInfo.h"
+#include "Framework/Commands/Commands.h"
+
+/**
+ *
+ */
+class FCustomizableObjectDebuggerCommands : public TCommands<FCustomizableObjectDebuggerCommands>
+{
+
+public:
+	FCustomizableObjectDebuggerCommands();
+
+	TSharedPtr< FUICommandInfo > GenerateMutableGraph;
+	TSharedPtr< FUICommandInfo > CompileMutableCode;
+	TSharedPtr< FUICommandInfo > CompileOptions_EnableTextureCompression;
+	TSharedPtr< FUICommandInfo > CompileOptions_UseParallelCompilation;
+	TSharedPtr< FUICommandInfo > CompileOptions_UseDiskCompilation;
+
+	/**
+	 * Initialize commands
+	 */
+	virtual void RegisterCommands() override;
+
+};
+
diff --git a/Engine/Plugins/Experimental/Mutable/Source/CustomizableObjectEditor/Public/CustomizableObjectEditorActions.h b/Engine/Plugins/Experimental/Mutable/Source/CustomizableObjectEditor/Public/CustomizableObjectEditorActions.h
new file mode 100644
--- /dev/null
+++ b/Engine/Plugins/Experimental/Mutable/Source/CustomizableObjectEditor/Public/CustomizableObjectEditorActions.h
@@ -0,0 +1,75 @@
+// Copyright Epic Games, Inc. All Rights Reserved.
+
+#pragma once
+
+#include "Framework/Commands/UICommandInfo.h"
+#include "Framework/Commands/Commands.h"
+
+/**
+ * 
+ */
+class FCustomizableObjectEditorCommands : public TCommands<FCustomizableObjectEditorCommands>
+{
+
+public:
+	FCustomizableObjectEditorCommands();
+	
+	TSharedPtr< FUICommandInfo > Compile;
+	TSharedPtr< FUICommandInfo > CompileOnlySelected;
+	TSharedPtr< FUICommandInfo > ResetCompileOptions;
+	TSharedPtr< FUICommandInfo > CompileOptions_EnableTextureCompression;
+	TSharedPtr< FUICommandInfo > CompileOptions_UseParallelCompilation;
+	TSharedPtr< FUICommandInfo > CompileOptions_UseDiskCompilation;
+
+	TSharedPtr< FUICommandInfo > Debug;
+	TSharedPtr< FUICommandInfo > DebugOptions_OnlySelected;
+	TSharedPtr< FUICommandInfo > DebugOptions_EnableTextureCompression;
+
+	TSharedPtr< FUICommandInfo > PerformanceReport;
+	TSharedPtr< FUICommandInfo > ResetPerformanceReportOptions;
+	TSharedPtr< FUICommandInfo > TextureAnalyzer;
+
+	/**
+	 * Initialize commands
+	 */
+	virtual void RegisterCommands() override;
+
+};
+
+
+
+/**
+ * 
+ */
+class FCustomizableObjectEditorViewportCommands : public TCommands<FCustomizableObjectEditorViewportCommands>
+{
+
+public:
+	FCustomizableObjectEditorViewportCommands();
+	
+	/**  */
+	TSharedPtr< FUICommandInfo > SetDrawUVs;
+	TSharedPtr< FUICommandInfo > SetShowGrid;
+	TSharedPtr< FUICommandInfo > SetShowSky;
+	TSharedPtr< FUICommandInfo > SetShowBounds;
+	TSharedPtr< FUICommandInfo > SetShowCollision;
+	TSharedPtr< FUICommandInfo > SetCameraLock;
+	TSharedPtr< FUICommandInfo > SaveThumbnail;
+
+	// View Menu Commands
+	TSharedPtr< FUICommandInfo > SetShowNormals;
+	TSharedPtr< FUICommandInfo > SetShowTangents;
+	TSharedPtr< FUICommandInfo > SetShowBinormals;
+	TSharedPtr< FUICommandInfo > SetShowPivot;
+
+	TSharedPtr< FUICommandInfo > BakeInstance;
+	TSharedPtr< FUICommandInfo > StateChangeTest;
+	TSharedPtr< FUICommandInfo > StateChangeShowData;
+	TSharedPtr< FUICommandInfo > StateChangeShowGeometryData;
+
+	/**
+	 * Initialize commands
+	 */
+	virtual void RegisterCommands() override;
+
+};
diff --git a/Engine/Plugins/Experimental/Mutable/Source/CustomizableObjectEditor/Public/CustomizableObjectEditorModule.h b/Engine/Plugins/Experimental/Mutable/Source/CustomizableObjectEditor/Public/CustomizableObjectEditorModule.h
new file mode 100644
--- /dev/null
+++ b/Engine/Plugins/Experimental/Mutable/Source/CustomizableObjectEditor/Public/CustomizableObjectEditorModule.h
@@ -0,0 +1,49 @@
+// Copyright Epic Games, Inc. All Rights Reserved.
+
+#pragma once
+
+#include "AssetTypeCategories.h"
+#include "CustomizableObjectEditorClasses.h"
+
+#include "Modules/ModuleInterface.h"
+#include "Modules/ModuleManager.h"
+#include "Toolkits/AssetEditorToolkit.h"
+#include "Toolkits/IToolkit.h"
+#include "ICustomizableObjectEditor.h"
+#include "ICustomizableObjectDebugger.h"
+#include "CustomizableObjectEditorLogger.h"
+
+extern const FName CustomizableObjectEditorAppIdentifier;
+extern const FName CustomizableObjectInstanceEditorAppIdentifier;
+extern const FName CustomizableObjectPopulationEditorAppIdentifier;
+extern const FName CustomizableObjectPopulationClassEditorAppIdentifier;
+extern const FName CustomizableObjectDebuggerAppIdentifier;
+
+/**
+ * Customizable object editor module interface
+ */
+class ICustomizableObjectEditorModule : public IModuleInterface
+{
+public:
+	static ICustomizableObjectEditorModule& Get()
+	{
+		return FModuleManager::LoadModuleChecked< ICustomizableObjectEditorModule >("CustomizableObjectEditor");
+	}
+
+	/**
+	 * Creates a new customizable object editor.
+	 */
+	virtual TSharedRef<ICustomizableObjectEditor> CreateCustomizableObjectEditor( const EToolkitMode::Type Mode, const TSharedPtr< IToolkitHost >& InitToolkitHost, UCustomizableObject* CustomizableObject ) = 0;
+	virtual TSharedRef<ICustomizableObjectInstanceEditor> CreateCustomizableObjectInstanceEditor(const EToolkitMode::Type Mode, const TSharedPtr< IToolkitHost >& InitToolkitHost, UCustomizableObjectInstance* CustomizableObjectInstance) = 0;
+	virtual TSharedRef<ICustomizableObjectDebugger> CreateCustomizableObjectDebugger(const EToolkitMode::Type Mode, const TSharedPtr< IToolkitHost >& InitToolkitHost, UCustomizableObject* CustomizableObjectInstance) = 0;
+	
+	virtual TSharedPtr<class FExtensibilityManager> GetCustomizableObjectEditorToolBarExtensibilityManager() { return nullptr; }
+	virtual TSharedPtr<class FExtensibilityManager> GetCustomizableObjectEditorMenuExtensibilityManager() { return nullptr; }
+	
+	/** Get the Customizable Object custom asset category. */
+   	virtual EAssetTypeCategories::Type GetAssetCategory() const = 0;
+
+	/** Returns the module logger. */
+	virtual FCustomizableObjectEditorLogger& GetLogger() = 0;
+};
+ 
\ No newline at end of file
diff --git a/Engine/Plugins/Experimental/Mutable/Source/CustomizableObjectEditor/Public/CustomizableObjectEditorViewportClient.h b/Engine/Plugins/Experimental/Mutable/Source/CustomizableObjectEditor/Public/CustomizableObjectEditorViewportClient.h
new file mode 100644
--- /dev/null
+++ b/Engine/Plugins/Experimental/Mutable/Source/CustomizableObjectEditor/Public/CustomizableObjectEditorViewportClient.h
@@ -0,0 +1,489 @@
+// Copyright Epic Games, Inc. All Rights Reserved.
+
+#pragma once
+
+#include "CustomizableObjectWidget.h"
+#include "EditorViewportClient.h"
+#include "Components/SkeletalMeshComponent.h"
+#include "Animation/AnimationAsset.h"
+#include "Nodes/CustomizableObjectNodeProjectorConstant.h"
+#include "Nodes/CustomizableObjectNodeProjectorParameter.h"
+#include "Nodes/CustomizableObjectNodeMeshClipMorph.h"
+#include "Nodes/CustomizableObjectNodeMeshClipWithMesh.h"
+#include "ICustomizableObjectInstanceEditor.h"
+
+
+extern void RemoveRestrictedChars(FString& String);
+
+/**
+* Class used to store all the possible gizmo edition options for projectors
+* (constant projector node, parameter projector node, and projector
+* parameter from a Customizable Object Instance
+*/
+class GizmoRTSProxy
+{
+public:
+	GizmoRTSProxy();
+
+	/** Updates the origin data of this gizmo proxy with the information
+    * present in GizmoRTSProxy::Value */
+    void CopyTransformFromOriginData();
+
+	/** Updates the origin data of this gizmo proxy with the information
+	* present in GizmoRTSProxy::Value */
+	bool UpdateOriginData();
+
+	/** Call to the corresponding Modify method when DataOriginParameter
+	* or DataOriginConstant is not nullptr */
+	bool Modify();
+
+	/** Call to the corresponding graph Modify method when DataOriginParameter
+	* or DataOriginConstant is not nullptr */
+	bool ModifyGraph();
+
+	/** Reinitializes the values of DataOriginParameter, DataOriginConstant
+	* and ProjectorParameterName and ProjectorRangeIndex */
+	void CleanOriginData();
+
+	/** Returns true if the values in the variable Parameter given has the ones assigned
+	* in the FCustomizableObjectProjector constructor (to identify uninitialized
+	* parameter proyector data outside the graph editor and initialize it) */
+	bool ProjectorHasInitialValues(FCustomizableObjectProjector& Parameter);
+
+	/** Returns the initial transform data for the projector */
+	FCustomizableObjectProjector SetProjectorInitialValue(TArray<TWeakObjectPtr<class UDebugSkelMeshComponent>>& SkeletalMeshComponents, float TotalLength);
+
+	/** Returns true if a projector parameter is selected in the viewport */
+	bool IsProjectorParameterSelected();
+
+	/** To notify of changes in any parameter projector node property in case is the one being used, currently
+	* just the projector type parameter is taken into account */
+	void ProjectorParameterChanged(UCustomizableObjectNodeProjectorParameter* Node);
+
+	/** To notify of changes in any parameter projector node property in case is the one being used, currently
+	* just the projector type parameter is taken into account */
+	void ProjectorParameterChanged(UCustomizableObjectNodeProjectorConstant* Node);
+
+	/** Setter for CallUpdateSkeletalMesh */
+	void SetCallUpdateSkeletalMesh(bool Value);
+
+	/** Set Projector updated in viewport to skip rebuilding the Instance Properties panel */
+	void SetProjectorUpdatedInViewport(bool Value);
+
+	/** Getter for ProjectorParameterName */
+	FString GetProjectorParameterName();
+
+	/** Getter for ProjectorParameterNameWithIndex */
+	FString GetProjectorParameterNameWithIndex();
+
+	/** Getter for HasAssignedData */
+	bool GetHasAssignedData();
+
+	/** Getter for AssignedDataIsFromNode */
+	bool GetAssignedDataIsFromNode();
+
+	// Projector information for edition
+	FCustomizableObjectProjector Value;
+
+	// True if there's a gizmo to be shown, false otherwise
+	bool AnyGizmoSelected;
+
+	// Pointer to the Customizable Object Node Projector Parameter
+	// in case this proxy is using gizmo data from an instance of this type
+	UCustomizableObjectNodeProjectorParameter* DataOriginParameter;
+
+	// Pointer to the Customizable Object Node Projector Constant
+	// in case this proxy is using gizmo data from an instance of this type
+	UCustomizableObjectNodeProjectorConstant* DataOriginConstant;
+
+	// Name of the projecto parameter in case this proxy is using gizmo
+	// data from a Customizable Object Instance Projector parameter data
+	// (an element of CustomizableObjectInstance::ProjectorParameters)
+	FString ProjectorParameterName;
+	FString ProjectorParameterNameWithIndex;
+	int32 ProjectorRangeIndex;
+
+	/** Pointer back to the editor tool that owns the FCustomizableObjectEditorViewportClient instance */
+	TWeakPtr<ICustomizableObjectInstanceEditor> CustomizableObjectEditorPtr;
+
+	/** Simple flag to cache when to update the original data source of the gizmo */
+	bool HasAssignedData;
+
+	/** Simple flag to when the gizmo assigned data is from a node */
+	bool AssignedDataIsFromNode;
+
+	/** Index of the projector parameter in the Customizable Object Instance projector
+	* parameter array in the case of editing a projector parameter */
+	int32 ProjectorParameterIndex;
+
+	/** To track edition of the projector gizmo */
+	bool ProjectorGizmoEdited;
+
+	/** In the case of selecting another parameter projector, avoid re-generating the skeletal mesh since no change is done in the selection event */
+	bool CallUpdateSkeletalMesh;
+
+	bool bManipulatingGizmo;
+
+	// Projection type for the associated texture.
+	ECustomizableObjectProjectorType ProjectionType;
+};
+
+
+/** Viewport Client for the preview viewport */
+class FCustomizableObjectEditorViewportClient : public FEditorViewportClient
+{
+public:
+	FCustomizableObjectEditorViewportClient(TWeakPtr<class ICustomizableObjectInstanceEditor> InCustomizableObjectEditor, FPreviewScene* InPreviewScene);
+	~FCustomizableObjectEditorViewportClient();
+
+	/** persona config options **/
+	class UPersonaOptions* ConfigOption;
+
+	// FEditorViewportClient interface
+	virtual FLinearColor GetBackgroundColor() const override;
+
+	virtual void Tick(float DeltaSeconds) override;
+	virtual void Draw(const FSceneView* View,FPrimitiveDrawInterface* PDI) override;
+	virtual void Draw(FViewport* Viewport,FCanvas* Canvas) override;
+
+	virtual bool InputKey(const FInputKeyEventArgs& EventArgs) override;
+	virtual bool InputWidgetDelta(FViewport* Viewport, EAxisList::Type CurrentAxis, FVector& Drag, FRotator& Rot, FVector& Scale) override;
+	virtual void TrackingStarted(const struct FInputEventState& InInputState, bool bIsDragging, bool bNudge) override;
+	virtual void TrackingStopped() override;
+	virtual UE::Widget::EWidgetMode GetWidgetMode() const override;
+	virtual bool CanSetWidgetMode(UE::Widget::EWidgetMode NewMode) const override;
+	virtual FVector GetWidgetLocation() const override;
+	virtual FMatrix GetWidgetCoordSystem() const override;
+	virtual ECoordSystem GetWidgetCoordSystemSpace() const override;
+	virtual void SetWidgetCoordSystemSpace(ECoordSystem NewCoordSystem) override;
+	virtual void SetViewportType(ELevelViewportType InViewportType) override;
+	// End of FEditorViewportClient
+
+	void UpdateCameraSetup();
+	void UpdateFloor();
+
+	/** */
+	void SetPreviewComponent(UStaticMeshComponent* InPreviewStaticMeshComponent);
+	void SetPreviewComponents(TArray<UDebugSkelMeshComponent*>& InPreviewSkeletalMeshComponents);
+
+	/** Sets an informative message in the viewport warning the user that the CustomizableObject has no reference mesh */
+	void SetReferenceMeshMissingWarningMessage(bool bVisible);
+
+	/**
+	 *	Draws the UV overlay for the current LOD.
+	 *
+	 *	@param	InViewport					The viewport to draw to
+	 *	@param	InCanvas					The canvas to draw to
+	 *	@param	InTextYPos					The Y-position to begin drawing the UV-Overlay at (due to text displayed above it).
+	 */
+	void DrawUVs(FViewport* InViewport, FCanvas* InCanvas, int32 InTextYPos, const FString& MaterialName);
+
+	// Callback for baking the current instance.
+	void BakeInstance();
+	
+	// Returns false if the resource has already been duplicated, otherwise, returns true and an unique ResourceName.
+	bool GetUniqueResourceName(UObject* InResource, FString& InOutResourceName, TArray<UObject*>& InCachedObjects, TArray<FString>& InCachedObjectNames);
+	
+	// If baking files that already exist, ask the user for permission to overwirte them
+	bool ManageBakingAction(const FString& Path, const FString& Name);
+
+	// Callback to show / hide instance geometry data
+	void StateChangeShowGeometryData();
+
+	/** Callback for toggling the UV overlay show flag. */
+	void SetDrawUVOverlay();
+	
+	/** Callback for checking the UV overlay show flag. */
+	bool IsSetDrawUVOverlayChecked() const;
+
+	// Set the material index whose UVs will be drawn.
+	void SetDrawUVOverlayMaterial(const FString& MaterialName, FString UVChannel);
+
+	/** Callback for toggling the grid show flag. */
+	void SetShowGrid();
+	
+	/** Callback for checking the grid show flag. */
+	bool IsSetShowGridChecked() const;
+
+	/** Callback for toggling the sky show flag. */
+	void SetShowSky();
+
+	/** Callback for checking the sky show flag. */
+	bool IsSetShowSkyChecked() const;
+
+	/** Callback for toggling the bounds show flag. */
+	void SetShowBounds();
+	
+	/** Callback for checking the bounds show flag. */
+	bool IsSetShowBoundsChecked() const;
+
+	/** Callback for toggling the collision geometry show flag. */
+	void SetShowCollision();
+	
+	/** Callback for checking the collision geometry show flag. */
+	bool IsSetShowCollisionChecked() const;
+
+	/** Callback for toggling the pivot show flag. */
+	void SetShowPivot();
+
+	/** Callback for checking the pivot show flag. */
+	bool IsSetShowPivotChecked() const;
+	
+	/** Returns the desired target of the camera */
+	FSphere GetCameraTarget();
+
+	/** Point the camera to the Skeletal Mesh Components. */
+	void ResetCamera();
+	
+	/* Returns the floor height offset */	
+	float GetFloorOffset() const;
+
+	/* Sets the floor height offset, saves it to config and invalidates the viewport so it shows up immediately */
+	void SetFloorOffset(float NewValue);
+
+	void SetClipMorphPlaneVisibility(bool bVisible, const FVector& Origin, const FVector& Normal, float MorphLength, const FBoxSphereBounds& Bounds, float InRadius1, float InRadius2, float InRotationAngle);
+	void SetClipMorphPlaneVisibility(bool bVisible, UCustomizableObjectNodeMeshClipMorph* ClipPlainNode);
+	void SetClipMeshVisibility(bool bVisible, UStaticMesh* ClipMesh, UCustomizableObjectNodeMeshClipWithMesh* ClipMeshNode);
+	void SetProjectorVisibility(bool bVisible, FString ProjectorParameterName, FString ProjectorParameterNameWithIndex, int32 ProjectorRangeIndex, const FCustomizableObjectProjector& Data, int32 ProjectorParameterIndex);
+	void SetProjectorType(bool bVisible, FString ProjectorParameterName, FString ProjectorParameterNameWithIndex, int32 ProjectorRangeIndex, const FCustomizableObjectProjector& Data, int32 ProjectorParameterIndex);
+	void SetProjectorVisibility(bool bVisible, UCustomizableObjectNodeProjectorConstant* Projector);
+	void SetProjectorParameterVisibility(bool bVisible, UCustomizableObjectNodeProjectorParameter* InProjectorParameter);
+	void ResetProjectorVisibility();
+	void SetVisibilityForWireframeMode(bool bIsWireframeMode);
+
+	void SetAnimation(class UAnimationAsset* Animation, EAnimationMode::Type AnimationType);
+
+	/** Set again the animation saved in AnimationBeingPlayed (if any) */
+	void ReSetAnimation();
+
+	/** Add Light Component to the scene */
+	void AddLightToScene(class ULightComponent* AddedLight);
+
+	/** Remove Light component from the scene. use destroy component instead? */
+	void RemoveLightFromScene(class ULightComponent* RemovedLight);
+
+	/** Select/Unselect lights, selected lights can be edited using gizmos */
+	void SetSelectedLight(class ULightComponent* SelectedLight);
+
+	void SetProjectorWidgetMode(UE::Widget::EWidgetMode InMode);
+
+	const GizmoRTSProxy& GetGizmoProxy();
+
+	/** Getter of GizmoRTSProxy::IsProjectorParameterSelected */
+	bool GetGizmoIsProjectorParameterSelected();
+
+	/** Getter for bManipulating */
+	bool GetIsManipulating();
+
+	/** Getter for WidgetVisibility */
+	bool GetWidgetVisibility();
+
+	/** Update the current gizmo data to its corresponding data origin */
+	void UpdateGizmoDataToOrigin();
+
+	/** Updates the origin data of this gizmo proxy with the information
+    * present in GizmoRTSProxy::Value */
+    void CopyTransformFromOriginData();
+
+	/** Returns true if any projector node is currently selected (in the case of the Customizable Object Editor) */
+	bool AnyProjectorNodeSelected();
+
+	/** Returns in ArrayBoxSize and ArrayBoxColor the data for drawing the visual clue of the performance for the state enter and update test */
+	void BuildMeanTimesBoxSizes(float MaxWidth, TArray<float>& ArrayData, TArray<float>& ArrayBoxSize, TArray<FLinearColor>& ArrayBoxColor);
+
+	/** Returns in BoxSize and BoxColor the data for drawing the visual clue of the performance test for state enter and parameter update */
+	void BuildMeanTimesBoxSize(const float MinValue, const float MaxValue, const float MeanValue, const float MaxWidth, const float Data, float& BoxSize, FLinearColor& BoxColor) const;
+
+	/** Setter of CustomizableObject */
+	void SetCustomizableObject(UCustomizableObject* CustomizableObjectParameter);
+
+	/** Setter for StateChangeShowDataFlag */
+	void SetStateChangeShowDataFlag(bool Value);
+
+	/** Getter for StateChangeShowDataFlag */
+	bool GetStateChangeShowDataFlag();
+
+	/** Helper method to draw shadowed string on viewport */
+	void DrawShadowedString(FCanvas* Canvas, float StartX, float StartY, const FLinearColor& Color, float TextScale, FString String);
+
+	// FSerializableObject interface
+	virtual void AddReferencedObjects(FReferenceCollector& Collector) override;
+	virtual FString GetReferencerName() const override
+	{
+		return TEXT("FCustomizableObjectEditorViewportClient");
+	}
+	// End of FSerializableObject interface
+
+	/** Setter of AssetRegistryLoaded */
+	void SetAssetRegistryLoaded(bool Value);
+
+	/** Show pero LOD geometric information of the instance */
+	void ShowInstanceGeometryInformation(FCanvas* InCanvas);
+
+	/** Sets up the ShowFlag according to the current preview scene profile */
+	void SetAdvancedShowFlagsForScene(const bool bAdvancedShowFlags);
+
+	/** Delegate for preview profile is changed (used for updating show flags) */
+	void OnAssetViewerSettingsChanged(const FName& InPropertyName);
+
+	/** To notify of changes in any parameter projector node property in case is the one being used, currently
+	* just the projector type parameter is taken into account */
+	void ProjectorParameterChanged(UCustomizableObjectNodeProjectorParameter* Node);
+
+	/** To notify of changes in any parameter projector node property in case is the one being used, currently
+	* just the projector type parameter is taken into account */
+	void ProjectorParameterChanged(UCustomizableObjectNodeProjectorConstant* Node);
+
+	/** Setter for LevelViewportClient::GizmoProxy::CallUpdateSkeletalMesh */
+	void SetGizmoCallUpdateSkeletalMesh(bool Value);
+
+	/** Getter for LevelViewportClient::GizmoProxy::ProjectorParameterName */
+	FString GetGizmoProjectorParameterName();
+
+	/** Getter for LevelViewportClient::GizmoProxy::ProjectorParameterNameWithIndex */
+	FString GetGizmoProjectorParameterNameWithIndex();
+
+	/** Getter for LevelViewportClient::GizmoProxy::HasAssignedData */
+	bool GetGizmoHasAssignedData();
+
+	/** Getter for LevelViewportClient::GizmoProxy::AssignedDataIsFromNode */
+	bool GetGizmoAssignedDataIsFromNode();
+
+	/** */
+	bool IsNodeMeshClipMorphSelected();
+
+	/** Debug draw a partial cylinder, given by MaxAngle in [0, 2PI] */
+	void DrawCylinderArc(FPrimitiveDrawInterface* PDI, const FMatrix& CylToWorld, const FVector& Base, const FVector& XAxis, const FVector& YAxis, const FVector& ZAxis, float Radius, float HalfHeight, uint32 Sides, const FMaterialRenderProxy* MaterialRenderProxy, uint8 DepthPriority, FColor Color, float MaxAngle);
+
+	/** Getter of viewport's floor visibility */
+	bool GetFloorVisibility();
+
+	/** Setter of viewport's floor visibility */
+	void SetFloorVisibility(bool Value);
+
+	/** Getter of viewport's grid visibility */
+	bool GetGridVisibility();
+
+	/** Getter of viewport's environment visibility */
+	bool GetEnvironmentMeshVisibility();
+
+	/** Setter of viewport's environment visibility */
+	void SetEnvironmentMeshVisibility(uint32 Value);
+
+	/** Get associated texture parameter names for the material given as parameter */
+	TArray<FName> GetTextureParameterNames(UMaterial* Material);
+
+	/** Returns camera mode */
+	bool IsOrbitalCameraActive() const; 
+
+	/** Sets camera mode */
+	void SetCameraMode(bool Value);
+
+	/** Getter of the Material name being drawn in the UV viewer */
+	FString GetMaterialToDrawInUVs()
+	{
+		return FullNameMaterialToDrawInUVs;
+	}
+
+
+private:
+	/** Component for the static/skeletal mesh. */
+	TWeakObjectPtr<class UStaticMeshComponent> StaticMeshComponent;
+	TArray<TWeakObjectPtr<class UDebugSkelMeshComponent>> SkeletalMeshComponents;
+		
+	TMap< int, TSharedPtr<FCustomizableObjectWidget> > ProjectorWidgets;
+
+	/** Last hit proxy data. */
+	bool bManipulating;
+	EAxis::Type LastHitProxyAxis;
+	FCustomizableObjectWidget* LastHitProxyWidget;
+
+	/** Drag values for handling moving and rotating widgets. */
+	FVector LocalManDir;
+	FVector WorldManDir;
+	float DragDirX;
+	float DragDirY;
+
+	/** Pointer back to the editor tool that owns us */
+	TWeakPtr<ICustomizableObjectInstanceEditor> CustomizableObjectEditorPtr;
+
+	/** Flags for various options in the editor. */
+	bool bCameraMove;
+	bool bDrawUVs;
+	bool bShowSockets;
+	bool bDrawNormals;
+	bool bDrawTangents;
+	bool bDrawBinormals;
+	bool bShowPivot;
+	bool bDrawSky;
+
+	// Material index to draw in the uv preview.
+	FString MaterialToDrawInUVs;
+	FString FullNameMaterialToDrawInUVs;
+	int32 MaterialToDrawInUVsLOD;
+	int32 MaterialToDrawInUVsIndex;
+	int32 UVChannelToDrawInUVs;
+	int32 MaterialToDrawInUVsComponent;
+
+	bool bReferenceMeshMissingWarningMessageVisible;
+
+	UCustomizableObjectNodeMeshClipMorph* ClipMorphNode;
+	UMaterial* ClipMorphMaterial;
+	bool bClipMorphVisible;
+	bool bClipMorphLocalStartOffset;
+	FVector ClipMorphOrigin;
+	FVector ClipMorphOffset;
+	FVector ClipMorphLocalOffset;
+	FVector ClipMorphNormal;
+	FVector ClipMorphXAxis;
+	FVector ClipMorphYAxis;
+	float MorphLength;
+	FBoxSphereBounds MorphBounds;
+	UCustomizableObjectNodeMeshClipWithMesh* ClipMeshNode;
+	UStaticMeshComponent* ClipMeshComp;
+	float Radius1;
+	float Radius2;
+	float RotationAngle;
+
+	UE::Widget::EWidgetMode WidgetMode;
+	FSphere BoundSphere;
+
+	/** Light beeing edited */
+	class ULightComponent* SelectedLightComponent;
+
+	bool bIsEditingLightEnabled;
+
+	/** Proxy class to store the gizmo information needed for gizmo edition for a projector
+	* (constant node, parameter node or Customizable Object Instance projector parameter */
+	GizmoRTSProxy GizmoProxy;
+
+	/** To track Widget::bDefaultVisibility value (there's no getter and FWidget::bDefaultVisibility is private) */
+	bool WidgetVisibility;
+
+	/** To know if an animation ia being played by the Customizable Object and restre it after compilation */
+	bool IsPlayingAnimation;
+
+	/** Animation being played by the Customizable Object, if any */
+	UAnimationAsset* AnimationBeingPlayed;
+
+	/** Customizable object being used */
+	UCustomizableObject* CustomizableObject;
+
+	/** Flag to control whther to show / hide the instance geometry information data */
+	bool StateChangeShowGeometryDataFlag;
+
+	/** To know if the user has given permission to overwrite files when baking if already present */
+	bool BakingOverwritePermission;
+
+	/** Flag to know when the asset registry initial loading has completed, value set by the Customizable Object / Customizable Object Instance editor */
+	bool AssetRegistryLoaded;
+
+	/** To know if the orbital camera is being used or not */
+	bool bActivateOrbitalCamera;
+
+	/** Material for cylinder arc solid render */
+	UMaterialInterface* TransparentPlaneMaterialXY;
+
+	/** bool to return the Camera mode to Orbital when changing the Camera view to Perspective */
+	bool bSetOrbitalOnPerspectiveMode;
+};
\ No newline at end of file
diff --git a/Engine/Plugins/Experimental/Mutable/Source/CustomizableObjectEditor/Public/CustomizableObjectEditor_Deprecated.h b/Engine/Plugins/Experimental/Mutable/Source/CustomizableObjectEditor/Public/CustomizableObjectEditor_Deprecated.h
new file mode 100644
--- /dev/null
+++ b/Engine/Plugins/Experimental/Mutable/Source/CustomizableObjectEditor/Public/CustomizableObjectEditor_Deprecated.h
@@ -0,0 +1,87 @@
+﻿// Copyright Epic Games, Inc. All Rights Reserved.
+
+#pragma once
+
+#include "Engine/Texture2D.h"
+
+#include "CustomizableObjectEditor_Deprecated.generated.h"
+
+
+// Place to hide all deprecated data structures. They are still needed for deserialization backwards compatibility.
+
+// UCustomizableObjectNodeMaterial
+// Deprecated, do not use!
+USTRUCT()
+struct CUSTOMIZABLEOBJECTEDITOR_API FCustomizableObjectNodeMaterialImage
+{
+	GENERATED_USTRUCT_BODY()
+
+	UPROPERTY()
+	FString Name;
+
+	UPROPERTY()
+	int32 UVLayout = 0;
+
+	UPROPERTY()
+	TObjectPtr<UTexture2D> ReferenceTexture = nullptr;
+	
+	UPROPERTY()
+	int32 LayerIndex = 0;
+
+	UPROPERTY()
+	FString PinName;
+};
+
+// Deprecated, do not use!
+USTRUCT()
+struct CUSTOMIZABLEOBJECTEDITOR_API FCustomizableObjectNodeMaterialVector
+{
+	GENERATED_USTRUCT_BODY()
+	
+	UPROPERTY()
+	FString Name;
+	
+	UPROPERTY()
+	int32 LayerIndex = 0;
+
+	UPROPERTY()
+	FString PinName;
+};
+
+// Deprecated, do not use!
+USTRUCT()
+struct CUSTOMIZABLEOBJECTEDITOR_API FCustomizableObjectNodeMaterialScalar
+{
+	GENERATED_USTRUCT_BODY()
+
+	UPROPERTY()
+	FString Name;
+	
+	UPROPERTY()
+	int32 LayerIndex = 0;
+
+	UPROPERTY()
+	FString PinName;
+};
+
+// UCustomizableObjectNodeEditMaterial
+// Deprecated, do not use!
+USTRUCT()
+struct CUSTOMIZABLEOBJECTEDITOR_API FCustomizableObjectNodeEditMaterialImage
+{
+	GENERATED_USTRUCT_BODY()
+
+	UPROPERTY()
+	FString Name;
+};
+
+// UCustomizableObjectNodeExtendMaterial
+// Deprecated, do not use!
+USTRUCT()
+struct CUSTOMIZABLEOBJECTEDITOR_API FCustomizableObjectNodeExtendMaterialImage
+{
+	GENERATED_USTRUCT_BODY()
+
+	UPROPERTY()
+	FString Name;
+};
\ No newline at end of file
diff --git a/Engine/Plugins/Experimental/Mutable/Source/CustomizableObjectEditor/Public/CustomizableObjectWidget.h b/Engine/Plugins/Experimental/Mutable/Source/CustomizableObjectEditor/Public/CustomizableObjectWidget.h
new file mode 100644
--- /dev/null
+++ b/Engine/Plugins/Experimental/Mutable/Source/CustomizableObjectEditor/Public/CustomizableObjectWidget.h
@@ -0,0 +1,415 @@
+// Copyright Epic Games, Inc. All Rights Reserved.
+
+#pragma once
+
+#include "UObject/GCObject.h"
+#include "SceneManagement.h"
+#include "Editor/UnrealEdTypes.h"
+
+#include "SceneView.h"
+#include "SceneManagement.h"
+#include "Editor/UnrealEdTypes.h"
+
+
+class FCustomizableObjectWidget : public FGCObject
+{
+public:
+	enum EWidgetMode
+	{
+		WM_None			= -1,
+		WM_Translate,
+		WM_Rotate,
+		WM_Scale,
+		WM_TranslateRotateZ,
+		WM_Max,
+	};
+
+	CUSTOMIZABLEOBJECTEDITOR_API FCustomizableObjectWidget( class UCustomizableObjectInstance* InInstance, int InParamIndex );
+
+	CUSTOMIZABLEOBJECTEDITOR_API void Render( const FSceneView* View ,FPrimitiveDrawInterface* PDI );
+
+	/**
+	 * Draws an arrow head line for a specific axis.
+	 * @param	bCubeHead		[opt] If true, render a cube at the axis tips.  If false (the default), render a cone.
+	 */
+	void Render_Axis( const FSceneView* View, FPrimitiveDrawInterface* PDI, EAxisList::Type InAxis, FMatrix& InMatrix, UMaterialInterface* InMaterial, FVector2D& OutAxisEnd, float InScale, bool bDrawWidget, bool bCubeHead=false );
+
+	/**
+	 * Draws a cube
+	 */
+	void Render_Cube( FPrimitiveDrawInterface* PDI, const FMatrix& InMatrix, const UMaterialInterface* InMaterial, float InScale );
+	void Render_Quad( FPrimitiveDrawInterface* PDI, const FMatrix& InMatrix, const UMaterialInterface* InMaterial, const FColor* InColor, FVector2D InScale );
+
+	/**
+	 * Draws the translation widget.
+	 */
+	void Render_Translate( const FSceneView* View, FPrimitiveDrawInterface* PDI, const FVector& InLocation, bool bDrawWidget );
+
+	/**
+	 * Draws the rotation widget.
+	 */
+	void Render_Rotate( const FSceneView* View, FPrimitiveDrawInterface* PDI, const FVector& InLocation, bool bDrawWidget );
+
+	/**
+	 * Draws the scaling widget.
+	 */
+	void Render_Scale( const FSceneView* View, FPrimitiveDrawInterface* PDI, const FVector& InLocation, bool bDrawWidget );
+
+	/**
+	* Draws the Translate & Rotate Z widget.
+	*/
+	void Render_TranslateRotateZ( const FSceneView* View, FPrimitiveDrawInterface* PDI, const FVector& InLocation, bool bDrawWidget );
+
+	/**
+	 * Converts mouse movement on the screen to widget axis movement/rotation.
+	 */
+	//void ConvertMouseMovementToAxisMovement( FViewportClient* InViewportClient, const FVector& InLocation, const FVector& InDiff, FVector& InDrag, FRotator& InRotation, FVector& InScale );
+
+	/**
+	 * Absolute Translation conversion from mouse movement on the screen to widget axis movement/rotation.
+	 */
+	//void AbsoluteTranslationConvertMouseMovementToAxisMovement(FSceneView* InView, FViewportClient* InViewportClient, const FVector& InLocation, const FVector2D& InMousePosition, FVector& OutDrag, FRotator& OutRotation, FVector& OutScale );
+
+	/** 
+	 * Grab the initial offset again first time input is captured
+	 */
+	void ResetInitialTranslationOffset (void)
+	{
+		bAbsoluteTranslationInitialOffsetCached = false;
+	}
+
+	/** Only some modes support Absolute Translation Movement.  Check current mode */
+	static bool AllowsAbsoluteTranslationMovement(void);
+
+	/**
+	 * Sets the axis currently being moused over.  Typically called by FMouseDeltaTracker or FEditorLevelViewportClient.
+	 *
+	 * @param	InCurrentAxis	The new axis value.
+	 */
+	void SetCurrentAxis(EAxisList::Type InCurrentAxis)
+	{
+		CurrentAxis = InCurrentAxis;
+	}
+
+	/**
+	 * @return	The axis currently being moused over.
+	 */
+	EAxisList::Type GetCurrentAxis() const
+	{
+		return CurrentAxis;
+	}
+
+	FVector2D GetOrigin() const
+	{
+		return Origin;
+	}
+
+	/**
+	 * Returns whether we are actively dragging
+	 */
+	bool IsDragging(void) const
+	{
+		return bDragging;
+	}
+
+	/**
+	 * Sets if we are currently engaging the widget in dragging
+	 */
+	void SetDragging (const bool InDragging) 
+	{ 
+		bDragging = InDragging;
+	}
+
+	/**
+	 * Sets if we are currently engaging the widget in dragging
+	 */
+	void SetSnapEnabled (const bool InSnapEnabled) 
+	{ 
+		bSnapEnabled = InSnapEnabled;
+	}
+
+	// FSerializableObject interface
+	virtual void AddReferencedObjects( FReferenceCollector& Collector ) override;
+	virtual FString GetReferencerName() const override
+	{
+		return TEXT("FCustomizableObjectWidget");
+	}
+
+	// Own interface
+	CUSTOMIZABLEOBJECTEDITOR_API void SetWidgetMode( EWidgetMode Mode );
+	CUSTOMIZABLEOBJECTEDITOR_API void SetTransform( FMatrix Transform );
+	CUSTOMIZABLEOBJECTEDITOR_API void SetScale( FVector2D Scale );
+
+	/**
+	 * An extra matrix to apply to the widget before drawing it (allows for local/custom coordinate systems).
+	 */
+	FMatrix CustomCoordSystem;
+	FVector2D ProjectorScale;
+	
+	/** */
+	EWidgetMode WidgetMode;
+
+	/** The axis currently being moused over */
+	EAxisList::Type CurrentAxis;
+
+	/** */
+	UCustomizableObjectInstance* Instance;
+	int ParamIndex;
+
+private:
+
+
+	struct FAbsoluteMovementParams
+	{
+		/** The normal of the plane to project onto */
+		FVector PlaneNormal;
+		/** A vector that represent any displacement we want to mute (remove an axis if we're doing axis movement)*/
+		FVector NormalToRemove;
+		/** The current position of the widget */
+		FVector Position;
+
+		//Coordinate System Axes
+		FVector XAxis;
+		FVector YAxis;
+		FVector ZAxis;
+
+		//true if camera movement is locked to the object
+		float bMovementLockedToCamera;
+
+		//Direction in world space to the current mouse location
+		FVector PixelDir;
+		//Direction in world space of the middle of the camera
+		FVector CameraDir;
+		FVector EyePos;
+
+		//whether to snap the requested positionto the grid
+		bool bPositionSnapping;
+	};
+
+	struct FThickArcParams
+	{
+		FThickArcParams(FPrimitiveDrawInterface* InPDI, const FVector& InPosition, UMaterialInterface* InMaterial, const float InInnerRadius, const float InOuterRadius)
+			: Position(InPosition)
+			, PDI(InPDI)
+			, Material(InMaterial)
+			, InnerRadius(InInnerRadius)
+			, OuterRadius(InOuterRadius)
+		{
+		}
+
+		/** The current position of the widget */
+		FVector Position;
+
+		//interface for Drawing
+		FPrimitiveDrawInterface* PDI;
+
+		//Material to use to render
+		UMaterialInterface* Material;
+
+		//Radii
+		float InnerRadius;
+		float OuterRadius;
+	};
+
+	/**
+	 * Returns the Delta from the current position that the absolute movement system wants the object to be at
+	 * @param InParams - Structure containing all the information needed for absolute movement
+	 * @return - The requested delta from the current position
+	 */
+	FVector GetAbsoluteTranslationDelta (const FAbsoluteMovementParams& InParams);
+	/**
+	 * Returns the offset from the initial selection point
+	 */
+	FVector GetAbsoluteTranslationInitialOffset(const FVector& InNewPosition, const FVector& InCurrentPosition);
+
+	/**
+	 * Returns true if we're in Local Space editing mode or editing BSP (which uses the World axes anyway
+	 */
+	bool IsRotationLocalSpace (void) const;
+
+	/**
+	 * Returns the "word" representation of how far we have just rotated
+	 */
+	float GetDeltaRotation (void) const;
+
+	/**
+	 * If actively dragging, draws a ring representing the potential rotation of the selected objects, snap ticks, and "delta" markers
+	 * If not actively dragging, draws a quarter ring representing the closest quadrant to the camera
+	 * @param View - Information about the scene/camera/etc
+	 * @param PDI - Drawing interface
+	 * @param InAxis - Enumeration of axis to rotate about
+	 * @param InLocation - The Origin of the widget
+	 * @param Axis0 - The Axis that describes a 0 degree rotation
+	 * @param Axis1 - The Axis that describes a 90 degree rotation
+	 * @param InDirectionToWidget - Direction from camera to the widget
+	 * @param InColor - The color associated with the axis of rotation
+	 * @param InScale - Multiplier to maintain a constant screen size for rendering the widget
+	 */
+	void DrawRotationArc(const FSceneView* View, FPrimitiveDrawInterface* PDI, EAxisList::Type InAxis, const FVector& InLocation, const FVector& Axis0, const FVector& Axis1, const FVector& InDirectionToWidget, const FColor& InColor, const float InScale);
+
+	/**
+	 * If actively dragging, draws a ring representing the potential rotation of the selected objects, snap ticks, and "delta" markers
+	 * If not actively dragging, draws a quarter ring representing the closest quadrant to the camera
+	 * @param View - Information about the scene/camera/etc
+	 * @param PDI - Drawing interface
+	 * @param InAxis - Enumeration of axis to rotate about
+	 * @param InLocation - The Origin of the widget
+	 * @param Axis0 - The Axis that describes a 0 degree rotation
+	 * @param Axis1 - The Axis that describes a 90 degree rotation
+	 * @param InStartAngle - The starting angle about (Axis0^Axis1) to render the arc
+	 * @param InEndAngle - The ending angle about (Axis0^Axis1) to render the arc
+	 * @param InColor - The color associated with the axis of rotation
+	 * @param InScale - Multiplier to maintain a constant screen size for rendering the widget
+	 */
+	void DrawPartialRotationArc(const FSceneView* View, FPrimitiveDrawInterface* PDI, EAxisList::Type InAxis, const FVector& InLocation, const FVector& Axis0, const FVector& Axis1, const float InStartAngle, const float InEndAngle, const FColor& InColor, const float InScale, const FVector& InDirectionToWidget);
+
+	/**
+	 * Renders a portion of an arc for the rotation widget
+	 * @param InParams - Material, Radii, etc
+	 * @param InStartAxis - Start of the arc
+	 * @param InEndAxis - End of the arc
+	 * @param InColor - Color to use for the arc
+	 */
+	void DrawThickArc (const FThickArcParams& InParams, const FVector& Axis0, const FVector& Axis1, const float InStartAngle, const float InEndAngle, const FColor& InColor, const FVector& InDirectionToWidget, bool bIsOrtho );
+
+	/**
+	 * Draws protractor like ticks where the rotation widget would snap too.
+	 * Also, used to draw the wider axis tick marks
+	 * @param PDI - Drawing interface
+	 * @param InLocation - The Origin of the widget
+	 * @param Axis0 - The Axis that describes a 0 degree rotation
+	 * @param Axis1 - The Axis that describes a 90 degree rotation
+	 * @param InAngle - The Angle to rotate about the axis of rotation, the vector (Axis0 ^ Axis1)
+	 * @param InColor - The color to use for line/poly drawing
+	 * @param InScale - Multiplier to maintain a constant screen size for rendering the widget
+	 * @param InWidthPercent - The percent of the distance between the outer ring and inner ring to use for tangential thickness
+	 * @param InPercentSize - The percent of the distance between the outer ring and inner ring to use for radial distance
+	 */
+	void DrawSnapMarker(FPrimitiveDrawInterface* PDI, const FVector& InLocation, const FVector& Axis0, const FVector& Axis1, const FColor& InColor, const float InScale, const float InWidthPercent=0.0f, const float InPercentSize = 1.0f);
+
+	/**
+	 * Draw Start/Stop Marker to show delta rotations along the arc of rotation
+	 * @param PDI - Drawing interface
+	 * @param InLocation - The Origin of the widget
+	 * @param Axis0 - The Axis that describes a 0 degree rotation
+	 * @param Axis1 - The Axis that describes a 90 degree rotation
+	 * @param InAngle - The Angle to rotate about the axis of rotation, the vector (Axis0 ^ Axis1)
+	 * @param InColor - The color to use for line/poly drawing
+	 * @param InScale - Multiplier to maintain a constant screen size for rendering the widget
+	 */
+	void DrawStartStopMarker(FPrimitiveDrawInterface* PDI, const FVector& InLocation, const FVector& Axis0, const FVector& Axis1, const float InAngle, const FColor& InColor, const float InScale);
+
+	/**
+	 * Gets the axis to use when converting mouse movement, accounting for Ortho views.
+	 *
+	 * @param InDiff Difference vector to determine dominant axis.
+	 * @param ViewportType Type of viewport for ortho checks.
+	 *
+	 * @return Index of the dominant axis.
+	 */
+	uint32 GetDominantAxisIndex( const FVector& InDiff, ELevelViewportType ViewportType ) const;
+
+	/** Locations of the various points on the widget */
+	FVector2D Origin, XAxisEnd, YAxisEnd, ZAxisEnd;
+
+	enum
+	{
+		AXIS_ARROW_SEGMENTS = 16
+	};
+
+	/** Materials and colors to be used when drawing the items for each axis */
+	UMaterialInterface* TransparentPlaneMaterialXY;
+	UMaterialInterface* GridMaterial;
+
+	UMaterialInstanceDynamic* AxisMaterialX;
+	UMaterialInstanceDynamic* AxisMaterialY;
+	UMaterialInstanceDynamic* AxisMaterialZ;
+	UMaterialInstanceDynamic* CurrentAxisMaterial;
+	UMaterialInstanceDynamic* OpaquePlaneMaterialXY;
+
+	FLinearColor AxisColorX, AxisColorY, AxisColorZ;
+	FColor PlaneColorXY, ScreenSpaceColor, CurrentColor;
+
+
+	/** Whether Absolute Translation cache position has been captured */
+	bool bAbsoluteTranslationInitialOffsetCached;
+	/** The initial offset where the widget was first clicked */
+	FVector InitialTranslationOffset;
+	/** The initial position of the widget before it was clicked */
+	FVector InitialTranslationPosition;
+	/** Whether or not the widget is actively dragging */
+	bool bDragging;
+	/** Whether or not snapping is enabled for all actors */
+	bool bSnapEnabled;
+};
+
+/**
+ * Widget hit proxy.
+ */
+struct HCustomizableObjectWidgetAxis : public HHitProxy
+{
+	DECLARE_HIT_PROXY( CUSTOMIZABLEOBJECTEDITOR_API );
+
+	EAxis::Type Axis;
+
+	FCustomizableObjectWidget* Widget;
+
+	HCustomizableObjectWidgetAxis(EAxis::Type InAxis, FCustomizableObjectWidget* InWidget )
+		: HHitProxy(HPP_UI)
+		, Axis(InAxis)
+		, Widget(InWidget)
+	{
+	}
+
+	virtual EMouseCursor::Type GetMouseCursor()
+	{
+		return EMouseCursor::CardinalCross;
+	}
+
+	virtual bool AlwaysAllowsTranslucentPrimitives() const
+	{
+		return true;
+	}
+
+
+	/** Utility for calculating drag direction when you click on this widget. */
+	void CalcVectors(FSceneView* SceneView, FVector& LocalManDir, FVector& WorldManDir, float& DragDirX, float& DragDirY)
+	{
+		FMatrix WidgetMatrix = Widget->CustomCoordSystem;
+
+		if(Axis == EAxis::X)
+		{
+			WorldManDir = WidgetMatrix.GetScaledAxis( EAxis::X );
+			LocalManDir = FVector(1,0,0);
+		}
+		else if(Axis == EAxis::Y)
+		{
+			WorldManDir = WidgetMatrix.GetScaledAxis( EAxis::Y );
+			LocalManDir = FVector(0,1,0);
+		}
+		else
+		{
+			WorldManDir = WidgetMatrix.GetScaledAxis( EAxis::Z );
+			LocalManDir = FVector(0,0,1);
+		}
+
+		FVector WorldDragDir = WorldManDir;
+
+		// Transform world-space drag dir to screen space.
+		FVector ScreenDir = SceneView->ViewMatrices.GetViewMatrix().TransformVector(WorldDragDir);
+		ScreenDir.Z = 0.0f;
+
+		if( ScreenDir.IsZero() )
+		{
+			DragDirX = 0.0f;
+			DragDirY = 0.0f;
+		}
+		else
+		{
+			ScreenDir.Normalize();
+			DragDirX = ScreenDir.X;
+			DragDirY = ScreenDir.Y;
+		}
+	}
+
+};
diff --git a/Engine/Plugins/Experimental/Mutable/Source/CustomizableObjectEditor/Public/GraphTraversal.h b/Engine/Plugins/Experimental/Mutable/Source/CustomizableObjectEditor/Public/GraphTraversal.h
new file mode 100644
--- /dev/null
+++ b/Engine/Plugins/Experimental/Mutable/Source/CustomizableObjectEditor/Public/GraphTraversal.h
@@ -0,0 +1,57 @@
+// Copyright Epic Games, Inc. All Rights Reserved.
+
+#pragma once
+
+#include "Containers/Array.h"
+
+#include "EdGraph/EdGraphPin.h"
+
+#include "Nodes/CustomizableObjectNodeObject.h"
+
+/** Follow the given input pin returning the output connected pin.
+ *
+ * - Skips all orphan pins.
+ * - Follows External Pin nodes.
+ *
+ * @param Pin Pin to follow.
+ * @param CycleDetected If provided, it will set to true if a cycle has been found.
+ * @return Connected pins.
+ */
+TArray<UEdGraphPin*> FollowInputPinArray(const UEdGraphPin& Pin, bool* CycleDetected = nullptr);
+
+/** Non-array version of FollowInputPinArray. The pin can only have one connected pin. */
+UEdGraphPin* FollowInputPin(const UEdGraphPin& Pin, bool* CycleDetected = nullptr);
+
+/** Follow the given input output returning the input connected pin.
+ * 
+ * - Skips all orphan pins.
+ * - Follows External Pin nodes.
+ * - It will only follow External Pin nodes of loaded CO (i.e., Expose Pin nodes of CO which are NOT loaded will not be found)!
+ *
+ * @param Pin Pin to follow.
+ * @param CycleDetected If provided, it will set to true if a cycle has been found.
+ * @return Connected pins.
+ */
+TArray<UEdGraphPin*> FollowOutputPinArray(const UEdGraphPin& Pin, bool* CycleDetected = nullptr);
+
+/** Non-array version of FollowOutputPinArray. The pin can only have one connected pin. */
+UEdGraphPin* FollowOutputPin(const UEdGraphPin& Pin, bool* CycleDetected = nullptr);
+
+/** Returns the root Object Node of the Customizable Object's graph */
+UCustomizableObjectNodeObject* GetRootNode(UCustomizableObject* Object, bool& bOutMultipleBaseObjectsFound);
+
+/** Return in ArrayNodeObject the roots nodes in each Customizable Object graph until the whole root node is found (i.e. the one with parent = nullptr)
+		return false if a cycle is found between Customizable Objects */
+bool GetParentsUntilRoot(UCustomizableObject* Object, TArray<UCustomizableObjectNodeObject*>& ArrayNodeObject, TArray<UCustomizableObject*>& ArrayCustomizableObject);
+
+/** Returns true if the Candidate is parent of the current Customizable Object */
+bool HasCandidateAsParent(UCustomizableObjectNodeObject* Node, UCustomizableObject* ParentCandidate);
+
+/** Return the full graph Customizable Object root of the node given as parameter */
+UCustomizableObject* GetFullGraphRootObject(UCustomizableObjectNodeObject* Node, TArray<UCustomizableObject*>& VisitedObjects);
+
+/** Return the full graph Customizable Object Node root of the node given as parameter */
+UCustomizableObjectNodeObject* GetFullGraphRootNodeObject(UCustomizableObjectNodeObject* Node, TArray<UCustomizableObject*>& VisitedObjects);
+
+/** Given an output pin, return the output pin where the mesh is located. */
+const UEdGraphPin* FindMeshBaseSource(const UEdGraphPin& Pin, const bool bOnlyLookForStaticMesh);
diff --git a/Engine/Plugins/Experimental/Mutable/Source/CustomizableObjectEditor/Public/ICustomizableObjectDebugger.h b/Engine/Plugins/Experimental/Mutable/Source/CustomizableObjectEditor/Public/ICustomizableObjectDebugger.h
new file mode 100644
--- /dev/null
+++ b/Engine/Plugins/Experimental/Mutable/Source/CustomizableObjectEditor/Public/ICustomizableObjectDebugger.h
@@ -0,0 +1,19 @@
+// Copyright Epic Games, Inc. All Rights Reserved.
+
+#pragma once
+
+#include "Toolkits/IToolkitHost.h"
+#include "Toolkits/AssetEditorToolkit.h"
+
+class UCustomizableObject;
+
+// Public interface to Customizable Object Editor
+class ICustomizableObjectDebugger : public FAssetEditorToolkit
+{
+public:
+	// Retrieves the current Customizable Object displayed in the Editor.
+	virtual UCustomizableObject* GetCustomizableObject() = 0;
+
+};
+
+
diff --git a/Engine/Plugins/Experimental/Mutable/Source/CustomizableObjectEditor/Public/ICustomizableObjectEditor.h b/Engine/Plugins/Experimental/Mutable/Source/CustomizableObjectEditor/Public/ICustomizableObjectEditor.h
new file mode 100644
--- /dev/null
+++ b/Engine/Plugins/Experimental/Mutable/Source/CustomizableObjectEditor/Public/ICustomizableObjectEditor.h
@@ -0,0 +1,37 @@
+// Copyright Epic Games, Inc. All Rights Reserved.
+
+#pragma once
+
+#include "Toolkits/IToolkitHost.h"
+#include "Toolkits/AssetEditorToolkit.h"
+#include "ICustomizableObjectInstanceEditor.h"
+
+class UCustomizableObject;
+class UCustomizableObjectInstance;
+
+// Public interface to Customizable Object Editor
+class ICustomizableObjectEditor : public ICustomizableObjectInstanceEditor
+{
+public:
+	// Retrieves the current Customizable Object displayed in the Editor.
+	virtual UCustomizableObject* GetCustomizableObject() = 0;
+	
+	// Checks whether nodes can currently be pasted */
+	virtual bool CanPasteNodes() const {return false;}
+
+	// Paste nodes at a specific location
+	virtual void PasteNodesHere(const FVector2D& Location) {};
+
+	virtual class UEdGraphNode* CreateCommentBox(const FVector2D& NodePos) { return nullptr; };
+
+	virtual void UpdateGraphNodeProperties() {};
+
+	virtual void UpdateObjectProperties() {};
+
+	// Make sure the graph editor selects a specific node
+	virtual void SelectNode(const class UCustomizableObjectNode* Node) {}
+
+	virtual void ReconstructNode(UEdGraphNode* Node) {};
+};
+
+
diff --git a/Engine/Plugins/Experimental/Mutable/Source/CustomizableObjectEditor/Public/ICustomizableObjectInstanceEditor.h b/Engine/Plugins/Experimental/Mutable/Source/CustomizableObjectEditor/Public/ICustomizableObjectInstanceEditor.h
new file mode 100644
--- /dev/null
+++ b/Engine/Plugins/Experimental/Mutable/Source/CustomizableObjectEditor/Public/ICustomizableObjectInstanceEditor.h
@@ -0,0 +1,31 @@
+// Copyright Epic Games, Inc. All Rights Reserved.
+
+#pragma once
+
+#include "Toolkits/IToolkitHost.h"
+#include "Toolkits/AssetEditorToolkit.h"
+
+class UCustomizableObject;
+class UCustomizableObjectInstance;
+
+/**
+ * Public interface to Customizable Object Instance Editor
+ */
+class ICustomizableObjectInstanceEditor : public FAssetEditorToolkit
+{
+public:
+	virtual UCustomizableObjectInstance* GetPreviewInstance() = 0;
+
+	/** Refreshes the Customizable Object Instance Editor's viewport. */
+	virtual void RefreshViewport() = 0;
+
+	/** Refreshes everything in the Customizable Object Instance Editor. */
+	virtual void RefreshTool() = 0;
+
+	/** Getter of AssetRegistryLoaded */
+	virtual bool GetAssetRegistryLoaded() = 0;
+
+	virtual void SetPoseAsset(class UPoseAsset* PoseAssetParameter) {};
+};
+
+
diff --git a/Engine/Plugins/Experimental/Mutable/Source/CustomizableObjectEditor/Public/MutableMeshPreviewUtils.h b/Engine/Plugins/Experimental/Mutable/Source/CustomizableObjectEditor/Public/MutableMeshPreviewUtils.h
new file mode 100644
--- /dev/null
+++ b/Engine/Plugins/Experimental/Mutable/Source/CustomizableObjectEditor/Public/MutableMeshPreviewUtils.h
@@ -0,0 +1,28 @@
+// Copyright Epic Games, Inc. All Rights Reserved.
+
+#pragma once
+
+#include "CustomizableObject/Private/CustomizableObjectSystemPrivate.h"
+#include "MutableRuntime/Public/Mesh.h"
+
+/**
+ * Namespace designed to contain a series of methods used for the generation of unreal meshes for previewing purposes.
+ * While they can be accessed to do other stuff it is not granted they will produce a complete USkeletalMesh since at the
+ * moment it only tries to generate something viewable, but may be missing bone, vertex or/and animation data.
+ */
+namespace MutableMeshPreviewUtils
+{
+	/**
+	 * Method designed to provide a way of getting a USkeletalMesh based on the data found on a mutable mesh. The
+	 * usage of a Reference Skeletal mesh is required to aid on the generation of the skeletal mesh.
+	 * @param InMutableMesh - The mutable mesh used as base for the generation of the USkeletalMesh
+	 * @param InReferenceSkeletalMesh - Skeletal mesh used as source for the data not stored by mutable mesh. It must be
+	 * "similar" in terms of bone structure to the mutable mesh. Example: If generating a mesh from a mutable mesh found
+	 * on an instance (using the mutable debugger) you may want to provide as reference skeletal mesh the one used by
+	 * the Customizable Object.
+	 * @return - A new Skeletal mesh based on the provided data. It may be a nullptr if the conversion fails.
+	 */
+	CUSTOMIZABLEOBJECTEDITOR_API USkeletalMesh* GenerateSkeletalMeshFromMutableMesh(
+		mu::MeshPtrConst InMutableMesh, const USkeletalMesh* InReferenceSkeletalMesh);
+
+}
diff --git a/Engine/Plugins/Experimental/Mutable/Source/CustomizableObjectEditor/Public/MutableUtils.h b/Engine/Plugins/Experimental/Mutable/Source/CustomizableObjectEditor/Public/MutableUtils.h
new file mode 100644
--- /dev/null
+++ b/Engine/Plugins/Experimental/Mutable/Source/CustomizableObjectEditor/Public/MutableUtils.h
@@ -0,0 +1,35 @@
+// Copyright Epic Games, Inc. All Rights Reserved.
+
+#pragma once
+
+#include "Engine/SkeletalMesh.h"
+#include "Engine/StaticMesh.h"
+#include "Math/Vector2D.h"
+
+
+/** Returns the mesh UV. */
+TArray<FVector2f> GetUV(const USkeletalMesh& SkeletalMesh, const int32 LODIndex, const int32 SectionIndex, const int32 UVIndex);
+
+
+/** Returns the mesh UV. */
+TArray<FVector2f> GetUV(const UStaticMesh& StaticMesh, const int32 LODIndex, const int32 SectionIndex, const int32 UVIndex);
+
+
+/** Returns if the given point is inside the 0-1 bounds. */
+bool HasNormalizedBounds(const FVector2f& Point);
+
+
+/** Returns if the given UV is normalized. */
+bool IsUVNormalized(const USkeletalMesh& SkeletalMesh, const int32 LODIndex, const int32 SectionIndex, const int32 UVIndex);
+
+
+/** Returns if the given UV is normalized. */
+bool IsUVNormalized(const UStaticMesh& SkeletalMesh, const int32 LODIndex, const int32 SectionIndex, const int32 UVIndex);
+
+
+/** Check if a given LOD and material in a skeletal mesh is closed. */
+bool IsMeshClosed(const USkeletalMesh* Mesh, int LOD, int MaterialIndex);
+
+
+/** Check if a given LOD and material in a skeletal mesh is closed. */
+bool IsMeshClosed(const UStaticMesh* Mesh, int LOD, int MaterialIndex);
diff --git a/Engine/Plugins/Experimental/Mutable/Source/CustomizableObjectEditor/Public/RemapPins/CustomizableObjectNodeRemapPinsByName.h b/Engine/Plugins/Experimental/Mutable/Source/CustomizableObjectEditor/Public/RemapPins/CustomizableObjectNodeRemapPinsByName.h
new file mode 100644
--- /dev/null
+++ b/Engine/Plugins/Experimental/Mutable/Source/CustomizableObjectEditor/Public/RemapPins/CustomizableObjectNodeRemapPinsByName.h
@@ -0,0 +1,26 @@
+// Copyright Epic Games, Inc. All Rights Reserved.
+
+#pragma once
+
+#include "CustomizableObjectNodeRemapPins.h"
+
+#include "CustomizableObjectNodeRemapPinsByName.generated.h"
+
+
+/** 
+ * Remap pins by name.
+ * 
+ * Remap pins by the field Pin->PinName. Output pins and input pins get remapped independently.
+ * 
+ * Use inheritance is a node requires a set of special rules when remapping by name.
+ */
+UCLASS()
+class UCustomizableObjectNodeRemapPinsByName : public UCustomizableObjectNodeRemapPins
+{
+public:
+	GENERATED_BODY()
+
+	virtual bool Equal(const UEdGraphPin& OldPin, const UEdGraphPin& NewPin) const;
+	
+	virtual void RemapPins(const TArray<UEdGraphPin*>& OldPins, const TArray<UEdGraphPin*>& NewPins, TMap<UEdGraphPin*, UEdGraphPin*>& PinsToRemap, TArray<UEdGraphPin*>& PinsToOrphan) override;
+};
\ No newline at end of file
diff --git a/Engine/Plugins/Experimental/Mutable/Source/CustomizableObjectEditor/Public/RemapPins/CustomizableObjectNodeRemapPinsByNameDefaultPin.h b/Engine/Plugins/Experimental/Mutable/Source/CustomizableObjectEditor/Public/RemapPins/CustomizableObjectNodeRemapPinsByNameDefaultPin.h
new file mode 100644
--- /dev/null
+++ b/Engine/Plugins/Experimental/Mutable/Source/CustomizableObjectEditor/Public/RemapPins/CustomizableObjectNodeRemapPinsByNameDefaultPin.h
@@ -0,0 +1,22 @@
+// Copyright Epic Games, Inc. All Rights Reserved.
+
+#pragma once
+
+#include "CustomizableObjectNodeRemapPinsByName.h"
+
+#include "CustomizableObjectNodeRemapPinsByNameDefaultPin.generated.h"
+
+
+/** If a default pin is given, remap it to the first pin. Otherwise, remap pins by name as usual. */
+UCLASS()
+class UCustomizableObjectNodeRemapPinsByNameDefaultPin : public UCustomizableObjectNodeRemapPinsByName
+{
+	GENERATED_BODY()
+
+public:
+	/** Default pin to remap with. */
+	UEdGraphPin* DefaultPin = nullptr;
+
+	virtual void RemapPins(const TArray<UEdGraphPin*>& OldPins, const TArray<UEdGraphPin*>& NewPins, TMap<UEdGraphPin*, UEdGraphPin*>& PinsToRemap, TArray<UEdGraphPin*>& PinsToOrphan) override;
+};
+
diff --git a/Engine/Plugins/Experimental/Mutable/Source/CustomizableObjectEditor/Public/RemapPins/CustomizableObjectNodeRemapPinsByPosition.h b/Engine/Plugins/Experimental/Mutable/Source/CustomizableObjectEditor/Public/RemapPins/CustomizableObjectNodeRemapPinsByPosition.h
new file mode 100644
--- /dev/null
+++ b/Engine/Plugins/Experimental/Mutable/Source/CustomizableObjectEditor/Public/RemapPins/CustomizableObjectNodeRemapPinsByPosition.h
@@ -0,0 +1,24 @@
+// Copyright Epic Games, Inc. All Rights Reserved.
+
+#pragma once
+
+#include "CustomizableObjectNodeRemapPins.h"
+
+#include "CustomizableObjectNodeRemapPinsByPosition.generated.h"
+
+
+/**
+ * Remap pins by position.
+ *
+ * Remap pins by their relative order. Output pins and input pins get remapped independently.
+ *
+ * Use inheritance is a node requires a set of special rules when remapping by position.
+ */
+UCLASS()
+class UCustomizableObjectNodeRemapPinsByPosition : public UCustomizableObjectNodeRemapPins
+{
+public:
+	GENERATED_BODY()
+
+	virtual void RemapPins(const TArray<UEdGraphPin*>& OldPins, const TArray<UEdGraphPin*>& NewPins, TMap<UEdGraphPin*, UEdGraphPin*>& PinsToRemap, TArray<UEdGraphPin*>& PinsToOrphan) override;
+};
\ No newline at end of file
diff --git a/Engine/Plugins/Experimental/Mutable/Source/CustomizableObjectEditor/Public/SCustomizableInstanceProperties.h b/Engine/Plugins/Experimental/Mutable/Source/CustomizableObjectEditor/Public/SCustomizableInstanceProperties.h
new file mode 100644
--- /dev/null
+++ b/Engine/Plugins/Experimental/Mutable/Source/CustomizableObjectEditor/Public/SCustomizableInstanceProperties.h
@@ -0,0 +1,261 @@
+// Copyright Epic Games, Inc. All Rights Reserved.
+
+#pragma once
+
+#include "CustomizableObject.h"
+#include "CustomizableObjectInstance.h"
+#include "Widgets/SCompoundWidget.h"
+#include "Widgets/Layout/SExpandableArea.h"
+#include "Widgets/DeclarativeSyntaxSupport.h"
+#include "Styling/SlateTypes.h"
+#include "Widgets/Input/SButton.h"
+#include "Widgets/Text/STextBlock.h"
+#include "Widgets/Input/SSearchBox.h"
+#include "Widgets/Input/SSpinBox.h"
+#include "CustomizableInstanceDetails.h"
+
+class UCustomizableObjectInstance;
+class SSlider;
+class SVerticalBox;
+class UUpdateClassHelper;
+class UInstanceUpdatedHelper;
+
+struct FMutableParamExpandableArea
+{
+	TSharedPtr<SExpandableArea> ExpandableArea;
+	TSharedPtr<SVerticalBox> VerticalBox;
+
+	FMutableParamExpandableArea(TSharedPtr<SExpandableArea> &InExpandableArea, TSharedPtr<SVerticalBox> &InVerticalBox)
+		: ExpandableArea(InExpandableArea), VerticalBox(InVerticalBox)
+	{ }
+};
+
+class CUSTOMIZABLEOBJECTEDITOR_API SCustomizableInstanceProperties : public SCompoundWidget
+{
+public:
+	friend class SCreateProfileParameters;
+	
+	SLATE_BEGIN_ARGS(SCustomizableInstanceProperties) {}
+	SLATE_ARGUMENT(TWeakObjectPtr<UCustomizableObjectInstance>, CustomInstance)
+	SLATE_ARGUMENT(TWeakPtr<FCustomizableInstanceDetails>, InstanceDetails)
+	SLATE_END_ARGS()
+	
+	/**
+	 * Construct the widget
+	 *
+	 * @param	InArgs			A declaration from which to construct the widget
+	 */
+	void Construct(const FArguments& InArgs);
+
+	// SWidget interface
+	virtual void Tick( const FGeometry& AllottedGeometry, const double InCurrentTime, const float InDeltaTime ) override;
+
+	/** */
+	const UCustomizableObjectInstance* GetInstance() const;
+	void SetNoInstanceMessage( const FText& Message );
+
+	/** Set SimplifiedProjectorInterface bool value - Deprecated */
+	void SetSimplifiedProjectorInterface( bool Enabled );
+
+	/** Redraws the Instance Properties Widget */
+	void ResetParamBox();
+
+	/** Bool to Redraw the Instance Properties Widget */
+	bool bShouldResetParamBox = true;
+
+private:
+	/** Bool that decides the type of interface of projectors - Deprecated */
+	bool SimplifiedProjectorInterface;
+
+	/** Message to to indicate that there is no instance and the CO needs to be compiled*/
+	FText NoInstanceMessage;
+
+	/** Pointer to the Instance Details. Owner of this widget. */
+	TWeakPtr<FCustomizableInstanceDetails> InstanceDetails;
+
+	/** Pointer to the Customizable Object Instance */
+	TWeakObjectPtr<UCustomizableObjectInstance> CustomInstance;
+
+	/** Array with all the possible states for the states combo box */
+	TArray< TSharedPtr<FString> > StateNames;
+	
+	/** Pointer to the widget that stores all the parameter */
+	TSharedPtr<SVerticalBox> ParamBox;
+
+	/** Array to store the option names of integers */
+	//TArray< TSharedPtr< TArray< TSharedPtr<FString> > > > IntOptionNames;
+	
+	/** Array to store dynamic brushes */
+	TArray< TSharedPtr<class FDeferredCleanupSlateBrush> > DynamicBrushes;
+	
+	/** Used to insert child params to a parent's expandable area */
+	TMap<FString, FMutableParamExpandableArea> ParamNameToExpandableAreaMap;
+	
+	/** Maps param name to children param indices, used to walk the params in order respecting parent/children relationships */
+	TMultiMap<FString, int32> ParamChildren;
+	
+	/** Maps param index to bool telling if it has parent, same use as previous line */
+	TMap<int32, bool> ParamHasParent; 
+
+	//int32 NumParams = 0; // Same use as previous line
+
+	// These arrays store the textures available for texture parameters of the model.
+	// These come from the texture generators registered in the CustomizableObjectSystem
+	TArray< TSharedPtr<FString> > TextureParameterValueNames;
+	TArray< int64_t > TextureParameterValues;
+
+	struct FSliderData
+	{
+		FSliderData( TSharedPtr<SSpinBox<float>> InSlider, const FString& InParamName, int32 InRangeIndex, float InLastValueSet)
+		{
+			Slider = InSlider;
+			ParameterName = InParamName;
+			RangeIndex = InRangeIndex;
+
+			LastValueSet = InLastValueSet;
+		}
+
+		TSharedPtr<SSpinBox<float>> Slider;
+		FString ParameterName;
+		int32 RangeIndex;
+		float LastValueSet;
+	};
+	TArray< FSliderData > FloatSliders;
+
+	// Get the last value set on a slider
+	float GetSliderValue( int Slider ) const;
+
+	/** Adds a widget for each the parameter */
+	void AddParameter(int32 ParamIndex);
+
+	/** Fills the children map with each parameter name and its index of the children objects */
+	void FillChildrenMap(int32 ParamIndex);
+
+	/** Fills recursively the paramchildren and adds the parameter to create its widget */
+	void RecursivelyAddParamAndChildren(int32 ParamIndexInObject);
+
+	/** CheckBox OnChecked callbacks - Determine the parameters that are shown */
+	void OnShowOnlyRuntimeSelectionChanged(ECheckBoxState InCheckboxState);
+	void OnShowOnlyRelevantSelectionChanged(ECheckBoxState InCheckboxState);
+
+	/** ComboBox OnSelectionChanged callback - Change the State of the instance */
+	void OnStateComboBoxSelectionChanged(TSharedPtr<FString> Selection, ESelectInfo::Type SelectInfo);
+
+	/** Copies the preview instance parameters */
+	FReply OnCopyAllParameters();
+
+	/** Pastes the preview instance parameters */
+	FReply OnPasteAllParameters();
+
+	/** Clears the preview instance parameters */
+	FReply OnResetAllParameters();
+
+	/** Buffer to store the copied parameters */
+	class FBufferArchive *ToBinary;
+
+	/** Check box callbacks for bool parameters */
+	void OnBoolParameterChanged(ECheckBoxState InCheckboxState, FString ParamName);
+	ECheckBoxState GetBoolParameterValue(FString ParamName) const;
+
+	/** Expandable Area OnAreaExpansionChanged callback */
+	void OnAreaExpansionChanged(bool bExpanded, FString ParamName) const;
+
+	/** Slider callbacks for float parameters */
+	float GetFloatParameterValue(FString ParamName, int RangeIndex) const;
+	void OnFloatParameterChanged(float Value, int SliderIndex);
+	void OnFloatParameterSliderBegin();
+	void OnFloatParameterSliderEnd(float Value);
+
+	/** Spinbox callback for int parameters */
+	void OnIntParameterChanged(int32 Value, FString ParamName); // Not implemented Yet
+
+	/** SearchBox OnTextCommitted Callback to search a parameter int by name */
+	void OnIntParameterComboBoxChanged(TSharedPtr<FString> Selection, ESelectInfo::Type SelectInfo, FString ParamName);
+
+	/** Gets the name value of a int parameter */
+	FString GetIntParameterValue(FString ParamName, int RangeIndex = -1) const;
+
+	/** CheckBox IsChecked Callbacks for projectors management */
+	ECheckBoxState IsProjectorHidden(FString ParamNameWithIndex) const;
+	ECheckBoxState IsProjectorTranslate(FString ParamNameWithIndex) const;
+	ECheckBoxState IsProjectorRotate(FString ParamNameWithIndex) const;
+	ECheckBoxState IsProjectorScale(FString ParamNameWithIndex) const;
+
+	/** CheckBox OnCheckStateChanged Callbacks for projectors management */
+	void OnProjectorTranslateStateChanged( ECheckBoxState InCheckboxState, FString ParamNameWithIndex) const;
+	void OnProjectorRotateStateChanged( ECheckBoxState InCheckboxState, FString ParamNameWithIndex) const;
+	void OnProjectorScaleStateChanged( ECheckBoxState InCheckboxState, FString ParamNameWithIndex) const;
+	
+	/** Button OnClicked Callback to select or unselect a projector */
+	FReply OnProjectorSelectStateChanged(FString ParamNameWithIndex);
+	
+	/** SearchBox OnTextCommitted Callback to search a parameter texture by name */
+	void OnProjectorTextureParameterComboBoxChanged(TSharedPtr<FString> Selection, ESelectInfo::Type SelectInfo, FString ParamNameWithIndex) const;
+
+	/** Slider OnValueChanged Callbacks for float parameters */
+	void OnProjectorFloatParameterChanged(float Value, int SliderIndex, FString ParamNameWithIndex);
+	void OnProjectorFloatParameterSliderBegin();
+	void OnProjectorFloatParameterSliderEnd(float Value);
+
+	/** OnClicked Button Callbacks - Manage Projector layers */
+	FReply OnProjectorLayerAdded(FString ParamName) const;
+	FReply OnProjectorLayerRemoved(FString ParamNameWithIndex) const;
+
+	/** OnClicked Button Callbacks - Manage Projectors Transform */
+    FReply OnProjectorCopyTransform(FString ParamName) const;
+    FReply OnProjectorPasteTransform(FString ParamName) const;
+    FReply OnProjectorResetTransform(FString ParamName) const;
+
+	/** Gets the value of a vector parameter */
+	FLinearColor GetColorParameterValue(FString ParamName) const;
+
+	/** Color Block OnMouseButtonDown Callback */
+	FReply OnColorBlockMouseButtonDown(const FGeometry& MyGeometry, const FPointerEvent& MouseEvent, FString ParamName);
+
+	/** Color Picker OnColorCommitted Callback - Changes the value of a vector parameter */
+	void OnSetColorFromColorPicker(FLinearColor NewColor, FString PickerParamName) const;
+
+	/** ComboBox OnSelectionChanged Callback - Changes the texture of a projector */
+	void OnTextureParameterComboBoxSelectionChanged(TSharedPtr<FString> Selection, ESelectInfo::Type SelectInfo, FString ParamName);
+
+	/** Array with the names of the current projector parameter names with their multidimensional (range) index */
+	TArray<FString> ArrayProjectorParameterNameWithIndex;
+
+	bool bSlidersChanged = false;
+
+	// Map where to store the instances of the select / unselect projector SButton
+	TMap<FString, TSharedPtr<SButton>> MapProjectorButton;
+	TMap<FString, TSharedPtr<STextBlock>> MapProjectorButtonTextBlock;
+
+
+	/** Called when the filter text changes.  This filters specific property nodes out of view */
+	void OnFilterTextChanged(const FText& InFilterText);
+
+	/*Search box that filters the Customizable Instance Properties*/
+	TSharedPtr<SSearchBox> SB_Properties;
+
+/*Profiles*/
+
+	/*Open the window to create a new profile*/
+	FReply CreateParameterProfileWindow();
+
+	/*Remove selected profile*/
+	FReply RemoveParameterProfile();
+
+	/*Stores the profile that has been selected*/
+	void OnProfileSelectedChanged(TSharedPtr<FString> Selection, ESelectInfo::Type SelectInfo);
+
+	/*Set all the profile parameters to the current COI*/
+	void SetParameterProfileNamesOnEditor();
+
+	/*Stores all the possible profiles for the current COI*/
+	TArray<TSharedPtr<FString>> ParameterProfileNames;
+
+	/*Check if there is the COI has any available parameter*/
+	bool HasAnyParameters() const;
+
+	void InstanceUpdated(UCustomizableObjectInstance* Instance) const;
+	
+	/*Check if SaveDescriptor has been called before LoadDescriptor*/
+	bool ParametersHaveBeenRead;
+};
\ No newline at end of file
diff --git a/Engine/Plugins/Experimental/Mutable/Source/CustomizableObjectEditor/Public/SCustomizableObjectEditorViewportToolBar.h b/Engine/Plugins/Experimental/Mutable/Source/CustomizableObjectEditor/Public/SCustomizableObjectEditorViewportToolBar.h
new file mode 100644
--- /dev/null
+++ b/Engine/Plugins/Experimental/Mutable/Source/CustomizableObjectEditor/Public/SCustomizableObjectEditorViewportToolBar.h
@@ -0,0 +1,196 @@
+// Copyright Epic Games, Inc. All Rights Reserved.
+
+#pragma once
+
+#include "SViewportToolBar.h"
+#include "Framework/MultiBox/MultiBoxBuilder.h"
+#include "Settings/LevelEditorViewportSettings.h"
+
+
+/**
+* A level viewport toolbar widget that is placed in a viewport
+*/
+class SCustomizableObjectEditorViewportToolBar : public SViewportToolBar
+{
+public:
+	SLATE_BEGIN_ARGS(SCustomizableObjectEditorViewportToolBar) {}
+	SLATE_END_ARGS()
+
+	void Construct(const FArguments& InArgs, TSharedPtr<class SCustomizableObjectEditorViewportTabBody> InViewport, TSharedPtr<class SEditorViewport> InRealViewport);
+
+private:
+	/**
+	* Generates the toolbar view menu content
+	*
+	* @return The widget containing the view menu content
+	*/
+	TSharedRef<SWidget> GenerateViewMenu() const;
+
+	/**
+	* Generates the toolbar show menu content
+	*
+	* @return The widget containing the show menu content
+	*/
+	//TSharedRef<SWidget> GenerateShowMenu() const;
+
+	/**
+	* Generates the Show -> Scene sub menu content
+	*/
+	//void FillShowSceneMenu(FMenuBuilder& MenuBuilder) const;
+
+	/**
+	* Generates the Show -> Advanced sub menu content
+	*/
+	//void FillShowAdvancedMenu(FMenuBuilder& MenuBuilder) const;
+
+	/**
+	* Generates the Show -> Bone sub menu content
+	*/
+	//void FillShowBoneDrawMenu(FMenuBuilder& MenuBuilder) const;
+
+	/**
+	* Generates the Show -> Overlay sub menu content
+	*/
+	//void FillShowOverlayDrawMenu(FMenuBuilder& MenuBuilder) const;
+
+	/**
+	* Generates the Show -> Clothing sub menu content
+	*/
+	//void FillShowClothingMenu(FMenuBuilder& MenuBuilder) const;
+
+	/**
+	* Generates the Show -> Display Info sub menu content
+	*/
+	//void FillShowDisplayInfoMenu(FMenuBuilder& MenuBuilder) const;
+
+	/** Add the projector Rotation, traslation and scale buttons to viewport toolbar */
+	TSharedRef<SWidget> GenerateRTSButtons();
+
+	/**
+	* Generates the toolbar LOD menu content
+	*
+	* @return The widget containing the LOD menu content based on LOD model count
+	*/
+	TSharedRef<SWidget> GenerateLODMenu() const;
+
+	/**
+	* Returns the label for the "LOD" tool bar menu, which changes depending on the current LOD selection
+	*
+	* @return	Label to use for this LOD label
+	*/
+	FText GetLODMenuLabel() const;
+
+	/**
+	* Generates the toolbar viewport type menu content
+	*
+	* @return The widget containing the viewport type menu content
+	*/
+	TSharedRef<SWidget> GenerateViewportTypeMenu() const;
+
+	/**
+	* Generates the toolbar playback menu content
+	*
+	* @return The widget containing the playback menu content
+	*/
+	TSharedRef<SWidget> GeneratePlaybackMenu() const;
+
+	/** Generate the turntable menu entries */
+	void GenerateTurnTableMenu(FMenuBuilder& MenuBuilder) const;
+
+	/** Generate the scene setup menu */
+	void GenerateSceneSetupMenu(FMenuBuilder& MenuBuilder);
+
+	/**
+	* Generates the toolbar Camera Mode menu content
+	*
+	* @return The widget containing the Camera Mode menu content
+	*/
+	TSharedRef<SWidget> GenerateCameraModeMenu() const;
+
+	/**
+	* Returns the label for the Camera Mode tool bar menu, which changes depending on the current Camera Mode selection
+	*
+	* @return	Label to use for this Camera Mode label
+	*/
+	FText GetCameraModeMenuLabel() const;
+
+
+	/** Customize the details of the scene setup object */
+	//TSharedRef<class IDetailCustomization> CustomizePreviewSceneDescription();
+
+	/** Customize a preview mesh collection entry */
+	//TSharedRef<class IPropertyTypeCustomization> CustomizePreviewMeshCollectionEntry();
+
+	/**
+	* Generate color of the text on the top
+	*/
+	FSlateColor GetFontColor() const;
+
+	/**
+	* Returns the label for the Playback tool bar menu, which changes depending on the current playback speed
+	*
+	* @return	Label to use for this Menu
+	*/
+	FText GetPlaybackMenuLabel() const;
+
+	/**
+	* Returns the label for the Viewport type tool bar menu, which changes depending on the current selected type
+	*
+	* @return	Label to use for this Menu
+	*/
+	FText GetCameraMenuLabel() const;
+	const FSlateBrush* GetCameraMenuLabelIcon() const;
+
+	/** Called by the FOV slider in the perspective viewport to get the FOV value */
+	float OnGetFOVValue() const;
+	/** Called when the FOV slider is adjusted in the perspective viewport */
+	void OnFOVValueChanged(float NewValue) const;
+	/** Called when a value is entered into the FOV slider/box in the perspective viewport */
+	void OnFOVValueCommitted(float NewValue, ETextCommit::Type CommitInfo);
+
+	/** Called by the floor offset slider in the perspective viewport to get the offset value */
+	TOptional<float> OnGetFloorOffset() const;
+	/** Called when the floor offset slider is adjusted in the perspective viewport */
+	void OnFloorOffsetChanged(float NewValue);
+
+	// Called to determine if the gizmos can be used in the current preview
+	EVisibility GetTransformToolbarVisibility() const;
+
+	// Called to show / hide Customizalbe Object compile error
+	EVisibility GetShowCompileErrorOverlay() const;
+
+	// Called to modify the instance compile status overlay information text
+	FText GetCompileErrorOverlayText() const;
+
+	// Called to show / hide Customizalbe Object info layout
+	EVisibility GetShowCompileInfoOverlay() const;
+
+	ECheckBoxState IsRotationGridSnapChecked() const;
+	void HandleToggleRotationGridSnap(ECheckBoxState InState);
+	FText GetRotationGridLabel() const;
+	TSharedRef<SWidget> FillRotationGridSnapMenu();
+	TSharedRef<SWidget> BuildRotationGridCheckBoxList(FName InExtentionHook, const FText& InHeading, const TArray<float>& InGridSizes, ERotationGridMode InGridMode) const;
+
+	/** Callback for drop-down menu with FOV and high resolution screenshot options currently */
+	FReply OnMenuClicked();
+
+	/** Generates drop-down menu with FOV and high resolution screenshot options currently */
+	TSharedRef<SWidget> GenerateOptionsMenu() const;
+
+	/** Generates widgets for viewport camera FOV control */
+	TSharedRef<SWidget> GenerateFOVMenu() const;
+
+private:
+	/** The viewport that we are in */
+	TWeakPtr<class SCustomizableObjectEditorViewportTabBody> Viewport;
+
+	// Layout to show information about instance skeletal mesh update / CO asset data
+	TSharedPtr<class SButton> CompileErrorLayout;
+
+	// Layout to show if the skeletal mesh of a functional CO is being updated
+	TSharedPtr<class SButton> UpdateInfoLayout;
+
+	SViewportToolBar* CastedViewportToolbar;
+
+	TSharedPtr<SMenuAnchor> MenuAnchor;
+};
diff --git a/Engine/Plugins/Experimental/Mutable/Source/CustomizableObjectEditor/Public/UnrealEditorPortabilityHelpers.h b/Engine/Plugins/Experimental/Mutable/Source/CustomizableObjectEditor/Public/UnrealEditorPortabilityHelpers.h
new file mode 100644
--- /dev/null
+++ b/Engine/Plugins/Experimental/Mutable/Source/CustomizableObjectEditor/Public/UnrealEditorPortabilityHelpers.h
@@ -0,0 +1,124 @@
+// Copyright Epic Games, Inc. All Rights Reserved.
+
+#pragma once
+
+#include "UnrealPortabilityHelpers.h"
+
+#include "ICustomizableObjectModule.h"
+#include "Runtime/Launch/Resources/Version.h"
+#include "Rendering/SkeletalMeshLODModel.h"
+#include "Rendering/SkeletalMeshModel.h"
+#include "Subsystems/AssetEditorSubsystem.h"
+#include "Framework/Commands/InputChord.h"
+
+#include "DragAndDrop/AssetDragDropOp.h"
+#include "Materials/Material.h"
+#include "UObject/UnrealType.h"
+#include "Editor.h"
+#include "Engine/StaticMesh.h"
+#include "Engine/DataTable.h"
+
+#include <functional>
+
+//---------------------------------------------------------------
+// Helpers to ease portability across unreal engine versions
+//---------------------------------------------------------------
+
+
+inline bool Helper_DoSectionsUse16BitBoneIndex( const FSkeletalMeshLODModel& m )
+{
+	for (int32 SectionIdx = 0; SectionIdx < m.Sections.Num(); ++SectionIdx)
+	{
+		if (m.Sections[SectionIdx].Use16BitBoneIndex())
+		{
+			return true;
+		}
+	}
+
+	return false;
+}
+
+inline int Helper_GetMaxBoneInfluences(const FSkeletalMeshLODModel& m)
+{
+	return m.GetMaxBoneInfluences();
+}
+
+
+inline FMaterialRenderProxy* Helper_GetMaterialProxy( const UMaterialInterface* Material )
+{
+	return Material->GetRenderProxy();
+}
+
+
+inline const TArray<FAssetData>& Helper_GetAssets(FAssetDragDropOp* DragDropOp)
+{
+	return DragDropOp->GetAssets();
+}
+
+inline FName Helper_GetPropertyName(FPropertyChangedEvent& PropertyChangedEvent)
+{
+	return PropertyChangedEvent.GetPropertyName();
+}
+
+inline FName Helper_GetPinFName(const UEdGraphPin* Pin)
+{
+	return Pin ? Pin->PinName: FName();
+}
+
+inline FString Helper_GetPinName(const UEdGraphPin* Pin)
+{
+	return Pin ? Pin->PinName.ToString() : FString();
+}
+
+inline void Helper_SetPinName( FName& To, const FString& From )
+{
+	To = FName(*From);
+}
+
+inline const FName& Helper_GetPinCategory(const FName& Cat)
+{
+	return Cat;
+}
+
+inline void Helper_GetMeshIndices(TArray<uint32>& Indices, USkeletalMesh* SkeletalMesh, int LODIndex)
+{
+	Indices = Helper_GetImportedModel(SkeletalMesh)->LODModels[LODIndex].IndexBuffer;
+}
+
+
+inline void Helper_ForEachTextureParameterName(const UMaterial* Material, const std::function<void(const FString&)>& f)
+{
+	TArray<FMaterialParameterInfo> ImageNames;
+	TArray<FGuid> ImageIds;
+	Material->GetAllTextureParameterInfo(ImageNames, ImageIds);
+	for (int32 ImageIndex = 0; ImageIndex < ImageNames.Num(); ++ImageIndex)
+	{
+		FString ImageName = ImageNames[ImageIndex].Name.ToString();
+		f(ImageName);
+	}
+}
+
+
+inline FStaticMeshVertexBuffer& Helper_GetStaticMeshVertexBuffer(UStaticMesh* StaticMesh, int LODIndex)
+{
+	return StaticMesh->GetRenderData()->LODResources[LODIndex].VertexBuffers.StaticMeshVertexBuffer;
+}
+
+
+inline const FPositionVertexBuffer& Helper_GetStaticMeshPositionVertexBuffer(UStaticMesh* StaticMesh, int LODIndex)
+{
+	return StaticMesh->GetRenderData()->LODResources[LODIndex].VertexBuffers.PositionVertexBuffer;
+}
+
+inline const TMap<FName, uint8*>& Helper_GetRowMap(UDataTable* Table)
+{
+	return Table->GetRowMap();
+}
+
+inline UAssetEditorSubsystem* Helper_GetEditorSubsystem()
+{
+	return GEditor->GetEditorSubsystem<UAssetEditorSubsystem>();
+}
+
+
+
diff --git a/Engine/Plugins/Experimental/Mutable/Source/CustomizableObjectPopulation/Private/CustomizableObjectPopulationCharacteristic.cpp b/Engine/Plugins/Experimental/Mutable/Source/CustomizableObjectPopulation/Private/CustomizableObjectPopulationCharacteristic.cpp
new file mode 100644
--- /dev/null
+++ b/Engine/Plugins/Experimental/Mutable/Source/CustomizableObjectPopulation/Private/CustomizableObjectPopulationCharacteristic.cpp
@@ -0,0 +1,9 @@
+// Copyright Epic Games, Inc. All Rights Reserved.
+
+#include "CustomizableObjectPopulationCharacteristic.h"
+#include "CustomizableObject.h"
+
+#define LOCTEXT_NAMESPACE "CustomizableObjectPopulationCharacteristic"
+
+
+#undef LOCTEXT_NAMESPACE
diff --git a/Engine/Plugins/Experimental/Mutable/Source/CustomizableObjectPopulation/Private/CustomizableObjectPopulationClass.cpp b/Engine/Plugins/Experimental/Mutable/Source/CustomizableObjectPopulation/Private/CustomizableObjectPopulationClass.cpp
new file mode 100644
--- /dev/null
+++ b/Engine/Plugins/Experimental/Mutable/Source/CustomizableObjectPopulation/Private/CustomizableObjectPopulationClass.cpp
@@ -0,0 +1,18 @@
+// Copyright Epic Games, Inc. All Rights Reserved.
+
+#include "CustomizableObjectPopulationClass.h"
+#include "CustomizableObjectPopulationCharacteristic.h"
+#include "CustomizableObjectPopulationCustomVersion.h"
+
+#define LOCTEXT_NAMESPACE "CustomizableObjectPopulationClass"
+
+using namespace CustomizableObjectPopulation;
+
+void UCustomizableObjectPopulationClass::Serialize(FArchive& Ar)
+{
+	Super::Serialize(Ar);
+
+	Ar.UsingCustomVersion(FCustomizableObjectPopulationCustomVersion::GUID);
+}
+
+#undef LOCTEXT_NAMESPACE
\ No newline at end of file
