[mutable] Move Mutable plugin from the restricted folder to the experimental plugins folder.

diff --git a/Engine/Plugins/Experimental/Mutable/Source/CustomizableObjectPopulationEditor/CustomizableObjectPopulationEditor.Build.cs b/Engine/Plugins/Experimental/Mutable/Source/CustomizableObjectPopulationEditor/CustomizableObjectPopulationEditor.Build.cs
new file mode 100644
--- /dev/null
+++ b/Engine/Plugins/Experimental/Mutable/Source/CustomizableObjectPopulationEditor/CustomizableObjectPopulationEditor.Build.cs
@@ -0,0 +1,128 @@
+// Copyright Epic Games, Inc. All Rights Reserved.
+
+using UnrealBuildTool;
+using System.IO;
+
+public class CustomizableObjectPopulationEditor : ModuleRules
+{
+	private int ReadEngineMajorVersion()
+	{
+		string line;
+
+		System.IO.StreamReader file =
+		   new System.IO.StreamReader("Runtime/Launch/Resources/Version.h");
+
+		while ((line = file.ReadLine()) != null)
+		{
+			if (line.StartsWith("#define ENGINE_MAJOR_VERSION"))
+			{
+				file.Close();
+				return int.Parse(line.Split()[2]);
+			}
+		}
+
+		file.Close();
+		return 0;
+	}
+
+	private int ReadEngineMinorVersion()
+	{
+		string line;
+
+		System.IO.StreamReader file =
+		   new System.IO.StreamReader("Runtime/Launch/Resources/Version.h");
+
+		while ((line = file.ReadLine()) != null)
+		{
+			if (line.StartsWith("#define ENGINE_MINOR_VERSION"))
+			{
+				file.Close();
+				return int.Parse(line.Split()[2]);
+			}
+		}
+
+		file.Close();
+		return 0;
+	}
+
+	public CustomizableObjectPopulationEditor(ReadOnlyTargetRules TargetRules) : base(TargetRules)
+	{
+		ShortName = "MuCOPE";
+
+		DefaultBuildSettings = BuildSettingsVersion.V2;
+
+		PrivateIncludePathModuleNames.AddRange(
+			new string[] {
+				"AssetRegistry",
+				"Kismet",
+				"EditorWidgets",
+				"MeshUtilities",
+				"ContentBrowser",
+				"SkeletonEditor",
+				"Persona",
+				"WorkspaceMenuStructure",
+				"AdvancedPreviewScene",
+
+				"MessageLog",
+				"KismetWidgets",
+			}
+		);
+
+		PrivateDependencyModuleNames.AddRange(
+			new string[] {
+				"Core",
+				"CoreUObject",
+				"Engine",
+				
+				"Slate",
+				"SlateCore",
+				"RenderCore",
+				"RHI",
+				"UnrealEd",
+				"TargetPlatform",
+				"RawMesh",
+				"PropertyEditor",
+				"LevelEditor",
+				"AssetTools",
+				"GraphEditor",
+				"InputCore",
+				"Kismet",
+				"AdvancedPreviewScene",
+				"AppFramework",
+				"Projects",
+				"ClothingSystemRuntimeCommon",
+
+				"MutableRuntime",
+				"MutableTools",
+			}
+		);
+
+		if (ReadEngineMajorVersion() >= 4 && ReadEngineMinorVersion() >= 18)
+		{
+			PrivateDependencyModuleNames.AddRange(new string[]
+			{
+				"ApplicationCore",
+				"ToolMenus",
+			});
+		}
+
+		DynamicallyLoadedModuleNames.AddRange(
+			new string[] {
+				"MainFrame",
+				"SceneOutliner",
+				"ClassViewer",
+				"ContentBrowser",
+				"WorkspaceMenuStructure",
+			}
+		);
+
+		PublicDependencyModuleNames.AddRange(
+			new string[] {
+				"PropertyEditor",
+				"CustomizableObject",
+				"CustomizableObjectEditor",
+				"CustomizableObjectPopulation",
+			}
+		);
+	}
+}
diff --git a/Engine/Plugins/Experimental/Mutable/Source/CustomizableObjectPopulationEditor/Private/AssetTypeActions_CustomizableObjectPopulation.cpp b/Engine/Plugins/Experimental/Mutable/Source/CustomizableObjectPopulationEditor/Private/AssetTypeActions_CustomizableObjectPopulation.cpp
new file mode 100644
--- /dev/null
+++ b/Engine/Plugins/Experimental/Mutable/Source/CustomizableObjectPopulationEditor/Private/AssetTypeActions_CustomizableObjectPopulation.cpp
@@ -0,0 +1,165 @@
+// Copyright Epic Games, Inc. All Rights Reserved.
+
+#include "AssetTypeActions_CustomizableObjectPopulation.h"
+
+#include "ContentBrowserModule.h"
+#include "IContentBrowserSingleton.h"
+#include "CustomizableObjectPopulationEditorModule.h"
+#include "CustomizableObjectPopulationEditor.h"
+#include "CustomizableObject.h"// For log function
+#include "CustomizableObjectEditorModule.h"
+
+#include "CustomizableObjectPopulation/Classes/CustomizableObjectPopulationGenerator.h"
+
+#include "Framework/MultiBox/MultiBoxBuilder.h"
+
+#define LOCTEXT_NAMESPACE "AssetTypeActions_CustomizableObjectPopulation"
+
+
+FText FAssetTypeActions_CustomizableObjectPopulation::GetName() const
+{
+	return NSLOCTEXT("AssetTypeActions", "AssetTypeActions_CustomizableObjectPopulation", "Customizable Population");
+}
+
+
+FColor FAssetTypeActions_CustomizableObjectPopulation::GetTypeColor() const
+{ 
+	return FColor(0, 0, 0); 
+}
+
+
+UClass* FAssetTypeActions_CustomizableObjectPopulation::GetSupportedClass() const
+{ 
+	return UCustomizableObjectPopulation::StaticClass(); 
+}
+
+
+uint32 FAssetTypeActions_CustomizableObjectPopulation::GetCategories()
+{ 
+	const ICustomizableObjectEditorModule* CustomizableObjectEditorModule = &FModuleManager::LoadModuleChecked<ICustomizableObjectEditorModule>("CustomizableObjectEditor");
+	return CustomizableObjectEditorModule->GetAssetCategory();
+}
+
+
+void FAssetTypeActions_CustomizableObjectPopulation::GetActions(const TArray<UObject*>& InObjects, FMenuBuilder& MenuBuilder)
+{
+	TArray<TWeakObjectPtr<UCustomizableObjectPopulation>> Objects = GetTypedWeakObjectPtrs<UCustomizableObjectPopulation>(InObjects);
+
+	MenuBuilder.AddMenuEntry(
+		NSLOCTEXT("AssetTypeActions_CustomizableObjectPopulation", "CustomizableObjectPopulation_RecompilePopulations", "Recompile Populations"),
+		NSLOCTEXT("AssetTypeActions_CustomizableObjectPopulation", "CustomizableObjectPopulation_RecompilePopulationsToolTip", "Recompiles all the Mutable Populations Assets"),
+		FSlateIcon(),
+		FUIAction(
+			FExecuteAction::CreateSP(this, &FAssetTypeActions_CustomizableObjectPopulation::RecompilePopulations, Objects),
+			FCanExecuteAction()
+		)
+	);
+
+	//MenuBuilder.AddMenuEntry(
+	//	NSLOCTEXT("AssetTypeActions_CustomizableObjectPopulation", "CustomizableObjectPopulation_CreatePopulationGenerator", "Create Population Generator"),
+	//	NSLOCTEXT("AssetTypeActions_CustomizableObjectPopulation", "CustomizableObjectPopulation_CreatePopulationGeneratorToolTip", "Compiles the Population and Creates a Population Generator Asset"),
+	//	FSlateIcon(),
+	//	FUIAction(
+	//		FExecuteAction::CreateSP(this, &FAssetTypeActions_CustomizableObjectPopulation::CreatePopulationGeneratorAsset, Objects),
+	//		FCanExecuteAction()
+	//	)
+	//);
+
+	//MenuBuilder.AddMenuEntry(
+	//	NSLOCTEXT("AssetTypeActions_CustomizableObjectPopulation", "CustomizableObjectPopulation_RecompilePopulationGenerator", "Recompile Population Generators"),
+	//	NSLOCTEXT("AssetTypeActions_CustomizableObjectPopulation", "CustomizableObjectPopulation_RecompilePopulationGeneratorToolTip", "Recompiles all the Population Generator Assets"),
+	//	FSlateIcon(),
+	//	FUIAction(
+	//		FExecuteAction::CreateSP(this, &FAssetTypeActions_CustomizableObjectPopulation::RecompilePopulationGenerators, Objects),
+	//		FCanExecuteAction()
+	//	)
+	//);
+}
+
+//void FAssetTypeActions_CustomizableObjectPopulation::CreatePopulationGeneratorAsset(TArray<TWeakObjectPtr<UCustomizableObjectPopulation>> Objects)
+//{
+//	const FString DefaultSuffix = TEXT("_Generator");
+//
+//	TArray<UObject*> ObjectsToSync;
+//
+//	for (auto ObjIt = Objects.CreateConstIterator(); ObjIt; ++ObjIt)
+//	{
+//		auto Object = (*ObjIt).Get();
+//		if (Object)
+//		{
+//			if (!Object->IsValidPopulation())
+//			{
+//				UE_LOG(LogMutable, Warning, TEXT("There are one or more unassigned Population Classes. Please review your %s Population."), *(Object->Name));
+//				continue;
+//			}
+//
+//			// Determine an appropriate name
+//			FString Name;
+//			FString PackageName;
+//			CreateUniqueAssetName(Object->GetOutermost()->GetName(), DefaultSuffix, PackageName, Name);
+//			UPackage* Pkg = CreatePackage(*PackageName);
+//
+//			// Create avoid asset
+//			UCustomizableObjectPopulation* CustomizableObject = Cast<UCustomizableObjectPopulation>(Object);
+//			UCustomizableObjectPopulationGenerator* Generator = NewObject<UCustomizableObjectPopulationGenerator>(Pkg, FName(*Name), RF_Public | RF_Standalone);
+//
+//			// Fill the Generator information
+//			FCustomizableObjectPopulationGeneratorCompiler GeneratorCompiler;
+//			GeneratorCompiler.CompilePopulation(Object, true, Generator);
+//
+//			if (Generator)
+//			{
+//				// Mark the package dirty...
+//				Pkg->MarkPackageDirty();
+//
+//				ObjectsToSync.Add(Generator);
+//			}
+//		}
+//	}
+//
+//	if (ObjectsToSync.Num() > 0)
+//	{
+//		FContentBrowserModule& ContentBrowserModule = FModuleManager::Get().LoadModuleChecked<FContentBrowserModule>("ContentBrowser");
+//		ContentBrowserModule.Get().SyncBrowserToAssets(ObjectsToSync, /*bAllowLockedBrowsers=*/true);
+//	}
+//}
+//
+//void FAssetTypeActions_CustomizableObjectPopulation::RecompilePopulationGenerators(TArray<TWeakObjectPtr<UCustomizableObjectPopulation>> Populations)
+//{
+//	FCustomizableObjectPopulationGeneratorCompiler GeneratorCompiler;
+//	for (int32 i = 0; i < Populations.Num(); ++i)
+//	{
+//		if (Populations[i].IsValid() && Populations[i]->IsValidPopulation())
+//		{
+//			GeneratorCompiler.CompilePopulation(Populations[i].Get(), false, Populations[i]->Generator);
+//		}
+//	}
+//}
+
+void FAssetTypeActions_CustomizableObjectPopulation::OpenAssetEditor(const TArray<UObject*>& InObjects, TSharedPtr<IToolkitHost> EditWithinLevelEditor)
+{
+	EToolkitMode::Type Mode = EditWithinLevelEditor.IsValid() ? EToolkitMode::WorldCentric : EToolkitMode::Standalone;
+
+	for (auto ObjIt = InObjects.CreateConstIterator(); ObjIt; ++ObjIt)
+	{
+		UCustomizableObjectPopulation* Object = Cast<UCustomizableObjectPopulation>(*ObjIt);
+		if (Object != NULL)
+		{
+			ICustomizableObjectPopulationEditorModule* CustomizableObjectPopulationEditorModule = &FModuleManager::LoadModuleChecked<ICustomizableObjectPopulationEditorModule>("CustomizableObjectPopulationEditor");
+			CustomizableObjectPopulationEditorModule->CreateCustomizableObjectPopulationEditor(Mode, EditWithinLevelEditor, Object);
+		}
+	}
+}
+
+void FAssetTypeActions_CustomizableObjectPopulation::RecompilePopulations(TArray<TWeakObjectPtr<UCustomizableObjectPopulation>> Populations)
+{
+	for (int32 i = 0; i < Populations.Num(); ++i)
+	{
+		if (Populations[i].IsValid() && Populations[i]->IsValidPopulation())
+		{
+			Populations[i].Get()->CompilePopulation();
+		}
+	}
+}
+
+#undef LOCTEXT_NAMESPACE
diff --git a/Engine/Plugins/Experimental/Mutable/Source/CustomizableObjectPopulationEditor/Private/AssetTypeActions_CustomizableObjectPopulation.h b/Engine/Plugins/Experimental/Mutable/Source/CustomizableObjectPopulationEditor/Private/AssetTypeActions_CustomizableObjectPopulation.h
new file mode 100644
--- /dev/null
+++ b/Engine/Plugins/Experimental/Mutable/Source/CustomizableObjectPopulationEditor/Private/AssetTypeActions_CustomizableObjectPopulation.h
@@ -0,0 +1,27 @@
+// Copyright Epic Games, Inc. All Rights Reserved.
+
+#pragma once
+
+#include "AssetTypeActions_Base.h"
+#include "../Classes/CustomizableObjectPopulation.h"
+
+
+class FAssetTypeActions_CustomizableObjectPopulation : public FAssetTypeActions_Base
+{
+	// IAssetTypeActions Implementation
+	FText GetName() const override;
+	FColor GetTypeColor() const override;
+	UClass* GetSupportedClass() const override;
+	uint32 GetCategories();
+
+	bool HasActions(const TArray<UObject*>& InObjects) const override { return true; }
+	void GetActions(const TArray<UObject*>& InObjects, FMenuBuilder& MenuBuilder) override;
+
+	void OpenAssetEditor(const TArray<UObject*>& InObjects, TSharedPtr<IToolkitHost> EditWithinLevelEditor) override;
+
+	//Actions
+	void RecompilePopulations(TArray<TWeakObjectPtr<UCustomizableObjectPopulation>> Populations);
+	//void CreatePopulationGeneratorAsset(TArray<TWeakObjectPtr<UCustomizableObjectPopulation>> Objects);
+	//void RecompilePopulationGenerators(TArray<TWeakObjectPtr<UCustomizableObjectPopulation>> Objects);
+
+};
diff --git a/Engine/Plugins/Experimental/Mutable/Source/CustomizableObjectPopulationEditor/Private/AssetTypeActions_CustomizableObjectPopulationClass.cpp b/Engine/Plugins/Experimental/Mutable/Source/CustomizableObjectPopulationEditor/Private/AssetTypeActions_CustomizableObjectPopulationClass.cpp
new file mode 100644
--- /dev/null
+++ b/Engine/Plugins/Experimental/Mutable/Source/CustomizableObjectPopulationEditor/Private/AssetTypeActions_CustomizableObjectPopulationClass.cpp
@@ -0,0 +1,53 @@
+// Copyright Epic Games, Inc. All Rights Reserved.
+
+
+#include "AssetTypeActions_CustomizableObjectPopulationClass.h"
+
+#include "CustomizableObjectEditorModule.h"
+#include "CustomizableObjectPopulationClassEditor.h"
+#include "CustomizableObjectPopulationEditorModule.h"
+
+#define LOCTEXT_NAMESPACE "AssetTypeActions_CustomizableObjectPopulationClass"
+
+
+FText FAssetTypeActions_CustomizableObjectPopulationClass::GetName() const 
+{ 
+	return NSLOCTEXT("AssetTypeActions", "AssetTypeActions_CustomizableObjectPopulationClass", "Customizable Population Class"); 
+}
+
+
+FColor FAssetTypeActions_CustomizableObjectPopulationClass::GetTypeColor() const 
+{ 
+	return FColor(255, 255, 255);
+}
+
+
+UClass* FAssetTypeActions_CustomizableObjectPopulationClass::GetSupportedClass() const 
+{ 
+	return UCustomizableObjectPopulationClass::StaticClass();
+}
+
+
+uint32 FAssetTypeActions_CustomizableObjectPopulationClass::GetCategories()
+{ 
+	const ICustomizableObjectEditorModule* CustomizableObjectEditorModule = &FModuleManager::LoadModuleChecked<ICustomizableObjectEditorModule>("CustomizableObjectEditor");
+	return CustomizableObjectEditorModule->GetAssetCategory();}
+
+
+void FAssetTypeActions_CustomizableObjectPopulationClass::OpenAssetEditor(const TArray<UObject*>& InObjects, TSharedPtr<IToolkitHost> EditWithinLevelEditor)
+{
+	EToolkitMode::Type Mode = EditWithinLevelEditor.IsValid() ? EToolkitMode::WorldCentric : EToolkitMode::Standalone;
+
+	for (auto ObjIt = InObjects.CreateConstIterator(); ObjIt; ++ObjIt)
+	{
+		UCustomizableObjectPopulationClass* Object = Cast<UCustomizableObjectPopulationClass>(*ObjIt);
+		if (Object != NULL)
+		{
+			ICustomizableObjectPopulationEditorModule* CustomizableObjectPopulationEditorModule = &FModuleManager::LoadModuleChecked<ICustomizableObjectPopulationEditorModule>("CustomizableObjectPopulationEditor");
+			CustomizableObjectPopulationEditorModule->CreateCustomizableObjectPopulationClassEditor(Mode, EditWithinLevelEditor, Object);
+		}
+	}
+}
+
+
+#undef LOCTEXT_NAMESPACE
\ No newline at end of file
diff --git a/Engine/Plugins/Experimental/Mutable/Source/CustomizableObjectPopulationEditor/Private/AssetTypeActions_CustomizableObjectPopulationClass.h b/Engine/Plugins/Experimental/Mutable/Source/CustomizableObjectPopulationEditor/Private/AssetTypeActions_CustomizableObjectPopulationClass.h
new file mode 100644
--- /dev/null
+++ b/Engine/Plugins/Experimental/Mutable/Source/CustomizableObjectPopulationEditor/Private/AssetTypeActions_CustomizableObjectPopulationClass.h
@@ -0,0 +1,19 @@
+// Copyright Epic Games, Inc. All Rights Reserved.
+
+#pragma once
+
+#include "AssetTypeActions_Base.h"
+#include "../Classes/CustomizableObjectPopulationClass.h"
+
+
+class FAssetTypeActions_CustomizableObjectPopulationClass : public FAssetTypeActions_Base
+{
+	// IAssetTypeActions Implementation
+	FText GetName() const override;
+	FColor GetTypeColor() const override;
+	UClass* GetSupportedClass() const override;
+	uint32 GetCategories();
+	void OpenAssetEditor(const TArray<UObject*>& InObjects, TSharedPtr<class IToolkitHost> EditWithinLevelEditor = TSharedPtr<IToolkitHost>()) override;
+
+};
+
diff --git a/Engine/Plugins/Experimental/Mutable/Source/CustomizableObjectPopulationEditor/Private/CustomizableObjectPopulationEditorModule.cpp b/Engine/Plugins/Experimental/Mutable/Source/CustomizableObjectPopulationEditor/Private/CustomizableObjectPopulationEditorModule.cpp
new file mode 100644
--- /dev/null
+++ b/Engine/Plugins/Experimental/Mutable/Source/CustomizableObjectPopulationEditor/Private/CustomizableObjectPopulationEditorModule.cpp
@@ -0,0 +1,144 @@
+// Copyright Epic Games, Inc. All Rights Reserved.
+
+#include "CustomizableObjectPopulationEditorModule.h"
+#include "CustomizableObjectPopulationEditor.h"
+#include "CustomizableObject.h"	// For the LogMutable log category
+#include "CustomizableObjectSystem.h"		// For defines related to memory function replacements.
+#include "CustomizableObjectIdentifierCustomization.h"
+#include "Modules/ModuleManager.h"
+#include "Interfaces/IPluginManager.h"
+#include "AssetToolsModule.h"
+#include "Styling/SlateStyle.h"
+#include "Brushes/SlateImageBrush.h"
+#include "MessageLogModule.h"
+#include "Logging/MessageLog.h"
+#include "CustomizableObjectEditorSettings.h"
+#include "ISettingsModule.h"
+#include "ISettingsSection.h"
+#include "AssetRegistry/AssetRegistryModule.h"
+#include "HAL/FileManager.h"
+#include "CoreGlobals.h"
+#include "Settings/ProjectPackagingSettings.h"
+#include "UnrealEditorPortabilityHelpers.h"
+#include "AssetTypeActions_CustomizableObjectPopulation.h"
+#include "AssetTypeActions_CustomizableObjectPopulationClass.h"
+#include "CustomizableObjectPopulationClassEditor.h"
+#include "CustomizableObjectPopulationEditor.h"
+#include "CustomizableObjectPopulationEditorStyle.h"
+
+#include "CustomizableObjectPopulationClassDetails.h"
+
+// This is necessary to register the cook delegate.
+#include "CookOnTheSide/CookOnTheFlyServer.h"
+
+
+const FName CustomizableObjectPopulationEditorAppIdentifier = FName(TEXT("CustomizableObjectPopulationEditorApp"));
+const FName CustomizableObjectPopulationClassEditorAppIdentifier = FName(TEXT("CustomizableObjectPopulationClassEditorApp"));
+
+#define LOCTEXT_NAMESPACE "MutableSettings"
+
+/**
+ * StaticMesh editor module
+ */
+class FCustomizableObjectPopulationEditorModule : public ICustomizableObjectPopulationEditorModule
+{
+public:
+
+	// IModuleInterface interface
+	void StartupModule() override;
+	void ShutdownModule() override;
+
+	// ICustomizableObjectEditorModule interface
+	TSharedRef<ICustomizableObjectPopulationEditor> CreateCustomizableObjectPopulationEditor(const EToolkitMode::Type Mode, const TSharedPtr< IToolkitHost >& InitToolkitHost, UCustomizableObjectPopulation* CustomizablePopulation) override;
+	TSharedRef<ICustomizableObjectPopulationClassEditor> CreateCustomizableObjectPopulationClassEditor(const EToolkitMode::Type Mode, const TSharedPtr< IToolkitHost >& InitToolkitHost, UCustomizableObjectPopulationClass* CustomizablePopulationClass) override;
+
+	virtual TSharedPtr<FExtensibilityManager> GetCustomizableObjectPopulationEditorToolBarExtensibilityManager() override { return CustomizableObjectPopulationEditor_ToolBarExtensibilityManager; }
+
+private:
+
+	TSharedPtr<FExtensibilityManager> CustomizableObjectPopulationEditor_ToolBarExtensibilityManager;
+
+};
+
+IMPLEMENT_MODULE( FCustomizableObjectPopulationEditorModule, CustomizableObjectPopulationEditor );
+
+static void LogWarning(const char* msg)
+{
+	// Can be called from any thread.
+	UE_LOG(LogMutable, Warning, TEXT("%s"), ANSI_TO_TCHAR(msg) );
+}
+
+
+static void LogError(const char* msg)
+{
+	// Can be called from any thread.
+	UE_LOG(LogMutable, Error, TEXT("%s"), ANSI_TO_TCHAR(msg));
+}
+
+
+static void* CustomMalloc(std::size_t Size_t, uint32_t Alignment)
+{
+	return FMemory::Malloc(Size_t, Alignment);
+}
+
+
+static void CustomFree(void* mem)
+{
+	return FMemory::Free(mem);
+}
+
+
+void FCustomizableObjectPopulationEditorModule::StartupModule()
+{	
+	// Property views
+	FPropertyEditorModule& PropertyModule = FModuleManager::LoadModuleChecked<FPropertyEditorModule>("PropertyEditor");
+	
+	PropertyModule.RegisterCustomClassLayout("CustomizableObjectPopulationClass", FOnGetDetailCustomizationInstance::CreateStatic(&FCustomizableObjectPopulationClassDetails::MakeInstance));
+	
+	// Asset actions
+	FAssetToolsModule& AssetToolsModule = FModuleManager::LoadModuleChecked< FAssetToolsModule >( "AssetTools" );
+
+	TSharedPtr<FAssetTypeActions_CustomizableObjectPopulation> CustomizableObjectPopulationAssetTypeActions = MakeShareable(new FAssetTypeActions_CustomizableObjectPopulation);
+	AssetToolsModule.Get().RegisterAssetTypeActions(CustomizableObjectPopulationAssetTypeActions.ToSharedRef());
+	
+	TSharedPtr<FAssetTypeActions_CustomizableObjectPopulationClass> CustomizableObjectPopulationClassAssetTypeActions = MakeShareable(new FAssetTypeActions_CustomizableObjectPopulationClass);
+	AssetToolsModule.Get().RegisterAssetTypeActions(CustomizableObjectPopulationClassAssetTypeActions.ToSharedRef());
+
+	// Additional UI style
+	FCustomizableObjectPopulationEditorStyle::Initialize();
+
+	CustomizableObjectPopulationEditor_ToolBarExtensibilityManager = MakeShareable(new FExtensibilityManager);
+}
+
+
+void FCustomizableObjectPopulationEditorModule::ShutdownModule()
+{
+	if (FModuleManager::Get().IsModuleLoaded("PropertyEditor"))
+	{
+		FPropertyEditorModule& PropertyModule = FModuleManager::GetModuleChecked<FPropertyEditorModule>("PropertyEditor");
+		PropertyModule.UnregisterCustomClassLayout("CustomizableObjectPopulationClass");
+	}
+
+	CustomizableObjectPopulationEditor_ToolBarExtensibilityManager.Reset();
+
+	FCustomizableObjectPopulationEditorStyle::Shutdown();
+}
+
+
+TSharedRef<ICustomizableObjectPopulationEditor> FCustomizableObjectPopulationEditorModule::CreateCustomizableObjectPopulationEditor(const EToolkitMode::Type Mode, const TSharedPtr< IToolkitHost >& InitToolkitHost, UCustomizableObjectPopulation* CustomizablePopulation)
+{
+	TSharedRef<FCustomizableObjectPopulationEditor> NewCustomizableObjectPopulationEditor(new FCustomizableObjectPopulationEditor());
+	NewCustomizableObjectPopulationEditor->InitCustomizableObjectPopulationEditor(Mode, InitToolkitHost, CustomizablePopulation);
+	return NewCustomizableObjectPopulationEditor;
+}
+
+
+TSharedRef<ICustomizableObjectPopulationClassEditor> FCustomizableObjectPopulationEditorModule::CreateCustomizableObjectPopulationClassEditor(const EToolkitMode::Type Mode, const TSharedPtr< IToolkitHost >& InitToolkitHost, UCustomizableObjectPopulationClass* CustomizablePopulationClass)
+{
+	TSharedRef<FCustomizableObjectPopulationClassEditor> NewCustomizableObjectPopulationClassEditor(new FCustomizableObjectPopulationClassEditor());
+	NewCustomizableObjectPopulationClassEditor->InitCustomizableObjectPopulationClassEditor(Mode, InitToolkitHost, CustomizablePopulationClass);
+	return NewCustomizableObjectPopulationClassEditor;
+}
+
+
+#undef LOCTEXT_NAMESPACE
\ No newline at end of file
diff --git a/Engine/Plugins/Experimental/Mutable/Source/CustomizableObjectPopulationEditor/Public/CustomizableObjectPopulationClassEditorActions.h b/Engine/Plugins/Experimental/Mutable/Source/CustomizableObjectPopulationEditor/Public/CustomizableObjectPopulationClassEditorActions.h
new file mode 100644
--- /dev/null
+++ b/Engine/Plugins/Experimental/Mutable/Source/CustomizableObjectPopulationEditor/Public/CustomizableObjectPopulationClassEditorActions.h
@@ -0,0 +1,39 @@
+// Copyright Epic Games, Inc. All Rights Reserved.
+
+#pragma once
+
+#include "Styling/AppStyle.h"
+#include "Framework/Commands/UICommandInfo.h"
+#include "Framework/Commands/Commands.h"
+#include "CustomizableObjectPopulationEditorStyle.h"
+
+
+class FCustomizableObjectPopulationClassEditorCommands : public TCommands<FCustomizableObjectPopulationClassEditorCommands>
+{
+public:
+
+	FCustomizableObjectPopulationClassEditorCommands() : TCommands<FCustomizableObjectPopulationClassEditorCommands>
+		(
+			"CustomizableObjectPopulationClassEditor", // Context name for fast lookup
+			NSLOCTEXT("Contexts", "CustomizableObjectPopulationClassEditor", "CustomizableObject Population Class Editor"), // Localized context name for displaying
+			NAME_None, // Parent
+			FCustomizableObjectPopulationEditorStyle::GetStyleSetName()
+			)
+	{
+	}
+
+	TSharedPtr< FUICommandInfo > Save;
+	TSharedPtr< FUICommandInfo > SaveCustomizableObject;
+	TSharedPtr< FUICommandInfo > OpenCustomizableObjectEditor;
+	TSharedPtr< FUICommandInfo > SaveEditorCurve;
+	TSharedPtr< FUICommandInfo > TestPopulationClass;
+	TSharedPtr< FUICommandInfo > GenerateInstances;
+	TSharedPtr< FUICommandInfo > InspectSelectedInstance;
+	TSharedPtr< FUICommandInfo > InspectSelectedSkeletalMesh;
+
+	/**
+	 * Initialize commands
+	 */
+	virtual void RegisterCommands() override;
+
+};
diff --git a/Engine/Plugins/Experimental/Mutable/Source/CustomizableObjectPopulationEditor/Public/CustomizableObjectPopulationEditorActions.h b/Engine/Plugins/Experimental/Mutable/Source/CustomizableObjectPopulationEditor/Public/CustomizableObjectPopulationEditorActions.h
new file mode 100644
--- /dev/null
+++ b/Engine/Plugins/Experimental/Mutable/Source/CustomizableObjectPopulationEditor/Public/CustomizableObjectPopulationEditorActions.h
@@ -0,0 +1,35 @@
+// Copyright Epic Games, Inc. All Rights Reserved.
+
+#pragma once
+
+#include "Styling/AppStyle.h"
+#include "Framework/Commands/UICommandInfo.h"
+#include "Framework/Commands/Commands.h"
+#include "CustomizableObjectPopulationEditorStyle.h"
+
+class FCustomizableObjectPopulationEditorCommands : public TCommands<FCustomizableObjectPopulationEditorCommands>
+{
+
+public:
+	FCustomizableObjectPopulationEditorCommands() : TCommands<FCustomizableObjectPopulationEditorCommands>
+		(
+			"CustomizableObjectPopulationEditor", // Context name for fast lookup
+			NSLOCTEXT("Contexts", "CustomizableObjectPopulationEditor", "CustomizableObject Population Editor"), // Localized context name for displaying
+			NAME_None, // Parent
+			FCustomizableObjectPopulationEditorStyle::GetStyleSetName()
+			)
+	{
+	}
+
+	TSharedPtr< FUICommandInfo > Save;
+	TSharedPtr< FUICommandInfo > TestPopulation;
+	TSharedPtr< FUICommandInfo > GenerateInstances;
+	TSharedPtr< FUICommandInfo > InspectSelectedInstance;
+	TSharedPtr< FUICommandInfo > InspectSelectedSkeletalMesh;
+
+	/**
+	 * Initialize commands
+	 */
+	virtual void RegisterCommands() override;
+
+};
