[mutable] Move Mutable plugin from the restricted folder to the experimental plugins folder.

diff --git a/Engine/Plugins/Experimental/Mutable/Source/MutableRuntime/Private/BlockCompression/Miro/Miro.cpp b/Engine/Plugins/Experimental/Mutable/Source/MutableRuntime/Private/BlockCompression/Miro/Miro.cpp
new file mode 100644
--- /dev/null
+++ b/Engine/Plugins/Experimental/Mutable/Source/MutableRuntime/Private/BlockCompression/Miro/Miro.cpp
@@ -0,0 +1,10566 @@
+// Copyright Epic Games, Inc. All Rights Reserved.
+
+#include "Miro.h"
+
+#include "MutableRuntime/Private/Platform.h"
+#include "Misc/AssertionMacros.h"
+#include "Async/ParallelFor.h"
+
+#include <cstdlib>
+#include <cmath>
+#include <memory>
+
+// debug
+//static bool s_debuglog = false;
+
+#if defined(MIRO_SSE)
+#if MIRO_SSE>=2
+#include <pmmintrin.h>
+#endif
+#endif
+
+#ifdef _MSC_VER
+#define MIRO_ALIGN __declspec(align(16))
+#else
+#define MIRO_ALIGN __attribute__((aligned(16)))
+#endif
+
+
+union MIRO_ALIGN miro_pixel_block
+{
+	uint8 bytes[4 * 4 * 4];
+};
+
+
+namespace impl
+{
+
+#if MIRO_INCLUDE_BC
+	static uint8 s_initialised_dxtc_decompress = 0;
+#endif
+
+#if MIRO_INCLUDE_ASTC
+	static uint8 s_initialised_astc_decompress = 0;
+#endif
+
+
+#if MIRO_INCLUDE_BC
+
+	//-------------------------------------------------------------------------------------------------
+	// Decide if a BC1 compressed block has alpha or not.
+	//-------------------------------------------------------------------------------------------------
+	inline bool BC1BlockHasAlpha(uint8 const* bytes)
+	{
+		int a = ((int)bytes[0] << 0) + ((int)bytes[1] << 8);
+		int b = ((int)bytes[2] << 0) + ((int)bytes[3] << 8);
+		return a <= b;
+	}
+
+
+	//-------------------------------------------------------------------------------------------------
+	//! STB precomputed tables.
+	//-------------------------------------------------------------------------------------------------
+	static unsigned char stb__Expand5[32];
+	static unsigned char stb__Expand6[64];
+	static unsigned char stb__OMatch5[256][2];
+	static unsigned char stb__OMatch6[256][2];
+	static unsigned char stb__QuantRBTab[256 + 16];
+	static unsigned char stb__QuantGTab[256 + 16];
+
+
+	//-------------------------------------------------------------------------------------------------
+	//! Unpack a 16 bit colour and also return it as int
+	//-------------------------------------------------------------------------------------------------
+	inline int Unpack565(uint8 const* packed, uint8* colour)
+	{
+		// build the packed value
+		int value = ((int)packed[0] << 0) + ((int)packed[1] << 8);
+
+		// get the components in the stored range
+		uint8 red = (uint8)((value >> 11) & 0x1F);
+		uint8 green = (uint8)((value >> 5) & 0x3F);
+		uint8 blue = (uint8)(value & 0x1F);
+
+		// scale up to 8 bits
+		colour[0] = stb__Expand5[red];
+		colour[1] = stb__Expand6[green];
+		colour[2] = stb__Expand5[blue];
+		colour[3] = 255;
+
+		// return the value
+		return value;
+	}
+
+
+	inline void Codebook3(uint8(&codes)[4 * 4])
+	{
+		// generate the midpoints
+		for (int i = 0; i < 3; ++i)
+		{
+			const int c = codes[0 + i];
+			const int d = codes[4 + i];
+
+			codes[8 + i] = (uint8)((1 * c + 1 * d) >> 1);
+			codes[12 + i] = 0;
+		}
+
+		// fill in alpha for the intermediate values
+		codes[8 + 3] = 255;
+		codes[12 + 3] = 0;
+	}
+
+
+	inline void Codebook4(uint8(&codes)[4 * 4])
+	{
+		// generate the midpoints
+		for (int i = 0; i < 3; ++i)
+		{
+			const int c = codes[0 + i];
+			const int d = codes[4 + i];
+
+			codes[8 + i] = (uint8)((2 * c + 1 * d) / 3);
+			codes[12 + i] = (uint8)((1 * c + 2 * d) / 3);
+		}
+
+		// fill in alpha for the intermediate values
+		codes[8 + 3] = 255;
+		codes[12 + 3] = 255;
+	}
+
+
+	inline void Codebook6(uint8(&codes)[8 * 1])
+	{
+		// generate the midpoints
+		for (int i = 0; i < 1; ++i)
+		{
+			const int c = codes[0 + i];
+			const int d = codes[1 + i];
+
+			codes[2 + i] = (uint8)((4 * c + 1 * d) / 5);
+			codes[3 + i] = (uint8)((3 * c + 2 * d) / 5);
+			codes[4 + i] = (uint8)((2 * c + 3 * d) / 5);
+			codes[5 + i] = (uint8)((1 * c + 4 * d) / 5);
+
+			codes[6 + i] = (uint8)0;
+			codes[7 + i] = (uint8)255;
+		}
+	}
+
+
+	inline void Codebook8(uint8(&codes)[8 * 1])
+	{
+		// generate the midpoints
+		for (int i = 0; i < 1; ++i)
+		{
+			const int c = codes[0 + i];
+			const int d = codes[1 + i];
+
+			codes[2 + i] = (uint8)((6 * c + 1 * d) / 7);
+			codes[3 + i] = (uint8)((5 * c + 2 * d) / 7);
+			codes[4 + i] = (uint8)((4 * c + 3 * d) / 7);
+			codes[5 + i] = (uint8)((3 * c + 4 * d) / 7);
+			codes[6 + i] = (uint8)((2 * c + 5 * d) / 7);
+			codes[7 + i] = (uint8)((1 * c + 6 * d) / 7);
+		}
+	}
+
+
+	inline void Codebook6(uint16(&codes)[8 * 1])
+	{
+		// generate the midpoints
+		for (int i = 0; i < 1; ++i)
+		{
+			const int c = codes[0 + i];
+			const int d = codes[1 + i];
+
+			codes[2 + i] = (uint16)((4 * c + 1 * d) / 5);
+			codes[3 + i] = (uint16)((3 * c + 2 * d) / 5);
+			codes[4 + i] = (uint16)((2 * c + 3 * d) / 5);
+			codes[5 + i] = (uint16)((1 * c + 4 * d) / 5);
+
+			codes[6 + i] = (uint16)0;
+			codes[7 + i] = (uint16)255;
+		}
+	}
+
+
+	inline void Codebook8(uint16(&codes)[8 * 1])
+	{
+		// generate the midpoints
+		for (int i = 0; i < 1; ++i)
+		{
+			const int c = codes[0 + i];
+			const int d = codes[1 + i];
+
+			codes[2 + i] = (uint16)((6 * c + 1 * d) / 7);
+			codes[3 + i] = (uint16)((5 * c + 2 * d) / 7);
+			codes[4 + i] = (uint16)((4 * c + 3 * d) / 7);
+			codes[5 + i] = (uint16)((3 * c + 4 * d) / 7);
+			codes[6 + i] = (uint16)((2 * c + 5 * d) / 7);
+			codes[7 + i] = (uint16)((1 * c + 6 * d) / 7);
+		}
+	}
+
+
+	inline void CompressAlphaBtc2u(uint8 const* rgba, void* block)
+	{
+		uint8* bytes = reinterpret_cast<uint8*>(block);
+
+		// quantize and pack the alpha values pairwise
+		for (int i = 0; i < 8; ++i)
+		{
+			// quantize down to 4 bits
+			float alpha1 = (float)rgba[8 * i + 3] * (15.0f / 255.0f);
+			float alpha2 = (float)rgba[8 * i + 7] * (15.0f / 255.0f);
+
+			int quant1 = (int)(alpha1 + 0.5f);
+			int quant2 = (int)(alpha2 + 0.5f);
+
+			// pack into the byte
+			bytes[i] = (uint8)((quant1 << 0) + (quant2 << 4));
+		}
+	}
+
+
+
+	//-------------------------------------------------------------------------------------------------
+	// Based on BCDec, under the MIT license
+	// See: [REDACTED_PROFILE_URL]/bcdec
+	//-------------------------------------------------------------------------------------------------
+	namespace bcdec
+	{
+#define BCDEC_STATIC 1
+#define BCDEC_IMPLEMENTATION
+
+		/* if BCDEC_STATIC causes problems, try defining BCDECDEF to 'inline' or 'static inline' */
+#ifndef BCDECDEF
+#ifdef BCDEC_STATIC
+#define BCDECDEF    static
+#else
+#ifdef __cplusplus
+#define BCDECDEF    extern "C"
+#else
+#define BCDECDEF    extern
+#endif
+#endif
+#endif
+
+#define BCDEC_BC1_BLOCK_SIZE    8
+#define BCDEC_BC2_BLOCK_SIZE    16
+#define BCDEC_BC3_BLOCK_SIZE    16
+#define BCDEC_BC4_BLOCK_SIZE    8
+#define BCDEC_BC5_BLOCK_SIZE    16
+#define BCDEC_BC6H_BLOCK_SIZE   16
+#define BCDEC_BC7_BLOCK_SIZE    16
+
+#define BCDEC_BC1_COMPRESSED_SIZE(w, h)     ((((w)>>2)*((h)>>2))*BCDEC_BC1_BLOCK_SIZE)
+#define BCDEC_BC2_COMPRESSED_SIZE(w, h)     ((((w)>>2)*((h)>>2))*BCDEC_BC2_BLOCK_SIZE)
+#define BCDEC_BC3_COMPRESSED_SIZE(w, h)     ((((w)>>2)*((h)>>2))*BCDEC_BC3_BLOCK_SIZE)
+#define BCDEC_BC4_COMPRESSED_SIZE(w, h)     ((((w)>>2)*((h)>>2))*BCDEC_BC4_BLOCK_SIZE)
+#define BCDEC_BC5_COMPRESSED_SIZE(w, h)     ((((w)>>2)*((h)>>2))*BCDEC_BC5_BLOCK_SIZE)
+#define BCDEC_BC6H_COMPRESSED_SIZE(w, h)    ((((w)>>2)*((h)>>2))*BCDEC_BC6H_BLOCK_SIZE)
+#define BCDEC_BC7_COMPRESSED_SIZE(w, h)     ((((w)>>2)*((h)>>2))*BCDEC_BC7_BLOCK_SIZE)
+
+		BCDECDEF void bcdec_bc1(const void* compressedBlock, void* decompressedBlock, int destinationPitch);
+		BCDECDEF void bcdec_bc2(const void* compressedBlock, void* decompressedBlock, int destinationPitch);
+		BCDECDEF void bcdec_bc3(const void* compressedBlock, void* decompressedBlock, int destinationPitch);
+		BCDECDEF void bcdec_bc4(const void* compressedBlock, void* decompressedBlock, int destinationPitch);
+		BCDECDEF void bcdec_bc5(const void* compressedBlock, void* decompressedBlock, int destinationPitch);
+		BCDECDEF void bcdec_bc6h_float(const void* compressedBlock, void* decompressedBlock, int destinationPitch, int isSigned);
+		BCDECDEF void bcdec_bc6h_half(const void* compressedBlock, void* decompressedBlock, int destinationPitch, int isSigned);
+		BCDECDEF void bcdec_bc7(const void* compressedBlock, void* decompressedBlock, int destinationPitch);
+
+
+#ifdef BCDEC_IMPLEMENTATION
+
+		static void bcdec__color_block(const void* compressedBlock, void* decompressedBlock, int destinationPitch, int onlyOpaqueMode) {
+			unsigned short c0, c1;
+			unsigned int refColors[4]; /* 0xAABBGGRR */
+			unsigned char* dstColors;
+			unsigned int colorIndices;
+			int i, j, idx;
+			unsigned int r0, g0, b0, r1, g1, b1, r, g, b;
+
+			c0 = ((unsigned short*)compressedBlock)[0];
+			c1 = ((unsigned short*)compressedBlock)[1];
+
+			/* Expand 565 ref colors to 888 */
+			r0 = (((c0 >> 11) & 0x1F) * 527 + 23) >> 6;
+			g0 = (((c0 >> 5) & 0x3F) * 259 + 33) >> 6;
+			b0 = ((c0 & 0x1F) * 527 + 23) >> 6;
+			refColors[0] = 0xFF000000 | (b0 << 16) | (g0 << 8) | r0;
+
+			r1 = (((c1 >> 11) & 0x1F) * 527 + 23) >> 6;
+			g1 = (((c1 >> 5) & 0x3F) * 259 + 33) >> 6;
+			b1 = ((c1 & 0x1F) * 527 + 23) >> 6;
+			refColors[1] = 0xFF000000 | (b1 << 16) | (g1 << 8) | r1;
+
+			if (c0 > c1 || onlyOpaqueMode) {    /* Standard BC1 mode (also BC3 color block uses ONLY this mode) */
+				/* color_2 = 2/3*color_0 + 1/3*color_1
+				   color_3 = 1/3*color_0 + 2/3*color_1 */
+				r = (2 * r0 + r1 + 1) / 3;
+				g = (2 * g0 + g1 + 1) / 3;
+				b = (2 * b0 + b1 + 1) / 3;
+				refColors[2] = 0xFF000000 | (b << 16) | (g << 8) | r;
+
+				r = (r0 + 2 * r1 + 1) / 3;
+				g = (g0 + 2 * g1 + 1) / 3;
+				b = (b0 + 2 * b1 + 1) / 3;
+				refColors[3] = 0xFF000000 | (b << 16) | (g << 8) | r;
+			}
+			else {                            /* Quite rare BC1A mode */
+			 /* color_2 = 1/2*color_0 + 1/2*color_1;
+				color_3 = 0;                         */
+				r = (r0 + r1 + 1) >> 1;
+				g = (g0 + g1 + 1) >> 1;
+				b = (b0 + b1 + 1) >> 1;
+				refColors[2] = 0xFF000000 | (b << 16) | (g << 8) | r;
+
+				refColors[3] = 0x00000000;
+			}
+
+			colorIndices = ((unsigned int*)compressedBlock)[1];
+
+			/* Fill out the decompressed color block */
+			dstColors = (unsigned char*)decompressedBlock;
+			for (i = 0; i < 4; ++i) {
+				for (j = 0; j < 4; ++j) {
+					idx = colorIndices & 0x03;
+					((unsigned int*)dstColors)[j] = refColors[idx];
+					colorIndices >>= 2;
+				}
+
+				dstColors += destinationPitch;
+			}
+		}
+
+		static void bcdec__sharp_alpha_block(const void* compressedBlock, void* decompressedBlock, int destinationPitch) {
+			unsigned short* alpha;
+			unsigned char* decompressed;
+			int i, j;
+
+			alpha = (unsigned short*)compressedBlock;
+			decompressed = (unsigned char*)decompressedBlock;
+
+			for (i = 0; i < 4; ++i) {
+				for (j = 0; j < 4; ++j) {
+					decompressed[j * 4] = ((alpha[i] >> (4 * j)) & 0x0F) * 17;
+				}
+
+				decompressed += destinationPitch;
+			}
+		}
+
+		static void bcdec__smooth_alpha_block(const void* compressedBlock, void* decompressedBlock, int destinationPitch, int pixelSize) {
+			unsigned char* decompressed;
+			unsigned char alpha[8];
+			int i, j;
+			unsigned long long block, indices;
+
+			block = *(unsigned long long*)compressedBlock;
+			decompressed = (unsigned char*)decompressedBlock;
+
+			alpha[0] = block & 0xFF;
+			alpha[1] = (block >> 8) & 0xFF;
+
+			if (alpha[0] > alpha[1]) {
+				/* 6 interpolated alpha values. */
+				alpha[2] = (6 * alpha[0] + alpha[1] + 1) / 7;   /* 6/7*alpha_0 + 1/7*alpha_1 */
+				alpha[3] = (5 * alpha[0] + 2 * alpha[1] + 1) / 7;   /* 5/7*alpha_0 + 2/7*alpha_1 */
+				alpha[4] = (4 * alpha[0] + 3 * alpha[1] + 1) / 7;   /* 4/7*alpha_0 + 3/7*alpha_1 */
+				alpha[5] = (3 * alpha[0] + 4 * alpha[1] + 1) / 7;   /* 3/7*alpha_0 + 4/7*alpha_1 */
+				alpha[6] = (2 * alpha[0] + 5 * alpha[1] + 1) / 7;   /* 2/7*alpha_0 + 5/7*alpha_1 */
+				alpha[7] = (alpha[0] + 6 * alpha[1] + 1) / 7;   /* 1/7*alpha_0 + 6/7*alpha_1 */
+			}
+			else {
+				/* 4 interpolated alpha values. */
+				alpha[2] = (4 * alpha[0] + alpha[1] + 1) / 5;   /* 4/5*alpha_0 + 1/5*alpha_1 */
+				alpha[3] = (3 * alpha[0] + 2 * alpha[1] + 1) / 5;   /* 3/5*alpha_0 + 2/5*alpha_1 */
+				alpha[4] = (2 * alpha[0] + 3 * alpha[1] + 1) / 5;   /* 2/5*alpha_0 + 3/5*alpha_1 */
+				alpha[5] = (alpha[0] + 4 * alpha[1] + 1) / 5;   /* 1/5*alpha_0 + 4/5*alpha_1 */
+				alpha[6] = 0x00;
+				alpha[7] = 0xFF;
+			}
+
+			indices = block >> 16;
+			for (i = 0; i < 4; ++i) {
+				for (j = 0; j < 4; ++j) {
+					decompressed[j * pixelSize] = alpha[indices & 0x07];
+					indices >>= 3;
+				}
+
+				decompressed += destinationPitch;
+			}
+		}
+
+		typedef struct bcdec__bitstream {
+			unsigned long long low;
+			unsigned long long high;
+		} bcdec__bitstream_t;
+
+		static int bcdec__bitstream_read_bits(bcdec__bitstream_t* bstream, int numBits) {
+			unsigned int mask = (1 << numBits) - 1;
+			/* Read the low N bits */
+			unsigned int bits = (bstream->low & mask);
+
+			bstream->low >>= numBits;
+			/* Put the low N bits of "high" into the high 64-N bits of "low". */
+			bstream->low |= (bstream->high & mask) << (sizeof(bstream->high) * 8 - numBits);
+			bstream->high >>= numBits;
+
+			return bits;
+		}
+
+		static int bcdec__bitstream_read_bit(bcdec__bitstream_t* bstream) {
+			return bcdec__bitstream_read_bits(bstream, 1);
+		}
+
+		/*  reversed bits pulling, used in BC6H decoding
+			why ?? just why ??? */
+		static int bcdec__bitstream_read_bits_r(bcdec__bitstream_t* bstream, int numBits) {
+			int bits = bcdec__bitstream_read_bits(bstream, numBits);
+			/* Reverse the bits. */
+			int result = 0;
+			while (numBits--) {
+				result <<= 1;
+				result |= (bits & 1);
+				bits >>= 1;
+			}
+			return result;
+		}
+
+
+
+		BCDECDEF void bcdec_bc1(const void* compressedBlock, void* decompressedBlock, int destinationPitch) {
+			bcdec__color_block(compressedBlock, decompressedBlock, destinationPitch, 0);
+		}
+
+		BCDECDEF void bcdec_bc2(const void* compressedBlock, void* decompressedBlock, int destinationPitch) {
+			bcdec__color_block(((char*)compressedBlock) + 8, decompressedBlock, destinationPitch, 1);
+			bcdec__sharp_alpha_block(compressedBlock, ((char*)decompressedBlock) + 3, destinationPitch);
+		}
+
+		BCDECDEF void bcdec_bc3(const void* compressedBlock, void* decompressedBlock, int destinationPitch) {
+			bcdec__color_block(((char*)compressedBlock) + 8, decompressedBlock, destinationPitch, 1);
+			bcdec__smooth_alpha_block(compressedBlock, ((char*)decompressedBlock) + 3, destinationPitch, 4);
+		}
+
+		BCDECDEF void bcdec_bc4(const void* compressedBlock, void* decompressedBlock, int destinationPitch) {
+			bcdec__smooth_alpha_block(compressedBlock, decompressedBlock, destinationPitch, 1);
+		}
+
+		BCDECDEF void bcdec_bc5(const void* compressedBlock, void* decompressedBlock, int destinationPitch) {
+			bcdec__smooth_alpha_block(compressedBlock, decompressedBlock, destinationPitch, 2);
+			bcdec__smooth_alpha_block(((char*)compressedBlock) + 8, ((char*)decompressedBlock) + 1, destinationPitch, 2);
+		}
+
+		/* http://graphics.stanford.edu/~seander/bithacks.html#VariableSignExtend */
+		static int bcdec__extend_sign(int val, int bits) {
+			return (val << (32 - bits)) >> (32 - bits);
+		}
+
+		static int bcdec__transform_inverse(int val, int a0, int bits, int isSigned) {
+			/* If the precision of A0 is "p" bits, then the transform algorithm is:
+			   B0 = (B0 + A0) & ((1 << p) - 1) */
+			val = (val + a0) & ((1 << bits) - 1);
+			if (isSigned) {
+				val = bcdec__extend_sign(val, bits);
+			}
+			return val;
+		}
+
+		/* pretty much copy-paste from documentation */
+		static int bcdec__unquantize(int val, int bits, int isSigned) {
+			int unq, s = 0;
+
+			if (!isSigned) {
+				if (bits >= 15) {
+					unq = val;
+				}
+				else if (!val) {
+					unq = 0;
+				}
+				else if (val == ((1 << bits) - 1)) {
+					unq = 0xFFFF;
+				}
+				else {
+					unq = ((val << 16) + 0x8000) >> bits;
+				}
+			}
+			else {
+				if (bits >= 16) {
+					unq = val;
+				}
+				else {
+					if (val < 0) {
+						s = 1;
+						val = -val;
+					}
+
+					if (val == 0) {
+						unq = 0;
+					}
+					else if (val >= ((1 << (bits - 1)) - 1)) {
+						unq = 0x7FFF;
+					}
+					else {
+						unq = ((val << 15) + 0x4000) >> (bits - 1);
+					}
+
+					if (s) {
+						unq = -unq;
+					}
+				}
+			}
+			return unq;
+		}
+
+		static int bcdec__interpolate(int a, int b, int* weights, int index) {
+			return (a * (64 - weights[index]) + b * weights[index] + 32) >> 6;
+		}
+
+		static unsigned short bcdec__finish_unquantize(int val, int isSigned) {
+			int s;
+
+			if (!isSigned) {
+				return (unsigned short)((val * 31) >> 6);                   /* scale the magnitude by 31 / 64 */
+			}
+			else {
+				val = (val < 0) ? -(((-val) * 31) >> 5) : (val * 31) >> 5;  /* scale the magnitude by 31 / 32 */
+				s = 0;
+				if (val < 0) {
+					s = 0x8000;
+					val = -val;
+				}
+				return (unsigned short)(s | val);
+			}
+		}
+
+		/* modified half_to_float_fast4 from https://gist.github.com/rygorous/2144712 */
+		static float bcdec__half_to_float_quick(unsigned short half) {
+			typedef union {
+				unsigned int u;
+				float f;
+			} FP32;
+
+			static const FP32 magic = { 113 << 23 };
+			static const unsigned int shifted_exp = 0x7c00 << 13;   /* exponent mask after shift */
+			FP32 o;
+			unsigned int exp;
+
+			o.u = (half & 0x7fff) << 13;                            /* exponent/mantissa bits */
+			exp = shifted_exp & o.u;                                /* just the exponent */
+			o.u += (127 - 15) << 23;                                /* exponent adjust */
+
+			/* handle exponent special cases */
+			if (exp == shifted_exp) {                               /* Inf/NaN? */
+				o.u += (128 - 16) << 23;                            /* extra exp adjust */
+			}
+			else if (exp == 0) {                                  /* Zero/Denormal? */
+				o.u += 1 << 23;                                     /* extra exp adjust */
+				o.f -= magic.f;                                     /* renormalize */
+			}
+
+			o.u |= (half & 0x8000) << 16;                           /* sign bit */
+			return o.f;
+		}
+
+		BCDECDEF void bcdec_bc6h_half(const void* compressedBlock, void* decompressedBlock, int destinationPitch, int isSigned) {
+			static char actual_bits_count[4][14] = {
+				{ 10, 7, 11, 11, 11, 9, 8, 8, 8, 6, 10, 11, 12, 16 },   /*  W */
+				{  5, 6,  5,  4,  4, 5, 6, 5, 5, 6, 10,  9,  8,  4 },   /* dR */
+				{  5, 6,  4,  5,  4, 5, 5, 6, 5, 6, 10,  9,  8,  4 },   /* dG */
+				{  5, 6,  4,  4,  5, 5, 5, 5, 6, 6, 10,  9,  8,  4 }    /* dB */
+			};
+
+			/* There are 32 possible partition sets for a two-region tile.
+			   Each 4x4 block represents a single shape.
+			   Here also every fix-up index has MSB bit set. */
+			static unsigned char partition_sets[32][4][4] = {
+				{ {128, 0,   1, 1}, {0, 0, 1, 1}, {  0, 0, 1, 1}, {0, 0, 1, 129} },   /*  0 */
+				{ {128, 0,   0, 1}, {0, 0, 0, 1}, {  0, 0, 0, 1}, {0, 0, 0, 129} },   /*  1 */
+				{ {128, 1,   1, 1}, {0, 1, 1, 1}, {  0, 1, 1, 1}, {0, 1, 1, 129} },   /*  2 */
+				{ {128, 0,   0, 1}, {0, 0, 1, 1}, {  0, 0, 1, 1}, {0, 1, 1, 129} },   /*  3 */
+				{ {128, 0,   0, 0}, {0, 0, 0, 1}, {  0, 0, 0, 1}, {0, 0, 1, 129} },   /*  4 */
+				{ {128, 0,   1, 1}, {0, 1, 1, 1}, {  0, 1, 1, 1}, {1, 1, 1, 129} },   /*  5 */
+				{ {128, 0,   0, 1}, {0, 0, 1, 1}, {  0, 1, 1, 1}, {1, 1, 1, 129} },   /*  6 */
+				{ {128, 0,   0, 0}, {0, 0, 0, 1}, {  0, 0, 1, 1}, {0, 1, 1, 129} },   /*  7 */
+				{ {128, 0,   0, 0}, {0, 0, 0, 0}, {  0, 0, 0, 1}, {0, 0, 1, 129} },   /*  8 */
+				{ {128, 0,   1, 1}, {0, 1, 1, 1}, {  1, 1, 1, 1}, {1, 1, 1, 129} },   /*  9 */
+				{ {128, 0,   0, 0}, {0, 0, 0, 1}, {  0, 1, 1, 1}, {1, 1, 1, 129} },   /* 10 */
+				{ {128, 0,   0, 0}, {0, 0, 0, 0}, {  0, 0, 0, 1}, {0, 1, 1, 129} },   /* 11 */
+				{ {128, 0,   0, 1}, {0, 1, 1, 1}, {  1, 1, 1, 1}, {1, 1, 1, 129} },   /* 12 */
+				{ {128, 0,   0, 0}, {0, 0, 0, 0}, {  1, 1, 1, 1}, {1, 1, 1, 129} },   /* 13 */
+				{ {128, 0,   0, 0}, {1, 1, 1, 1}, {  1, 1, 1, 1}, {1, 1, 1, 129} },   /* 14 */
+				{ {128, 0,   0, 0}, {0, 0, 0, 0}, {  0, 0, 0, 0}, {1, 1, 1, 129} },   /* 15 */
+				{ {128, 0,   0, 0}, {1, 0, 0, 0}, {  1, 1, 1, 0}, {1, 1, 1, 129} },   /* 16 */
+				{ {128, 1, 129, 1}, {0, 0, 0, 1}, {  0, 0, 0, 0}, {0, 0, 0,   0} },   /* 17 */
+				{ {128, 0,   0, 0}, {0, 0, 0, 0}, {129, 0, 0, 0}, {1, 1, 1,   0} },   /* 18 */
+				{ {128, 1, 129, 1}, {0, 0, 1, 1}, {  0, 0, 0, 1}, {0, 0, 0,   0} },   /* 19 */
+				{ {128, 0, 129, 1}, {0, 0, 0, 1}, {  0, 0, 0, 0}, {0, 0, 0,   0} },   /* 20 */
+				{ {128, 0,   0, 0}, {1, 0, 0, 0}, {129, 1, 0, 0}, {1, 1, 1,   0} },   /* 21 */
+				{ {128, 0,   0, 0}, {0, 0, 0, 0}, {129, 0, 0, 0}, {1, 1, 0,   0} },   /* 22 */
+				{ {128, 1,   1, 1}, {0, 0, 1, 1}, {  0, 0, 1, 1}, {0, 0, 0, 129} },   /* 23 */
+				{ {128, 0, 129, 1}, {0, 0, 0, 1}, {  0, 0, 0, 1}, {0, 0, 0,   0} },   /* 24 */
+				{ {128, 0,   0, 0}, {1, 0, 0, 0}, {129, 0, 0, 0}, {1, 1, 0,   0} },   /* 25 */
+				{ {128, 1, 129, 0}, {0, 1, 1, 0}, {  0, 1, 1, 0}, {0, 1, 1,   0} },   /* 26 */
+				{ {128, 0, 129, 1}, {0, 1, 1, 0}, {  0, 1, 1, 0}, {1, 1, 0,   0} },   /* 27 */
+				{ {128, 0,   0, 1}, {0, 1, 1, 1}, {129, 1, 1, 0}, {1, 0, 0,   0} },   /* 28 */
+				{ {128, 0,   0, 0}, {1, 1, 1, 1}, {129, 1, 1, 1}, {0, 0, 0,   0} },   /* 29 */
+				{ {128, 1, 129, 1}, {0, 0, 0, 1}, {  1, 0, 0, 0}, {1, 1, 1,   0} },   /* 30 */
+				{ {128, 0, 129, 1}, {1, 0, 0, 1}, {  1, 0, 0, 1}, {1, 1, 0,   0} }    /* 31 */
+			};
+
+			static int aWeight3[8] = { 0, 9, 18, 27, 37, 46, 55, 64 };
+			static int aWeight4[16] = { 0, 4, 9, 13, 17, 21, 26, 30, 34, 38, 43, 47, 51, 55, 60, 64 };
+
+			bcdec__bitstream_t bstream;
+			int mode, partition, numPartitions, i, j, partitionSet, indexBits, index;
+			int r[4], g[4], b[4];       /* wxyz */
+			int epR[2], epG[2], epB[2]; /* endpoints A and B */
+			unsigned short* decompressed;
+
+			decompressed = (unsigned short*)decompressedBlock;
+
+			bstream.low = ((unsigned long long*)compressedBlock)[0];
+			bstream.high = ((unsigned long long*)compressedBlock)[1];
+
+			r[0] = r[1] = r[2] = r[3] = 0;
+			g[0] = g[1] = g[2] = g[3] = 0;
+			b[0] = b[1] = b[2] = b[3] = 0;
+
+			mode = bcdec__bitstream_read_bits(&bstream, 2);
+			if (mode > 1) {
+				mode |= (bcdec__bitstream_read_bits(&bstream, 3) << 2);
+			}
+
+			switch (mode) {
+				/* mode 1 */
+			case 0b00: {
+				/* Partitition indices: 46 bits
+				   Partition: 5 bits
+				   Color Endpoints: 75 bits (10.555, 10.555, 10.555) */
+				g[2] |= bcdec__bitstream_read_bit(&bstream) << 4;       /* gy[4]   */
+				b[2] |= bcdec__bitstream_read_bit(&bstream) << 4;       /* by[4]   */
+				b[3] |= bcdec__bitstream_read_bit(&bstream) << 4;       /* bz[4]   */
+				r[0] |= bcdec__bitstream_read_bits(&bstream, 10);       /* rw[9:0] */
+				g[0] |= bcdec__bitstream_read_bits(&bstream, 10);       /* gw[9:0] */
+				b[0] |= bcdec__bitstream_read_bits(&bstream, 10);       /* bw[9:0] */
+				r[1] |= bcdec__bitstream_read_bits(&bstream, 5);        /* rx[4:0] */
+				g[3] |= bcdec__bitstream_read_bit(&bstream) << 4;       /* gz[4]   */
+				g[2] |= bcdec__bitstream_read_bits(&bstream, 4);        /* gy[3:0] */
+				g[1] |= bcdec__bitstream_read_bits(&bstream, 5);        /* gx[4:0] */
+				b[3] |= bcdec__bitstream_read_bit(&bstream);            /* bz[0]   */
+				g[3] |= bcdec__bitstream_read_bits(&bstream, 4);        /* gz[3:0] */
+				b[1] |= bcdec__bitstream_read_bits(&bstream, 5);        /* bx[4:0] */
+				b[3] |= bcdec__bitstream_read_bit(&bstream) << 1;       /* bz[1]   */
+				b[2] |= bcdec__bitstream_read_bits(&bstream, 4);        /* by[3:0] */
+				r[2] |= bcdec__bitstream_read_bits(&bstream, 5);        /* ry[4:0] */
+				b[3] |= bcdec__bitstream_read_bit(&bstream) << 2;       /* bz[2]   */
+				r[3] |= bcdec__bitstream_read_bits(&bstream, 5);        /* rz[4:0] */
+				b[3] |= bcdec__bitstream_read_bit(&bstream) << 3;       /* bz[3]   */
+				partition = bcdec__bitstream_read_bits(&bstream, 5);    /* d[4:0]  */
+				mode = 0;
+			} break;
+
+				/* mode 2 */
+			case 0b01: {
+				/* Partitition indices: 46 bits
+				   Partition: 5 bits
+				   Color Endpoints: 75 bits (7666, 7666, 7666) */
+				g[2] |= bcdec__bitstream_read_bit(&bstream) << 5;       /* gy[5]   */
+				g[3] |= bcdec__bitstream_read_bit(&bstream) << 4;       /* gz[4]   */
+				g[3] |= bcdec__bitstream_read_bit(&bstream) << 5;       /* gz[5]   */
+				r[0] |= bcdec__bitstream_read_bits(&bstream, 7);        /* rw[6:0] */
+				b[3] |= bcdec__bitstream_read_bit(&bstream);            /* bz[0]   */
+				b[3] |= bcdec__bitstream_read_bit(&bstream) << 1;       /* bz[1]   */
+				b[2] |= bcdec__bitstream_read_bit(&bstream) << 4;       /* by[4]   */
+				g[0] |= bcdec__bitstream_read_bits(&bstream, 7);        /* gw[6:0] */
+				b[2] |= bcdec__bitstream_read_bit(&bstream) << 5;       /* by[5]   */
+				b[3] |= bcdec__bitstream_read_bit(&bstream) << 2;       /* bz[2]   */
+				g[2] |= bcdec__bitstream_read_bit(&bstream) << 4;       /* gy[4]   */
+				b[0] |= bcdec__bitstream_read_bits(&bstream, 7);        /* bw[6:0] */
+				b[3] |= bcdec__bitstream_read_bit(&bstream) << 3;       /* bz[3]   */
+				b[3] |= bcdec__bitstream_read_bit(&bstream) << 5;       /* bz[5]   */
+				b[3] |= bcdec__bitstream_read_bit(&bstream) << 4;       /* bz[4]   */
+				r[1] |= bcdec__bitstream_read_bits(&bstream, 6);        /* rx[5:0] */
+				g[2] |= bcdec__bitstream_read_bits(&bstream, 4);        /* gy[3:0] */
+				g[1] |= bcdec__bitstream_read_bits(&bstream, 6);        /* gx[5:0] */
+				g[3] |= bcdec__bitstream_read_bits(&bstream, 4);        /* gz[3:0] */
+				b[1] |= bcdec__bitstream_read_bits(&bstream, 6);        /* bx[5:0] */
+				b[2] |= bcdec__bitstream_read_bits(&bstream, 4);        /* by[3:0] */
+				r[2] |= bcdec__bitstream_read_bits(&bstream, 6);        /* ry[5:0] */
+				r[3] |= bcdec__bitstream_read_bits(&bstream, 6);        /* rz[5:0] */
+				partition = bcdec__bitstream_read_bits(&bstream, 5);    /* d[4:0]  */
+				mode = 1;
+			} break;
+
+				/* mode 3 */
+			case 0b00010: {
+				/* Partitition indices: 46 bits
+				   Partition: 5 bits
+				   Color Endpoints: 72 bits (11.555, 11.444, 11.444) */
+				r[0] |= bcdec__bitstream_read_bits(&bstream, 10);       /* rw[9:0] */
+				g[0] |= bcdec__bitstream_read_bits(&bstream, 10);       /* gw[9:0] */
+				b[0] |= bcdec__bitstream_read_bits(&bstream, 10);       /* bw[9:0] */
+				r[1] |= bcdec__bitstream_read_bits(&bstream, 5);        /* rx[4:0] */
+				r[0] |= bcdec__bitstream_read_bit(&bstream) << 10;      /* rw[10]  */
+				g[2] |= bcdec__bitstream_read_bits(&bstream, 4);        /* gy[3:0] */
+				g[1] |= bcdec__bitstream_read_bits(&bstream, 4);        /* gx[3:0] */
+				g[0] |= bcdec__bitstream_read_bit(&bstream) << 10;      /* gw[10]  */
+				b[3] |= bcdec__bitstream_read_bit(&bstream);            /* bz[0]   */
+				g[3] |= bcdec__bitstream_read_bits(&bstream, 4);        /* gz[3:0] */
+				b[1] |= bcdec__bitstream_read_bits(&bstream, 4);        /* bx[3:0] */
+				b[0] |= bcdec__bitstream_read_bit(&bstream) << 10;      /* bw[10]  */
+				b[3] |= bcdec__bitstream_read_bit(&bstream) << 1;       /* bz[1]   */
+				b[2] |= bcdec__bitstream_read_bits(&bstream, 4);        /* by[3:0] */
+				r[2] |= bcdec__bitstream_read_bits(&bstream, 5);        /* ry[4:0] */
+				b[3] |= bcdec__bitstream_read_bit(&bstream) << 2;       /* bz[2]   */
+				r[3] |= bcdec__bitstream_read_bits(&bstream, 5);        /* rz[4:0] */
+				b[3] |= bcdec__bitstream_read_bit(&bstream) << 3;       /* bz[3]   */
+				partition = bcdec__bitstream_read_bits(&bstream, 5);    /* d[4:0]  */
+				mode = 2;
+			} break;
+
+				/* mode 4 */
+			case 0b00110: {
+				/* Partitition indices: 46 bits
+				   Partition: 5 bits
+				   Color Endpoints: 72 bits (11.444, 11.555, 11.444) */
+				r[0] |= bcdec__bitstream_read_bits(&bstream, 10);       /* rw[9:0] */
+				g[0] |= bcdec__bitstream_read_bits(&bstream, 10);       /* gw[9:0] */
+				b[0] |= bcdec__bitstream_read_bits(&bstream, 10);       /* bw[9:0] */
+				r[1] |= bcdec__bitstream_read_bits(&bstream, 4);        /* rx[3:0] */
+				r[0] |= bcdec__bitstream_read_bit(&bstream) << 10;      /* rw[10]  */
+				g[3] |= bcdec__bitstream_read_bit(&bstream) << 4;       /* gz[4]   */
+				g[2] |= bcdec__bitstream_read_bits(&bstream, 4);        /* gy[3:0] */
+				g[1] |= bcdec__bitstream_read_bits(&bstream, 5);        /* gx[4:0] */
+				g[0] |= bcdec__bitstream_read_bit(&bstream) << 10;      /* gw[10]  */
+				g[3] |= bcdec__bitstream_read_bits(&bstream, 4);        /* gz[3:0] */
+				b[1] |= bcdec__bitstream_read_bits(&bstream, 4);        /* bx[3:0] */
+				b[0] |= bcdec__bitstream_read_bit(&bstream) << 10;      /* bw[10]  */
+				b[3] |= bcdec__bitstream_read_bit(&bstream) << 1;       /* bz[1]   */
+				b[2] |= bcdec__bitstream_read_bits(&bstream, 4);        /* by[3:0] */
+				r[2] |= bcdec__bitstream_read_bits(&bstream, 4);        /* ry[3:0] */
+				b[3] |= bcdec__bitstream_read_bit(&bstream);            /* bz[0]   */
+				b[3] |= bcdec__bitstream_read_bit(&bstream) << 2;       /* bz[2]   */
+				r[3] |= bcdec__bitstream_read_bits(&bstream, 4);        /* rz[3:0] */
+				g[2] |= bcdec__bitstream_read_bit(&bstream) << 4;       /* gy[4]   */
+				b[3] |= bcdec__bitstream_read_bit(&bstream) << 3;       /* bz[3]   */
+				partition = bcdec__bitstream_read_bits(&bstream, 5);    /* d[4:0]  */
+				mode = 3;
+			} break;
+
+				/* mode 5 */
+			case 0b01010: {
+				/* Partitition indices: 46 bits
+				   Partition: 5 bits
+				   Color Endpoints: 72 bits (11.444, 11.444, 11.555) */
+				r[0] |= bcdec__bitstream_read_bits(&bstream, 10);       /* rw[9:0] */
+				g[0] |= bcdec__bitstream_read_bits(&bstream, 10);       /* gw[9:0] */
+				b[0] |= bcdec__bitstream_read_bits(&bstream, 10);       /* bw[9:0] */
+				r[1] |= bcdec__bitstream_read_bits(&bstream, 4);        /* rx[3:0] */
+				r[0] |= bcdec__bitstream_read_bit(&bstream) << 10;      /* rw[10]  */
+				b[2] |= bcdec__bitstream_read_bit(&bstream) << 4;       /* by[4]   */
+				g[2] |= bcdec__bitstream_read_bits(&bstream, 4);        /* gy[3:0] */
+				g[1] |= bcdec__bitstream_read_bits(&bstream, 4);        /* gx[3:0] */
+				g[0] |= bcdec__bitstream_read_bit(&bstream) << 10;      /* gw[10]  */
+				b[3] |= bcdec__bitstream_read_bit(&bstream);            /* bz[0]   */
+				g[3] |= bcdec__bitstream_read_bits(&bstream, 4);        /* gz[3:0] */
+				b[1] |= bcdec__bitstream_read_bits(&bstream, 5);        /* bx[4:0] */
+				b[0] |= bcdec__bitstream_read_bit(&bstream) << 10;      /* bw[10]  */
+				b[2] |= bcdec__bitstream_read_bits(&bstream, 4);        /* by[3:0] */
+				r[2] |= bcdec__bitstream_read_bits(&bstream, 4);        /* ry[3:0] */
+				b[3] |= bcdec__bitstream_read_bit(&bstream) << 1;       /* bz[1]   */
+				b[3] |= bcdec__bitstream_read_bit(&bstream) << 2;       /* bz[2]   */
+				r[3] |= bcdec__bitstream_read_bits(&bstream, 4);        /* rz[3:0] */
+				b[3] |= bcdec__bitstream_read_bit(&bstream) << 4;       /* bz[4]   */
+				b[3] |= bcdec__bitstream_read_bit(&bstream) << 3;       /* bz[3]   */
+				partition = bcdec__bitstream_read_bits(&bstream, 5);    /* d[4:0]  */
+				mode = 4;
+			} break;
+
+				/* mode 6 */
+			case 0b01110: {
+				/* Partitition indices: 46 bits
+				   Partition: 5 bits
+				   Color Endpoints: 72 bits (9555, 9555, 9555) */
+				r[0] |= bcdec__bitstream_read_bits(&bstream, 9);        /* rw[8:0] */
+				b[2] |= bcdec__bitstream_read_bit(&bstream) << 4;       /* by[4]   */
+				g[0] |= bcdec__bitstream_read_bits(&bstream, 9);        /* gw[8:0] */
+				g[2] |= bcdec__bitstream_read_bit(&bstream) << 4;       /* gy[4]   */
+				b[0] |= bcdec__bitstream_read_bits(&bstream, 9);        /* bw[8:0] */
+				b[3] |= bcdec__bitstream_read_bit(&bstream) << 4;       /* bz[4]   */
+				r[1] |= bcdec__bitstream_read_bits(&bstream, 5);        /* rx[4:0] */
+				g[3] |= bcdec__bitstream_read_bit(&bstream) << 4;       /* gz[4]   */
+				g[2] |= bcdec__bitstream_read_bits(&bstream, 4);        /* gy[3:0] */
+				g[1] |= bcdec__bitstream_read_bits(&bstream, 5);        /* gx[4:0] */
+				b[3] |= bcdec__bitstream_read_bit(&bstream);            /* bz[0]   */
+				g[3] |= bcdec__bitstream_read_bits(&bstream, 4);        /* gx[3:0] */
+				b[1] |= bcdec__bitstream_read_bits(&bstream, 5);        /* bx[4:0] */
+				b[3] |= bcdec__bitstream_read_bit(&bstream) << 1;       /* bz[1]   */
+				b[2] |= bcdec__bitstream_read_bits(&bstream, 4);        /* by[3:0] */
+				r[2] |= bcdec__bitstream_read_bits(&bstream, 5);        /* ry[4:0] */
+				b[3] |= bcdec__bitstream_read_bit(&bstream) << 2;       /* bz[2]   */
+				r[3] |= bcdec__bitstream_read_bits(&bstream, 5);        /* rz[4:0] */
+				b[3] |= bcdec__bitstream_read_bit(&bstream) << 3;       /* bz[3]   */
+				partition = bcdec__bitstream_read_bits(&bstream, 5);    /* d[4:0]  */
+				mode = 5;
+			} break;
+
+				/* mode 7 */
+			case 0b10010: {
+				/* Partitition indices: 46 bits
+				   Partition: 5 bits
+				   Color Endpoints: 72 bits (8666, 8555, 8555) */
+				r[0] |= bcdec__bitstream_read_bits(&bstream, 8);        /* rw[7:0] */
+				g[3] |= bcdec__bitstream_read_bit(&bstream) << 4;       /* gz[4]   */
+				b[2] |= bcdec__bitstream_read_bit(&bstream) << 4;       /* by[4]   */
+				g[0] |= bcdec__bitstream_read_bits(&bstream, 8);        /* gw[7:0] */
+				b[3] |= bcdec__bitstream_read_bit(&bstream) << 2;       /* bz[2]   */
+				g[2] |= bcdec__bitstream_read_bit(&bstream) << 4;       /* gy[4]   */
+				b[0] |= bcdec__bitstream_read_bits(&bstream, 8);        /* bw[7:0] */
+				b[3] |= bcdec__bitstream_read_bit(&bstream) << 3;       /* bz[3]   */
+				b[3] |= bcdec__bitstream_read_bit(&bstream) << 4;       /* bz[4]   */
+				r[1] |= bcdec__bitstream_read_bits(&bstream, 6);        /* rx[5:0] */
+				g[2] |= bcdec__bitstream_read_bits(&bstream, 4);        /* gy[3:0] */
+				g[1] |= bcdec__bitstream_read_bits(&bstream, 5);        /* gx[4:0] */
+				b[3] |= bcdec__bitstream_read_bit(&bstream);            /* bz[0]   */
+				g[3] |= bcdec__bitstream_read_bits(&bstream, 4);        /* gz[3:0] */
+				b[1] |= bcdec__bitstream_read_bits(&bstream, 5);        /* bx[4:0] */
+				b[3] |= bcdec__bitstream_read_bit(&bstream) << 1;       /* bz[1]   */
+				b[2] |= bcdec__bitstream_read_bits(&bstream, 4);        /* by[3:0] */
+				r[2] |= bcdec__bitstream_read_bits(&bstream, 6);        /* ry[5:0] */
+				r[3] |= bcdec__bitstream_read_bits(&bstream, 6);        /* rz[5:0] */
+				partition = bcdec__bitstream_read_bits(&bstream, 5);    /* d[4:0]  */
+				mode = 6;
+			} break;
+
+				/* mode 8 */
+			case 0b10110: {
+				/* Partitition indices: 46 bits
+				   Partition: 5 bits
+				   Color Endpoints: 72 bits (8555, 8666, 8555) */
+				r[0] |= bcdec__bitstream_read_bits(&bstream, 8);        /* rw[7:0] */
+				b[3] |= bcdec__bitstream_read_bit(&bstream);            /* bz[0]   */
+				b[2] |= bcdec__bitstream_read_bit(&bstream) << 4;       /* by[4]   */
+				g[0] |= bcdec__bitstream_read_bits(&bstream, 8);        /* gw[7:0] */
+				g[2] |= bcdec__bitstream_read_bit(&bstream) << 5;       /* gy[5]   */
+				g[2] |= bcdec__bitstream_read_bit(&bstream) << 4;       /* gy[4]   */
+				b[0] |= bcdec__bitstream_read_bits(&bstream, 8);        /* bw[7:0] */
+				g[3] |= bcdec__bitstream_read_bit(&bstream) << 5;       /* gz[5]   */
+				b[3] |= bcdec__bitstream_read_bit(&bstream) << 4;       /* bz[4]   */
+				r[1] |= bcdec__bitstream_read_bits(&bstream, 5);        /* rx[4:0] */
+				g[3] |= bcdec__bitstream_read_bit(&bstream) << 4;       /* gz[4]   */
+				g[2] |= bcdec__bitstream_read_bits(&bstream, 4);        /* gy[3:0] */
+				g[1] |= bcdec__bitstream_read_bits(&bstream, 6);        /* gx[5:0] */
+				g[3] |= bcdec__bitstream_read_bits(&bstream, 4);        /* zx[3:0] */
+				b[1] |= bcdec__bitstream_read_bits(&bstream, 5);        /* bx[4:0] */
+				b[3] |= bcdec__bitstream_read_bit(&bstream) << 1;       /* bz[1]   */
+				b[2] |= bcdec__bitstream_read_bits(&bstream, 4);        /* by[3:0] */
+				r[2] |= bcdec__bitstream_read_bits(&bstream, 5);        /* ry[4:0] */
+				b[3] |= bcdec__bitstream_read_bit(&bstream) << 2;       /* bz[2]   */
+				r[3] |= bcdec__bitstream_read_bits(&bstream, 5);        /* rz[4:0] */
+				b[3] |= bcdec__bitstream_read_bit(&bstream) << 3;       /* bz[3]   */
+				partition = bcdec__bitstream_read_bits(&bstream, 5);    /* d[4:0]  */
+				mode = 7;
+			} break;
+
+				/* mode 9 */
+			case 0b11010: {
+				/* Partitition indices: 46 bits
+				   Partition: 5 bits
+				   Color Endpoints: 72 bits (8555, 8555, 8666) */
+				r[0] |= bcdec__bitstream_read_bits(&bstream, 8);        /* rw[7:0] */
+				b[3] |= bcdec__bitstream_read_bit(&bstream) << 1;       /* bz[1]   */
+				b[2] |= bcdec__bitstream_read_bit(&bstream) << 4;       /* by[4]   */
+				g[0] |= bcdec__bitstream_read_bits(&bstream, 8);        /* gw[7:0] */
+				b[2] |= bcdec__bitstream_read_bit(&bstream) << 5;       /* by[5]   */
+				g[2] |= bcdec__bitstream_read_bit(&bstream) << 4;       /* gy[4]   */
+				b[0] |= bcdec__bitstream_read_bits(&bstream, 8);        /* bw[7:0] */
+				b[3] |= bcdec__bitstream_read_bit(&bstream) << 5;       /* bz[5]   */
+				b[3] |= bcdec__bitstream_read_bit(&bstream) << 4;       /* bz[4]   */
+				r[1] |= bcdec__bitstream_read_bits(&bstream, 5);        /* bw[4:0] */
+				g[3] |= bcdec__bitstream_read_bit(&bstream) << 4;       /* gz[4]   */
+				g[2] |= bcdec__bitstream_read_bits(&bstream, 4);        /* gy[3:0] */
+				g[1] |= bcdec__bitstream_read_bits(&bstream, 5);        /* gx[4:0] */
+				b[3] |= bcdec__bitstream_read_bit(&bstream);            /* bz[0]   */
+				g[3] |= bcdec__bitstream_read_bits(&bstream, 4);        /* gz[3:0] */
+				b[1] |= bcdec__bitstream_read_bits(&bstream, 6);        /* bx[5:0] */
+				b[2] |= bcdec__bitstream_read_bits(&bstream, 4);        /* by[3:0] */
+				r[2] |= bcdec__bitstream_read_bits(&bstream, 5);        /* ry[4:0] */
+				b[3] |= bcdec__bitstream_read_bit(&bstream) << 2;       /* bz[2]   */
+				r[3] |= bcdec__bitstream_read_bits(&bstream, 5);        /* rz[4:0] */
+				b[3] |= bcdec__bitstream_read_bit(&bstream) << 3;       /* bz[3]   */
+				partition = bcdec__bitstream_read_bits(&bstream, 5);    /* d[4:0]  */
+				mode = 8;
+			} break;
+
+				/* mode 10 */
+			case 0b11110: {
+				/* Partitition indices: 46 bits
+				   Partition: 5 bits
+				   Color Endpoints: 72 bits (6666, 6666, 6666) */
+				r[0] |= bcdec__bitstream_read_bits(&bstream, 6);        /* rw[5:0] */
+				g[3] |= bcdec__bitstream_read_bit(&bstream) << 4;       /* gz[4]   */
+				b[3] |= bcdec__bitstream_read_bit(&bstream);            /* bz[0]   */
+				b[3] |= bcdec__bitstream_read_bit(&bstream) << 1;       /* bz[1]   */
+				b[2] |= bcdec__bitstream_read_bit(&bstream) << 4;       /* by[4]   */
+				g[0] |= bcdec__bitstream_read_bits(&bstream, 6);        /* gw[5:0] */
+				g[2] |= bcdec__bitstream_read_bit(&bstream) << 5;       /* gy[5]   */
+				b[2] |= bcdec__bitstream_read_bit(&bstream) << 5;       /* by[5]   */
+				b[3] |= bcdec__bitstream_read_bit(&bstream) << 2;       /* bz[2]   */
+				g[2] |= bcdec__bitstream_read_bit(&bstream) << 4;       /* gy[4]   */
+				b[0] |= bcdec__bitstream_read_bits(&bstream, 6);        /* bw[5:0] */
+				g[3] |= bcdec__bitstream_read_bit(&bstream) << 5;       /* gz[5]   */
+				b[3] |= bcdec__bitstream_read_bit(&bstream) << 3;       /* bz[3]   */
+				b[3] |= bcdec__bitstream_read_bit(&bstream) << 5;       /* bz[5]   */
+				b[3] |= bcdec__bitstream_read_bit(&bstream) << 4;       /* bz[4]   */
+				r[1] |= bcdec__bitstream_read_bits(&bstream, 6);        /* rx[5:0] */
+				g[2] |= bcdec__bitstream_read_bits(&bstream, 4);        /* gy[3:0] */
+				g[1] |= bcdec__bitstream_read_bits(&bstream, 6);        /* gx[5:0] */
+				g[3] |= bcdec__bitstream_read_bits(&bstream, 4);        /* gz[3:0] */
+				b[1] |= bcdec__bitstream_read_bits(&bstream, 6);        /* bx[5:0] */
+				b[2] |= bcdec__bitstream_read_bits(&bstream, 4);        /* by[3:0] */
+				r[2] |= bcdec__bitstream_read_bits(&bstream, 6);        /* ry[5:0] */
+				r[3] |= bcdec__bitstream_read_bits(&bstream, 6);        /* rz[5:0] */
+				partition = bcdec__bitstream_read_bits(&bstream, 5);    /* d[4:0]  */
+				mode = 9;
+			} break;
+
+				/* mode 11 */
+			case 0b00011: {
+				/* Partitition indices: 63 bits
+				   Partition: 0 bits
+				   Color Endpoints: 60 bits (10.10, 10.10, 10.10) */
+				r[0] |= bcdec__bitstream_read_bits(&bstream, 10);       /* rw[9:0] */
+				g[0] |= bcdec__bitstream_read_bits(&bstream, 10);       /* gw[9:0] */
+				b[0] |= bcdec__bitstream_read_bits(&bstream, 10);       /* bw[9:0] */
+				r[1] |= bcdec__bitstream_read_bits(&bstream, 10);       /* rx[9:0] */
+				g[1] |= bcdec__bitstream_read_bits(&bstream, 10);       /* gx[9:0] */
+				b[1] |= bcdec__bitstream_read_bits(&bstream, 10);       /* bx[9:0] */
+				mode = 10;
+			} break;
+
+				/* mode 12 */
+			case 0b00111: {
+				/* Partitition indices: 63 bits
+				   Partition: 0 bits
+				   Color Endpoints: 60 bits (11.9, 11.9, 11.9) */
+				r[0] |= bcdec__bitstream_read_bits(&bstream, 10);       /* rw[9:0] */
+				g[0] |= bcdec__bitstream_read_bits(&bstream, 10);       /* gw[9:0] */
+				b[0] |= bcdec__bitstream_read_bits(&bstream, 10);       /* bw[9:0] */
+				r[1] |= bcdec__bitstream_read_bits(&bstream, 9);        /* rx[8:0] */
+				r[0] |= bcdec__bitstream_read_bit(&bstream) << 10;      /* rw[10]  */
+				g[1] |= bcdec__bitstream_read_bits(&bstream, 9);        /* gx[8:0] */
+				g[0] |= bcdec__bitstream_read_bit(&bstream) << 10;      /* gw[10]  */
+				b[1] |= bcdec__bitstream_read_bits(&bstream, 9);        /* bx[8:0] */
+				b[0] |= bcdec__bitstream_read_bit(&bstream) << 10;      /* bw[10]  */
+				mode = 11;
+			} break;
+
+				/* mode 13 */
+			case 0b01011: {
+				/* Partitition indices: 63 bits
+				   Partition: 0 bits
+				   Color Endpoints: 60 bits (12.8, 12.8, 12.8) */
+				r[0] |= bcdec__bitstream_read_bits(&bstream, 10);       /* rw[9:0] */
+				g[0] |= bcdec__bitstream_read_bits(&bstream, 10);       /* gw[9:0] */
+				b[0] |= bcdec__bitstream_read_bits(&bstream, 10);       /* bw[9:0] */
+				r[1] |= bcdec__bitstream_read_bits(&bstream, 8);        /* rx[7:0] */
+				r[0] |= bcdec__bitstream_read_bits_r(&bstream, 2) << 10;/* rx[10:11] */
+				g[1] |= bcdec__bitstream_read_bits(&bstream, 8);        /* gx[7:0] */
+				g[0] |= bcdec__bitstream_read_bits_r(&bstream, 2) << 10;/* gx[10:11] */
+				b[1] |= bcdec__bitstream_read_bits(&bstream, 8);        /* bx[7:0] */
+				b[0] |= bcdec__bitstream_read_bits_r(&bstream, 2) << 10;/* bx[10:11] */
+				mode = 12;
+			} break;
+
+				/* mode 14 */
+			case 0b01111: {
+				/* Partitition indices: 63 bits
+				   Partition: 0 bits
+				   Color Endpoints: 60 bits (16.4, 16.4, 16.4) */
+				r[0] |= bcdec__bitstream_read_bits(&bstream, 10);       /* rw[9:0] */
+				g[0] |= bcdec__bitstream_read_bits(&bstream, 10);       /* gw[9:0] */
+				b[0] |= bcdec__bitstream_read_bits(&bstream, 10);       /* bw[9:0] */
+				r[1] |= bcdec__bitstream_read_bits(&bstream, 4);        /* rx[3:0] */
+				r[0] |= bcdec__bitstream_read_bits_r(&bstream, 6) << 10;/* rw[10:15] */
+				g[1] |= bcdec__bitstream_read_bits(&bstream, 4);        /* gx[3:0] */
+				g[0] |= bcdec__bitstream_read_bits_r(&bstream, 6) << 10;/* gw[10:15] */
+				b[1] |= bcdec__bitstream_read_bits(&bstream, 4);        /* bx[3:0] */
+				b[0] |= bcdec__bitstream_read_bits_r(&bstream, 6) << 10;/* bw[10:15] */
+				mode = 13;
+			} break;
+
+			default: {
+				/* Modes 10011, 10111, 11011, and 11111 (not shown) are reserved.
+				   Do not use these in your encoder. If the hardware is passed blocks
+				   with one of these modes specified, the resulting decompressed block
+				   must contain all zeroes in all channels except for the alpha channel. */
+				for (i = 0; i < 4; ++i) {
+					for (j = 0; j < 4; ++j) {
+						decompressed[j * 3 + 0] = 0;
+						decompressed[j * 3 + 1] = 0;
+						decompressed[j * 3 + 2] = 0;
+					}
+					decompressed += destinationPitch;
+				}
+
+				return;
+			}
+			}
+
+			if (mode >= 10) {
+				partition = 0;
+				numPartitions = 0;
+			}
+			else {
+				numPartitions = 1;
+			}
+
+			if (isSigned) {
+				r[0] = bcdec__extend_sign(r[0], actual_bits_count[0][mode]);
+				g[0] = bcdec__extend_sign(g[0], actual_bits_count[0][mode]);
+				b[0] = bcdec__extend_sign(b[0], actual_bits_count[0][mode]);
+			}
+
+			/* Mode 11 (like Mode 10) does not use delta compression,
+			   and instead stores both color endpoints explicitly.  */
+			if ((mode != 9 && mode != 10) || isSigned) {
+				for (i = 1; i < (numPartitions + 1) * 2; ++i) {
+					r[i] = bcdec__extend_sign(r[i], actual_bits_count[1][mode]);
+					g[i] = bcdec__extend_sign(g[i], actual_bits_count[2][mode]);
+					b[i] = bcdec__extend_sign(b[i], actual_bits_count[3][mode]);
+				}
+			}
+
+			if (mode != 9 && mode != 10) {
+				for (i = 1; i < (numPartitions + 1) * 2; ++i) {
+					r[i] = bcdec__transform_inverse(r[i], r[0], actual_bits_count[0][mode], isSigned);
+					g[i] = bcdec__transform_inverse(g[i], g[0], actual_bits_count[0][mode], isSigned);
+					b[i] = bcdec__transform_inverse(b[i], b[0], actual_bits_count[0][mode], isSigned);
+				}
+			}
+
+			for (i = 0; i < 4; ++i) {
+				for (j = 0; j < 4; ++j) {
+					partitionSet = (mode >= 10) ? ((i | j) ? 0 : 128) : partition_sets[partition][i][j];
+
+					indexBits = (mode >= 10) ? 4 : 3;
+					/* fix-up index is specified with one less bit */
+					/* The fix-up index for subset 0 is always index 0 */
+					if (partitionSet & 0x80) {
+						indexBits--;
+					}
+					partitionSet &= 0x01;
+
+					index = bcdec__bitstream_read_bits(&bstream, indexBits);
+
+					epR[0] = bcdec__unquantize(r[partitionSet * 2 + 0], actual_bits_count[0][mode], isSigned);
+					epG[0] = bcdec__unquantize(g[partitionSet * 2 + 0], actual_bits_count[0][mode], isSigned);
+					epB[0] = bcdec__unquantize(b[partitionSet * 2 + 0], actual_bits_count[0][mode], isSigned);
+					epR[1] = bcdec__unquantize(r[partitionSet * 2 + 1], actual_bits_count[0][mode], isSigned);
+					epG[1] = bcdec__unquantize(g[partitionSet * 2 + 1], actual_bits_count[0][mode], isSigned);
+					epB[1] = bcdec__unquantize(b[partitionSet * 2 + 1], actual_bits_count[0][mode], isSigned);
+
+					decompressed[j * 3 + 0] = bcdec__finish_unquantize(
+						bcdec__interpolate(epR[0], epR[1], (mode >= 10) ? aWeight4 : aWeight3, index), isSigned);
+					decompressed[j * 3 + 1] = bcdec__finish_unquantize(
+						bcdec__interpolate(epG[0], epG[1], (mode >= 10) ? aWeight4 : aWeight3, index), isSigned);
+					decompressed[j * 3 + 2] = bcdec__finish_unquantize(
+						bcdec__interpolate(epB[0], epB[1], (mode >= 10) ? aWeight4 : aWeight3, index), isSigned);
+				}
+
+				decompressed += destinationPitch;
+			}
+		}
+
+		BCDECDEF void bcdec_bc6h_float(const void* compressedBlock, void* decompressedBlock, int destinationPitch, int isSigned) {
+			unsigned short block[16 * 3];
+			float* decompressed;
+			const unsigned short* b;
+			int i, j;
+
+			bcdec_bc6h_half(compressedBlock, block, 4 * 3, isSigned);
+			b = block;
+			decompressed = (float*)decompressedBlock;
+			for (i = 0; i < 4; ++i) {
+				for (j = 0; j < 4; ++j) {
+					decompressed[j * 3 + 0] = bcdec__half_to_float_quick(*b++);
+					decompressed[j * 3 + 1] = bcdec__half_to_float_quick(*b++);
+					decompressed[j * 3 + 2] = bcdec__half_to_float_quick(*b++);
+				}
+				decompressed += destinationPitch;
+			}
+		}
+
+		static void bcdec__swap_values(int* a, int* b) {
+			a[0] ^= b[0], b[0] ^= a[0], a[0] ^= b[0];
+		}
+
+		BCDECDEF void bcdec_bc7(const void* compressedBlock, void* decompressedBlock, int destinationPitch) {
+			static char actual_bits_count[2][8] = {
+				{ 4, 6, 5, 7, 5, 7, 7, 5 },     /* RGBA  */
+				{ 0, 0, 0, 0, 6, 8, 7, 5 },     /* Alpha */
+			};
+
+			/* There are 64 possible partition sets for a two-region tile.
+			   Each 4x4 block represents a single shape.
+			   Here also every fix-up index has MSB bit set. */
+			static unsigned char partition_sets[2][64][4][4] = {
+				{   /* Partition table for 2-subset BPTC */
+					{ {128, 0,   1, 1}, {0, 0,   1, 1}, {  0, 0, 1, 1}, {0, 0, 1, 129} }, /*  0 */
+					{ {128, 0,   0, 1}, {0, 0,   0, 1}, {  0, 0, 0, 1}, {0, 0, 0, 129} }, /*  1 */
+					{ {128, 1,   1, 1}, {0, 1,   1, 1}, {  0, 1, 1, 1}, {0, 1, 1, 129} }, /*  2 */
+					{ {128, 0,   0, 1}, {0, 0,   1, 1}, {  0, 0, 1, 1}, {0, 1, 1, 129} }, /*  3 */
+					{ {128, 0,   0, 0}, {0, 0,   0, 1}, {  0, 0, 0, 1}, {0, 0, 1, 129} }, /*  4 */
+					{ {128, 0,   1, 1}, {0, 1,   1, 1}, {  0, 1, 1, 1}, {1, 1, 1, 129} }, /*  5 */
+					{ {128, 0,   0, 1}, {0, 0,   1, 1}, {  0, 1, 1, 1}, {1, 1, 1, 129} }, /*  6 */
+					{ {128, 0,   0, 0}, {0, 0,   0, 1}, {  0, 0, 1, 1}, {0, 1, 1, 129} }, /*  7 */
+					{ {128, 0,   0, 0}, {0, 0,   0, 0}, {  0, 0, 0, 1}, {0, 0, 1, 129} }, /*  8 */
+					{ {128, 0,   1, 1}, {0, 1,   1, 1}, {  1, 1, 1, 1}, {1, 1, 1, 129} }, /*  9 */
+					{ {128, 0,   0, 0}, {0, 0,   0, 1}, {  0, 1, 1, 1}, {1, 1, 1, 129} }, /* 10 */
+					{ {128, 0,   0, 0}, {0, 0,   0, 0}, {  0, 0, 0, 1}, {0, 1, 1, 129} }, /* 11 */
+					{ {128, 0,   0, 1}, {0, 1,   1, 1}, {  1, 1, 1, 1}, {1, 1, 1, 129} }, /* 12 */
+					{ {128, 0,   0, 0}, {0, 0,   0, 0}, {  1, 1, 1, 1}, {1, 1, 1, 129} }, /* 13 */
+					{ {128, 0,   0, 0}, {1, 1,   1, 1}, {  1, 1, 1, 1}, {1, 1, 1, 129} }, /* 14 */
+					{ {128, 0,   0, 0}, {0, 0,   0, 0}, {  0, 0, 0, 0}, {1, 1, 1, 129} }, /* 15 */
+					{ {128, 0,   0, 0}, {1, 0,   0, 0}, {  1, 1, 1, 0}, {1, 1, 1, 129} }, /* 16 */
+					{ {128, 1, 129, 1}, {0, 0,   0, 1}, {  0, 0, 0, 0}, {0, 0, 0,   0} }, /* 17 */
+					{ {128, 0,   0, 0}, {0, 0,   0, 0}, {129, 0, 0, 0}, {1, 1, 1,   0} }, /* 18 */
+					{ {128, 1, 129, 1}, {0, 0,   1, 1}, {  0, 0, 0, 1}, {0, 0, 0,   0} }, /* 19 */
+					{ {128, 0, 129, 1}, {0, 0,   0, 1}, {  0, 0, 0, 0}, {0, 0, 0,   0} }, /* 20 */
+					{ {128, 0,   0, 0}, {1, 0,   0, 0}, {129, 1, 0, 0}, {1, 1, 1,   0} }, /* 21 */
+					{ {128, 0,   0, 0}, {0, 0,   0, 0}, {129, 0, 0, 0}, {1, 1, 0,   0} }, /* 22 */
+					{ {128, 1,   1, 1}, {0, 0,   1, 1}, {  0, 0, 1, 1}, {0, 0, 0, 129} }, /* 23 */
+					{ {128, 0, 129, 1}, {0, 0,   0, 1}, {  0, 0, 0, 1}, {0, 0, 0,   0} }, /* 24 */
+					{ {128, 0,   0, 0}, {1, 0,   0, 0}, {129, 0, 0, 0}, {1, 1, 0,   0} }, /* 25 */
+					{ {128, 1, 129, 0}, {0, 1,   1, 0}, {  0, 1, 1, 0}, {0, 1, 1,   0} }, /* 26 */
+					{ {128, 0, 129, 1}, {0, 1,   1, 0}, {  0, 1, 1, 0}, {1, 1, 0,   0} }, /* 27 */
+					{ {128, 0,   0, 1}, {0, 1,   1, 1}, {129, 1, 1, 0}, {1, 0, 0,   0} }, /* 28 */
+					{ {128, 0,   0, 0}, {1, 1,   1, 1}, {129, 1, 1, 1}, {0, 0, 0,   0} }, /* 29 */
+					{ {128, 1, 129, 1}, {0, 0,   0, 1}, {  1, 0, 0, 0}, {1, 1, 1,   0} }, /* 30 */
+					{ {128, 0, 129, 1}, {1, 0,   0, 1}, {  1, 0, 0, 1}, {1, 1, 0,   0} }, /* 31 */
+					{ {128, 1,   0, 1}, {0, 1,   0, 1}, {  0, 1, 0, 1}, {0, 1, 0, 129} }, /* 32 */
+					{ {128, 0,   0, 0}, {1, 1,   1, 1}, {  0, 0, 0, 0}, {1, 1, 1, 129} }, /* 33 */
+					{ {128, 1,   0, 1}, {1, 0, 129, 0}, {  0, 1, 0, 1}, {1, 0, 1,   0} }, /* 34 */
+					{ {128, 0,   1, 1}, {0, 0,   1, 1}, {129, 1, 0, 0}, {1, 1, 0,   0} }, /* 35 */
+					{ {128, 0, 129, 1}, {1, 1,   0, 0}, {  0, 0, 1, 1}, {1, 1, 0,   0} }, /* 36 */
+					{ {128, 1,   0, 1}, {0, 1,   0, 1}, {129, 0, 1, 0}, {1, 0, 1,   0} }, /* 37 */
+					{ {128, 1,   1, 0}, {1, 0,   0, 1}, {  0, 1, 1, 0}, {1, 0, 0, 129} }, /* 38 */
+					{ {128, 1,   0, 1}, {1, 0,   1, 0}, {  1, 0, 1, 0}, {0, 1, 0, 129} }, /* 39 */
+					{ {128, 1, 129, 1}, {0, 0,   1, 1}, {  1, 1, 0, 0}, {1, 1, 1,   0} }, /* 40 */
+					{ {128, 0,   0, 1}, {0, 0,   1, 1}, {129, 1, 0, 0}, {1, 0, 0,   0} }, /* 41 */
+					{ {128, 0, 129, 1}, {0, 0,   1, 0}, {  0, 1, 0, 0}, {1, 1, 0,   0} }, /* 42 */
+					{ {128, 0, 129, 1}, {1, 0,   1, 1}, {  1, 1, 0, 1}, {1, 1, 0,   0} }, /* 43 */
+					{ {128, 1, 129, 0}, {1, 0,   0, 1}, {  1, 0, 0, 1}, {0, 1, 1,   0} }, /* 44 */
+					{ {128, 0,   1, 1}, {1, 1,   0, 0}, {  1, 1, 0, 0}, {0, 0, 1, 129} }, /* 45 */
+					{ {128, 1,   1, 0}, {0, 1,   1, 0}, {  1, 0, 0, 1}, {1, 0, 0, 129} }, /* 46 */
+					{ {128, 0,   0, 0}, {0, 1, 129, 0}, {  0, 1, 1, 0}, {0, 0, 0,   0} }, /* 47 */
+					{ {128, 1,   0, 0}, {1, 1, 129, 0}, {  0, 1, 0, 0}, {0, 0, 0,   0} }, /* 48 */
+					{ {128, 0, 129, 0}, {0, 1,   1, 1}, {  0, 0, 1, 0}, {0, 0, 0,   0} }, /* 49 */
+					{ {128, 0,   0, 0}, {0, 0, 129, 0}, {  0, 1, 1, 1}, {0, 0, 1,   0} }, /* 50 */
+					{ {128, 0,   0, 0}, {0, 1,   0, 0}, {129, 1, 1, 0}, {0, 1, 0,   0} }, /* 51 */
+					{ {128, 1,   1, 0}, {1, 1,   0, 0}, {  1, 0, 0, 1}, {0, 0, 1, 129} }, /* 52 */
+					{ {128, 0,   1, 1}, {0, 1,   1, 0}, {  1, 1, 0, 0}, {1, 0, 0, 129} }, /* 53 */
+					{ {128, 1, 129, 0}, {0, 0,   1, 1}, {  1, 0, 0, 1}, {1, 1, 0,   0} }, /* 54 */
+					{ {128, 0, 129, 1}, {1, 0,   0, 1}, {  1, 1, 0, 0}, {0, 1, 1,   0} }, /* 55 */
+					{ {128, 1,   1, 0}, {1, 1,   0, 0}, {  1, 1, 0, 0}, {1, 0, 0, 129} }, /* 56 */
+					{ {128, 1,   1, 0}, {0, 0,   1, 1}, {  0, 0, 1, 1}, {1, 0, 0, 129} }, /* 57 */
+					{ {128, 1,   1, 1}, {1, 1,   1, 0}, {  1, 0, 0, 0}, {0, 0, 0, 129} }, /* 58 */
+					{ {128, 0,   0, 1}, {1, 0,   0, 0}, {  1, 1, 1, 0}, {0, 1, 1, 129} }, /* 59 */
+					{ {128, 0,   0, 0}, {1, 1,   1, 1}, {  0, 0, 1, 1}, {0, 0, 1, 129} }, /* 60 */
+					{ {128, 0, 129, 1}, {0, 0,   1, 1}, {  1, 1, 1, 1}, {0, 0, 0,   0} }, /* 61 */
+					{ {128, 0, 129, 0}, {0, 0,   1, 0}, {  1, 1, 1, 0}, {1, 1, 1,   0} }, /* 62 */
+					{ {128, 1,   0, 0}, {0, 1,   0, 0}, {  0, 1, 1, 1}, {0, 1, 1, 129} }  /* 63 */
+				},
+				{   /* Partition table for 3-subset BPTC */
+					{ {128, 0, 1, 129}, {0,   0,   1, 1}, {  0,   2,   2, 1}, {  2,   2, 2, 130} }, /*  0 */
+					{ {128, 0, 0, 129}, {0,   0,   1, 1}, {130,   2,   1, 1}, {  2,   2, 2,   1} }, /*  1 */
+					{ {128, 0, 0,   0}, {2,   0,   0, 1}, {130,   2,   1, 1}, {  2,   2, 1, 129} }, /*  2 */
+					{ {128, 2, 2, 130}, {0,   0,   2, 2}, {  0,   0,   1, 1}, {  0,   1, 1, 129} }, /*  3 */
+					{ {128, 0, 0,   0}, {0,   0,   0, 0}, {129,   1,   2, 2}, {  1,   1, 2, 130} }, /*  4 */
+					{ {128, 0, 1, 129}, {0,   0,   1, 1}, {  0,   0,   2, 2}, {  0,   0, 2, 130} }, /*  5 */
+					{ {128, 0, 2, 130}, {0,   0,   2, 2}, {  1,   1,   1, 1}, {  1,   1, 1, 129} }, /*  6 */
+					{ {128, 0, 1,   1}, {0,   0,   1, 1}, {130,   2,   1, 1}, {  2,   2, 1, 129} }, /*  7 */
+					{ {128, 0, 0,   0}, {0,   0,   0, 0}, {129,   1,   1, 1}, {  2,   2, 2, 130} }, /*  8 */
+					{ {128, 0, 0,   0}, {1,   1,   1, 1}, {129,   1,   1, 1}, {  2,   2, 2, 130} }, /*  9 */
+					{ {128, 0, 0,   0}, {1,   1, 129, 1}, {  2,   2,   2, 2}, {  2,   2, 2, 130} }, /* 10 */
+					{ {128, 0, 1,   2}, {0,   0, 129, 2}, {  0,   0,   1, 2}, {  0,   0, 1, 130} }, /* 11 */
+					{ {128, 1, 1,   2}, {0,   1, 129, 2}, {  0,   1,   1, 2}, {  0,   1, 1, 130} }, /* 12 */
+					{ {128, 1, 2,   2}, {0, 129,   2, 2}, {  0,   1,   2, 2}, {  0,   1, 2, 130} }, /* 13 */
+					{ {128, 0, 1, 129}, {0,   1,   1, 2}, {  1,   1,   2, 2}, {  1,   2, 2, 130} }, /* 14 */
+					{ {128, 0, 1, 129}, {2,   0,   0, 1}, {130,   2,   0, 0}, {  2,   2, 2,   0} }, /* 15 */
+					{ {128, 0, 0, 129}, {0,   0,   1, 1}, {  0,   1,   1, 2}, {  1,   1, 2, 130} }, /* 16 */
+					{ {128, 1, 1, 129}, {0,   0,   1, 1}, {130,   0,   0, 1}, {  2,   2, 0,   0} }, /* 17 */
+					{ {128, 0, 0,   0}, {1,   1,   2, 2}, {129,   1,   2, 2}, {  1,   1, 2, 130} }, /* 18 */
+					{ {128, 0, 2, 130}, {0,   0,   2, 2}, {  0,   0,   2, 2}, {  1,   1, 1, 129} }, /* 19 */
+					{ {128, 1, 1, 129}, {0,   1,   1, 1}, {  0,   2,   2, 2}, {  0,   2, 2, 130} }, /* 20 */
+					{ {128, 0, 0, 129}, {0,   0,   0, 1}, {130,   2,   2, 1}, {  2,   2, 2,   1} }, /* 21 */
+					{ {128, 0, 0,   0}, {0,   0, 129, 1}, {  0,   1,   2, 2}, {  0,   1, 2, 130} }, /* 22 */
+					{ {128, 0, 0,   0}, {1,   1,   0, 0}, {130,   2, 129, 0}, {  2,   2, 1,   0} }, /* 23 */
+					{ {128, 1, 2, 130}, {0, 129,   2, 2}, {  0,   0,   1, 1}, {  0,   0, 0,   0} }, /* 24 */
+					{ {128, 0, 1,   2}, {0,   0,   1, 2}, {129,   1,   2, 2}, {  2,   2, 2, 130} }, /* 25 */
+					{ {128, 1, 1,   0}, {1,   2, 130, 1}, {129,   2,   2, 1}, {  0,   1, 1,   0} }, /* 26 */
+					{ {128, 0, 0,   0}, {0,   1, 129, 0}, {  1,   2, 130, 1}, {  1,   2, 2,   1} }, /* 27 */
+					{ {128, 0, 2,   2}, {1,   1,   0, 2}, {129,   1,   0, 2}, {  0,   0, 2, 130} }, /* 28 */
+					{ {128, 1, 1,   0}, {0, 129,   1, 0}, {  2,   0,   0, 2}, {  2,   2, 2, 130} }, /* 29 */
+					{ {128, 0, 1,   1}, {0,   1,   2, 2}, {  0,   1, 130, 2}, {  0,   0, 1, 129} }, /* 30 */
+					{ {128, 0, 0,   0}, {2,   0,   0, 0}, {130,   2,   1, 1}, {  2,   2, 2, 129} }, /* 31 */
+					{ {128, 0, 0,   0}, {0,   0,   0, 2}, {129,   1,   2, 2}, {  1,   2, 2, 130} }, /* 32 */
+					{ {128, 2, 2, 130}, {0,   0,   2, 2}, {  0,   0,   1, 2}, {  0,   0, 1, 129} }, /* 33 */
+					{ {128, 0, 1, 129}, {0,   0,   1, 2}, {  0,   0,   2, 2}, {  0,   2, 2, 130} }, /* 34 */
+					{ {128, 1, 2,   0}, {0, 129,   2, 0}, {  0,   1, 130, 0}, {  0,   1, 2,   0} }, /* 35 */
+					{ {128, 0, 0,   0}, {1,   1, 129, 1}, {  2,   2, 130, 2}, {  0,   0, 0,   0} }, /* 36 */
+					{ {128, 1, 2,   0}, {1,   2,   0, 1}, {130,   0, 129, 2}, {  0,   1, 2,   0} }, /* 37 */
+					{ {128, 1, 2,   0}, {2,   0,   1, 2}, {129, 130,   0, 1}, {  0,   1, 2,   0} }, /* 38 */
+					{ {128, 0, 1,   1}, {2,   2,   0, 0}, {  1,   1, 130, 2}, {  0,   0, 1, 129} }, /* 39 */
+					{ {128, 0, 1,   1}, {1,   1, 130, 2}, {  2,   2,   0, 0}, {  0,   0, 1, 129} }, /* 40 */
+					{ {128, 1, 0, 129}, {0,   1,   0, 1}, {  2,   2,   2, 2}, {  2,   2, 2, 130} }, /* 41 */
+					{ {128, 0, 0,   0}, {0,   0,   0, 0}, {130,   1,   2, 1}, {  2,   1, 2, 129} }, /* 42 */
+					{ {128, 0, 2,   2}, {1, 129,   2, 2}, {  0,   0,   2, 2}, {  1,   1, 2, 130} }, /* 43 */
+					{ {128, 0, 2, 130}, {0,   0,   1, 1}, {  0,   0,   2, 2}, {  0,   0, 1, 129} }, /* 44 */
+					{ {128, 2, 2,   0}, {1,   2, 130, 1}, {  0,   2,   2, 0}, {  1,   2, 2, 129} }, /* 45 */
+					{ {128, 1, 0,   1}, {2,   2, 130, 2}, {  2,   2,   2, 2}, {  0,   1, 0, 129} }, /* 46 */
+					{ {128, 0, 0,   0}, {2,   1,   2, 1}, {130,   1,   2, 1}, {  2,   1, 2, 129} }, /* 47 */
+					{ {128, 1, 0, 129}, {0,   1,   0, 1}, {  0,   1,   0, 1}, {  2,   2, 2, 130} }, /* 48 */
+					{ {128, 2, 2, 130}, {0,   1,   1, 1}, {  0,   2,   2, 2}, {  0,   1, 1, 129} }, /* 49 */
+					{ {128, 0, 0,   2}, {1, 129,   1, 2}, {  0,   0,   0, 2}, {  1,   1, 1, 130} }, /* 50 */
+					{ {128, 0, 0,   0}, {2, 129,   1, 2}, {  2,   1,   1, 2}, {  2,   1, 1, 130} }, /* 51 */
+					{ {128, 2, 2,   2}, {0, 129,   1, 1}, {  0,   1,   1, 1}, {  0,   2, 2, 130} }, /* 52 */
+					{ {128, 0, 0,   2}, {1,   1,   1, 2}, {129,   1,   1, 2}, {  0,   0, 0, 130} }, /* 53 */
+					{ {128, 1, 1,   0}, {0, 129,   1, 0}, {  0,   1,   1, 0}, {  2,   2, 2, 130} }, /* 54 */
+					{ {128, 0, 0,   0}, {0,   0,   0, 0}, {  2,   1, 129, 2}, {  2,   1, 1, 130} }, /* 55 */
+					{ {128, 1, 1,   0}, {0, 129,   1, 0}, {  2,   2,   2, 2}, {  2,   2, 2, 130} }, /* 56 */
+					{ {128, 0, 2,   2}, {0,   0,   1, 1}, {  0,   0, 129, 1}, {  0,   0, 2, 130} }, /* 57 */
+					{ {128, 0, 2,   2}, {1,   1,   2, 2}, {129,   1,   2, 2}, {  0,   0, 2, 130} }, /* 58 */
+					{ {128, 0, 0,   0}, {0,   0,   0, 0}, {  0,   0,   0, 0}, {  2, 129, 1, 130} }, /* 59 */
+					{ {128, 0, 0, 130}, {0,   0,   0, 1}, {  0,   0,   0, 2}, {  0,   0, 0, 129} }, /* 60 */
+					{ {128, 2, 2,   2}, {1,   2,   2, 2}, {  0,   2,   2, 2}, {129,   2, 2, 130} }, /* 61 */
+					{ {128, 1, 0, 129}, {2,   2,   2, 2}, {  2,   2,   2, 2}, {  2,   2, 2, 130} }, /* 62 */
+					{ {128, 1, 1, 129}, {2,   0,   1, 1}, {130,   2,   0, 1}, {  2,   2, 2,   0} }  /* 63 */
+				}
+			};
+
+			static int aWeight2[] = { 0, 21, 43, 64 };
+			static int aWeight3[] = { 0, 9, 18, 27, 37, 46, 55, 64 };
+			static int aWeight4[] = { 0, 4, 9, 13, 17, 21, 26, 30, 34, 38, 43, 47, 51, 55, 60, 64 };
+
+			static unsigned char sModeHasPBits = 0b11001011;
+
+			bcdec__bitstream_t bstream;
+			int mode, partition, numPartitions, numEndpoints, i, j, k, rotation, partitionSet;
+			int indexSelectionBit, indexBits, indexBits2, index, index2;
+			int endpoints[6][4];
+			char indices[4][4];
+			int r, g, b, a;
+			int* weights, * weights2;
+			unsigned char* decompressed;
+
+			decompressed = (unsigned char*)decompressedBlock;
+
+			bstream.low = ((unsigned long long*)compressedBlock)[0];
+			bstream.high = ((unsigned long long*)compressedBlock)[1];
+
+			for (mode = 0; mode < 8 && (0 == bcdec__bitstream_read_bit(&bstream)); ++mode);
+
+			/* unexpected mode, clear the block (transparent black) */
+			if (mode >= 8) {
+				for (i = 0; i < 4; ++i) {
+					for (j = 0; j < 4; ++j) {
+						decompressed[j * 4 + 0] = 0;
+						decompressed[j * 4 + 1] = 0;
+						decompressed[j * 4 + 2] = 0;
+						decompressed[j * 4 + 3] = 0;
+					}
+					decompressed += destinationPitch;
+				}
+
+				return;
+			}
+
+			partition = 0;
+			numPartitions = 1;
+			rotation = 0;
+			indexSelectionBit = 0;
+
+			if (mode == 0 || mode == 1 || mode == 2 || mode == 3 || mode == 7) {
+				numPartitions = (mode == 0 || mode == 2) ? 3 : 2;
+				partition = bcdec__bitstream_read_bits(&bstream, (mode == 0) ? 4 : 6);
+			}
+
+			numEndpoints = numPartitions * 2;
+
+			if (mode == 4 || mode == 5) {
+				rotation = bcdec__bitstream_read_bits(&bstream, 2);
+
+				if (mode == 4) {
+					indexSelectionBit = bcdec__bitstream_read_bit(&bstream);
+				}
+			}
+
+			/* Extract endpoints */
+			/* RGB */
+			for (i = 0; i < 3; ++i) {
+				for (j = 0; j < numEndpoints; ++j) {
+					endpoints[j][i] = bcdec__bitstream_read_bits(&bstream, actual_bits_count[0][mode]);
+				}
+			}
+			/* Alpha (if any) */
+			if (actual_bits_count[1][mode] > 0) {
+				for (j = 0; j < numEndpoints; ++j) {
+					endpoints[j][3] = bcdec__bitstream_read_bits(&bstream, actual_bits_count[1][mode]);
+				}
+			}
+
+			/* Fully decode endpoints */
+			/* First handle modes that have P-bits */
+			if (mode == 0 || mode == 1 || mode == 3 || mode == 6 || mode == 7) {
+				for (i = 0; i < numEndpoints; ++i) {
+					/* component-wise left-shift */
+					for (j = 0; j < 4; ++j) {
+						endpoints[i][j] <<= 1;
+					}
+				}
+
+				/* if P-bit is shared */
+				if (mode == 1) {
+					i = bcdec__bitstream_read_bit(&bstream);
+					j = bcdec__bitstream_read_bit(&bstream);
+
+					/* rgb component-wise insert pbits */
+					for (k = 0; k < 3; ++k) {
+						endpoints[0][k] |= i;
+						endpoints[1][k] |= i;
+						endpoints[2][k] |= j;
+						endpoints[3][k] |= j;
+					}
+				}
+				else if (sModeHasPBits & (1 << mode)) {
+					/* unique P-bit per endpoint */
+					for (i = 0; i < numEndpoints; ++i) {
+						j = bcdec__bitstream_read_bit(&bstream);
+						for (k = 0; k < 4; ++k) {
+							endpoints[i][k] |= j;
+						}
+					}
+				}
+			}
+
+			for (i = 0; i < numEndpoints; ++i) {
+				/* get color components precision including pbit */
+				j = actual_bits_count[0][mode] + ((sModeHasPBits >> mode) & 1);
+
+				for (k = 0; k < 3; ++k) {
+					/* left shift endpoint components so that their MSB lies in bit 7 */
+					endpoints[i][k] = endpoints[i][k] << (8 - j);
+					/* Replicate each component's MSB into the LSBs revealed by the left-shift operation above */
+					endpoints[i][k] = endpoints[i][k] | (endpoints[i][k] >> j);
+				}
+
+				/* get alpha component precision including pbit */
+				j = actual_bits_count[1][mode] + ((sModeHasPBits >> mode) & 1);
+
+				/* left shift endpoint components so that their MSB lies in bit 7 */
+				endpoints[i][3] = endpoints[i][3] << (8 - j);
+				/* Replicate each component's MSB into the LSBs revealed by the left-shift operation above */
+				endpoints[i][3] = endpoints[i][3] | (endpoints[i][3] >> j);
+			}
+
+			/* If this mode does not explicitly define the alpha component */
+			/* set alpha equal to 1.0 */
+			if (!actual_bits_count[1][mode]) {
+				for (j = 0; j < numEndpoints; ++j) {
+					endpoints[j][3] = 0xFF;
+				}
+			}
+
+			/* Determine weights tables */
+			indexBits = (mode == 0 || mode == 1) ? 3 : ((mode == 6) ? 4 : 2);
+			indexBits2 = (mode == 4) ? 3 : ((mode == 5) ? 2 : 0);
+			weights = (indexBits == 2) ? aWeight2 : ((indexBits == 3) ? aWeight3 : aWeight4);
+			weights2 = (indexBits2 == 2) ? aWeight2 : aWeight3;
+
+			/* Quite inconvenient that indices aren't interleaved so we have to make 2 passes here */
+			/* Pass #1: collecting color indices */
+			for (i = 0; i < 4; ++i) {
+				for (j = 0; j < 4; ++j) {
+					partitionSet = (numPartitions == 1) ? ((i | j) ? 0 : 128) : partition_sets[numPartitions - 2][partition][i][j];
+
+					indexBits = (mode == 0 || mode == 1) ? 3 : ((mode == 6) ? 4 : 2);
+					/* fix-up index is specified with one less bit */
+					/* The fix-up index for subset 0 is always index 0 */
+					if (partitionSet & 0x80) {
+						indexBits--;
+					}
+
+					indices[i][j] = bcdec__bitstream_read_bits(&bstream, indexBits);
+				}
+			}
+
+			/* Pass #2: reading alpha indices (if any) and interpolating & rotating */
+			for (i = 0; i < 4; ++i) {
+				for (j = 0; j < 4; ++j) {
+					partitionSet = (numPartitions == 1) ? ((i | j) ? 0 : 128) : partition_sets[numPartitions - 2][partition][i][j];
+					partitionSet &= 0x03;
+
+					index = indices[i][j];
+
+					if (!indexBits2) {
+						r = bcdec__interpolate(endpoints[partitionSet * 2][0], endpoints[partitionSet * 2 + 1][0], weights, index);
+						g = bcdec__interpolate(endpoints[partitionSet * 2][1], endpoints[partitionSet * 2 + 1][1], weights, index);
+						b = bcdec__interpolate(endpoints[partitionSet * 2][2], endpoints[partitionSet * 2 + 1][2], weights, index);
+						a = bcdec__interpolate(endpoints[partitionSet * 2][3], endpoints[partitionSet * 2 + 1][3], weights, index);
+					}
+					else {
+						index2 = bcdec__bitstream_read_bits(&bstream, (i | j) ? indexBits2 : (indexBits2 - 1));
+						/* The index value for interpolating color comes from the secondary index bits for the texel
+						   if the mode has an index selection bit and its value is one, and from the primary index bits otherwise.
+						   The alpha index comes from the secondary index bits if the block has a secondary index and
+						   the block either doesn’t have an index selection bit or that bit is zero, and from the primary index bits otherwise. */
+						if (!indexSelectionBit) {
+							r = bcdec__interpolate(endpoints[partitionSet * 2][0], endpoints[partitionSet * 2 + 1][0], weights, index);
+							g = bcdec__interpolate(endpoints[partitionSet * 2][1], endpoints[partitionSet * 2 + 1][1], weights, index);
+							b = bcdec__interpolate(endpoints[partitionSet * 2][2], endpoints[partitionSet * 2 + 1][2], weights, index);
+							a = bcdec__interpolate(endpoints[partitionSet * 2][3], endpoints[partitionSet * 2 + 1][3], weights2, index2);
+						}
+						else {
+							r = bcdec__interpolate(endpoints[partitionSet * 2][0], endpoints[partitionSet * 2 + 1][0], weights2, index2);
+							g = bcdec__interpolate(endpoints[partitionSet * 2][1], endpoints[partitionSet * 2 + 1][1], weights2, index2);
+							b = bcdec__interpolate(endpoints[partitionSet * 2][2], endpoints[partitionSet * 2 + 1][2], weights2, index2);
+							a = bcdec__interpolate(endpoints[partitionSet * 2][3], endpoints[partitionSet * 2 + 1][3], weights, index);
+						}
+					}
+
+					switch (rotation) {
+					case 1: {   /* 01 – Block format is Scalar(R) Vector(AGB) - swap A and R */
+						bcdec__swap_values(&a, &r);
+					} break;
+					case 2: {   /* 10 – Block format is Scalar(G) Vector(RAB) - swap A and G */
+						bcdec__swap_values(&a, &g);
+					} break;
+					case 3: {   /* 11 - Block format is Scalar(B) Vector(RGA) - swap A and B */
+						bcdec__swap_values(&a, &b);
+					} break;
+					}
+
+					decompressed[j * 4 + 0] = r;
+					decompressed[j * 4 + 1] = g;
+					decompressed[j * 4 + 2] = b;
+					decompressed[j * 4 + 3] = a;
+				}
+
+				decompressed += destinationPitch;
+			}
+		}
+
+#endif /* BCDEC_IMPLEMENTATION */
+	}
+
+
+	//-------------------------------------------------------------------------------------------------
+	// Based on STB v1.08, under the MIT license
+	// See [REDACTED_PROFILE_URL]/stb/
+	//-------------------------------------------------------------------------------------------------
+	namespace stb
+	{
+		// compression mode (bitflags)
+#define STB_DXT_NORMAL    0
+#define STB_DXT_DITHER    1   // use dithering. dubious win. never use for normal maps and the like!
+#define STB_DXT_HIGHQUAL  2   // high quality mode, does two refinement steps instead of 1. ~30-40% slower.
+
+		static void stb_compress_dxt_block(unsigned char* dest, const unsigned char* src_rgba_four_bytes_per_pixel, int alpha, int mode);
+		static void stb_compress_bc4_block(unsigned char* dest, const unsigned char* src_r_one_byte_per_pixel);
+		static void stb_compress_bc5_block(unsigned char* dest, const unsigned char* src_rg_two_byte_per_pixel);
+
+#define STB_COMPRESS_DXT_BLOCK
+
+
+		// configuration options for DXT encoder. set them in the project/makefile or just define
+		// them at the top.
+
+		// STB_DXT_USE_ROUNDING_BIAS
+		//     use a rounding bias during color interpolation. this is closer to what "ideal"
+		//     interpolation would do but doesn't match the S3TC/DX10 spec. old versions (pre-1.03)
+		//     implicitly had this turned on.
+		//
+		//     in case you're targeting a specific type of hardware (e.g. console programmers):
+		//     NVidia and Intel GPUs (as of 2010) as well as DX9 ref use DXT decoders that are closer
+		//     to STB_DXT_USE_ROUNDING_BIAS. AMD/ATI, S3 and DX10 ref are closer to rounding with no bias.
+		//     you also see "(a*5 + b*3) / 8" on some old GPU designs.
+		// #define STB_DXT_USE_ROUNDING_BIAS
+
+#ifndef STBD_ABS
+#define STBD_ABS(i)           abs(i)
+#endif
+
+#ifndef STBD_FABS
+#define STBD_FABS(x)          fabs(x)
+#endif
+
+#ifndef STBD_MEMSET
+#define STBD_MEMSET           memset
+#endif
+
+		static int stb__Mul8Bit(int a, int b)
+		{
+			int t = a * b + 128;
+			return (t + (t >> 8)) >> 8;
+		}
+
+		static void stb__From16Bit(unsigned char* out, unsigned short v)
+		{
+			int rv = (v & 0xf800) >> 11;
+			int gv = (v & 0x07e0) >> 5;
+			int bv = (v & 0x001f) >> 0;
+
+			out[0] = stb__Expand5[rv];
+			out[1] = stb__Expand6[gv];
+			out[2] = stb__Expand5[bv];
+			out[3] = 0;
+		}
+
+		static unsigned short stb__Aint16_tBit(int r, int g, int b)
+		{
+			return uint16((stb__Mul8Bit(r, 31) << 11) + (stb__Mul8Bit(g, 63) << 5) + stb__Mul8Bit(b, 31));
+		}
+
+		// linear interpolation at 1/3 point between a and b, using desired rounding type
+		static int stb__Lerp13(int a, int b)
+		{
+#ifdef STB_DXT_USE_ROUNDING_BIAS
+			// with rounding bias
+			return a + stb__Mul8Bit(b - a, 0x55);
+#else
+			// without rounding bias
+			// replace "/ 3" by "* 0xaaab) >> 17" if your compiler sucks or you really need every ounce of speed.
+			return (2 * a + b) / 3;
+#endif
+		}
+
+		// lerp RGB color
+		static void stb__Lerp13RGB(unsigned char* out, unsigned char* p1, unsigned char* p2)
+		{
+			out[0] = (uint8)stb__Lerp13(p1[0], p2[0]);
+			out[1] = (uint8)stb__Lerp13(p1[1], p2[1]);
+			out[2] = (uint8)stb__Lerp13(p1[2], p2[2]);
+		}
+
+		/****************************************************************************/
+
+		// compute table to reproduce constant colors as accurately as possible
+		static void stb__PrepareOptTable(unsigned char* Table, const unsigned char* expand, int size)
+		{
+			int i, mn, mx;
+			for (i = 0; i < 256; i++) {
+				int bestErr = 256;
+				for (mn = 0; mn < size; mn++) {
+					for (mx = 0; mx < size; mx++) {
+						int mine = expand[mn];
+						int maxe = expand[mx];
+						int err = STBD_ABS(stb__Lerp13(maxe, mine) - i);
+
+						// DX10 spec says that interpolation must be within 3% of "correct" result,
+						// add this as error term. (normally we'd expect a random distribution of
+						// +-1.5% error, but nowhere in the spec does it say that the error has to be
+						// unbiased - better safe than sorry).
+						err += STBD_ABS(maxe - mine) * 3 / 100;
+
+						if (err < bestErr)
+						{
+							Table[i * 2 + 0] = (uint8)mx;
+							Table[i * 2 + 1] = (uint8)mn;
+							bestErr = err;
+						}
+					}
+				}
+			}
+		}
+
+		static void stb__EvalColors(unsigned char* color, unsigned short c0, unsigned short c1)
+		{
+			stb__From16Bit(color + 0, c0);
+			stb__From16Bit(color + 4, c1);
+			stb__Lerp13RGB(color + 8, color + 0, color + 4);
+			stb__Lerp13RGB(color + 12, color + 4, color + 0);
+		}
+
+		// Block dithering function. Simply dithers a block to 565 RGB.
+		// (Floyd-Steinberg)
+		static void stb__DitherBlock(unsigned char* dest, unsigned char* block)
+		{
+			int err[8], * ep1 = err, * ep2 = err + 4, * et;
+			int ch, y;
+
+			// process channels separately
+			for (ch = 0; ch < 3; ++ch) {
+				unsigned char* bp = block + ch, * dp = dest + ch;
+				unsigned char* quant = (ch == 1) ? stb__QuantGTab + 8 : stb__QuantRBTab + 8;
+				STBD_MEMSET(err, 0, sizeof(err));
+				for (y = 0; y < 4; ++y) {
+					dp[0] = quant[bp[0] + ((3 * ep2[1] + 5 * ep2[0]) >> 4)];
+					ep1[0] = bp[0] - dp[0];
+					dp[4] = quant[bp[4] + ((7 * ep1[0] + 3 * ep2[2] + 5 * ep2[1] + ep2[0]) >> 4)];
+					ep1[1] = bp[4] - dp[4];
+					dp[8] = quant[bp[8] + ((7 * ep1[1] + 3 * ep2[3] + 5 * ep2[2] + ep2[1]) >> 4)];
+					ep1[2] = bp[8] - dp[8];
+					dp[12] = quant[bp[12] + ((7 * ep1[2] + 5 * ep2[3] + ep2[2]) >> 4)];
+					ep1[3] = bp[12] - dp[12];
+					bp += 16;
+					dp += 16;
+					et = ep1, ep1 = ep2, ep2 = et; // swap
+				}
+			}
+		}
+
+		// The color matching function
+		static unsigned int stb__MatchColorsBlock(unsigned char* block, unsigned char* color, int dither)
+		{
+			unsigned int mask = 0;
+			int dirr = color[0 * 4 + 0] - color[1 * 4 + 0];
+			int dirg = color[0 * 4 + 1] - color[1 * 4 + 1];
+			int dirb = color[0 * 4 + 2] - color[1 * 4 + 2];
+			int dots[16];
+			int stops[4];
+			int i;
+			int c0Point, halfPoint, c3Point;
+
+			for (i = 0; i < 16; i++)
+				dots[i] = block[i * 4 + 0] * dirr + block[i * 4 + 1] * dirg + block[i * 4 + 2] * dirb;
+
+			for (i = 0; i < 4; i++)
+				stops[i] = color[i * 4 + 0] * dirr + color[i * 4 + 1] * dirg + color[i * 4 + 2] * dirb;
+
+			// think of the colors as arranged on a line; project point onto that line, then choose
+			// next color out of available ones. we compute the crossover points for "best color in top
+			// half"/"best in bottom half" and then the same inside that subinterval.
+			//
+			// relying on this 1d approximation isn't always optimal in terms of euclidean distance,
+			// but it's very close and a lot faster.
+			// http://cbloomrants.blogspot.com/2008/12/12-08-08-dxtc-summary.html
+
+			c0Point = (stops[1] + stops[3]) >> 1;
+			halfPoint = (stops[3] + stops[2]) >> 1;
+			c3Point = (stops[2] + stops[0]) >> 1;
+
+			if (!dither) {
+				// the version without dithering is straightforward
+				for (i = 15; i >= 0; i--) {
+					int dot = dots[i];
+					mask <<= 2;
+
+					if (dot < halfPoint)
+						mask |= (dot < c0Point) ? 1 : 3;
+					else
+						mask |= (dot < c3Point) ? 2 : 0;
+				}
+			}
+			else {
+				// with floyd-steinberg dithering
+				int err[8], * ep1 = err, * ep2 = err + 4;
+				int* dp = dots, y;
+
+				c0Point <<= 4;
+				halfPoint <<= 4;
+				c3Point <<= 4;
+				for (i = 0; i < 8; i++)
+					err[i] = 0;
+
+				for (y = 0; y < 4; y++)
+				{
+					int dot, lmask, step;
+
+					dot = (dp[0] << 4) + (3 * ep2[1] + 5 * ep2[0]);
+					if (dot < halfPoint)
+						step = (dot < c0Point) ? 1 : 3;
+					else
+						step = (dot < c3Point) ? 2 : 0;
+					ep1[0] = dp[0] - stops[step];
+					lmask = step;
+
+					dot = (dp[1] << 4) + (7 * ep1[0] + 3 * ep2[2] + 5 * ep2[1] + ep2[0]);
+					if (dot < halfPoint)
+						step = (dot < c0Point) ? 1 : 3;
+					else
+						step = (dot < c3Point) ? 2 : 0;
+					ep1[1] = dp[1] - stops[step];
+					lmask |= step << 2;
+
+					dot = (dp[2] << 4) + (7 * ep1[1] + 3 * ep2[3] + 5 * ep2[2] + ep2[1]);
+					if (dot < halfPoint)
+						step = (dot < c0Point) ? 1 : 3;
+					else
+						step = (dot < c3Point) ? 2 : 0;
+					ep1[2] = dp[2] - stops[step];
+					lmask |= step << 4;
+
+					dot = (dp[3] << 4) + (7 * ep1[2] + 5 * ep2[3] + ep2[2]);
+					if (dot < halfPoint)
+						step = (dot < c0Point) ? 1 : 3;
+					else
+						step = (dot < c3Point) ? 2 : 0;
+					ep1[3] = dp[3] - stops[step];
+					lmask |= step << 6;
+
+					dp += 4;
+					mask |= lmask << (y * 8);
+					{ int* et = ep1; ep1 = ep2; ep2 = et; } // swap
+				}
+			}
+
+			return mask;
+		}
+
+		// The color optimization function. (Clever code, part 1)
+		static void stb__OptimizeColorsBlock(unsigned char* block, unsigned short* pmax16, unsigned short* pmin16)
+		{
+			int mind = 0x7fffffff, maxd = -0x7fffffff;
+			unsigned char* minp, * maxp;
+			double magn;
+			int v_r, v_g, v_b;
+			static const int nIterPower = 4;
+			float covf[6], vfr, vfg, vfb;
+
+			// determine color distribution
+			int cov[6];
+			int mu[3], min[3], max[3];
+			int ch, i, iter;
+
+			for (ch = 0; ch < 3; ch++)
+			{
+				const unsigned char* bp = ((const unsigned char*)block) + ch;
+				int muv, minv, maxv;
+
+				muv = minv = maxv = bp[0];
+				for (i = 4; i < 64; i += 4)
+				{
+					muv += bp[i];
+					if (bp[i] < minv) minv = bp[i];
+					else if (bp[i] > maxv) maxv = bp[i];
+				}
+
+				mu[ch] = (muv + 8) >> 4;
+				min[ch] = minv;
+				max[ch] = maxv;
+			}
+
+			// determine covariance matrix
+			for (i = 0; i < 6; i++)
+				cov[i] = 0;
+
+			for (i = 0; i < 16; i++)
+			{
+				int r = block[i * 4 + 0] - mu[0];
+				int g = block[i * 4 + 1] - mu[1];
+				int b = block[i * 4 + 2] - mu[2];
+
+				cov[0] += r * r;
+				cov[1] += r * g;
+				cov[2] += r * b;
+				cov[3] += g * g;
+				cov[4] += g * b;
+				cov[5] += b * b;
+			}
+
+			// convert covariance matrix to float, find principal axis via power iter
+			for (i = 0; i < 6; i++)
+				covf[i] = cov[i] / 255.0f;
+
+			vfr = (float)(max[0] - min[0]);
+			vfg = (float)(max[1] - min[1]);
+			vfb = (float)(max[2] - min[2]);
+
+			for (iter = 0; iter < nIterPower; iter++)
+			{
+				float r = vfr * covf[0] + vfg * covf[1] + vfb * covf[2];
+				float g = vfr * covf[1] + vfg * covf[3] + vfb * covf[4];
+				float b = vfr * covf[2] + vfg * covf[4] + vfb * covf[5];
+
+				vfr = r;
+				vfg = g;
+				vfb = b;
+			}
+
+			magn = STBD_FABS(vfr);
+			if (STBD_FABS(vfg) > magn) magn = STBD_FABS(vfg);
+			if (STBD_FABS(vfb) > magn) magn = STBD_FABS(vfb);
+
+			if (magn < 4.0f) { // too small, default to luminance
+				v_r = 299; // JPEG YCbCr luma coefs, scaled by 1000.
+				v_g = 587;
+				v_b = 114;
+			}
+			else {
+				magn = 512.0 / magn;
+				v_r = (int)(vfr * magn);
+				v_g = (int)(vfg * magn);
+				v_b = (int)(vfb * magn);
+			}
+
+			// Pick colors at extreme points
+			minp = maxp = block;
+			for (i = 0; i < 16; i++)
+			{
+				int dot = block[i * 4 + 0] * v_r + block[i * 4 + 1] * v_g + block[i * 4 + 2] * v_b;
+
+				if (dot < mind) {
+					mind = dot;
+					minp = block + i * 4;
+				}
+
+				if (dot > maxd) {
+					maxd = dot;
+					maxp = block + i * 4;
+				}
+			}
+
+			*pmax16 = stb__Aint16_tBit(maxp[0], maxp[1], maxp[2]);
+			*pmin16 = stb__Aint16_tBit(minp[0], minp[1], minp[2]);
+		}
+
+		static int stb__sclamp(float y, int p0, int p1)
+		{
+			int x = (int)y;
+			if (x < p0) return p0;
+			if (x > p1) return p1;
+			return x;
+		}
+
+		// The refinement function. (Clever code, part 2)
+		// Tries to optimize colors to suit block contents better.
+		// (By solving a least squares system via normal equations+Cramer's rule)
+		static int stb__RefineBlock(unsigned char* block, unsigned short* pmax16, unsigned short* pmin16, unsigned int mask)
+		{
+			static const int w1Tab[4] = { 3,0,2,1 };
+			static const int prods[4] = { 0x090000,0x000900,0x040102,0x010402 };
+			// ^some magic to save a lot of multiplies in the accumulating loop...
+			// (precomputed products of weights for least squares system, accumulated inside one 32-bit register)
+
+			float frb, fg;
+			unsigned short oldMin, oldMax, min16, max16;
+			int i, akku = 0, xx, xy, yy;
+			int At1_r, At1_g, At1_b;
+			int At2_r, At2_g, At2_b;
+			unsigned int cm = mask;
+
+			oldMin = *pmin16;
+			oldMax = *pmax16;
+
+			if ((mask ^ (mask << 2)) < 4) // all pixels have the same index?
+			{
+				// yes, linear system would be singular; solve using optimal
+				// single-color match on average color
+				int r = 8, g = 8, b = 8;
+				for (i = 0; i < 16; ++i) {
+					r += block[i * 4 + 0];
+					g += block[i * 4 + 1];
+					b += block[i * 4 + 2];
+				}
+
+				r >>= 4; g >>= 4; b >>= 4;
+
+				max16 = (stb__OMatch5[r][0] << 11) | (stb__OMatch6[g][0] << 5) | stb__OMatch5[b][0];
+				min16 = (stb__OMatch5[r][1] << 11) | (stb__OMatch6[g][1] << 5) | stb__OMatch5[b][1];
+			}
+			else {
+				At1_r = At1_g = At1_b = 0;
+				At2_r = At2_g = At2_b = 0;
+				for (i = 0; i < 16; ++i, cm >>= 2) {
+					int step = cm & 3;
+					int w1 = w1Tab[step];
+					int r = block[i * 4 + 0];
+					int g = block[i * 4 + 1];
+					int b = block[i * 4 + 2];
+
+					akku += prods[step];
+					At1_r += w1 * r;
+					At1_g += w1 * g;
+					At1_b += w1 * b;
+					At2_r += r;
+					At2_g += g;
+					At2_b += b;
+				}
+
+				At2_r = 3 * At2_r - At1_r;
+				At2_g = 3 * At2_g - At1_g;
+				At2_b = 3 * At2_b - At1_b;
+
+				// extract solutions and decide solvability
+				xx = akku >> 16;
+				yy = (akku >> 8) & 0xff;
+				xy = (akku >> 0) & 0xff;
+
+				frb = 3.0f * 31.0f / 255.0f / (xx * yy - xy * xy);
+				fg = frb * 63.0f / 31.0f;
+
+				// solve.
+				max16 = (uint16)stb__sclamp((At1_r * yy - At2_r * xy) * frb + 0.5f, 0, 31) << 11;
+				max16 |= (uint16)stb__sclamp((At1_g * yy - At2_g * xy) * fg + 0.5f, 0, 63) << 5;
+				max16 |= (uint16)stb__sclamp((At1_b * yy - At2_b * xy) * frb + 0.5f, 0, 31) << 0;
+
+				min16 = (uint16)stb__sclamp((At2_r * xx - At1_r * xy) * frb + 0.5f, 0, 31) << 11;
+				min16 |= (uint16)stb__sclamp((At2_g * xx - At1_g * xy) * fg + 0.5f, 0, 63) << 5;
+				min16 |= (uint16)stb__sclamp((At2_b * xx - At1_b * xy) * frb + 0.5f, 0, 31) << 0;
+			}
+
+			*pmin16 = min16;
+			*pmax16 = max16;
+			return oldMin != min16 || oldMax != max16;
+		}
+
+		// Color block compression
+		static void stb__CompressColorBlock(unsigned char* dest, unsigned char* block, int mode)
+		{
+			unsigned int mask;
+			int i;
+			int dither;
+			int refinecount;
+			unsigned short max16, min16;
+			unsigned char dblock[16 * 4], color[4 * 4];
+
+			dither = mode & STB_DXT_DITHER;
+			refinecount = (mode & STB_DXT_HIGHQUAL) ? 2 : 1;
+
+			// check if block is constant
+			for (i = 1; i < 16; i++)
+				if (((unsigned int*)block)[i] != ((unsigned int*)block)[0])
+					break;
+
+			if (i == 16)
+			{ // constant color
+				int r = block[0], g = block[1], b = block[2];
+				mask = 0xaaaaaaaa;
+				max16 = (stb__OMatch5[r][0] << 11) | (stb__OMatch6[g][0] << 5) | stb__OMatch5[b][0];
+				min16 = (stb__OMatch5[r][1] << 11) | (stb__OMatch6[g][1] << 5) | stb__OMatch5[b][1];
+			}
+			else
+			{
+				// first step: compute dithered version for PCA if desired
+				if (dither)
+					stb__DitherBlock(dblock, block);
+
+				// second step: pca+map along principal axis
+				stb__OptimizeColorsBlock(dither ? dblock : block, &max16, &min16);
+				if (max16 != min16)
+				{
+					stb__EvalColors(color, max16, min16);
+					mask = stb__MatchColorsBlock(block, color, dither);
+				}
+				else
+					mask = 0;
+
+				// third step: refine (multiple times if requested)
+				for (i = 0; i < refinecount; i++)
+				{
+					unsigned int lastmask = mask;
+
+					if (stb__RefineBlock(dither ? dblock : block, &max16, &min16, mask))
+					{
+						if (max16 != min16)
+						{
+							stb__EvalColors(color, max16, min16);
+							mask = stb__MatchColorsBlock(block, color, dither);
+						}
+						else
+						{
+							mask = 0;
+							break;
+						}
+					}
+
+					if (mask == lastmask)
+						break;
+				}
+			}
+
+			// write the color block
+			if (max16 < min16)
+			{
+				unsigned short t = min16;
+				min16 = max16;
+				max16 = t;
+				mask ^= 0x55555555;
+			}
+
+			dest[0] = (unsigned char)(max16);
+			dest[1] = (unsigned char)(max16 >> 8);
+			dest[2] = (unsigned char)(min16);
+			dest[3] = (unsigned char)(min16 >> 8);
+			dest[4] = (unsigned char)(mask);
+			dest[5] = (unsigned char)(mask >> 8);
+			dest[6] = (unsigned char)(mask >> 16);
+			dest[7] = (unsigned char)(mask >> 24);
+		}
+
+
+		// Alpha block compression (this is easy for a change)
+		static void stb__CompressAlphaBlock(unsigned char* dest, unsigned char* src, int stride)
+		{
+			int i, dist, bias, dist4, dist2, bits, mask;
+
+			// find min/max color
+			int mn, mx;
+			mn = mx = src[0];
+
+			for (i = 1; i < 16; i++)
+			{
+				if (src[i * stride] < mn) mn = src[i * stride];
+				else if (src[i * stride] > mx) mx = src[i * stride];
+			}
+
+			// encode them
+			((uint8*)dest)[0] = (uint8)mx;
+			((uint8*)dest)[1] = (uint8)mn;
+			dest += 2;
+
+			// determine bias and emit color indices
+			// given the choice of mx/mn, these indices are optimal:
+			// http://fgiesen.wordpress.com/2009/12/15/dxt5-alpha-block-index-determination/
+			dist = mx - mn;
+			dist4 = dist * 4;
+			dist2 = dist * 2;
+			bias = (dist < 8) ? (dist - 1) : (dist / 2 + 2);
+			bias -= mn * 7;
+			bits = 0, mask = 0;
+
+			for (i = 0; i < 16; i++)
+			{
+				int a = src[i * stride] * 7 + bias;
+				int ind, t;
+
+				// select index. this is a "linear scale" lerp factor between 0 (val=min) and 7 (val=max).
+				t = (a >= dist4) ? -1 : 0; ind = t & 4; a -= dist4 & t;
+				t = (a >= dist2) ? -1 : 0; ind += t & 2; a -= dist2 & t;
+				ind += (a >= dist);
+
+				// turn linear scale into DXT index (0/1 are extremal pts)
+				ind = -ind & 7;
+				ind ^= (2 > ind);
+
+				// write index
+				mask |= ind << bits;
+				if ((bits += 3) >= 8) {
+					*dest++ = (uint8)mask;
+					mask >>= 8;
+					bits -= 8;
+				}
+			}
+		}
+
+
+		static void stb__InitDXT()
+		{
+			int i;
+			for (i = 0; i < 32; i++)
+				stb__Expand5[i] = uint8((i << 3) | (i >> 2));
+
+			for (i = 0; i < 64; i++)
+				stb__Expand6[i] = uint8((i << 2) | (i >> 4));
+
+			for (i = 0; i < 256 + 16; i++)
+			{
+				int v = i - 8 < 0 ? 0 : i - 8 > 255 ? 255 : i - 8;
+				stb__QuantRBTab[i] = stb__Expand5[stb__Mul8Bit(v, 31)];
+				stb__QuantGTab[i] = stb__Expand6[stb__Mul8Bit(v, 63)];
+			}
+
+			stb__PrepareOptTable(&stb__OMatch5[0][0], stb__Expand5, 32);
+			stb__PrepareOptTable(&stb__OMatch6[0][0], stb__Expand6, 64);
+		}
+
+
+		void stb_compress_dxt_block(unsigned char* dest, const unsigned char* src, int alpha, int mode)
+		{
+			// If this assert fails it means miro::initialize() has not been called.
+			check(s_initialised_dxtc_decompress);
+
+			unsigned char data[16][4];
+
+			if (alpha > 0)
+			{
+				int i;
+				if (alpha == 2)
+				{
+					CompressAlphaBtc2u(src, dest);
+				}
+				else
+				{
+					stb__CompressAlphaBlock(dest, (unsigned char*)src + 3, 4);
+				}
+				dest += 8;
+				// make a new copy of the data in which alpha is opaque,
+				// because code uses a fast test for color constancy
+				FMemory::Memcpy(data, src, 4 * 16);
+				for (i = 0; i < 16; ++i)
+					data[i][3] = 255;
+				src = &data[0][0];
+			}
+
+			stb__CompressColorBlock(dest, (unsigned char*)src, mode);
+		}
+
+		void stb_compress_bc4_block(unsigned char* dest, const unsigned char* src)
+		{
+			stb__CompressAlphaBlock(dest, (unsigned char*)src, 1);
+		}
+
+		void stb_compress_bc5_block(unsigned char* dest, const unsigned char* src)
+		{
+			stb__CompressAlphaBlock(dest, (unsigned char*)src, 2);
+			stb__CompressAlphaBlock(dest + 8, (unsigned char*)src + 1, 2);
+		}
+
+	}
+
+
+#endif //MIRO_INCLUDE_BC
+
+
+#if MIRO_INCLUDE_ASTC
+
+	//-------------------------------------------------------------------------------------------------
+	// Based on ASTCRT, under the BSD 3-Clause License
+	// See: [REDACTED_PROFILE_URL]/astcrt
+	//-------------------------------------------------------------------------------------------------
+	inline void DCHECK(bool x)
+	{
+		(void)x;
+		// if (!x)
+		// {
+		//     assert(x);
+		// }
+	}
+
+
+	namespace astcrt
+	{
+		constexpr int APPROX_COLOR_EPSILON = 50;
+		constexpr size_t BLOCK_WIDTH = 4;
+		constexpr size_t BLOCK_HEIGHT = 4;
+		constexpr size_t BLOCK_TEXEL_COUNT = BLOCK_WIDTH * BLOCK_HEIGHT;
+		constexpr size_t BLOCK_BYTES = 16;
+
+		constexpr size_t MAXIMUM_ENCODED_WEIGHT_BITS = 96;
+		constexpr size_t MAXIMUM_ENCODED_WEIGHT_BYTES = 12;
+
+		constexpr size_t MAXIMUM_ENCODED_COLOR_ENDPOINT_BYTES = 12;
+
+
+		inline bool getbit(size_t number, size_t n) {
+			return (number >> n) & 1;
+		}
+
+		inline uint8 getbits(uint8 number, uint8 msb, uint8 lsb) {
+			int count = msb - lsb + 1;
+			return static_cast<uint8>((number >> lsb) & ((1 << count) - 1));
+		}
+
+		inline size_t getbits(size_t number, size_t msb, size_t lsb) {
+			size_t count = msb - lsb + 1;
+			return (number >> lsb) & ((size_t(1) << count) - 1);
+		}
+
+		inline void orbits8_ptr(uint8* ptr,
+			size_t bitoffset,
+			size_t number,
+			size_t bitcount) {
+			DCHECK(bitcount <= 8);
+			DCHECK((number >> bitcount) == 0);
+
+			size_t index = bitoffset / 8;
+			size_t shift = bitoffset % 8;
+
+			// Depending on the offset we might have to consider two bytes when
+			// writing, for instance if we are writing 8 bits and the offset is 4,
+			// then we have to write 4 bits to the first byte (ptr[index]) and 4 bits
+			// to the second byte (ptr[index+1]).
+			//
+			// FIXME: Writing to the last byte when the number of bytes is a multiple of 2
+			// will write past the allocated memory.
+
+			uint8* p = ptr + index;
+			size_t mask = number << shift;
+
+			DCHECK((p[0] & mask) == 0);
+			DCHECK((p[1] & (mask >> 8)) == 0);
+
+			p[0] |= static_cast<uint8>(mask & 0xFF);
+			p[1] |= static_cast<uint8>((mask >> 8) & 0xFF);
+		}
+
+		inline uint16 getbytes2(const uint8* ptr, size_t byteoffset) {
+			const uint8* p = ptr + byteoffset;
+			return static_cast<uint16>((p[1] << 8) | p[0]);
+		}
+
+		inline void setbytes2(uint8* ptr, size_t byteoffset, uint16 bytes) {
+			ptr[byteoffset + 0] = static_cast<uint8>(bytes & 0xFF);
+			ptr[byteoffset + 1] = static_cast<uint8>((bytes >> 8) & 0xFF);
+		}
+
+		inline void split_high_low(uint8 n, size_t i, uint8& high, uint8& low) {
+			DCHECK(i < 8);
+
+			uint8 low_mask = static_cast<uint8>((1 << i) - 1);
+
+			low = n & low_mask;
+			high = static_cast<uint8>(n >> i);
+		}
+
+		class bitwriter {
+		public:
+			explicit bitwriter(uint8* ptr) : ptr_(ptr), bitoffset_(0) {
+				// assumption that all bits in ptr are zero after the offset
+
+				// writing beyound the bounds of the allocated memory is undefined
+				// behaviour
+			}
+
+			// Specialized function that can't write more than 8 bits.
+			void write8(uint8 number, size_t bitcount) {
+				orbits8_ptr(ptr_, bitoffset_, number, bitcount);
+
+				bitoffset_ += bitcount;
+			}
+
+			size_t offset() const { return bitoffset_; }
+
+		private:
+			uint8* ptr_;
+			size_t bitoffset_;  // in bits
+		};
+
+		const uint8 bit_reverse_table[256] = {
+			0x00, 0x80, 0x40, 0xC0, 0x20, 0xA0, 0x60, 0xE0, 0x10, 0x90, 0x50, 0xD0,
+			0x30, 0xB0, 0x70, 0xF0, 0x08, 0x88, 0x48, 0xC8, 0x28, 0xA8, 0x68, 0xE8,
+			0x18, 0x98, 0x58, 0xD8, 0x38, 0xB8, 0x78, 0xF8, 0x04, 0x84, 0x44, 0xC4,
+			0x24, 0xA4, 0x64, 0xE4, 0x14, 0x94, 0x54, 0xD4, 0x34, 0xB4, 0x74, 0xF4,
+			0x0C, 0x8C, 0x4C, 0xCC, 0x2C, 0xAC, 0x6C, 0xEC, 0x1C, 0x9C, 0x5C, 0xDC,
+			0x3C, 0xBC, 0x7C, 0xFC, 0x02, 0x82, 0x42, 0xC2, 0x22, 0xA2, 0x62, 0xE2,
+			0x12, 0x92, 0x52, 0xD2, 0x32, 0xB2, 0x72, 0xF2, 0x0A, 0x8A, 0x4A, 0xCA,
+			0x2A, 0xAA, 0x6A, 0xEA, 0x1A, 0x9A, 0x5A, 0xDA, 0x3A, 0xBA, 0x7A, 0xFA,
+			0x06, 0x86, 0x46, 0xC6, 0x26, 0xA6, 0x66, 0xE6, 0x16, 0x96, 0x56, 0xD6,
+			0x36, 0xB6, 0x76, 0xF6, 0x0E, 0x8E, 0x4E, 0xCE, 0x2E, 0xAE, 0x6E, 0xEE,
+			0x1E, 0x9E, 0x5E, 0xDE, 0x3E, 0xBE, 0x7E, 0xFE, 0x01, 0x81, 0x41, 0xC1,
+			0x21, 0xA1, 0x61, 0xE1, 0x11, 0x91, 0x51, 0xD1, 0x31, 0xB1, 0x71, 0xF1,
+			0x09, 0x89, 0x49, 0xC9, 0x29, 0xA9, 0x69, 0xE9, 0x19, 0x99, 0x59, 0xD9,
+			0x39, 0xB9, 0x79, 0xF9, 0x05, 0x85, 0x45, 0xC5, 0x25, 0xA5, 0x65, 0xE5,
+			0x15, 0x95, 0x55, 0xD5, 0x35, 0xB5, 0x75, 0xF5, 0x0D, 0x8D, 0x4D, 0xCD,
+			0x2D, 0xAD, 0x6D, 0xED, 0x1D, 0x9D, 0x5D, 0xDD, 0x3D, 0xBD, 0x7D, 0xFD,
+			0x03, 0x83, 0x43, 0xC3, 0x23, 0xA3, 0x63, 0xE3, 0x13, 0x93, 0x53, 0xD3,
+			0x33, 0xB3, 0x73, 0xF3, 0x0B, 0x8B, 0x4B, 0xCB, 0x2B, 0xAB, 0x6B, 0xEB,
+			0x1B, 0x9B, 0x5B, 0xDB, 0x3B, 0xBB, 0x7B, 0xFB, 0x07, 0x87, 0x47, 0xC7,
+			0x27, 0xA7, 0x67, 0xE7, 0x17, 0x97, 0x57, 0xD7, 0x37, 0xB7, 0x77, 0xF7,
+			0x0F, 0x8F, 0x4F, 0xCF, 0x2F, 0xAF, 0x6F, 0xEF, 0x1F, 0x9F, 0x5F, 0xDF,
+			0x3F, 0xBF, 0x7F, 0xFF };
+
+		/**
+		 * Reverse a byte, total function.
+		 */
+		inline uint8 reverse_byte(uint8 number) {
+			return bit_reverse_table[number];
+		}
+
+		/**
+		 * Reverse a sequence of bytes.
+		 *
+		 * Assumes that the bits written to (using bitwise or) are zero and that they
+		 * will not clash with bits already written to target sequence. That is it is
+		 * possible to write to a non-zero byte as long as the bits that are actually
+		 * written to are zero.
+		 */
+		inline void reverse_bytes(const uint8* source,
+			size_t bytecount,
+			uint8* target)
+		{
+			for (int i = 0; i < static_cast<int>(bytecount); ++i)
+			{
+				DCHECK((reverse_byte(source[i]) & target[-i]) == 0);
+				target[-i] = target[-i] | reverse_byte(source[i]);
+			}
+		}
+
+		inline void copy_bytes(const uint8* source,
+			size_t bytecount,
+			uint8* target,
+			size_t bitoffset)
+		{
+			for (size_t i = 0; i < bytecount; ++i)
+			{
+				orbits8_ptr(target, bitoffset + i * 8, source[i], 8);
+			}
+		}
+
+
+		union unorm8_t
+		{
+			struct RgbaColorType
+			{
+				uint8 b, g, r, a;
+			} channels;
+			uint8 components[4];
+			uint32 bits;
+		};
+
+		union unorm16_t
+		{
+			struct RgbaColorType
+			{
+				uint16 b, g, r, a;
+			} channels;
+			uint16 components[4];
+			uint64 bits;
+		};
+
+
+
+		template <typename T>
+		struct vec3_t
+		{
+		public:
+			vec3_t() {}
+			vec3_t(T x_, T y_, T z_) : r(x_), g(y_), b(z_) {}
+			T& components(size_t i) { return ((T*)this)[i]; }
+			const T& components(size_t i) const { return ((T*)this)[i]; }
+			T r, g, b;
+		};
+
+		typedef vec3_t<float> vec3f_t;
+		typedef vec3_t<int> vec3i_t;
+
+		template <typename T>
+		struct vec4_t
+		{
+		public:
+			vec4_t() {}
+			vec4_t(T x_, T y_, T z_, T w_) : r(x_), g(y_), b(z_), a(w_) {}
+			vec3_t<T> rgb() const { return vec3_t<T>(r, g, b); }
+			T components(size_t i) { return ((T*)this)[i]; }
+			T r, g, b, a;
+		};
+
+		typedef vec4_t<float> vec4f_t;
+		typedef vec4_t<int> vec4i_t;
+
+		template <typename T>
+		vec3_t<T> operator+(vec3_t<T> a, vec3_t<T> b) {
+			vec3_t<T> result;
+			result.r = a.r + b.r;
+			result.g = a.g + b.g;
+			result.b = a.b + b.b;
+			return result;
+		}
+
+		template <typename T>
+		vec3_t<T> operator-(vec3_t<T> a, vec3_t<T> b) {
+			vec3_t<T> result;
+			result.r = a.r - b.r;
+			result.g = a.g - b.g;
+			result.b = a.b - b.b;
+			return result;
+		}
+
+		template <typename T>
+		vec4_t<T> operator-(vec4_t<T> a, vec4_t<T> b) {
+			vec4_t<T> result;
+			result.r = a.r - b.r;
+			result.g = a.g - b.g;
+			result.b = a.b - b.b;
+			result.a = a.a - b.a;
+			return result;
+		}
+
+		template <typename T>
+		vec3_t<T> operator*(vec3_t<T> a, vec3_t<T> b) {
+			vec3_t<T> result;
+			result.r = a.r * b.r;
+			result.g = a.g * b.g;
+			result.b = a.b * b.b;
+			return result;
+		}
+
+		template <typename T>
+		vec3_t<T> operator*(vec3_t<T> a, T b)
+		{
+			vec3_t<T> result;
+			result.r = a.r * b;
+			result.g = a.g * b;
+			result.b = a.b * b;
+			return result;
+		}
+
+		template <typename T>
+		vec3_t<T> operator/(vec3_t<T> a, T b)
+		{
+			vec3_t<T> result;
+			result.r = a.r / b;
+			result.g = a.g / b;
+			result.b = a.b / b;
+			return result;
+		}
+
+		template <typename T>
+		vec3_t<T> operator/(vec3_t<T> a, vec3_t<T> b) {
+			vec3_t<T> result;
+			result.x = a.x / b.x;
+			result.y = a.y / b.y;
+			result.z = a.z / b.z;
+			return result;
+		}
+
+		template <typename T>
+		bool operator==(vec3_t<T> a, vec3_t<T> b)
+		{
+			return a.r == b.r && a.g == b.g && a.b == b.b;
+		}
+
+		template <typename T>
+		bool operator==(vec4_t<T> a, vec4_t<T> b)
+		{
+			return a.r == b.r && a.g == b.g && a.b == b.b && a.a == b.a;
+		}
+
+		template <typename T>
+		bool operator!=(vec3_t<T> a, vec3_t<T> b)
+		{
+			return a.r != b.r || a.g != b.g || a.b != b.b;
+		}
+
+		template <typename T>
+		T dot(vec3_t<T> a, vec3_t<T> b)
+		{
+			return a.r * b.r + a.g * b.g + a.b * b.b;
+		}
+		template <typename T>
+		T dot(vec4_t<T> a, vec4_t<T> b)
+		{
+			return a.r * b.r + a.g * b.g + a.b * b.b + a.a * b.a;
+		}
+
+		template <typename T>
+		T quadrance(vec3_t<T> a)
+		{
+			return dot(a, a);
+		}
+
+		template <typename T>
+		T quadrance(vec4_t<T> a)
+		{
+			return dot(a, a);
+		}
+
+		template <typename T>
+		T norm(vec3_t<T> a) {
+			return static_cast<T>(sqrt(quadrance(a)));
+		}
+
+		template <typename T>
+		T distance(vec3_t<T> a, vec3_t<T> b) {
+			return norm(a - b);
+		}
+
+		template <typename T>
+		T qd(vec3_t<T> a, vec3_t<T> b) {
+			return quadrance(a - b);
+		}
+
+		template <typename T>
+		vec3_t<T> signorm(vec3_t<T> a) {
+			T x = norm(a);
+
+			// Safety fix for degenerated cases.
+			// \todo This should be intercepted earlier.
+			// DCHECK(x != 0.0);
+			if (x == 0.0)
+				return vec3_t<T>(0, 1, 0);
+
+			return a / x;
+		}
+
+		template <typename T>
+		vec3_t<T> vecmin(vec3_t<T> a, vec3_t<T> b) {
+			vec3_t<T> result;
+			result.x = FMath::Min(a.x, b.x);
+			result.y = FMath::Min(a.y, b.y);
+			result.z = FMath::Min(a.z, b.z);
+			return result;
+		}
+
+		template <typename T>
+		vec3_t<T> vecmax(vec3_t<T> a, vec3_t<T> b) {
+			vec3_t<T> result;
+			result.x = FMath::Max(a.x, b.x);
+			result.y = FMath::Max(a.y, b.y);
+			result.z = FMath::Max(a.z, b.z);
+			return result;
+		}
+
+		template <typename T>
+		T qd_to_line(vec3_t<T> m, vec3_t<T> k, T kk, vec3_t<T> p) {
+			T t = dot(p - m, k) / kk;
+			vec3_t<T> q = k * t + m;
+			return qd(p, q);
+		}
+
+
+		inline bool is_greyscale(vec3i_t color) {
+			// integer equality is transitive
+			return color.r == color.g && color.g == color.b;
+		}
+
+		inline int luminance(vec3i_t color) {
+			return (color.r + color.g + color.b) / 3;
+		}
+
+		inline bool approx_equal(vec3i_t a, vec3i_t b) {
+			return quadrance(a - b) <= APPROX_COLOR_EPSILON;
+		}
+
+		inline bool approx_equal(vec4i_t a, vec4i_t b) {
+			return quadrance(a - b) <= APPROX_COLOR_EPSILON;
+		}
+
+		inline vec3i_t clamp_rgb(vec3i_t color) {
+			vec3i_t result;
+			result.r = FMath::Min(255, FMath::Max(0, color.r));
+			result.g = FMath::Min(255, FMath::Max(0, color.g));
+			result.b = FMath::Min(255, FMath::Max(0, color.b));
+			return result;
+		}
+
+		inline vec3f_t clamp_rgb(vec3f_t color) {
+			vec3f_t result;
+			result.r = FMath::Min(255.0f, FMath::Max(0.0f, color.r));
+			result.g = FMath::Min(255.0f, FMath::Max(0.0f, color.g));
+			result.b = FMath::Min(255.0f, FMath::Max(0.0f, color.b));
+			return result;
+		}
+
+		inline bool is_rgb(float color) {
+			return color >= 0.0f && color <= 255.0f;
+		}
+
+		inline bool is_rgb(vec3f_t color) {
+			return is_rgb(color.r) && is_rgb(color.g) && is_rgb(color.b);
+		}
+
+		inline vec3i_t floor(vec3f_t color) {
+			vec3i_t result;
+			result.r = static_cast<int>(std::floor(color.r));
+			result.g = static_cast<int>(std::floor(color.g));
+			result.b = static_cast<int>(std::floor(color.b));
+			return result;
+		}
+
+		inline vec3i_t round(vec3f_t color) {
+			vec3i_t result;
+			result.r = static_cast<int>(std::round(color.r));
+			result.g = static_cast<int>(std::round(color.g));
+			result.b = static_cast<int>(std::round(color.b));
+			return result;
+		}
+
+		inline vec4i_t round(vec4f_t color) {
+			vec4i_t result;
+			result.r = static_cast<int>(std::round(color.r));
+			result.g = static_cast<int>(std::round(color.g));
+			result.b = static_cast<int>(std::round(color.b));
+			result.a = static_cast<int>(std::round(color.a));
+			return result;
+		}
+
+		inline vec3i_t to_vec3i(unorm8_t color) {
+			vec3i_t result;
+			result.r = color.channels.r;
+			result.g = color.channels.g;
+			result.b = color.channels.b;
+			return result;
+		}
+
+		inline vec4i_t to_vec4i(unorm8_t color) {
+			vec4i_t result;
+			result.r = color.channels.r;
+			result.g = color.channels.g;
+			result.b = color.channels.b;
+			result.a = color.channels.a;
+			return result;
+		}
+
+		inline vec3i_t to_vec3i(vec3f_t color) {
+			vec3i_t result;
+			result.r = static_cast<int>(color.r);
+			result.g = static_cast<int>(color.g);
+			result.b = static_cast<int>(color.b);
+			return result;
+		}
+
+		inline vec3f_t to_vec3f(unorm8_t color) {
+			vec3f_t result;
+			result.r = color.channels.r;
+			result.g = color.channels.g;
+			result.b = color.channels.b;
+			return result;
+		}
+
+		inline vec3f_t to_vec3f(vec3i_t color) {
+			vec3f_t result;
+			result.r = static_cast<float>(color.r);
+			result.g = static_cast<float>(color.g);
+			result.b = static_cast<float>(color.b);
+			return result;
+		}
+
+		inline unorm8_t to_unorm8(vec3i_t color) {
+			unorm8_t result;
+			result.channels.r = static_cast<uint8>(color.r);
+			result.channels.g = static_cast<uint8>(color.g);
+			result.channels.b = static_cast<uint8>(color.b);
+			result.channels.a = 255;
+			return result;
+		}
+
+		inline unorm8_t to_unorm8(vec4i_t color) {
+			unorm8_t result;
+			result.channels.r = static_cast<uint8>(color.r);
+			result.channels.g = static_cast<uint8>(color.g);
+			result.channels.b = static_cast<uint8>(color.b);
+			result.channels.a = static_cast<uint8>(color.a);
+			return result;
+		}
+
+		inline unorm16_t unorm8_to_unorm16(unorm8_t c8) {
+			// (x / 255) * (2^16-1) = x * 65535 / 255 = x * 257
+			unorm16_t result;
+			result.channels.r = static_cast<uint16>(c8.channels.r * 257);
+			result.channels.g = static_cast<uint16>(c8.channels.g * 257);
+			result.channels.b = static_cast<uint16>(c8.channels.b * 257);
+			result.channels.a = static_cast<uint16>(c8.channels.a * 257);
+			return result;
+		}
+
+		struct PhysicalBlock
+		{
+			uint8 data[BLOCK_BYTES];
+		};
+
+		inline void void_extent_to_physical(unorm16_t color, PhysicalBlock* pb)
+		{
+			pb->data[0] = 0xFC;
+			pb->data[1] = 0xFD;
+			pb->data[2] = 0xFF;
+			pb->data[3] = 0xFF;
+			pb->data[4] = 0xFF;
+			pb->data[5] = 0xFF;
+			pb->data[6] = 0xFF;
+			pb->data[7] = 0xFF;
+
+			setbytes2(pb->data, 8, color.channels.r);
+			setbytes2(pb->data, 10, color.channels.g);
+			setbytes2(pb->data, 12, color.channels.b);
+			setbytes2(pb->data, 14, color.channels.a);
+		}
+
+		enum color_endpoint_mode_t
+		{
+			CEM_LDR_LUMINANCE_DIRECT = 0,
+			CEM_LDR_LUMINANCE_BASE_OFFSET = 1,
+			CEM_HDR_LUMINANCE_LARGE_RANGE = 2,
+			CEM_HDR_LUMINANCE_SMALL_RANGE = 3,
+			CEM_LDR_LUMINANCE_ALPHA_DIRECT = 4,
+			CEM_LDR_LUMINANCE_ALPHA_BASE_OFFSET = 5,
+			CEM_LDR_RGB_BASE_SCALE = 6,
+			CEM_HDR_RGB_BASE_SCALE = 7,
+			CEM_LDR_RGB_DIRECT = 8,
+			CEM_LDR_RGB_BASE_OFFSET = 9,
+			CEM_LDR_RGB_BASE_SCALE_PLUS_TWO_ALPHA = 10,
+			CEM_HDR_RGB = 11,
+			CEM_LDR_RGBA_DIRECT = 12,
+			CEM_LDR_RGBA_BASE_OFFSET = 13,
+			CEM_HDR_RGB_LDR_ALPHA = 14,
+			CEM_HDR_RGB_HDR_ALPHA = 15,
+			CEM_MAX = 16
+		};
+
+
+		/**
+		 * Define normalized (starting at zero) numeric ranges that can be represented
+		 * with 8 bits or less.
+		 */
+		enum range_t
+		{
+			RANGE_2,
+			RANGE_3,
+			RANGE_4,
+			RANGE_5,
+			RANGE_6,
+			RANGE_8,
+			RANGE_10,
+			RANGE_12,
+			RANGE_16,
+			RANGE_20,
+			RANGE_24,
+			RANGE_32,
+			RANGE_40,
+			RANGE_48,
+			RANGE_64,
+			RANGE_80,
+			RANGE_96,
+			RANGE_128,
+			RANGE_160,
+			RANGE_192,
+			RANGE_256,
+			RANGE_MAX
+		};
+
+		/**
+		 * Table of maximum value for each range, minimum is always zero.
+		 */
+#ifndef NDEBUG
+		const uint8 range_max_table[RANGE_MAX] = { 1,  2,  3,  4,   5,   7,   9,
+													11, 15, 19, 23,  31,  39,  47,
+													63, 79, 95, 127, 159, 191, 255 };
+#endif
+
+
+
+
+		const uint8 integer_from_trits[3][3][3][3][3] =
+		{
+			{
+				{{{0,1,2},{4,5,6},{8,9,10}},{{16,17,18},{20,21,22},{24,25,26}},{{3,7,15},{19,23,27},{12,13,14}}},
+				{{{32,33,34},{36,37,38},{40,41,42}},{{48,49,50},{52,53,54},{56,57,58}},{{35,39,47},{51,55,59},{44,45,46}}},
+				{{{64,65,66},{68,69,70},{72,73,74}},{{80,81,82},{84,85,86},{88,89,90}},{{67,71,79},{83,87,91},{76,77,78}}}
+			},
+			{
+				{{{128,129,130},{132,133,134},{136,137,138}},{{144,145,146},{148,149,150},{152,153,154}},{{131,135,143},{147,151,155},{140,141,142}}},
+				{{{160,161,162},{164,165,166},{168,169,170}},{{176,177,178},{180,181,182},{184,185,186}},{{163,167,175},{179,183,187},{172,173,174}}},
+				{{{192,193,194},{196,197,198},{200,201,202}},{{208,209,210},{212,213,214},{216,217,218}},{{195,199,207},{211,215,219},{204,205,206}}}
+			},
+			{
+				{{{96,97,98},{100,101,102},{104,105,106}},{{112,113,114},{116,117,118},{120,121,122}},{{99,103,111},{115,119,123},{108,109,110}}},
+				{{{224,225,226},{228,229,230},{232,233,234}},{{240,241,242},{244,245,246},{248,249,250}},{{227,231,239},{243,247,251},{236,237,238}}},
+				{{{28,29,30},{60,61,62},{92,93,94}},{{156,157,158},{188,189,190},{220,221,222}},{{31,63,127},{159,191,255},{252,253,254}}}
+			}
+		};
+
+		const uint8 integer_from_quints[5][5][5] =
+		{
+			{{0,1,2,3,4},{8,9,10,11,12},{16,17,18,19,20},{24,25,26,27,28},{5,13,21,29,6}},
+			{{32,33,34,35,36},{40,41,42,43,44},{48,49,50,51,52},{56,57,58,59,60},{37,45,53,61,14}},
+			{{64,65,66,67,68},{72,73,74,75,76},{80,81,82,83,84},{88,89,90,91,92},{69,77,85,93,22}},
+			{{96,97,98,99,100},{104,105,106,107,108},{112,113,114,115,116},{120,121,122,123,124},{101,109,117,125,30}},
+			{{102,103,70,71,38},{110,111,78,79,46},{118,119,86,87,54},{126,127,94,95,62},{39,47,55,63,31}}
+		};
+
+		const int8_t color_endpoint_range_table[2][12][16] =
+		{
+			{
+				{-1,-1,-1,-1,-1,-1,-1,-1,-1,-1,-1,-1,-1,-1,-1,-1},
+				{20,20,20,20,20,20,20,20,20,20,20,20,20,20,20,20},
+				{20,20,20,20,20,20,20,20,20,20,20,20,20,20,20,20},
+				{20,20,20,20,20,20,20,20,20,20,20,20,20,20,20,20},
+				{20,20,20,20,20,20,20,20,20,20,20,20,20,20,20,20},
+				{20,20,20,20,20,20,20,20,20,20,20,20,19,19,19,19},
+				{20,20,20,20,20,20,20,20,20,20,20,20,17,17,17,17},
+				{20,20,20,20,20,20,20,20,20,20,20,20,16,16,16,16},
+				{20,20,20,20,20,20,20,20,19,19,19,19,13,13,13,13},
+				{20,20,20,20,20,20,20,20,16,16,16,16,11,11,11,11},
+				{20,20,20,20,20,20,20,20,14,14,14,14,10,10,10,10},
+				{20,20,20,20,19,19,19,19,11,11,11,11,7,7,7,7}
+			}
+			,
+			{
+				{-1,-1,-1,-1,-1,-1,-1,-1,-1,-1,-1,-1,-1,-1,-1,-1},
+				{20,20,20,20,20,20,20,20,14,14,14,14,9,9,9,9},
+				{20,20,20,20,20,20,20,20,12,12,12,12,8,8,8,8},
+				{20,20,20,20,19,19,19,19,11,11,11,11,7,7,7,7},
+				{20,20,20,20,17,17,17,17,10,10,10,10,6,6,6,6},
+				{20,20,20,20,15,15,15,15,8,8,8,8,5,5,5,5},
+				{20,20,20,20,13,13,13,13,7,7,7,7,4,4,4,4},
+				{20,20,20,20,11,11,11,11,6,6,6,6,3,3,3,3},
+				{20,20,20,20,9,9,9,9,4,4,4,4,2,2,2,2},
+				{17,17,17,17,7,7,7,7,3,3,3,3,1,1,1,1},
+				{14,14,14,14,5,5,5,5,2,2,2,2,0,0,0,0},
+				{10,10,10,10,3,3,3,3,0,0,0,0,0,0,0,0}
+			}
+		};
+
+		const uint8 color_unquantize_table[21][256] =
+		{
+			{0,255},
+			{0,128,255},
+			{0,85,170,255},
+			{0,64,128,192,255},
+			{0,255,51,204,102,153},
+			{0,36,73,109,146,182,219,255},
+			{0,255,28,227,56,199,84,171,113,142},
+			{0,255,69,186,23,232,92,163,46,209,116,139},
+			{0,17,34,51,68,85,102,119,136,153,170,187,204,221,238,255},
+			{0,255,67,188,13,242,80,175,27,228,94,161,40,215,107,148,54,201,121,134},
+			{0,255,33,222,66,189,99,156,11,244,44,211,77,178,110,145,22,233,55,200,88,167,121,134},
+			{0,8,16,24,33,41,49,57,66,74,82,90,99,107,115,123,132,140,148,156,165,173,181,189,198,206,214,222,231,239,247,255},
+			{0,255,32,223,65,190,97,158,6,249,39,216,71,184,104,151,13,242,45,210,78,177,110,145,19,236,52,203,84,171,117,138,26,229,58,197,91,164,123,132},
+			{0,255,16,239,32,223,48,207,65,190,81,174,97,158,113,142,5,250,21,234,38,217,54,201,70,185,86,169,103,152,119,136,11,244,27,228,43,212,59,196,76,179,92,163,108,147,124,131},
+			{0,4,8,12,16,20,24,28,32,36,40,44,48,52,56,60,65,69,73,77,81,85,89,93,97,101,105,109,113,117,121,125,130,134,138,142,146,150,154,158,162,166,170,174,178,182,186,190,195,199,203,207,211,215,219,223,227,231,235,239,243,247,251,255},
+			{0,255,16,239,32,223,48,207,64,191,80,175,96,159,112,143,3,252,19,236,35,220,51,204,67,188,83,172,100,155,116,139,6,249,22,233,38,217,54,201,71,184,87,168,103,152,119,136,9,246,25,230,42,213,58,197,74,181,90,165,106,149,122,133,13,242,29,226,45,210,61,194,77,178,93,162,109,146,125,130},
+			{0,255,8,247,16,239,24,231,32,223,40,215,48,207,56,199,64,191,72,183,80,175,88,167,96,159,104,151,112,143,120,135,2,253,10,245,18,237,26,229,35,220,43,212,51,204,59,196,67,188,75,180,83,172,91,164,99,156,107,148,115,140,123,132,5,250,13,242,21,234,29,226,37,218,45,210,53,202,61,194,70,185,78,177,86,169,94,161,102,153,110,145,118,137,126,129},
+			{0,2,4,6,8,10,12,14,16,18,20,22,24,26,28,30,32,34,36,38,40,42,44,46,48,50,52,54,56,58,60,62,64,66,68,70,72,74,76,78,80,82,84,86,88,90,92,94,96,98,100,102,104,106,108,110,112,114,116,118,120,122,124,126,129,131,133,135,137,139,141,143,145,147,149,151,153,155,157,159,161,163,165,167,169,171,173,175,177,179,181,183,185,187,189,191,193,195,197,199,201,203,205,207,209,211,213,215,217,219,221,223,225,227,229,231,233,235,237,239,241,243,245,247,249,251,253,255},
+			{0,255,8,247,16,239,24,231,32,223,40,215,48,207,56,199,64,191,72,183,80,175,88,167,96,159,104,151,112,143,120,135,1,254,9,246,17,238,25,230,33,222,41,214,49,206,57,198,65,190,73,182,81,174,89,166,97,158,105,150,113,142,121,134,3,252,11,244,19,236,27,228,35,220,43,212,51,204,59,196,67,188,75,180,83,172,91,164,99,156,107,148,115,140,123,132,4,251,12,243,20,235,28,227,36,219,44,211,52,203,60,195,68,187,76,179,84,171,92,163,100,155,108,147,116,139,124,131,6,249,14,241,22,233,30,225,38,217,46,209,54,201,62,193,70,185,78,177,86,169,94,161,102,153,110,145,118,137,126,129},
+			{0,255,4,251,8,247,12,243,16,239,20,235,24,231,28,227,32,223,36,219,40,215,44,211,48,207,52,203,56,199,60,195,64,191,68,187,72,183,76,179,80,175,84,171,88,167,92,163,96,159,100,155,104,151,108,147,112,143,116,139,120,135,124,131,1,254,5,250,9,246,13,242,17,238,21,234,25,230,29,226,33,222,37,218,41,214,45,210,49,206,53,202,57,198,61,194,65,190,69,186,73,182,77,178,81,174,85,170,89,166,93,162,97,158,101,154,105,150,109,146,113,142,117,138,121,134,125,130,2,253,6,249,10,245,14,241,18,237,22,233,26,229,30,225,34,221,38,217,42,213,46,209,50,205,54,201,58,197,62,193,66,189,70,185,74,181,78,177,82,173,86,169,90,165,94,161,98,157,102,153,106,149,110,145,114,141,118,137,122,133,126,129},
+			{0,1,2,3,4,5,6,7,8,9,10,11,12,13,14,15,16,17,18,19,20,21,22,23,24,25,26,27,28,29,30,31,32,33,34,35,36,37,38,39,40,41,42,43,44,45,46,47,48,49,50,51,52,53,54,55,56,57,58,59,60,61,62,63,64,65,66,67,68,69,70,71,72,73,74,75,76,77,78,79,80,81,82,83,84,85,86,87,88,89,90,91,92,93,94,95,96,97,98,99,100,101,102,103,104,105,106,107,108,109,110,111,112,113,114,115,116,117,118,119,120,121,122,123,124,125,126,127,128,129,130,131,132,133,134,135,136,137,138,139,140,141,142,143,144,145,146,147,148,149,150,151,152,153,154,155,156,157,158,159,160,161,162,163,164,165,166,167,168,169,170,171,172,173,174,175,176,177,178,179,180,181,182,183,184,185,186,187,188,189,190,191,192,193,194,195,196,197,198,199,200,201,202,203,204,205,206,207,208,209,210,211,212,213,214,215,216,217,218,219,220,221,222,223,224,225,226,227,228,229,230,231,232,233,234,235,236,237,238,239,240,241,242,243,244,245,246,247,248,249,250,251,252,253,254,255}
+		};
+
+		const uint8 color_quantize_table[21][256] =
+		{
+			{0,0,0,0,0,0,0,0,0,0,0,0,0,0,0,0,0,0,0,0,0,0,0,0,0,0,0,0,0,0,0,0,0,0,0,0,0,0,0,0,0,0,0,0,0,0,0,0,0,0,0,0,0,0,0,0,0,0,0,0,0,0,0,0,0,0,0,0,0,0,0,0,0,0,0,0,0,0,0,0,0,0,0,0,0,0,0,0,0,0,0,0,0,0,0,0,0,0,0,0,0,0,0,0,0,0,0,0,0,0,0,0,0,0,0,0,0,0,0,0,0,0,0,0,0,0,0,0,1,1,1,1,1,1,1,1,1,1,1,1,1,1,1,1,1,1,1,1,1,1,1,1,1,1,1,1,1,1,1,1,1,1,1,1,1,1,1,1,1,1,1,1,1,1,1,1,1,1,1,1,1,1,1,1,1,1,1,1,1,1,1,1,1,1,1,1,1,1,1,1,1,1,1,1,1,1,1,1,1,1,1,1,1,1,1,1,1,1,1,1,1,1,1,1,1,1,1,1,1,1,1,1,1,1,1,1,1,1,1,1,1,1,1,1,1,1,1,1,1,1,1,1,1,1,1,1},
+			{0,0,0,0,0,0,0,0,0,0,0,0,0,0,0,0,0,0,0,0,0,0,0,0,0,0,0,0,0,0,0,0,0,0,0,0,0,0,0,0,0,0,0,0,0,0,0,0,0,0,0,0,0,0,0,0,0,0,0,0,0,0,0,0,0,1,1,1,1,1,1,1,1,1,1,1,1,1,1,1,1,1,1,1,1,1,1,1,1,1,1,1,1,1,1,1,1,1,1,1,1,1,1,1,1,1,1,1,1,1,1,1,1,1,1,1,1,1,1,1,1,1,1,1,1,1,1,1,1,1,1,1,1,1,1,1,1,1,1,1,1,1,1,1,1,1,1,1,1,1,1,1,1,1,1,1,1,1,1,1,1,1,1,1,1,1,1,1,1,1,1,1,1,1,1,1,1,1,1,1,1,1,1,1,1,1,1,1,1,1,1,1,2,2,2,2,2,2,2,2,2,2,2,2,2,2,2,2,2,2,2,2,2,2,2,2,2,2,2,2,2,2,2,2,2,2,2,2,2,2,2,2,2,2,2,2,2,2,2,2,2,2,2,2,2,2,2,2,2,2,2,2,2,2,2,2},
+			{0,0,0,0,0,0,0,0,0,0,0,0,0,0,0,0,0,0,0,0,0,0,0,0,0,0,0,0,0,0,0,0,0,0,0,0,0,0,0,0,0,0,0,1,1,1,1,1,1,1,1,1,1,1,1,1,1,1,1,1,1,1,1,1,1,1,1,1,1,1,1,1,1,1,1,1,1,1,1,1,1,1,1,1,1,1,1,1,1,1,1,1,1,1,1,1,1,1,1,1,1,1,1,1,1,1,1,1,1,1,1,1,1,1,1,1,1,1,1,1,1,1,1,1,1,1,1,1,2,2,2,2,2,2,2,2,2,2,2,2,2,2,2,2,2,2,2,2,2,2,2,2,2,2,2,2,2,2,2,2,2,2,2,2,2,2,2,2,2,2,2,2,2,2,2,2,2,2,2,2,2,2,2,2,2,2,2,2,2,2,2,2,2,2,2,2,2,2,2,2,2,2,2,2,2,2,2,2,2,2,2,2,2,3,3,3,3,3,3,3,3,3,3,3,3,3,3,3,3,3,3,3,3,3,3,3,3,3,3,3,3,3,3,3,3,3,3,3,3,3,3,3,3,3,3,3},
+			{0,0,0,0,0,0,0,0,0,0,0,0,0,0,0,0,0,0,0,0,0,0,0,0,0,0,0,0,0,0,0,0,0,1,1,1,1,1,1,1,1,1,1,1,1,1,1,1,1,1,1,1,1,1,1,1,1,1,1,1,1,1,1,1,1,1,1,1,1,1,1,1,1,1,1,1,1,1,1,1,1,1,1,1,1,1,1,1,1,1,1,1,1,1,1,1,1,2,2,2,2,2,2,2,2,2,2,2,2,2,2,2,2,2,2,2,2,2,2,2,2,2,2,2,2,2,2,2,2,2,2,2,2,2,2,2,2,2,2,2,2,2,2,2,2,2,2,2,2,2,2,2,2,2,2,2,2,2,2,2,2,3,3,3,3,3,3,3,3,3,3,3,3,3,3,3,3,3,3,3,3,3,3,3,3,3,3,3,3,3,3,3,3,3,3,3,3,3,3,3,3,3,3,3,3,3,3,3,3,3,3,3,3,3,3,3,3,3,3,3,3,3,3,3,4,4,4,4,4,4,4,4,4,4,4,4,4,4,4,4,4,4,4,4,4,4,4,4,4,4,4,4,4,4,4,4},
+			{0,0,0,0,0,0,0,0,0,0,0,0,0,0,0,0,0,0,0,0,0,0,0,0,0,0,2,2,2,2,2,2,2,2,2,2,2,2,2,2,2,2,2,2,2,2,2,2,2,2,2,2,2,2,2,2,2,2,2,2,2,2,2,2,2,2,2,2,2,2,2,2,2,2,2,2,2,4,4,4,4,4,4,4,4,4,4,4,4,4,4,4,4,4,4,4,4,4,4,4,4,4,4,4,4,4,4,4,4,4,4,4,4,4,4,4,4,4,4,4,4,4,4,4,4,4,4,4,5,5,5,5,5,5,5,5,5,5,5,5,5,5,5,5,5,5,5,5,5,5,5,5,5,5,5,5,5,5,5,5,5,5,5,5,5,5,5,5,5,5,5,5,5,5,5,5,5,5,5,3,3,3,3,3,3,3,3,3,3,3,3,3,3,3,3,3,3,3,3,3,3,3,3,3,3,3,3,3,3,3,3,3,3,3,3,3,3,3,3,3,3,3,3,3,3,3,3,3,3,3,1,1,1,1,1,1,1,1,1,1,1,1,1,1,1,1,1,1,1,1,1,1,1,1,1,1},
+			{0,0,0,0,0,0,0,0,0,0,0,0,0,0,0,0,0,0,0,1,1,1,1,1,1,1,1,1,1,1,1,1,1,1,1,1,1,1,1,1,1,1,1,1,1,1,1,1,1,1,1,1,1,1,1,2,2,2,2,2,2,2,2,2,2,2,2,2,2,2,2,2,2,2,2,2,2,2,2,2,2,2,2,2,2,2,2,2,2,2,2,2,3,3,3,3,3,3,3,3,3,3,3,3,3,3,3,3,3,3,3,3,3,3,3,3,3,3,3,3,3,3,3,3,3,3,3,3,4,4,4,4,4,4,4,4,4,4,4,4,4,4,4,4,4,4,4,4,4,4,4,4,4,4,4,4,4,4,4,4,4,4,4,4,4,5,5,5,5,5,5,5,5,5,5,5,5,5,5,5,5,5,5,5,5,5,5,5,5,5,5,5,5,5,5,5,5,5,5,5,5,6,6,6,6,6,6,6,6,6,6,6,6,6,6,6,6,6,6,6,6,6,6,6,6,6,6,6,6,6,6,6,6,6,6,6,6,6,7,7,7,7,7,7,7,7,7,7,7,7,7,7,7,7,7,7},
+			{0,0,0,0,0,0,0,0,0,0,0,0,0,0,0,2,2,2,2,2,2,2,2,2,2,2,2,2,2,2,2,2,2,2,2,2,2,2,2,2,2,2,2,4,4,4,4,4,4,4,4,4,4,4,4,4,4,4,4,4,4,4,4,4,4,4,4,4,4,4,4,6,6,6,6,6,6,6,6,6,6,6,6,6,6,6,6,6,6,6,6,6,6,6,6,6,6,6,6,8,8,8,8,8,8,8,8,8,8,8,8,8,8,8,8,8,8,8,8,8,8,8,8,8,8,8,8,8,9,9,9,9,9,9,9,9,9,9,9,9,9,9,9,9,9,9,9,9,9,9,9,9,9,9,9,9,9,7,7,7,7,7,7,7,7,7,7,7,7,7,7,7,7,7,7,7,7,7,7,7,7,7,7,7,7,5,5,5,5,5,5,5,5,5,5,5,5,5,5,5,5,5,5,5,5,5,5,5,5,5,5,5,5,3,3,3,3,3,3,3,3,3,3,3,3,3,3,3,3,3,3,3,3,3,3,3,3,3,3,3,3,1,1,1,1,1,1,1,1,1,1,1,1,1,1,1},
+			{0,0,0,0,0,0,0,0,0,0,0,0,4,4,4,4,4,4,4,4,4,4,4,4,4,4,4,4,4,4,4,4,4,4,4,8,8,8,8,8,8,8,8,8,8,8,8,8,8,8,8,8,8,8,8,8,8,8,2,2,2,2,2,2,2,2,2,2,2,2,2,2,2,2,2,2,2,2,2,2,2,6,6,6,6,6,6,6,6,6,6,6,6,6,6,6,6,6,6,6,6,6,6,6,6,10,10,10,10,10,10,10,10,10,10,10,10,10,10,10,10,10,10,10,10,10,10,10,11,11,11,11,11,11,11,11,11,11,11,11,11,11,11,11,11,11,11,11,11,11,11,7,7,7,7,7,7,7,7,7,7,7,7,7,7,7,7,7,7,7,7,7,7,7,7,3,3,3,3,3,3,3,3,3,3,3,3,3,3,3,3,3,3,3,3,3,3,3,9,9,9,9,9,9,9,9,9,9,9,9,9,9,9,9,9,9,9,9,9,9,9,5,5,5,5,5,5,5,5,5,5,5,5,5,5,5,5,5,5,5,5,5,5,5,1,1,1,1,1,1,1,1,1,1,1,1},
+			{0,0,0,0,0,0,0,0,0,1,1,1,1,1,1,1,1,1,1,1,1,1,1,1,1,1,2,2,2,2,2,2,2,2,2,2,2,2,2,2,2,2,2,3,3,3,3,3,3,3,3,3,3,3,3,3,3,3,3,3,4,4,4,4,4,4,4,4,4,4,4,4,4,4,4,4,4,5,5,5,5,5,5,5,5,5,5,5,5,5,5,5,5,5,6,6,6,6,6,6,6,6,6,6,6,6,6,6,6,6,6,7,7,7,7,7,7,7,7,7,7,7,7,7,7,7,7,7,8,8,8,8,8,8,8,8,8,8,8,8,8,8,8,8,8,9,9,9,9,9,9,9,9,9,9,9,9,9,9,9,9,9,10,10,10,10,10,10,10,10,10,10,10,10,10,10,10,10,10,11,11,11,11,11,11,11,11,11,11,11,11,11,11,11,11,11,12,12,12,12,12,12,12,12,12,12,12,12,12,12,12,12,12,13,13,13,13,13,13,13,13,13,13,13,13,13,13,13,13,13,14,14,14,14,14,14,14,14,14,14,14,14,14,14,14,14,14,15,15,15,15,15,15,15,15,15},
+			{0,0,0,0,0,0,0,4,4,4,4,4,4,4,4,4,4,4,4,4,4,8,8,8,8,8,8,8,8,8,8,8,8,8,12,12,12,12,12,12,12,12,12,12,12,12,12,12,16,16,16,16,16,16,16,16,16,16,16,16,16,2,2,2,2,2,2,2,2,2,2,2,2,2,6,6,6,6,6,6,6,6,6,6,6,6,6,6,10,10,10,10,10,10,10,10,10,10,10,10,10,14,14,14,14,14,14,14,14,14,14,14,14,14,14,18,18,18,18,18,18,18,18,18,18,18,18,18,19,19,19,19,19,19,19,19,19,19,19,19,19,15,15,15,15,15,15,15,15,15,15,15,15,15,15,11,11,11,11,11,11,11,11,11,11,11,11,11,7,7,7,7,7,7,7,7,7,7,7,7,7,7,3,3,3,3,3,3,3,3,3,3,3,3,3,17,17,17,17,17,17,17,17,17,17,17,17,17,13,13,13,13,13,13,13,13,13,13,13,13,13,13,9,9,9,9,9,9,9,9,9,9,9,9,9,5,5,5,5,5,5,5,5,5,5,5,5,5,5,1,1,1,1,1,1,1},
+			{0,0,0,0,0,0,8,8,8,8,8,8,8,8,8,8,8,16,16,16,16,16,16,16,16,16,16,16,2,2,2,2,2,2,2,2,2,2,2,10,10,10,10,10,10,10,10,10,10,10,18,18,18,18,18,18,18,18,18,18,18,4,4,4,4,4,4,4,4,4,4,4,12,12,12,12,12,12,12,12,12,12,12,20,20,20,20,20,20,20,20,20,20,20,6,6,6,6,6,6,6,6,6,6,6,14,14,14,14,14,14,14,14,14,14,14,22,22,22,22,22,22,22,22,22,22,22,22,23,23,23,23,23,23,23,23,23,23,23,23,15,15,15,15,15,15,15,15,15,15,15,7,7,7,7,7,7,7,7,7,7,7,21,21,21,21,21,21,21,21,21,21,21,13,13,13,13,13,13,13,13,13,13,13,5,5,5,5,5,5,5,5,5,5,5,19,19,19,19,19,19,19,19,19,19,19,11,11,11,11,11,11,11,11,11,11,11,3,3,3,3,3,3,3,3,3,3,3,17,17,17,17,17,17,17,17,17,17,17,9,9,9,9,9,9,9,9,9,9,9,1,1,1,1,1,1},
+			{0,0,0,0,0,1,1,1,1,1,1,1,1,2,2,2,2,2,2,2,2,3,3,3,3,3,3,3,3,4,4,4,4,4,4,4,4,4,5,5,5,5,5,5,5,5,6,6,6,6,6,6,6,6,7,7,7,7,7,7,7,7,8,8,8,8,8,8,8,8,8,9,9,9,9,9,9,9,9,10,10,10,10,10,10,10,10,11,11,11,11,11,11,11,11,12,12,12,12,12,12,12,12,12,13,13,13,13,13,13,13,13,14,14,14,14,14,14,14,14,15,15,15,15,15,15,15,15,16,16,16,16,16,16,16,16,16,17,17,17,17,17,17,17,17,18,18,18,18,18,18,18,18,19,19,19,19,19,19,19,19,20,20,20,20,20,20,20,20,20,21,21,21,21,21,21,21,21,22,22,22,22,22,22,22,22,23,23,23,23,23,23,23,23,24,24,24,24,24,24,24,24,24,25,25,25,25,25,25,25,25,26,26,26,26,26,26,26,26,27,27,27,27,27,27,27,27,28,28,28,28,28,28,28,28,28,29,29,29,29,29,29,29,29,30,30,30,30,30,30,30,30,31,31,31,31},
+			{0,0,0,0,8,8,8,8,8,8,16,16,16,16,16,16,16,24,24,24,24,24,24,32,32,32,32,32,32,2,2,2,2,2,2,2,10,10,10,10,10,10,10,18,18,18,18,18,18,26,26,26,26,26,26,26,34,34,34,34,34,34,4,4,4,4,4,4,4,12,12,12,12,12,12,20,20,20,20,20,20,20,28,28,28,28,28,28,36,36,36,36,36,36,6,6,6,6,6,6,6,14,14,14,14,14,14,14,22,22,22,22,22,22,30,30,30,30,30,30,30,38,38,38,38,38,38,38,39,39,39,39,39,39,39,31,31,31,31,31,31,31,23,23,23,23,23,23,15,15,15,15,15,15,15,7,7,7,7,7,7,7,37,37,37,37,37,37,29,29,29,29,29,29,21,21,21,21,21,21,21,13,13,13,13,13,13,5,5,5,5,5,5,5,35,35,35,35,35,35,27,27,27,27,27,27,27,19,19,19,19,19,19,11,11,11,11,11,11,11,3,3,3,3,3,3,3,33,33,33,33,33,33,25,25,25,25,25,25,17,17,17,17,17,17,17,9,9,9,9,9,9,1,1,1,1},
+			{0,0,0,16,16,16,16,16,16,32,32,32,32,32,2,2,2,2,2,18,18,18,18,18,18,34,34,34,34,34,4,4,4,4,4,4,20,20,20,20,20,36,36,36,36,36,6,6,6,6,6,6,22,22,22,22,22,38,38,38,38,38,8,8,8,8,8,8,24,24,24,24,24,24,40,40,40,40,40,10,10,10,10,10,26,26,26,26,26,26,42,42,42,42,42,12,12,12,12,12,12,28,28,28,28,28,44,44,44,44,44,14,14,14,14,14,14,30,30,30,30,30,46,46,46,46,46,46,47,47,47,47,47,47,31,31,31,31,31,15,15,15,15,15,15,45,45,45,45,45,29,29,29,29,29,13,13,13,13,13,13,43,43,43,43,43,27,27,27,27,27,27,11,11,11,11,11,41,41,41,41,41,25,25,25,25,25,25,9,9,9,9,9,9,39,39,39,39,39,23,23,23,23,23,7,7,7,7,7,7,37,37,37,37,37,21,21,21,21,21,5,5,5,5,5,5,35,35,35,35,35,19,19,19,19,19,19,3,3,3,3,3,33,33,33,33,33,17,17,17,17,17,17,1,1,1},
+			{0,0,0,1,1,1,1,2,2,2,2,3,3,3,3,4,4,4,4,5,5,5,5,6,6,6,6,7,7,7,7,8,8,8,8,9,9,9,9,10,10,10,10,11,11,11,11,12,12,12,12,13,13,13,13,14,14,14,14,15,15,15,15,16,16,16,16,16,17,17,17,17,18,18,18,18,19,19,19,19,20,20,20,20,21,21,21,21,22,22,22,22,23,23,23,23,24,24,24,24,25,25,25,25,26,26,26,26,27,27,27,27,28,28,28,28,29,29,29,29,30,30,30,30,31,31,31,31,32,32,32,32,32,33,33,33,33,34,34,34,34,35,35,35,35,36,36,36,36,37,37,37,37,38,38,38,38,39,39,39,39,40,40,40,40,41,41,41,41,42,42,42,42,43,43,43,43,44,44,44,44,45,45,45,45,46,46,46,46,47,47,47,47,48,48,48,48,48,49,49,49,49,50,50,50,50,51,51,51,51,52,52,52,52,53,53,53,53,54,54,54,54,55,55,55,55,56,56,56,56,57,57,57,57,58,58,58,58,59,59,59,59,60,60,60,60,61,61,61,61,62,62,62,62,63,63},
+			{0,0,16,16,16,32,32,32,48,48,48,48,64,64,64,2,2,2,18,18,18,34,34,34,50,50,50,50,66,66,66,4,4,4,20,20,20,36,36,36,36,52,52,52,68,68,68,6,6,6,22,22,22,38,38,38,38,54,54,54,70,70,70,8,8,8,24,24,24,24,40,40,40,56,56,56,72,72,72,10,10,10,26,26,26,26,42,42,42,58,58,58,74,74,74,12,12,12,12,28,28,28,44,44,44,60,60,60,76,76,76,14,14,14,14,30,30,30,46,46,46,62,62,62,78,78,78,78,79,79,79,79,63,63,63,47,47,47,31,31,31,15,15,15,15,77,77,77,61,61,61,45,45,45,29,29,29,13,13,13,13,75,75,75,59,59,59,43,43,43,27,27,27,27,11,11,11,73,73,73,57,57,57,41,41,41,25,25,25,25,9,9,9,71,71,71,55,55,55,39,39,39,39,23,23,23,7,7,7,69,69,69,53,53,53,37,37,37,37,21,21,21,5,5,5,67,67,67,51,51,51,51,35,35,35,19,19,19,3,3,3,65,65,65,49,49,49,49,33,33,33,17,17,17,1,1},
+			{0,0,32,32,64,64,64,2,2,2,34,34,66,66,66,4,4,4,36,36,68,68,68,6,6,6,38,38,70,70,70,8,8,8,40,40,40,72,72,10,10,10,42,42,42,74,74,12,12,12,44,44,44,76,76,14,14,14,46,46,46,78,78,16,16,16,48,48,48,80,80,18,18,18,50,50,50,82,82,20,20,20,52,52,52,84,84,22,22,22,54,54,54,86,86,24,24,24,56,56,56,88,88,26,26,26,58,58,58,90,90,28,28,28,60,60,60,92,92,30,30,30,62,62,62,94,94,94,95,95,95,63,63,63,31,31,31,93,93,61,61,61,29,29,29,91,91,59,59,59,27,27,27,89,89,57,57,57,25,25,25,87,87,55,55,55,23,23,23,85,85,53,53,53,21,21,21,83,83,51,51,51,19,19,19,81,81,49,49,49,17,17,17,79,79,47,47,47,15,15,15,77,77,45,45,45,13,13,13,75,75,43,43,43,11,11,11,73,73,41,41,41,9,9,9,71,71,71,39,39,7,7,7,69,69,69,37,37,5,5,5,67,67,67,35,35,3,3,3,65,65,65,33,33,1,1},
+			{0,0,1,1,2,2,3,3,4,4,5,5,6,6,7,7,8,8,9,9,10,10,11,11,12,12,13,13,14,14,15,15,16,16,17,17,18,18,19,19,20,20,21,21,22,22,23,23,24,24,25,25,26,26,27,27,28,28,29,29,30,30,31,31,32,32,33,33,34,34,35,35,36,36,37,37,38,38,39,39,40,40,41,41,42,42,43,43,44,44,45,45,46,46,47,47,48,48,49,49,50,50,51,51,52,52,53,53,54,54,55,55,56,56,57,57,58,58,59,59,60,60,61,61,62,62,63,63,64,64,64,65,65,66,66,67,67,68,68,69,69,70,70,71,71,72,72,73,73,74,74,75,75,76,76,77,77,78,78,79,79,80,80,81,81,82,82,83,83,84,84,85,85,86,86,87,87,88,88,89,89,90,90,91,91,92,92,93,93,94,94,95,95,96,96,97,97,98,98,99,99,100,100,101,101,102,102,103,103,104,104,105,105,106,106,107,107,108,108,109,109,110,110,111,111,112,112,113,113,114,114,115,115,116,116,117,117,118,118,119,119,120,120,121,121,122,122,123,123,124,124,125,125,126,126,127},
+			{0,32,32,64,96,96,128,2,2,34,34,66,98,98,130,4,4,36,36,68,100,100,132,6,6,38,38,70,102,102,134,8,8,40,40,72,104,104,136,10,10,42,42,74,106,106,138,12,12,44,44,76,108,108,140,14,14,46,46,78,110,110,142,16,16,48,48,80,112,112,144,18,18,50,50,82,114,114,146,20,20,52,52,84,116,116,148,22,22,54,54,86,118,118,150,24,24,56,56,88,120,120,152,26,26,58,58,90,122,122,154,28,28,60,60,92,124,124,156,30,30,62,62,94,126,126,158,158,159,159,127,127,95,63,63,31,31,157,125,125,93,61,61,29,29,155,123,123,91,59,59,27,27,153,121,121,89,57,57,25,25,151,119,119,87,55,55,23,23,149,117,117,85,53,53,21,21,147,115,115,83,51,51,19,19,145,113,113,81,49,49,17,17,143,111,111,79,47,47,15,15,141,109,109,77,45,45,13,13,139,107,107,75,43,43,11,11,137,105,105,73,41,41,9,9,135,103,103,71,39,39,7,7,133,101,101,69,37,37,5,5,131,99,99,67,35,35,3,3,129,97,97,65,33,33,1},
+			{0,64,128,2,2,66,130,4,4,68,132,6,6,70,134,8,8,72,136,10,10,74,138,12,12,76,140,14,14,78,142,16,16,80,144,18,18,82,146,20,20,84,148,22,22,86,150,24,24,88,152,26,26,90,154,28,28,92,156,30,30,94,158,32,32,96,160,34,34,98,162,36,36,100,164,38,38,102,166,40,40,104,168,42,42,106,170,44,44,108,172,46,46,110,174,48,48,112,176,50,50,114,178,52,52,116,180,54,54,118,182,56,56,120,184,58,58,122,186,60,60,124,188,62,62,126,190,190,191,191,127,63,63,189,125,61,61,187,123,59,59,185,121,57,57,183,119,55,55,181,117,53,53,179,115,51,51,177,113,49,49,175,111,47,47,173,109,45,45,171,107,43,43,169,105,41,41,167,103,39,39,165,101,37,37,163,99,35,35,161,97,33,33,159,95,31,31,157,93,29,29,155,91,27,27,153,89,25,25,151,87,23,23,149,85,21,21,147,83,19,19,145,81,17,17,143,79,15,15,141,77,13,13,139,75,11,11,137,73,9,9,135,71,7,7,133,69,5,5,131,67,3,3,129,65,1},
+			{0,1,2,3,4,5,6,7,8,9,10,11,12,13,14,15,16,17,18,19,20,21,22,23,24,25,26,27,28,29,30,31,32,33,34,35,36,37,38,39,40,41,42,43,44,45,46,47,48,49,50,51,52,53,54,55,56,57,58,59,60,61,62,63,64,65,66,67,68,69,70,71,72,73,74,75,76,77,78,79,80,81,82,83,84,85,86,87,88,89,90,91,92,93,94,95,96,97,98,99,100,101,102,103,104,105,106,107,108,109,110,111,112,113,114,115,116,117,118,119,120,121,122,123,124,125,126,127,128,129,130,131,132,133,134,135,136,137,138,139,140,141,142,143,144,145,146,147,148,149,150,151,152,153,154,155,156,157,158,159,160,161,162,163,164,165,166,167,168,169,170,171,172,173,174,175,176,177,178,179,180,181,182,183,184,185,186,187,188,189,190,191,192,193,194,195,196,197,198,199,200,201,202,203,204,205,206,207,208,209,210,211,212,213,214,215,216,217,218,219,220,221,222,223,224,225,226,227,228,229,230,231,232,233,234,235,236,237,238,239,240,241,242,243,244,245,246,247,248,249,250,251,252,253,254,255}
+		};
+
+
+		const uint8 weight_quantize_table[12][1025] = {
+			{0, 0, 0, 0, 0, 0, 0, 0, 0, 0, 0, 0, 0, 0, 0, 0, 0, 0, 0, 0, 0, 0, 0, 0, 0,
+			 0, 0, 0, 0, 0, 0, 0, 0, 0, 0, 0, 0, 0, 0, 0, 0, 0, 0, 0, 0, 0, 0, 0, 0, 0,
+			 0, 0, 0, 0, 0, 0, 0, 0, 0, 0, 0, 0, 0, 0, 0, 0, 0, 0, 0, 0, 0, 0, 0, 0, 0,
+			 0, 0, 0, 0, 0, 0, 0, 0, 0, 0, 0, 0, 0, 0, 0, 0, 0, 0, 0, 0, 0, 0, 0, 0, 0,
+			 0, 0, 0, 0, 0, 0, 0, 0, 0, 0, 0, 0, 0, 0, 0, 0, 0, 0, 0, 0, 0, 0, 0, 0, 0,
+			 0, 0, 0, 0, 0, 0, 0, 0, 0, 0, 0, 0, 0, 0, 0, 0, 0, 0, 0, 0, 0, 0, 0, 0, 0,
+			 0, 0, 0, 0, 0, 0, 0, 0, 0, 0, 0, 0, 0, 0, 0, 0, 0, 0, 0, 0, 0, 0, 0, 0, 0,
+			 0, 0, 0, 0, 0, 0, 0, 0, 0, 0, 0, 0, 0, 0, 0, 0, 0, 0, 0, 0, 0, 0, 0, 0, 0,
+			 0, 0, 0, 0, 0, 0, 0, 0, 0, 0, 0, 0, 0, 0, 0, 0, 0, 0, 0, 0, 0, 0, 0, 0, 0,
+			 0, 0, 0, 0, 0, 0, 0, 0, 0, 0, 0, 0, 0, 0, 0, 0, 0, 0, 0, 0, 0, 0, 0, 0, 0,
+			 0, 0, 0, 0, 0, 0, 0, 0, 0, 0, 0, 0, 0, 0, 0, 0, 0, 0, 0, 0, 0, 0, 0, 0, 0,
+			 0, 0, 0, 0, 0, 0, 0, 0, 0, 0, 0, 0, 0, 0, 0, 0, 0, 0, 0, 0, 0, 0, 0, 0, 0,
+			 0, 0, 0, 0, 0, 0, 0, 0, 0, 0, 0, 0, 0, 0, 0, 0, 0, 0, 0, 0, 0, 0, 0, 0, 0,
+			 0, 0, 0, 0, 0, 0, 0, 0, 0, 0, 0, 0, 0, 0, 0, 0, 0, 0, 0, 0, 0, 0, 0, 0, 0,
+			 0, 0, 0, 0, 0, 0, 0, 0, 0, 0, 0, 0, 0, 0, 0, 0, 0, 0, 0, 0, 0, 0, 0, 0, 0,
+			 0, 0, 0, 0, 0, 0, 0, 0, 0, 0, 0, 0, 0, 0, 0, 0, 0, 0, 0, 0, 0, 0, 0, 0, 0,
+			 0, 0, 0, 0, 0, 0, 0, 0, 0, 0, 0, 0, 0, 0, 0, 0, 0, 0, 0, 0, 0, 0, 0, 0, 0,
+			 0, 0, 0, 0, 0, 0, 0, 0, 0, 0, 0, 0, 0, 0, 0, 0, 0, 0, 0, 0, 0, 0, 0, 0, 0,
+			 0, 0, 0, 0, 0, 0, 0, 0, 0, 0, 0, 0, 0, 0, 0, 0, 0, 0, 0, 0, 0, 0, 0, 0, 0,
+			 0, 0, 0, 0, 0, 0, 0, 0, 0, 0, 0, 0, 0, 0, 0, 0, 0, 0, 0, 0, 0, 0, 0, 0, 0,
+			 0, 0, 0, 0, 0, 0, 0, 0, 0, 0, 0, 0, 0, 1, 1, 1, 1, 1, 1, 1, 1, 1, 1, 1, 1,
+			 1, 1, 1, 1, 1, 1, 1, 1, 1, 1, 1, 1, 1, 1, 1, 1, 1, 1, 1, 1, 1, 1, 1, 1, 1,
+			 1, 1, 1, 1, 1, 1, 1, 1, 1, 1, 1, 1, 1, 1, 1, 1, 1, 1, 1, 1, 1, 1, 1, 1, 1,
+			 1, 1, 1, 1, 1, 1, 1, 1, 1, 1, 1, 1, 1, 1, 1, 1, 1, 1, 1, 1, 1, 1, 1, 1, 1,
+			 1, 1, 1, 1, 1, 1, 1, 1, 1, 1, 1, 1, 1, 1, 1, 1, 1, 1, 1, 1, 1, 1, 1, 1, 1,
+			 1, 1, 1, 1, 1, 1, 1, 1, 1, 1, 1, 1, 1, 1, 1, 1, 1, 1, 1, 1, 1, 1, 1, 1, 1,
+			 1, 1, 1, 1, 1, 1, 1, 1, 1, 1, 1, 1, 1, 1, 1, 1, 1, 1, 1, 1, 1, 1, 1, 1, 1,
+			 1, 1, 1, 1, 1, 1, 1, 1, 1, 1, 1, 1, 1, 1, 1, 1, 1, 1, 1, 1, 1, 1, 1, 1, 1,
+			 1, 1, 1, 1, 1, 1, 1, 1, 1, 1, 1, 1, 1, 1, 1, 1, 1, 1, 1, 1, 1, 1, 1, 1, 1,
+			 1, 1, 1, 1, 1, 1, 1, 1, 1, 1, 1, 1, 1, 1, 1, 1, 1, 1, 1, 1, 1, 1, 1, 1, 1,
+			 1, 1, 1, 1, 1, 1, 1, 1, 1, 1, 1, 1, 1, 1, 1, 1, 1, 1, 1, 1, 1, 1, 1, 1, 1,
+			 1, 1, 1, 1, 1, 1, 1, 1, 1, 1, 1, 1, 1, 1, 1, 1, 1, 1, 1, 1, 1, 1, 1, 1, 1,
+			 1, 1, 1, 1, 1, 1, 1, 1, 1, 1, 1, 1, 1, 1, 1, 1, 1, 1, 1, 1, 1, 1, 1, 1, 1,
+			 1, 1, 1, 1, 1, 1, 1, 1, 1, 1, 1, 1, 1, 1, 1, 1, 1, 1, 1, 1, 1, 1, 1, 1, 1,
+			 1, 1, 1, 1, 1, 1, 1, 1, 1, 1, 1, 1, 1, 1, 1, 1, 1, 1, 1, 1, 1, 1, 1, 1, 1,
+			 1, 1, 1, 1, 1, 1, 1, 1, 1, 1, 1, 1, 1, 1, 1, 1, 1, 1, 1, 1, 1, 1, 1, 1, 1,
+			 1, 1, 1, 1, 1, 1, 1, 1, 1, 1, 1, 1, 1, 1, 1, 1, 1, 1, 1, 1, 1, 1, 1, 1, 1,
+			 1, 1, 1, 1, 1, 1, 1, 1, 1, 1, 1, 1, 1, 1, 1, 1, 1, 1, 1, 1, 1, 1, 1, 1, 1,
+			 1, 1, 1, 1, 1, 1, 1, 1, 1, 1, 1, 1, 1, 1, 1, 1, 1, 1, 1, 1, 1, 1, 1, 1, 1,
+			 1, 1, 1, 1, 1, 1, 1, 1, 1, 1, 1, 1, 1, 1, 1, 1, 1, 1, 1, 1, 1, 1, 1, 1, 1,
+			 1, 1, 1, 1, 1, 1, 1, 1, 1, 1, 1, 1, 1, 1, 1, 1, 1, 1, 1, 1, 1, 1, 1, 1, 1},
+			{0, 0, 0, 0, 0, 0, 0, 0, 0, 0, 0, 0, 0, 0, 0, 0, 0, 0, 0, 0, 0, 0, 0, 0, 0,
+			 0, 0, 0, 0, 0, 0, 0, 0, 0, 0, 0, 0, 0, 0, 0, 0, 0, 0, 0, 0, 0, 0, 0, 0, 0,
+			 0, 0, 0, 0, 0, 0, 0, 0, 0, 0, 0, 0, 0, 0, 0, 0, 0, 0, 0, 0, 0, 0, 0, 0, 0,
+			 0, 0, 0, 0, 0, 0, 0, 0, 0, 0, 0, 0, 0, 0, 0, 0, 0, 0, 0, 0, 0, 0, 0, 0, 0,
+			 0, 0, 0, 0, 0, 0, 0, 0, 0, 0, 0, 0, 0, 0, 0, 0, 0, 0, 0, 0, 0, 0, 0, 0, 0,
+			 0, 0, 0, 0, 0, 0, 0, 0, 0, 0, 0, 0, 0, 0, 0, 0, 0, 0, 0, 0, 0, 0, 0, 0, 0,
+			 0, 0, 0, 0, 0, 0, 0, 0, 0, 0, 0, 0, 0, 0, 0, 0, 0, 0, 0, 0, 0, 0, 0, 0, 0,
+			 0, 0, 0, 0, 0, 0, 0, 0, 0, 0, 0, 0, 0, 0, 0, 0, 0, 0, 0, 0, 0, 0, 0, 0, 0,
+			 0, 0, 0, 0, 0, 0, 0, 0, 0, 0, 0, 0, 0, 0, 0, 0, 0, 0, 0, 0, 0, 0, 0, 0, 0,
+			 0, 0, 0, 0, 0, 0, 0, 0, 0, 0, 0, 0, 0, 0, 0, 0, 0, 0, 0, 0, 0, 0, 0, 0, 0,
+			 0, 0, 0, 0, 0, 0, 0, 1, 1, 1, 1, 1, 1, 1, 1, 1, 1, 1, 1, 1, 1, 1, 1, 1, 1,
+			 1, 1, 1, 1, 1, 1, 1, 1, 1, 1, 1, 1, 1, 1, 1, 1, 1, 1, 1, 1, 1, 1, 1, 1, 1,
+			 1, 1, 1, 1, 1, 1, 1, 1, 1, 1, 1, 1, 1, 1, 1, 1, 1, 1, 1, 1, 1, 1, 1, 1, 1,
+			 1, 1, 1, 1, 1, 1, 1, 1, 1, 1, 1, 1, 1, 1, 1, 1, 1, 1, 1, 1, 1, 1, 1, 1, 1,
+			 1, 1, 1, 1, 1, 1, 1, 1, 1, 1, 1, 1, 1, 1, 1, 1, 1, 1, 1, 1, 1, 1, 1, 1, 1,
+			 1, 1, 1, 1, 1, 1, 1, 1, 1, 1, 1, 1, 1, 1, 1, 1, 1, 1, 1, 1, 1, 1, 1, 1, 1,
+			 1, 1, 1, 1, 1, 1, 1, 1, 1, 1, 1, 1, 1, 1, 1, 1, 1, 1, 1, 1, 1, 1, 1, 1, 1,
+			 1, 1, 1, 1, 1, 1, 1, 1, 1, 1, 1, 1, 1, 1, 1, 1, 1, 1, 1, 1, 1, 1, 1, 1, 1,
+			 1, 1, 1, 1, 1, 1, 1, 1, 1, 1, 1, 1, 1, 1, 1, 1, 1, 1, 1, 1, 1, 1, 1, 1, 1,
+			 1, 1, 1, 1, 1, 1, 1, 1, 1, 1, 1, 1, 1, 1, 1, 1, 1, 1, 1, 1, 1, 1, 1, 1, 1,
+			 1, 1, 1, 1, 1, 1, 1, 1, 1, 1, 1, 1, 1, 1, 1, 1, 1, 1, 1, 1, 1, 1, 1, 1, 1,
+			 1, 1, 1, 1, 1, 1, 1, 1, 1, 1, 1, 1, 1, 1, 1, 1, 1, 1, 1, 1, 1, 1, 1, 1, 1,
+			 1, 1, 1, 1, 1, 1, 1, 1, 1, 1, 1, 1, 1, 1, 1, 1, 1, 1, 1, 1, 1, 1, 1, 1, 1,
+			 1, 1, 1, 1, 1, 1, 1, 1, 1, 1, 1, 1, 1, 1, 1, 1, 1, 1, 1, 1, 1, 1, 1, 1, 1,
+			 1, 1, 1, 1, 1, 1, 1, 1, 1, 1, 1, 1, 1, 1, 1, 1, 1, 1, 1, 1, 1, 1, 1, 1, 1,
+			 1, 1, 1, 1, 1, 1, 1, 1, 1, 1, 1, 1, 1, 1, 1, 1, 1, 1, 1, 1, 1, 1, 1, 1, 1,
+			 1, 1, 1, 1, 1, 1, 1, 1, 1, 1, 1, 1, 1, 1, 1, 1, 1, 1, 1, 1, 1, 1, 1, 1, 1,
+			 1, 1, 1, 1, 1, 1, 1, 1, 1, 1, 1, 1, 1, 1, 1, 1, 1, 1, 1, 1, 1, 1, 1, 1, 1,
+			 1, 1, 1, 1, 1, 1, 1, 1, 1, 1, 1, 1, 1, 1, 1, 1, 1, 1, 1, 1, 1, 1, 1, 1, 1,
+			 1, 1, 1, 1, 1, 1, 1, 1, 1, 1, 1, 1, 1, 1, 1, 1, 1, 1, 1, 1, 1, 1, 1, 1, 1,
+			 1, 1, 1, 1, 1, 1, 1, 1, 1, 1, 1, 1, 1, 1, 1, 1, 1, 1, 1, 2, 2, 2, 2, 2, 2,
+			 2, 2, 2, 2, 2, 2, 2, 2, 2, 2, 2, 2, 2, 2, 2, 2, 2, 2, 2, 2, 2, 2, 2, 2, 2,
+			 2, 2, 2, 2, 2, 2, 2, 2, 2, 2, 2, 2, 2, 2, 2, 2, 2, 2, 2, 2, 2, 2, 2, 2, 2,
+			 2, 2, 2, 2, 2, 2, 2, 2, 2, 2, 2, 2, 2, 2, 2, 2, 2, 2, 2, 2, 2, 2, 2, 2, 2,
+			 2, 2, 2, 2, 2, 2, 2, 2, 2, 2, 2, 2, 2, 2, 2, 2, 2, 2, 2, 2, 2, 2, 2, 2, 2,
+			 2, 2, 2, 2, 2, 2, 2, 2, 2, 2, 2, 2, 2, 2, 2, 2, 2, 2, 2, 2, 2, 2, 2, 2, 2,
+			 2, 2, 2, 2, 2, 2, 2, 2, 2, 2, 2, 2, 2, 2, 2, 2, 2, 2, 2, 2, 2, 2, 2, 2, 2,
+			 2, 2, 2, 2, 2, 2, 2, 2, 2, 2, 2, 2, 2, 2, 2, 2, 2, 2, 2, 2, 2, 2, 2, 2, 2,
+			 2, 2, 2, 2, 2, 2, 2, 2, 2, 2, 2, 2, 2, 2, 2, 2, 2, 2, 2, 2, 2, 2, 2, 2, 2,
+			 2, 2, 2, 2, 2, 2, 2, 2, 2, 2, 2, 2, 2, 2, 2, 2, 2, 2, 2, 2, 2, 2, 2, 2, 2,
+			 2, 2, 2, 2, 2, 2, 2, 2, 2, 2, 2, 2, 2, 2, 2, 2, 2, 2, 2, 2, 2, 2, 2, 2, 2},
+			{0, 0, 0, 0, 0, 0, 0, 0, 0, 0, 0, 0, 0, 0, 0, 0, 0, 0, 0, 0, 0, 0, 0, 0, 0,
+			 0, 0, 0, 0, 0, 0, 0, 0, 0, 0, 0, 0, 0, 0, 0, 0, 0, 0, 0, 0, 0, 0, 0, 0, 0,
+			 0, 0, 0, 0, 0, 0, 0, 0, 0, 0, 0, 0, 0, 0, 0, 0, 0, 0, 0, 0, 0, 0, 0, 0, 0,
+			 0, 0, 0, 0, 0, 0, 0, 0, 0, 0, 0, 0, 0, 0, 0, 0, 0, 0, 0, 0, 0, 0, 0, 0, 0,
+			 0, 0, 0, 0, 0, 0, 0, 0, 0, 0, 0, 0, 0, 0, 0, 0, 0, 0, 0, 0, 0, 0, 0, 0, 0,
+			 0, 0, 0, 0, 0, 0, 0, 0, 0, 0, 0, 0, 0, 0, 0, 0, 0, 0, 0, 0, 0, 0, 0, 0, 0,
+			 0, 0, 0, 0, 0, 0, 0, 0, 0, 0, 0, 0, 0, 0, 0, 0, 0, 0, 0, 1, 1, 1, 1, 1, 1,
+			 1, 1, 1, 1, 1, 1, 1, 1, 1, 1, 1, 1, 1, 1, 1, 1, 1, 1, 1, 1, 1, 1, 1, 1, 1,
+			 1, 1, 1, 1, 1, 1, 1, 1, 1, 1, 1, 1, 1, 1, 1, 1, 1, 1, 1, 1, 1, 1, 1, 1, 1,
+			 1, 1, 1, 1, 1, 1, 1, 1, 1, 1, 1, 1, 1, 1, 1, 1, 1, 1, 1, 1, 1, 1, 1, 1, 1,
+			 1, 1, 1, 1, 1, 1, 1, 1, 1, 1, 1, 1, 1, 1, 1, 1, 1, 1, 1, 1, 1, 1, 1, 1, 1,
+			 1, 1, 1, 1, 1, 1, 1, 1, 1, 1, 1, 1, 1, 1, 1, 1, 1, 1, 1, 1, 1, 1, 1, 1, 1,
+			 1, 1, 1, 1, 1, 1, 1, 1, 1, 1, 1, 1, 1, 1, 1, 1, 1, 1, 1, 1, 1, 1, 1, 1, 1,
+			 1, 1, 1, 1, 1, 1, 1, 1, 1, 1, 1, 1, 1, 1, 1, 1, 1, 1, 1, 1, 1, 1, 1, 1, 1,
+			 1, 1, 1, 1, 1, 1, 1, 1, 1, 1, 1, 1, 1, 1, 1, 1, 1, 1, 1, 1, 1, 1, 1, 1, 1,
+			 1, 1, 1, 1, 1, 1, 1, 1, 1, 1, 1, 1, 1, 1, 1, 1, 1, 1, 1, 1, 1, 1, 1, 1, 1,
+			 1, 1, 1, 1, 1, 1, 1, 1, 1, 1, 1, 1, 1, 1, 1, 1, 1, 1, 1, 1, 1, 1, 1, 1, 1,
+			 1, 1, 1, 1, 1, 1, 1, 1, 1, 1, 1, 1, 1, 1, 1, 1, 1, 1, 1, 1, 1, 1, 1, 1, 1,
+			 1, 1, 1, 1, 1, 1, 1, 1, 1, 1, 1, 1, 1, 1, 1, 1, 1, 1, 1, 1, 1, 1, 1, 1, 1,
+			 1, 1, 1, 1, 1, 1, 1, 1, 1, 1, 1, 1, 1, 1, 1, 1, 1, 1, 1, 1, 1, 1, 1, 1, 1,
+			 1, 1, 1, 1, 1, 1, 1, 1, 1, 1, 1, 1, 1, 2, 2, 2, 2, 2, 2, 2, 2, 2, 2, 2, 2,
+			 2, 2, 2, 2, 2, 2, 2, 2, 2, 2, 2, 2, 2, 2, 2, 2, 2, 2, 2, 2, 2, 2, 2, 2, 2,
+			 2, 2, 2, 2, 2, 2, 2, 2, 2, 2, 2, 2, 2, 2, 2, 2, 2, 2, 2, 2, 2, 2, 2, 2, 2,
+			 2, 2, 2, 2, 2, 2, 2, 2, 2, 2, 2, 2, 2, 2, 2, 2, 2, 2, 2, 2, 2, 2, 2, 2, 2,
+			 2, 2, 2, 2, 2, 2, 2, 2, 2, 2, 2, 2, 2, 2, 2, 2, 2, 2, 2, 2, 2, 2, 2, 2, 2,
+			 2, 2, 2, 2, 2, 2, 2, 2, 2, 2, 2, 2, 2, 2, 2, 2, 2, 2, 2, 2, 2, 2, 2, 2, 2,
+			 2, 2, 2, 2, 2, 2, 2, 2, 2, 2, 2, 2, 2, 2, 2, 2, 2, 2, 2, 2, 2, 2, 2, 2, 2,
+			 2, 2, 2, 2, 2, 2, 2, 2, 2, 2, 2, 2, 2, 2, 2, 2, 2, 2, 2, 2, 2, 2, 2, 2, 2,
+			 2, 2, 2, 2, 2, 2, 2, 2, 2, 2, 2, 2, 2, 2, 2, 2, 2, 2, 2, 2, 2, 2, 2, 2, 2,
+			 2, 2, 2, 2, 2, 2, 2, 2, 2, 2, 2, 2, 2, 2, 2, 2, 2, 2, 2, 2, 2, 2, 2, 2, 2,
+			 2, 2, 2, 2, 2, 2, 2, 2, 2, 2, 2, 2, 2, 2, 2, 2, 2, 2, 2, 2, 2, 2, 2, 2, 2,
+			 2, 2, 2, 2, 2, 2, 2, 2, 2, 2, 2, 2, 2, 2, 2, 2, 2, 2, 2, 2, 2, 2, 2, 2, 2,
+			 2, 2, 2, 2, 2, 2, 2, 2, 2, 2, 2, 2, 2, 2, 2, 2, 2, 2, 2, 2, 2, 2, 2, 2, 2,
+			 2, 2, 2, 2, 2, 2, 2, 2, 2, 2, 2, 2, 2, 2, 2, 2, 2, 2, 2, 2, 2, 2, 2, 2, 2,
+			 2, 2, 2, 2, 2, 2, 2, 3, 3, 3, 3, 3, 3, 3, 3, 3, 3, 3, 3, 3, 3, 3, 3, 3, 3,
+			 3, 3, 3, 3, 3, 3, 3, 3, 3, 3, 3, 3, 3, 3, 3, 3, 3, 3, 3, 3, 3, 3, 3, 3, 3,
+			 3, 3, 3, 3, 3, 3, 3, 3, 3, 3, 3, 3, 3, 3, 3, 3, 3, 3, 3, 3, 3, 3, 3, 3, 3,
+			 3, 3, 3, 3, 3, 3, 3, 3, 3, 3, 3, 3, 3, 3, 3, 3, 3, 3, 3, 3, 3, 3, 3, 3, 3,
+			 3, 3, 3, 3, 3, 3, 3, 3, 3, 3, 3, 3, 3, 3, 3, 3, 3, 3, 3, 3, 3, 3, 3, 3, 3,
+			 3, 3, 3, 3, 3, 3, 3, 3, 3, 3, 3, 3, 3, 3, 3, 3, 3, 3, 3, 3, 3, 3, 3, 3, 3,
+			 3, 3, 3, 3, 3, 3, 3, 3, 3, 3, 3, 3, 3, 3, 3, 3, 3, 3, 3, 3, 3, 3, 3, 3, 3},
+			{0, 0, 0, 0, 0, 0, 0, 0, 0, 0, 0, 0, 0, 0, 0, 0, 0, 0, 0, 0, 0, 0, 0, 0, 0,
+			 0, 0, 0, 0, 0, 0, 0, 0, 0, 0, 0, 0, 0, 0, 0, 0, 0, 0, 0, 0, 0, 0, 0, 0, 0,
+			 0, 0, 0, 0, 0, 0, 0, 0, 0, 0, 0, 0, 0, 0, 0, 0, 0, 0, 0, 0, 0, 0, 0, 0, 0,
+			 0, 0, 0, 0, 0, 0, 0, 0, 0, 0, 0, 0, 0, 0, 0, 0, 0, 0, 0, 0, 0, 0, 0, 0, 0,
+			 0, 0, 0, 0, 0, 0, 0, 0, 0, 0, 0, 0, 0, 0, 0, 0, 0, 0, 0, 0, 0, 0, 0, 0, 0,
+			 0, 0, 0, 0, 1, 1, 1, 1, 1, 1, 1, 1, 1, 1, 1, 1, 1, 1, 1, 1, 1, 1, 1, 1, 1,
+			 1, 1, 1, 1, 1, 1, 1, 1, 1, 1, 1, 1, 1, 1, 1, 1, 1, 1, 1, 1, 1, 1, 1, 1, 1,
+			 1, 1, 1, 1, 1, 1, 1, 1, 1, 1, 1, 1, 1, 1, 1, 1, 1, 1, 1, 1, 1, 1, 1, 1, 1,
+			 1, 1, 1, 1, 1, 1, 1, 1, 1, 1, 1, 1, 1, 1, 1, 1, 1, 1, 1, 1, 1, 1, 1, 1, 1,
+			 1, 1, 1, 1, 1, 1, 1, 1, 1, 1, 1, 1, 1, 1, 1, 1, 1, 1, 1, 1, 1, 1, 1, 1, 1,
+			 1, 1, 1, 1, 1, 1, 1, 1, 1, 1, 1, 1, 1, 1, 1, 1, 1, 1, 1, 1, 1, 1, 1, 1, 1,
+			 1, 1, 1, 1, 1, 1, 1, 1, 1, 1, 1, 1, 1, 1, 1, 1, 1, 1, 1, 1, 1, 1, 1, 1, 1,
+			 1, 1, 1, 1, 1, 1, 1, 1, 1, 1, 1, 1, 1, 1, 1, 1, 1, 1, 1, 1, 1, 1, 1, 1, 1,
+			 1, 1, 1, 1, 1, 1, 1, 1, 1, 1, 1, 1, 1, 1, 1, 1, 1, 1, 1, 1, 1, 1, 1, 1, 1,
+			 1, 1, 1, 1, 1, 1, 1, 1, 1, 1, 1, 1, 1, 1, 1, 1, 1, 1, 1, 1, 1, 1, 1, 1, 1,
+			 1, 1, 1, 1, 1, 1, 1, 1, 1, 1, 2, 2, 2, 2, 2, 2, 2, 2, 2, 2, 2, 2, 2, 2, 2,
+			 2, 2, 2, 2, 2, 2, 2, 2, 2, 2, 2, 2, 2, 2, 2, 2, 2, 2, 2, 2, 2, 2, 2, 2, 2,
+			 2, 2, 2, 2, 2, 2, 2, 2, 2, 2, 2, 2, 2, 2, 2, 2, 2, 2, 2, 2, 2, 2, 2, 2, 2,
+			 2, 2, 2, 2, 2, 2, 2, 2, 2, 2, 2, 2, 2, 2, 2, 2, 2, 2, 2, 2, 2, 2, 2, 2, 2,
+			 2, 2, 2, 2, 2, 2, 2, 2, 2, 2, 2, 2, 2, 2, 2, 2, 2, 2, 2, 2, 2, 2, 2, 2, 2,
+			 2, 2, 2, 2, 2, 2, 2, 2, 2, 2, 2, 2, 2, 2, 2, 2, 2, 2, 2, 2, 2, 2, 2, 2, 2,
+			 2, 2, 2, 2, 2, 2, 2, 2, 2, 2, 2, 2, 2, 2, 2, 2, 2, 2, 2, 2, 2, 2, 2, 2, 2,
+			 2, 2, 2, 2, 2, 2, 2, 2, 2, 2, 2, 2, 2, 2, 2, 2, 2, 2, 2, 2, 2, 2, 2, 2, 2,
+			 2, 2, 2, 2, 2, 2, 2, 2, 2, 2, 2, 2, 2, 2, 2, 2, 2, 2, 2, 2, 2, 2, 2, 2, 2,
+			 2, 2, 2, 2, 2, 2, 2, 2, 2, 2, 2, 2, 2, 2, 2, 2, 2, 2, 2, 2, 2, 2, 2, 2, 2,
+			 2, 2, 2, 2, 2, 2, 2, 2, 2, 2, 2, 2, 2, 2, 2, 2, 3, 3, 3, 3, 3, 3, 3, 3, 3,
+			 3, 3, 3, 3, 3, 3, 3, 3, 3, 3, 3, 3, 3, 3, 3, 3, 3, 3, 3, 3, 3, 3, 3, 3, 3,
+			 3, 3, 3, 3, 3, 3, 3, 3, 3, 3, 3, 3, 3, 3, 3, 3, 3, 3, 3, 3, 3, 3, 3, 3, 3,
+			 3, 3, 3, 3, 3, 3, 3, 3, 3, 3, 3, 3, 3, 3, 3, 3, 3, 3, 3, 3, 3, 3, 3, 3, 3,
+			 3, 3, 3, 3, 3, 3, 3, 3, 3, 3, 3, 3, 3, 3, 3, 3, 3, 3, 3, 3, 3, 3, 3, 3, 3,
+			 3, 3, 3, 3, 3, 3, 3, 3, 3, 3, 3, 3, 3, 3, 3, 3, 3, 3, 3, 3, 3, 3, 3, 3, 3,
+			 3, 3, 3, 3, 3, 3, 3, 3, 3, 3, 3, 3, 3, 3, 3, 3, 3, 3, 3, 3, 3, 3, 3, 3, 3,
+			 3, 3, 3, 3, 3, 3, 3, 3, 3, 3, 3, 3, 3, 3, 3, 3, 3, 3, 3, 3, 3, 3, 3, 3, 3,
+			 3, 3, 3, 3, 3, 3, 3, 3, 3, 3, 3, 3, 3, 3, 3, 3, 3, 3, 3, 3, 3, 3, 3, 3, 3,
+			 3, 3, 3, 3, 3, 3, 3, 3, 3, 3, 3, 3, 3, 3, 3, 3, 3, 3, 3, 3, 3, 3, 3, 3, 3,
+			 3, 3, 3, 3, 3, 3, 3, 3, 3, 3, 3, 3, 3, 3, 3, 3, 3, 3, 3, 3, 3, 3, 4, 4, 4,
+			 4, 4, 4, 4, 4, 4, 4, 4, 4, 4, 4, 4, 4, 4, 4, 4, 4, 4, 4, 4, 4, 4, 4, 4, 4,
+			 4, 4, 4, 4, 4, 4, 4, 4, 4, 4, 4, 4, 4, 4, 4, 4, 4, 4, 4, 4, 4, 4, 4, 4, 4,
+			 4, 4, 4, 4, 4, 4, 4, 4, 4, 4, 4, 4, 4, 4, 4, 4, 4, 4, 4, 4, 4, 4, 4, 4, 4,
+			 4, 4, 4, 4, 4, 4, 4, 4, 4, 4, 4, 4, 4, 4, 4, 4, 4, 4, 4, 4, 4, 4, 4, 4, 4,
+			 4, 4, 4, 4, 4, 4, 4, 4, 4, 4, 4, 4, 4, 4, 4, 4, 4, 4, 4, 4, 4, 4, 4, 4, 4},
+			{0, 0, 0, 0, 0, 0, 0, 0, 0, 0, 0, 0, 0, 0, 0, 0, 0, 0, 0, 0, 0, 0, 0, 0, 0,
+			 0, 0, 0, 0, 0, 0, 0, 0, 0, 0, 0, 0, 0, 0, 0, 0, 0, 0, 0, 0, 0, 0, 0, 0, 0,
+			 0, 0, 0, 0, 0, 0, 0, 0, 0, 0, 0, 0, 0, 0, 0, 0, 0, 0, 0, 0, 0, 0, 0, 0, 0,
+			 0, 0, 0, 0, 0, 0, 0, 0, 0, 0, 0, 0, 0, 0, 0, 0, 0, 0, 0, 0, 0, 0, 2, 2, 2,
+			 2, 2, 2, 2, 2, 2, 2, 2, 2, 2, 2, 2, 2, 2, 2, 2, 2, 2, 2, 2, 2, 2, 2, 2, 2,
+			 2, 2, 2, 2, 2, 2, 2, 2, 2, 2, 2, 2, 2, 2, 2, 2, 2, 2, 2, 2, 2, 2, 2, 2, 2,
+			 2, 2, 2, 2, 2, 2, 2, 2, 2, 2, 2, 2, 2, 2, 2, 2, 2, 2, 2, 2, 2, 2, 2, 2, 2,
+			 2, 2, 2, 2, 2, 2, 2, 2, 2, 2, 2, 2, 2, 2, 2, 2, 2, 2, 2, 2, 2, 2, 2, 2, 2,
+			 2, 2, 2, 2, 2, 2, 2, 2, 2, 2, 2, 2, 2, 2, 2, 2, 2, 2, 2, 2, 2, 2, 2, 2, 2,
+			 2, 2, 2, 2, 2, 2, 2, 2, 2, 2, 2, 2, 2, 2, 2, 2, 2, 2, 2, 2, 2, 2, 2, 2, 2,
+			 2, 2, 2, 2, 2, 2, 2, 2, 2, 2, 2, 2, 2, 2, 2, 2, 2, 2, 2, 2, 2, 2, 2, 2, 2,
+			 2, 2, 2, 2, 2, 2, 2, 2, 2, 2, 2, 2, 2, 2, 2, 2, 2, 2, 2, 2, 2, 2, 4, 4, 4,
+			 4, 4, 4, 4, 4, 4, 4, 4, 4, 4, 4, 4, 4, 4, 4, 4, 4, 4, 4, 4, 4, 4, 4, 4, 4,
+			 4, 4, 4, 4, 4, 4, 4, 4, 4, 4, 4, 4, 4, 4, 4, 4, 4, 4, 4, 4, 4, 4, 4, 4, 4,
+			 4, 4, 4, 4, 4, 4, 4, 4, 4, 4, 4, 4, 4, 4, 4, 4, 4, 4, 4, 4, 4, 4, 4, 4, 4,
+			 4, 4, 4, 4, 4, 4, 4, 4, 4, 4, 4, 4, 4, 4, 4, 4, 4, 4, 4, 4, 4, 4, 4, 4, 4,
+			 4, 4, 4, 4, 4, 4, 4, 4, 4, 4, 4, 4, 4, 4, 4, 4, 4, 4, 4, 4, 4, 4, 4, 4, 4,
+			 4, 4, 4, 4, 4, 4, 4, 4, 4, 4, 4, 4, 4, 4, 4, 4, 4, 4, 4, 4, 4, 4, 4, 4, 4,
+			 4, 4, 4, 4, 4, 4, 4, 4, 4, 4, 4, 4, 4, 4, 4, 4, 4, 4, 4, 4, 4, 4, 4, 4, 4,
+			 4, 4, 4, 4, 4, 4, 4, 4, 4, 4, 4, 4, 4, 4, 4, 4, 4, 4, 4, 4, 4, 4, 4, 4, 4,
+			 4, 4, 4, 4, 4, 4, 4, 4, 4, 4, 4, 4, 4, 5, 5, 5, 5, 5, 5, 5, 5, 5, 5, 5, 5,
+			 5, 5, 5, 5, 5, 5, 5, 5, 5, 5, 5, 5, 5, 5, 5, 5, 5, 5, 5, 5, 5, 5, 5, 5, 5,
+			 5, 5, 5, 5, 5, 5, 5, 5, 5, 5, 5, 5, 5, 5, 5, 5, 5, 5, 5, 5, 5, 5, 5, 5, 5,
+			 5, 5, 5, 5, 5, 5, 5, 5, 5, 5, 5, 5, 5, 5, 5, 5, 5, 5, 5, 5, 5, 5, 5, 5, 5,
+			 5, 5, 5, 5, 5, 5, 5, 5, 5, 5, 5, 5, 5, 5, 5, 5, 5, 5, 5, 5, 5, 5, 5, 5, 5,
+			 5, 5, 5, 5, 5, 5, 5, 5, 5, 5, 5, 5, 5, 5, 5, 5, 5, 5, 5, 5, 5, 5, 5, 5, 5,
+			 5, 5, 5, 5, 5, 5, 5, 5, 5, 5, 5, 5, 5, 5, 5, 5, 5, 5, 5, 5, 5, 5, 5, 5, 5,
+			 5, 5, 5, 5, 5, 5, 5, 5, 5, 5, 5, 5, 5, 5, 5, 5, 5, 5, 5, 5, 5, 5, 5, 5, 5,
+			 5, 5, 5, 5, 5, 5, 5, 5, 5, 5, 5, 5, 5, 5, 5, 5, 5, 5, 5, 5, 5, 5, 5, 5, 5,
+			 5, 5, 5, 5, 3, 3, 3, 3, 3, 3, 3, 3, 3, 3, 3, 3, 3, 3, 3, 3, 3, 3, 3, 3, 3,
+			 3, 3, 3, 3, 3, 3, 3, 3, 3, 3, 3, 3, 3, 3, 3, 3, 3, 3, 3, 3, 3, 3, 3, 3, 3,
+			 3, 3, 3, 3, 3, 3, 3, 3, 3, 3, 3, 3, 3, 3, 3, 3, 3, 3, 3, 3, 3, 3, 3, 3, 3,
+			 3, 3, 3, 3, 3, 3, 3, 3, 3, 3, 3, 3, 3, 3, 3, 3, 3, 3, 3, 3, 3, 3, 3, 3, 3,
+			 3, 3, 3, 3, 3, 3, 3, 3, 3, 3, 3, 3, 3, 3, 3, 3, 3, 3, 3, 3, 3, 3, 3, 3, 3,
+			 3, 3, 3, 3, 3, 3, 3, 3, 3, 3, 3, 3, 3, 3, 3, 3, 3, 3, 3, 3, 3, 3, 3, 3, 3,
+			 3, 3, 3, 3, 3, 3, 3, 3, 3, 3, 3, 3, 3, 3, 3, 3, 3, 3, 3, 3, 3, 3, 3, 3, 3,
+			 3, 3, 3, 3, 3, 3, 3, 3, 3, 3, 3, 3, 3, 3, 3, 3, 3, 3, 3, 3, 3, 3, 3, 3, 3,
+			 3, 3, 3, 3, 1, 1, 1, 1, 1, 1, 1, 1, 1, 1, 1, 1, 1, 1, 1, 1, 1, 1, 1, 1, 1,
+			 1, 1, 1, 1, 1, 1, 1, 1, 1, 1, 1, 1, 1, 1, 1, 1, 1, 1, 1, 1, 1, 1, 1, 1, 1,
+			 1, 1, 1, 1, 1, 1, 1, 1, 1, 1, 1, 1, 1, 1, 1, 1, 1, 1, 1, 1, 1, 1, 1, 1, 1,
+			 1, 1, 1, 1, 1, 1, 1, 1, 1, 1, 1, 1, 1, 1, 1, 1, 1, 1, 1, 1, 1, 1, 1, 1, 1},
+			{0, 0, 0, 0, 0, 0, 0, 0, 0, 0, 0, 0, 0, 0, 0, 0, 0, 0, 0, 0, 0, 0, 0, 0, 0,
+			 0, 0, 0, 0, 0, 0, 0, 0, 0, 0, 0, 0, 0, 0, 0, 0, 0, 0, 0, 0, 0, 0, 0, 0, 0,
+			 0, 0, 0, 0, 0, 0, 0, 0, 0, 0, 0, 0, 0, 0, 0, 0, 0, 0, 0, 0, 0, 0, 0, 1, 1,
+			 1, 1, 1, 1, 1, 1, 1, 1, 1, 1, 1, 1, 1, 1, 1, 1, 1, 1, 1, 1, 1, 1, 1, 1, 1,
+			 1, 1, 1, 1, 1, 1, 1, 1, 1, 1, 1, 1, 1, 1, 1, 1, 1, 1, 1, 1, 1, 1, 1, 1, 1,
+			 1, 1, 1, 1, 1, 1, 1, 1, 1, 1, 1, 1, 1, 1, 1, 1, 1, 1, 1, 1, 1, 1, 1, 1, 1,
+			 1, 1, 1, 1, 1, 1, 1, 1, 1, 1, 1, 1, 1, 1, 1, 1, 1, 1, 1, 1, 1, 1, 1, 1, 1,
+			 1, 1, 1, 1, 1, 1, 1, 1, 1, 1, 1, 1, 1, 1, 1, 1, 1, 1, 1, 1, 1, 1, 1, 1, 1,
+			 1, 1, 1, 1, 1, 1, 1, 1, 1, 1, 1, 1, 1, 1, 1, 1, 1, 2, 2, 2, 2, 2, 2, 2, 2,
+			 2, 2, 2, 2, 2, 2, 2, 2, 2, 2, 2, 2, 2, 2, 2, 2, 2, 2, 2, 2, 2, 2, 2, 2, 2,
+			 2, 2, 2, 2, 2, 2, 2, 2, 2, 2, 2, 2, 2, 2, 2, 2, 2, 2, 2, 2, 2, 2, 2, 2, 2,
+			 2, 2, 2, 2, 2, 2, 2, 2, 2, 2, 2, 2, 2, 2, 2, 2, 2, 2, 2, 2, 2, 2, 2, 2, 2,
+			 2, 2, 2, 2, 2, 2, 2, 2, 2, 2, 2, 2, 2, 2, 2, 2, 2, 2, 2, 2, 2, 2, 2, 2, 2,
+			 2, 2, 2, 2, 2, 2, 2, 2, 2, 2, 2, 2, 2, 2, 2, 2, 2, 2, 2, 2, 2, 2, 2, 2, 2,
+			 2, 2, 2, 2, 2, 2, 2, 2, 2, 2, 2, 3, 3, 3, 3, 3, 3, 3, 3, 3, 3, 3, 3, 3, 3,
+			 3, 3, 3, 3, 3, 3, 3, 3, 3, 3, 3, 3, 3, 3, 3, 3, 3, 3, 3, 3, 3, 3, 3, 3, 3,
+			 3, 3, 3, 3, 3, 3, 3, 3, 3, 3, 3, 3, 3, 3, 3, 3, 3, 3, 3, 3, 3, 3, 3, 3, 3,
+			 3, 3, 3, 3, 3, 3, 3, 3, 3, 3, 3, 3, 3, 3, 3, 3, 3, 3, 3, 3, 3, 3, 3, 3, 3,
+			 3, 3, 3, 3, 3, 3, 3, 3, 3, 3, 3, 3, 3, 3, 3, 3, 3, 3, 3, 3, 3, 3, 3, 3, 3,
+			 3, 3, 3, 3, 3, 3, 3, 3, 3, 3, 3, 3, 3, 3, 3, 3, 3, 3, 3, 3, 3, 3, 3, 3, 3,
+			 3, 3, 3, 3, 3, 3, 3, 3, 3, 3, 3, 3, 3, 4, 4, 4, 4, 4, 4, 4, 4, 4, 4, 4, 4,
+			 4, 4, 4, 4, 4, 4, 4, 4, 4, 4, 4, 4, 4, 4, 4, 4, 4, 4, 4, 4, 4, 4, 4, 4, 4,
+			 4, 4, 4, 4, 4, 4, 4, 4, 4, 4, 4, 4, 4, 4, 4, 4, 4, 4, 4, 4, 4, 4, 4, 4, 4,
+			 4, 4, 4, 4, 4, 4, 4, 4, 4, 4, 4, 4, 4, 4, 4, 4, 4, 4, 4, 4, 4, 4, 4, 4, 4,
+			 4, 4, 4, 4, 4, 4, 4, 4, 4, 4, 4, 4, 4, 4, 4, 4, 4, 4, 4, 4, 4, 4, 4, 4, 4,
+			 4, 4, 4, 4, 4, 4, 4, 4, 4, 4, 4, 4, 4, 4, 4, 4, 4, 4, 4, 4, 4, 4, 4, 4, 4,
+			 4, 4, 4, 4, 4, 4, 4, 4, 4, 4, 4, 4, 4, 4, 4, 5, 5, 5, 5, 5, 5, 5, 5, 5, 5,
+			 5, 5, 5, 5, 5, 5, 5, 5, 5, 5, 5, 5, 5, 5, 5, 5, 5, 5, 5, 5, 5, 5, 5, 5, 5,
+			 5, 5, 5, 5, 5, 5, 5, 5, 5, 5, 5, 5, 5, 5, 5, 5, 5, 5, 5, 5, 5, 5, 5, 5, 5,
+			 5, 5, 5, 5, 5, 5, 5, 5, 5, 5, 5, 5, 5, 5, 5, 5, 5, 5, 5, 5, 5, 5, 5, 5, 5,
+			 5, 5, 5, 5, 5, 5, 5, 5, 5, 5, 5, 5, 5, 5, 5, 5, 5, 5, 5, 5, 5, 5, 5, 5, 5,
+			 5, 5, 5, 5, 5, 5, 5, 5, 5, 5, 5, 5, 5, 5, 5, 5, 5, 5, 5, 5, 5, 5, 5, 5, 5,
+			 5, 5, 5, 5, 5, 5, 5, 5, 5, 6, 6, 6, 6, 6, 6, 6, 6, 6, 6, 6, 6, 6, 6, 6, 6,
+			 6, 6, 6, 6, 6, 6, 6, 6, 6, 6, 6, 6, 6, 6, 6, 6, 6, 6, 6, 6, 6, 6, 6, 6, 6,
+			 6, 6, 6, 6, 6, 6, 6, 6, 6, 6, 6, 6, 6, 6, 6, 6, 6, 6, 6, 6, 6, 6, 6, 6, 6,
+			 6, 6, 6, 6, 6, 6, 6, 6, 6, 6, 6, 6, 6, 6, 6, 6, 6, 6, 6, 6, 6, 6, 6, 6, 6,
+			 6, 6, 6, 6, 6, 6, 6, 6, 6, 6, 6, 6, 6, 6, 6, 6, 6, 6, 6, 6, 6, 6, 6, 6, 6,
+			 6, 6, 6, 6, 6, 6, 6, 6, 6, 6, 6, 6, 6, 6, 6, 6, 6, 6, 6, 6, 6, 6, 6, 6, 6,
+			 6, 6, 6, 7, 7, 7, 7, 7, 7, 7, 7, 7, 7, 7, 7, 7, 7, 7, 7, 7, 7, 7, 7, 7, 7,
+			 7, 7, 7, 7, 7, 7, 7, 7, 7, 7, 7, 7, 7, 7, 7, 7, 7, 7, 7, 7, 7, 7, 7, 7, 7,
+			 7, 7, 7, 7, 7, 7, 7, 7, 7, 7, 7, 7, 7, 7, 7, 7, 7, 7, 7, 7, 7, 7, 7, 7, 7},
+			{0, 0, 0, 0, 0, 0, 0, 0, 0, 0, 0, 0, 0, 0, 0, 0, 0, 0, 0, 0, 0, 0, 0, 0, 0,
+			 0, 0, 0, 0, 0, 0, 0, 0, 0, 0, 0, 0, 0, 0, 0, 0, 0, 0, 0, 0, 0, 0, 0, 0, 0,
+			 0, 0, 0, 0, 0, 0, 0, 2, 2, 2, 2, 2, 2, 2, 2, 2, 2, 2, 2, 2, 2, 2, 2, 2, 2,
+			 2, 2, 2, 2, 2, 2, 2, 2, 2, 2, 2, 2, 2, 2, 2, 2, 2, 2, 2, 2, 2, 2, 2, 2, 2,
+			 2, 2, 2, 2, 2, 2, 2, 2, 2, 2, 2, 2, 2, 2, 2, 2, 2, 2, 2, 2, 2, 2, 2, 2, 2,
+			 2, 2, 2, 2, 2, 2, 2, 2, 2, 2, 2, 2, 2, 2, 2, 2, 2, 2, 2, 2, 2, 2, 2, 2, 2,
+			 2, 2, 2, 2, 2, 2, 2, 2, 2, 2, 2, 2, 2, 2, 2, 2, 2, 2, 2, 4, 4, 4, 4, 4, 4,
+			 4, 4, 4, 4, 4, 4, 4, 4, 4, 4, 4, 4, 4, 4, 4, 4, 4, 4, 4, 4, 4, 4, 4, 4, 4,
+			 4, 4, 4, 4, 4, 4, 4, 4, 4, 4, 4, 4, 4, 4, 4, 4, 4, 4, 4, 4, 4, 4, 4, 4, 4,
+			 4, 4, 4, 4, 4, 4, 4, 4, 4, 4, 4, 4, 4, 4, 4, 4, 4, 4, 4, 4, 4, 4, 4, 4, 4,
+			 4, 4, 4, 4, 4, 4, 4, 4, 4, 4, 4, 4, 4, 4, 4, 4, 4, 4, 4, 4, 4, 4, 4, 4, 4,
+			 4, 4, 4, 4, 4, 4, 6, 6, 6, 6, 6, 6, 6, 6, 6, 6, 6, 6, 6, 6, 6, 6, 6, 6, 6,
+			 6, 6, 6, 6, 6, 6, 6, 6, 6, 6, 6, 6, 6, 6, 6, 6, 6, 6, 6, 6, 6, 6, 6, 6, 6,
+			 6, 6, 6, 6, 6, 6, 6, 6, 6, 6, 6, 6, 6, 6, 6, 6, 6, 6, 6, 6, 6, 6, 6, 6, 6,
+			 6, 6, 6, 6, 6, 6, 6, 6, 6, 6, 6, 6, 6, 6, 6, 6, 6, 6, 6, 6, 6, 6, 6, 6, 6,
+			 6, 6, 6, 6, 6, 6, 6, 6, 6, 6, 6, 6, 6, 6, 6, 6, 6, 6, 8, 8, 8, 8, 8, 8, 8,
+			 8, 8, 8, 8, 8, 8, 8, 8, 8, 8, 8, 8, 8, 8, 8, 8, 8, 8, 8, 8, 8, 8, 8, 8, 8,
+			 8, 8, 8, 8, 8, 8, 8, 8, 8, 8, 8, 8, 8, 8, 8, 8, 8, 8, 8, 8, 8, 8, 8, 8, 8,
+			 8, 8, 8, 8, 8, 8, 8, 8, 8, 8, 8, 8, 8, 8, 8, 8, 8, 8, 8, 8, 8, 8, 8, 8, 8,
+			 8, 8, 8, 8, 8, 8, 8, 8, 8, 8, 8, 8, 8, 8, 8, 8, 8, 8, 8, 8, 8, 8, 8, 8, 8,
+			 8, 8, 8, 8, 8, 8, 8, 8, 8, 8, 8, 8, 8, 9, 9, 9, 9, 9, 9, 9, 9, 9, 9, 9, 9,
+			 9, 9, 9, 9, 9, 9, 9, 9, 9, 9, 9, 9, 9, 9, 9, 9, 9, 9, 9, 9, 9, 9, 9, 9, 9,
+			 9, 9, 9, 9, 9, 9, 9, 9, 9, 9, 9, 9, 9, 9, 9, 9, 9, 9, 9, 9, 9, 9, 9, 9, 9,
+			 9, 9, 9, 9, 9, 9, 9, 9, 9, 9, 9, 9, 9, 9, 9, 9, 9, 9, 9, 9, 9, 9, 9, 9, 9,
+			 9, 9, 9, 9, 9, 9, 9, 9, 9, 9, 9, 9, 9, 9, 9, 9, 9, 9, 9, 9, 9, 9, 9, 9, 9,
+			 9, 9, 9, 9, 9, 9, 9, 9, 7, 7, 7, 7, 7, 7, 7, 7, 7, 7, 7, 7, 7, 7, 7, 7, 7,
+			 7, 7, 7, 7, 7, 7, 7, 7, 7, 7, 7, 7, 7, 7, 7, 7, 7, 7, 7, 7, 7, 7, 7, 7, 7,
+			 7, 7, 7, 7, 7, 7, 7, 7, 7, 7, 7, 7, 7, 7, 7, 7, 7, 7, 7, 7, 7, 7, 7, 7, 7,
+			 7, 7, 7, 7, 7, 7, 7, 7, 7, 7, 7, 7, 7, 7, 7, 7, 7, 7, 7, 7, 7, 7, 7, 7, 7,
+			 7, 7, 7, 7, 7, 7, 7, 7, 7, 7, 7, 7, 7, 7, 7, 7, 7, 7, 7, 7, 5, 5, 5, 5, 5,
+			 5, 5, 5, 5, 5, 5, 5, 5, 5, 5, 5, 5, 5, 5, 5, 5, 5, 5, 5, 5, 5, 5, 5, 5, 5,
+			 5, 5, 5, 5, 5, 5, 5, 5, 5, 5, 5, 5, 5, 5, 5, 5, 5, 5, 5, 5, 5, 5, 5, 5, 5,
+			 5, 5, 5, 5, 5, 5, 5, 5, 5, 5, 5, 5, 5, 5, 5, 5, 5, 5, 5, 5, 5, 5, 5, 5, 5,
+			 5, 5, 5, 5, 5, 5, 5, 5, 5, 5, 5, 5, 5, 5, 5, 5, 5, 5, 5, 5, 5, 5, 5, 5, 5,
+			 5, 5, 5, 5, 5, 5, 5, 3, 3, 3, 3, 3, 3, 3, 3, 3, 3, 3, 3, 3, 3, 3, 3, 3, 3,
+			 3, 3, 3, 3, 3, 3, 3, 3, 3, 3, 3, 3, 3, 3, 3, 3, 3, 3, 3, 3, 3, 3, 3, 3, 3,
+			 3, 3, 3, 3, 3, 3, 3, 3, 3, 3, 3, 3, 3, 3, 3, 3, 3, 3, 3, 3, 3, 3, 3, 3, 3,
+			 3, 3, 3, 3, 3, 3, 3, 3, 3, 3, 3, 3, 3, 3, 3, 3, 3, 3, 3, 3, 3, 3, 3, 3, 3,
+			 3, 3, 3, 3, 3, 3, 3, 3, 3, 3, 3, 3, 3, 3, 3, 3, 3, 3, 3, 1, 1, 1, 1, 1, 1,
+			 1, 1, 1, 1, 1, 1, 1, 1, 1, 1, 1, 1, 1, 1, 1, 1, 1, 1, 1, 1, 1, 1, 1, 1, 1,
+			 1, 1, 1, 1, 1, 1, 1, 1, 1, 1, 1, 1, 1, 1, 1, 1, 1, 1, 1, 1, 1, 1, 1, 1, 1},
+			{0,  0,  0,  0,  0,  0,  0,  0,  0,  0,  0,  0,  0,  0,  0,  0,  0,  0,  0,
+			 0,  0,  0,  0,  0,  0,  0,  0,  0,  0,  0,  0,  0,  0,  0,  0,  0,  0,  0,
+			 0,  0,  0,  4,  4,  4,  4,  4,  4,  4,  4,  4,  4,  4,  4,  4,  4,  4,  4,
+			 4,  4,  4,  4,  4,  4,  4,  4,  4,  4,  4,  4,  4,  4,  4,  4,  4,  4,  4,
+			 4,  4,  4,  4,  4,  4,  4,  4,  4,  4,  4,  4,  4,  4,  4,  4,  4,  4,  4,
+			 4,  4,  4,  4,  4,  4,  4,  4,  4,  4,  4,  4,  4,  4,  4,  4,  4,  4,  4,
+			 4,  4,  4,  4,  4,  4,  4,  4,  4,  4,  4,  4,  4,  4,  4,  8,  8,  8,  8,
+			 8,  8,  8,  8,  8,  8,  8,  8,  8,  8,  8,  8,  8,  8,  8,  8,  8,  8,  8,
+			 8,  8,  8,  8,  8,  8,  8,  8,  8,  8,  8,  8,  8,  8,  8,  8,  8,  8,  8,
+			 8,  8,  8,  8,  8,  8,  8,  8,  8,  8,  8,  8,  8,  8,  8,  8,  8,  8,  8,
+			 8,  8,  8,  8,  8,  8,  8,  8,  8,  8,  8,  8,  8,  8,  8,  8,  8,  8,  8,
+			 8,  8,  8,  8,  8,  8,  8,  8,  8,  8,  8,  8,  8,  8,  8,  8,  2,  2,  2,
+			 2,  2,  2,  2,  2,  2,  2,  2,  2,  2,  2,  2,  2,  2,  2,  2,  2,  2,  2,
+			 2,  2,  2,  2,  2,  2,  2,  2,  2,  2,  2,  2,  2,  2,  2,  2,  2,  2,  2,
+			 2,  2,  2,  2,  2,  2,  2,  2,  2,  2,  2,  2,  2,  2,  2,  2,  2,  2,  2,
+			 2,  2,  2,  2,  2,  2,  2,  2,  2,  2,  2,  2,  2,  2,  2,  2,  2,  2,  2,
+			 2,  2,  2,  2,  2,  2,  2,  2,  2,  2,  2,  2,  2,  2,  2,  2,  2,  6,  6,
+			 6,  6,  6,  6,  6,  6,  6,  6,  6,  6,  6,  6,  6,  6,  6,  6,  6,  6,  6,
+			 6,  6,  6,  6,  6,  6,  6,  6,  6,  6,  6,  6,  6,  6,  6,  6,  6,  6,  6,
+			 6,  6,  6,  6,  6,  6,  6,  6,  6,  6,  6,  6,  6,  6,  6,  6,  6,  6,  6,
+			 6,  6,  6,  6,  6,  6,  6,  6,  6,  6,  6,  6,  6,  6,  6,  6,  6,  6,  6,
+			 6,  6,  6,  6,  6,  6,  6,  6,  6,  6,  10, 10, 10, 10, 10, 10, 10, 10, 10,
+			 10, 10, 10, 10, 10, 10, 10, 10, 10, 10, 10, 10, 10, 10, 10, 10, 10, 10, 10,
+			 10, 10, 10, 10, 10, 10, 10, 10, 10, 10, 10, 10, 10, 10, 10, 10, 10, 10, 10,
+			 10, 10, 10, 10, 10, 10, 10, 10, 10, 10, 10, 10, 10, 10, 10, 10, 10, 10, 10,
+			 10, 10, 10, 10, 10, 10, 10, 10, 10, 10, 10, 10, 10, 10, 10, 10, 10, 10, 10,
+			 10, 10, 10, 10, 10, 10, 10, 10, 10, 10, 10, 10, 10, 10, 10, 10, 10, 10, 10,
+			 11, 11, 11, 11, 11, 11, 11, 11, 11, 11, 11, 11, 11, 11, 11, 11, 11, 11, 11,
+			 11, 11, 11, 11, 11, 11, 11, 11, 11, 11, 11, 11, 11, 11, 11, 11, 11, 11, 11,
+			 11, 11, 11, 11, 11, 11, 11, 11, 11, 11, 11, 11, 11, 11, 11, 11, 11, 11, 11,
+			 11, 11, 11, 11, 11, 11, 11, 11, 11, 11, 11, 11, 11, 11, 11, 11, 11, 11, 11,
+			 11, 11, 11, 11, 11, 11, 11, 11, 11, 11, 11, 11, 11, 11, 11, 11, 11, 11, 11,
+			 11, 11, 11, 11, 11, 11, 11, 11, 11, 7,  7,  7,  7,  7,  7,  7,  7,  7,  7,
+			 7,  7,  7,  7,  7,  7,  7,  7,  7,  7,  7,  7,  7,  7,  7,  7,  7,  7,  7,
+			 7,  7,  7,  7,  7,  7,  7,  7,  7,  7,  7,  7,  7,  7,  7,  7,  7,  7,  7,
+			 7,  7,  7,  7,  7,  7,  7,  7,  7,  7,  7,  7,  7,  7,  7,  7,  7,  7,  7,
+			 7,  7,  7,  7,  7,  7,  7,  7,  7,  7,  7,  7,  7,  7,  7,  7,  7,  7,  7,
+			 7,  7,  3,  3,  3,  3,  3,  3,  3,  3,  3,  3,  3,  3,  3,  3,  3,  3,  3,
+			 3,  3,  3,  3,  3,  3,  3,  3,  3,  3,  3,  3,  3,  3,  3,  3,  3,  3,  3,
+			 3,  3,  3,  3,  3,  3,  3,  3,  3,  3,  3,  3,  3,  3,  3,  3,  3,  3,  3,
+			 3,  3,  3,  3,  3,  3,  3,  3,  3,  3,  3,  3,  3,  3,  3,  3,  3,  3,  3,
+			 3,  3,  3,  3,  3,  3,  3,  3,  3,  3,  3,  3,  3,  3,  3,  3,  3,  3,  3,
+			 3,  3,  3,  9,  9,  9,  9,  9,  9,  9,  9,  9,  9,  9,  9,  9,  9,  9,  9,
+			 9,  9,  9,  9,  9,  9,  9,  9,  9,  9,  9,  9,  9,  9,  9,  9,  9,  9,  9,
+			 9,  9,  9,  9,  9,  9,  9,  9,  9,  9,  9,  9,  9,  9,  9,  9,  9,  9,  9,
+			 9,  9,  9,  9,  9,  9,  9,  9,  9,  9,  9,  9,  9,  9,  9,  9,  9,  9,  9,
+			 9,  9,  9,  9,  9,  9,  9,  9,  9,  9,  9,  9,  9,  9,  9,  9,  9,  9,  9,
+			 9,  9,  9,  9,  5,  5,  5,  5,  5,  5,  5,  5,  5,  5,  5,  5,  5,  5,  5,
+			 5,  5,  5,  5,  5,  5,  5,  5,  5,  5,  5,  5,  5,  5,  5,  5,  5,  5,  5,
+			 5,  5,  5,  5,  5,  5,  5,  5,  5,  5,  5,  5,  5,  5,  5,  5,  5,  5,  5,
+			 5,  5,  5,  5,  5,  5,  5,  5,  5,  5,  5,  5,  5,  5,  5,  5,  5,  5,  5,
+			 5,  5,  5,  5,  5,  5,  5,  5,  5,  5,  5,  5,  5,  5,  5,  5,  1,  1,  1,
+			 1,  1,  1,  1,  1,  1,  1,  1,  1,  1,  1,  1,  1,  1,  1,  1,  1,  1,  1,
+			 1,  1,  1,  1,  1,  1,  1,  1,  1,  1,  1,  1,  1,  1,  1,  1,  1,  1},
+			{0,  0,  0,  0,  0,  0,  0,  0,  0,  0,  0,  0,  0,  0,  0,  0,  0,  0,  0,
+			 0,  0,  0,  0,  0,  0,  0,  0,  0,  0,  0,  0,  0,  0,  1,  1,  1,  1,  1,
+			 1,  1,  1,  1,  1,  1,  1,  1,  1,  1,  1,  1,  1,  1,  1,  1,  1,  1,  1,
+			 1,  1,  1,  1,  1,  1,  1,  1,  1,  1,  1,  1,  1,  1,  1,  1,  1,  1,  1,
+			 1,  1,  1,  1,  1,  1,  1,  1,  1,  1,  1,  1,  1,  1,  1,  1,  1,  1,  1,
+			 1,  1,  2,  2,  2,  2,  2,  2,  2,  2,  2,  2,  2,  2,  2,  2,  2,  2,  2,
+			 2,  2,  2,  2,  2,  2,  2,  2,  2,  2,  2,  2,  2,  2,  2,  2,  2,  2,  2,
+			 2,  2,  2,  2,  2,  2,  2,  2,  2,  2,  2,  2,  2,  2,  2,  2,  2,  2,  2,
+			 2,  2,  2,  2,  2,  2,  2,  2,  2,  3,  3,  3,  3,  3,  3,  3,  3,  3,  3,
+			 3,  3,  3,  3,  3,  3,  3,  3,  3,  3,  3,  3,  3,  3,  3,  3,  3,  3,  3,
+			 3,  3,  3,  3,  3,  3,  3,  3,  3,  3,  3,  3,  3,  3,  3,  3,  3,  3,  3,
+			 3,  3,  3,  3,  3,  3,  3,  3,  3,  3,  3,  3,  3,  3,  3,  3,  3,  3,  3,
+			 3,  3,  3,  3,  3,  4,  4,  4,  4,  4,  4,  4,  4,  4,  4,  4,  4,  4,  4,
+			 4,  4,  4,  4,  4,  4,  4,  4,  4,  4,  4,  4,  4,  4,  4,  4,  4,  4,  4,
+			 4,  4,  4,  4,  4,  4,  4,  4,  4,  4,  4,  4,  4,  4,  4,  4,  4,  4,  4,
+			 4,  4,  4,  4,  4,  4,  4,  4,  4,  4,  4,  4,  4,  4,  4,  4,  4,  4,  4,
+			 4,  5,  5,  5,  5,  5,  5,  5,  5,  5,  5,  5,  5,  5,  5,  5,  5,  5,  5,
+			 5,  5,  5,  5,  5,  5,  5,  5,  5,  5,  5,  5,  5,  5,  5,  5,  5,  5,  5,
+			 5,  5,  5,  5,  5,  5,  5,  5,  5,  5,  5,  5,  5,  5,  5,  5,  5,  5,  5,
+			 5,  5,  5,  5,  5,  5,  5,  5,  6,  6,  6,  6,  6,  6,  6,  6,  6,  6,  6,
+			 6,  6,  6,  6,  6,  6,  6,  6,  6,  6,  6,  6,  6,  6,  6,  6,  6,  6,  6,
+			 6,  6,  6,  6,  6,  6,  6,  6,  6,  6,  6,  6,  6,  6,  6,  6,  6,  6,  6,
+			 6,  6,  6,  6,  6,  6,  6,  6,  6,  6,  6,  6,  6,  6,  6,  7,  7,  7,  7,
+			 7,  7,  7,  7,  7,  7,  7,  7,  7,  7,  7,  7,  7,  7,  7,  7,  7,  7,  7,
+			 7,  7,  7,  7,  7,  7,  7,  7,  7,  7,  7,  7,  7,  7,  7,  7,  7,  7,  7,
+			 7,  7,  7,  7,  7,  7,  7,  7,  7,  7,  7,  7,  7,  7,  7,  7,  7,  7,  7,
+			 7,  7,  7,  7,  7,  7,  7,  7,  7,  7,  7,  7,  7,  7,  7,  7,  7,  7,  7,
+			 8,  8,  8,  8,  8,  8,  8,  8,  8,  8,  8,  8,  8,  8,  8,  8,  8,  8,  8,
+			 8,  8,  8,  8,  8,  8,  8,  8,  8,  8,  8,  8,  8,  8,  8,  8,  8,  8,  8,
+			 8,  8,  8,  8,  8,  8,  8,  8,  8,  8,  8,  8,  8,  8,  8,  8,  8,  8,  8,
+			 8,  8,  8,  8,  8,  8,  8,  8,  8,  8,  8,  8,  8,  8,  8,  8,  8,  8,  8,
+			 8,  8,  8,  8,  9,  9,  9,  9,  9,  9,  9,  9,  9,  9,  9,  9,  9,  9,  9,
+			 9,  9,  9,  9,  9,  9,  9,  9,  9,  9,  9,  9,  9,  9,  9,  9,  9,  9,  9,
+			 9,  9,  9,  9,  9,  9,  9,  9,  9,  9,  9,  9,  9,  9,  9,  9,  9,  9,  9,
+			 9,  9,  9,  9,  9,  9,  9,  9,  9,  9,  9,  10, 10, 10, 10, 10, 10, 10, 10,
+			 10, 10, 10, 10, 10, 10, 10, 10, 10, 10, 10, 10, 10, 10, 10, 10, 10, 10, 10,
+			 10, 10, 10, 10, 10, 10, 10, 10, 10, 10, 10, 10, 10, 10, 10, 10, 10, 10, 10,
+			 10, 10, 10, 10, 10, 10, 10, 10, 10, 10, 10, 10, 10, 10, 10, 10, 10, 10, 11,
+			 11, 11, 11, 11, 11, 11, 11, 11, 11, 11, 11, 11, 11, 11, 11, 11, 11, 11, 11,
+			 11, 11, 11, 11, 11, 11, 11, 11, 11, 11, 11, 11, 11, 11, 11, 11, 11, 11, 11,
+			 11, 11, 11, 11, 11, 11, 11, 11, 11, 11, 11, 11, 11, 11, 11, 11, 11, 11, 11,
+			 11, 11, 11, 11, 11, 11, 11, 11, 11, 11, 11, 11, 11, 11, 12, 12, 12, 12, 12,
+			 12, 12, 12, 12, 12, 12, 12, 12, 12, 12, 12, 12, 12, 12, 12, 12, 12, 12, 12,
+			 12, 12, 12, 12, 12, 12, 12, 12, 12, 12, 12, 12, 12, 12, 12, 12, 12, 12, 12,
+			 12, 12, 12, 12, 12, 12, 12, 12, 12, 12, 12, 12, 12, 12, 12, 12, 12, 12, 12,
+			 12, 12, 12, 12, 12, 12, 12, 12, 12, 12, 13, 13, 13, 13, 13, 13, 13, 13, 13,
+			 13, 13, 13, 13, 13, 13, 13, 13, 13, 13, 13, 13, 13, 13, 13, 13, 13, 13, 13,
+			 13, 13, 13, 13, 13, 13, 13, 13, 13, 13, 13, 13, 13, 13, 13, 13, 13, 13, 13,
+			 13, 13, 13, 13, 13, 13, 13, 13, 13, 13, 13, 13, 13, 13, 13, 13, 13, 14, 14,
+			 14, 14, 14, 14, 14, 14, 14, 14, 14, 14, 14, 14, 14, 14, 14, 14, 14, 14, 14,
+			 14, 14, 14, 14, 14, 14, 14, 14, 14, 14, 14, 14, 14, 14, 14, 14, 14, 14, 14,
+			 14, 14, 14, 14, 14, 14, 14, 14, 14, 14, 14, 14, 14, 14, 14, 14, 14, 14, 14,
+			 14, 14, 14, 14, 14, 15, 15, 15, 15, 15, 15, 15, 15, 15, 15, 15, 15, 15, 15,
+			 15, 15, 15, 15, 15, 15, 15, 15, 15, 15, 15, 15, 15, 15, 15, 15, 15, 15},
+			{0,  0,  0,  0,  0,  0,  0,  0,  0,  0,  0,  0,  0,  0,  0,  0,  0,  0,  0,
+			 0,  0,  0,  0,  0,  0,  4,  4,  4,  4,  4,  4,  4,  4,  4,  4,  4,  4,  4,
+			 4,  4,  4,  4,  4,  4,  4,  4,  4,  4,  4,  4,  4,  4,  4,  4,  4,  4,  4,
+			 4,  4,  4,  4,  4,  4,  4,  4,  4,  4,  4,  4,  4,  4,  4,  4,  8,  8,  8,
+			 8,  8,  8,  8,  8,  8,  8,  8,  8,  8,  8,  8,  8,  8,  8,  8,  8,  8,  8,
+			 8,  8,  8,  8,  8,  8,  8,  8,  8,  8,  8,  8,  8,  8,  8,  8,  8,  8,  8,
+			 8,  8,  8,  8,  8,  8,  8,  12, 12, 12, 12, 12, 12, 12, 12, 12, 12, 12, 12,
+			 12, 12, 12, 12, 12, 12, 12, 12, 12, 12, 12, 12, 12, 12, 12, 12, 12, 12, 12,
+			 12, 12, 12, 12, 12, 12, 12, 12, 12, 12, 12, 12, 12, 12, 12, 12, 12, 12, 12,
+			 12, 12, 12, 12, 12, 12, 16, 16, 16, 16, 16, 16, 16, 16, 16, 16, 16, 16, 16,
+			 16, 16, 16, 16, 16, 16, 16, 16, 16, 16, 16, 16, 16, 16, 16, 16, 16, 16, 16,
+			 16, 16, 16, 16, 16, 16, 16, 16, 16, 16, 16, 16, 16, 16, 16, 16, 16, 16, 16,
+			 16, 16, 16, 16, 16, 2,  2,  2,  2,  2,  2,  2,  2,  2,  2,  2,  2,  2,  2,
+			 2,  2,  2,  2,  2,  2,  2,  2,  2,  2,  2,  2,  2,  2,  2,  2,  2,  2,  2,
+			 2,  2,  2,  2,  2,  2,  2,  2,  2,  2,  2,  2,  2,  2,  2,  6,  6,  6,  6,
+			 6,  6,  6,  6,  6,  6,  6,  6,  6,  6,  6,  6,  6,  6,  6,  6,  6,  6,  6,
+			 6,  6,  6,  6,  6,  6,  6,  6,  6,  6,  6,  6,  6,  6,  6,  6,  6,  6,  6,
+			 6,  6,  6,  6,  6,  6,  6,  6,  6,  6,  6,  6,  6,  6,  10, 10, 10, 10, 10,
+			 10, 10, 10, 10, 10, 10, 10, 10, 10, 10, 10, 10, 10, 10, 10, 10, 10, 10, 10,
+			 10, 10, 10, 10, 10, 10, 10, 10, 10, 10, 10, 10, 10, 10, 10, 10, 10, 10, 10,
+			 10, 10, 10, 10, 10, 10, 10, 10, 10, 10, 10, 10, 10, 14, 14, 14, 14, 14, 14,
+			 14, 14, 14, 14, 14, 14, 14, 14, 14, 14, 14, 14, 14, 14, 14, 14, 14, 14, 14,
+			 14, 14, 14, 14, 14, 14, 14, 14, 14, 14, 14, 14, 14, 14, 14, 14, 14, 14, 14,
+			 14, 14, 14, 14, 18, 18, 18, 18, 18, 18, 18, 18, 18, 18, 18, 18, 18, 18, 18,
+			 18, 18, 18, 18, 18, 18, 18, 18, 18, 18, 18, 18, 18, 18, 18, 18, 18, 18, 18,
+			 18, 18, 18, 18, 18, 18, 18, 18, 18, 18, 18, 18, 18, 18, 18, 18, 18, 18, 18,
+			 18, 18, 18, 18, 18, 18, 18, 18, 18, 18, 18, 18, 18, 18, 18, 18, 18, 18, 18,
+			 19, 19, 19, 19, 19, 19, 19, 19, 19, 19, 19, 19, 19, 19, 19, 19, 19, 19, 19,
+			 19, 19, 19, 19, 19, 19, 19, 19, 19, 19, 19, 19, 19, 19, 19, 19, 19, 19, 19,
+			 19, 19, 19, 19, 19, 19, 19, 19, 19, 19, 19, 19, 19, 19, 19, 19, 19, 19, 19,
+			 19, 19, 19, 19, 19, 19, 19, 19, 19, 19, 19, 19, 19, 19, 19, 15, 15, 15, 15,
+			 15, 15, 15, 15, 15, 15, 15, 15, 15, 15, 15, 15, 15, 15, 15, 15, 15, 15, 15,
+			 15, 15, 15, 15, 15, 15, 15, 15, 15, 15, 15, 15, 15, 15, 15, 15, 15, 15, 15,
+			 15, 15, 15, 15, 15, 15, 11, 11, 11, 11, 11, 11, 11, 11, 11, 11, 11, 11, 11,
+			 11, 11, 11, 11, 11, 11, 11, 11, 11, 11, 11, 11, 11, 11, 11, 11, 11, 11, 11,
+			 11, 11, 11, 11, 11, 11, 11, 11, 11, 11, 11, 11, 11, 11, 11, 11, 11, 11, 11,
+			 11, 11, 11, 11, 11, 7,  7,  7,  7,  7,  7,  7,  7,  7,  7,  7,  7,  7,  7,
+			 7,  7,  7,  7,  7,  7,  7,  7,  7,  7,  7,  7,  7,  7,  7,  7,  7,  7,  7,
+			 7,  7,  7,  7,  7,  7,  7,  7,  7,  7,  7,  7,  7,  7,  7,  7,  7,  7,  7,
+			 7,  7,  7,  7,  3,  3,  3,  3,  3,  3,  3,  3,  3,  3,  3,  3,  3,  3,  3,
+			 3,  3,  3,  3,  3,  3,  3,  3,  3,  3,  3,  3,  3,  3,  3,  3,  3,  3,  3,
+			 3,  3,  3,  3,  3,  3,  3,  3,  3,  3,  3,  3,  3,  3,  17, 17, 17, 17, 17,
+			 17, 17, 17, 17, 17, 17, 17, 17, 17, 17, 17, 17, 17, 17, 17, 17, 17, 17, 17,
+			 17, 17, 17, 17, 17, 17, 17, 17, 17, 17, 17, 17, 17, 17, 17, 17, 17, 17, 17,
+			 17, 17, 17, 17, 17, 17, 17, 17, 17, 17, 17, 17, 17, 13, 13, 13, 13, 13, 13,
+			 13, 13, 13, 13, 13, 13, 13, 13, 13, 13, 13, 13, 13, 13, 13, 13, 13, 13, 13,
+			 13, 13, 13, 13, 13, 13, 13, 13, 13, 13, 13, 13, 13, 13, 13, 13, 13, 13, 13,
+			 13, 13, 13, 13, 13, 13, 13, 13, 13, 13, 13, 13, 9,  9,  9,  9,  9,  9,  9,
+			 9,  9,  9,  9,  9,  9,  9,  9,  9,  9,  9,  9,  9,  9,  9,  9,  9,  9,  9,
+			 9,  9,  9,  9,  9,  9,  9,  9,  9,  9,  9,  9,  9,  9,  9,  9,  9,  9,  9,
+			 9,  9,  9,  5,  5,  5,  5,  5,  5,  5,  5,  5,  5,  5,  5,  5,  5,  5,  5,
+			 5,  5,  5,  5,  5,  5,  5,  5,  5,  5,  5,  5,  5,  5,  5,  5,  5,  5,  5,
+			 5,  5,  5,  5,  5,  5,  5,  5,  5,  5,  5,  5,  5,  1,  1,  1,  1,  1,  1,
+			 1,  1,  1,  1,  1,  1,  1,  1,  1,  1,  1,  1,  1,  1,  1,  1,  1,  1},
+			{0,  0,  0,  0,  0,  0,  0,  0,  0,  0,  0,  0,  0,  0,  0,  0,  0,  8,  8,
+			 8,  8,  8,  8,  8,  8,  8,  8,  8,  8,  8,  8,  8,  8,  8,  8,  8,  8,  8,
+			 8,  8,  8,  8,  8,  8,  8,  8,  8,  8,  8,  8,  8,  8,  8,  8,  8,  8,  8,
+			 16, 16, 16, 16, 16, 16, 16, 16, 16, 16, 16, 16, 16, 16, 16, 16, 16, 16, 16,
+			 16, 16, 16, 16, 16, 16, 16, 16, 16, 16, 16, 16, 16, 16, 16, 16, 16, 16, 16,
+			 16, 16, 16, 16, 16, 16, 16, 16, 16, 16, 2,  2,  2,  2,  2,  2,  2,  2,  2,
+			 2,  2,  2,  2,  2,  2,  2,  2,  2,  2,  2,  2,  2,  2,  2,  2,  2,  2,  2,
+			 2,  2,  2,  2,  2,  2,  2,  2,  2,  2,  2,  2,  2,  2,  2,  2,  2,  2,  2,
+			 2,  10, 10, 10, 10, 10, 10, 10, 10, 10, 10, 10, 10, 10, 10, 10, 10, 10, 10,
+			 10, 10, 10, 10, 10, 10, 10, 10, 10, 10, 10, 10, 10, 10, 10, 10, 10, 10, 10,
+			 10, 10, 10, 18, 18, 18, 18, 18, 18, 18, 18, 18, 18, 18, 18, 18, 18, 18, 18,
+			 18, 18, 18, 18, 18, 18, 18, 18, 18, 18, 18, 18, 18, 18, 18, 18, 18, 18, 18,
+			 18, 18, 18, 18, 18, 4,  4,  4,  4,  4,  4,  4,  4,  4,  4,  4,  4,  4,  4,
+			 4,  4,  4,  4,  4,  4,  4,  4,  4,  4,  4,  4,  4,  4,  4,  4,  4,  4,  4,
+			 4,  4,  4,  4,  4,  4,  4,  4,  4,  4,  4,  4,  4,  4,  4,  12, 12, 12, 12,
+			 12, 12, 12, 12, 12, 12, 12, 12, 12, 12, 12, 12, 12, 12, 12, 12, 12, 12, 12,
+			 12, 12, 12, 12, 12, 12, 12, 12, 12, 12, 12, 12, 12, 12, 12, 12, 12, 12, 12,
+			 12, 12, 12, 12, 12, 12, 20, 20, 20, 20, 20, 20, 20, 20, 20, 20, 20, 20, 20,
+			 20, 20, 20, 20, 20, 20, 20, 20, 20, 20, 20, 20, 20, 20, 20, 20, 20, 20, 20,
+			 20, 20, 20, 20, 20, 20, 20, 20, 6,  6,  6,  6,  6,  6,  6,  6,  6,  6,  6,
+			 6,  6,  6,  6,  6,  6,  6,  6,  6,  6,  6,  6,  6,  6,  6,  6,  6,  6,  6,
+			 6,  6,  6,  6,  6,  6,  6,  6,  6,  6,  14, 14, 14, 14, 14, 14, 14, 14, 14,
+			 14, 14, 14, 14, 14, 14, 14, 14, 14, 14, 14, 14, 14, 14, 14, 14, 14, 14, 14,
+			 14, 14, 14, 14, 14, 14, 14, 14, 14, 14, 14, 14, 14, 14, 14, 14, 14, 14, 14,
+			 14, 22, 22, 22, 22, 22, 22, 22, 22, 22, 22, 22, 22, 22, 22, 22, 22, 22, 22,
+			 22, 22, 22, 22, 22, 22, 22, 22, 22, 22, 22, 22, 22, 22, 22, 22, 22, 22, 22,
+			 22, 22, 22, 22, 22, 22, 22, 22, 22, 22, 22, 22, 22, 22, 22, 22, 22, 22, 22,
+			 23, 23, 23, 23, 23, 23, 23, 23, 23, 23, 23, 23, 23, 23, 23, 23, 23, 23, 23,
+			 23, 23, 23, 23, 23, 23, 23, 23, 23, 23, 23, 23, 23, 23, 23, 23, 23, 23, 23,
+			 23, 23, 23, 23, 23, 23, 23, 23, 23, 23, 23, 23, 23, 23, 23, 23, 23, 23, 15,
+			 15, 15, 15, 15, 15, 15, 15, 15, 15, 15, 15, 15, 15, 15, 15, 15, 15, 15, 15,
+			 15, 15, 15, 15, 15, 15, 15, 15, 15, 15, 15, 15, 15, 15, 15, 15, 15, 15, 15,
+			 15, 15, 15, 15, 15, 15, 15, 15, 15, 7,  7,  7,  7,  7,  7,  7,  7,  7,  7,
+			 7,  7,  7,  7,  7,  7,  7,  7,  7,  7,  7,  7,  7,  7,  7,  7,  7,  7,  7,
+			 7,  7,  7,  7,  7,  7,  7,  7,  7,  7,  7,  21, 21, 21, 21, 21, 21, 21, 21,
+			 21, 21, 21, 21, 21, 21, 21, 21, 21, 21, 21, 21, 21, 21, 21, 21, 21, 21, 21,
+			 21, 21, 21, 21, 21, 21, 21, 21, 21, 21, 21, 21, 21, 13, 13, 13, 13, 13, 13,
+			 13, 13, 13, 13, 13, 13, 13, 13, 13, 13, 13, 13, 13, 13, 13, 13, 13, 13, 13,
+			 13, 13, 13, 13, 13, 13, 13, 13, 13, 13, 13, 13, 13, 13, 13, 13, 13, 13, 13,
+			 13, 13, 13, 13, 5,  5,  5,  5,  5,  5,  5,  5,  5,  5,  5,  5,  5,  5,  5,
+			 5,  5,  5,  5,  5,  5,  5,  5,  5,  5,  5,  5,  5,  5,  5,  5,  5,  5,  5,
+			 5,  5,  5,  5,  5,  5,  5,  5,  5,  5,  5,  5,  5,  5,  19, 19, 19, 19, 19,
+			 19, 19, 19, 19, 19, 19, 19, 19, 19, 19, 19, 19, 19, 19, 19, 19, 19, 19, 19,
+			 19, 19, 19, 19, 19, 19, 19, 19, 19, 19, 19, 19, 19, 19, 19, 19, 11, 11, 11,
+			 11, 11, 11, 11, 11, 11, 11, 11, 11, 11, 11, 11, 11, 11, 11, 11, 11, 11, 11,
+			 11, 11, 11, 11, 11, 11, 11, 11, 11, 11, 11, 11, 11, 11, 11, 11, 11, 11, 3,
+			 3,  3,  3,  3,  3,  3,  3,  3,  3,  3,  3,  3,  3,  3,  3,  3,  3,  3,  3,
+			 3,  3,  3,  3,  3,  3,  3,  3,  3,  3,  3,  3,  3,  3,  3,  3,  3,  3,  3,
+			 3,  3,  3,  3,  3,  3,  3,  3,  3,  17, 17, 17, 17, 17, 17, 17, 17, 17, 17,
+			 17, 17, 17, 17, 17, 17, 17, 17, 17, 17, 17, 17, 17, 17, 17, 17, 17, 17, 17,
+			 17, 17, 17, 17, 17, 17, 17, 17, 17, 17, 17, 17, 17, 17, 17, 17, 17, 17, 17,
+			 9,  9,  9,  9,  9,  9,  9,  9,  9,  9,  9,  9,  9,  9,  9,  9,  9,  9,  9,
+			 9,  9,  9,  9,  9,  9,  9,  9,  9,  9,  9,  9,  9,  9,  9,  9,  9,  9,  9,
+			 9,  9,  1,  1,  1,  1,  1,  1,  1,  1,  1,  1,  1,  1,  1,  1,  1,  1},
+			{0,  0,  0,  0,  0,  0,  0,  0,  0,  0,  0,  0,  0,  0,  0,  0,  0,  1,  1,
+			 1,  1,  1,  1,  1,  1,  1,  1,  1,  1,  1,  1,  1,  1,  1,  1,  1,  1,  1,
+			 1,  1,  1,  1,  1,  1,  1,  1,  1,  1,  1,  2,  2,  2,  2,  2,  2,  2,  2,
+			 2,  2,  2,  2,  2,  2,  2,  2,  2,  2,  2,  2,  2,  2,  2,  2,  2,  2,  2,
+			 2,  2,  2,  2,  2,  3,  3,  3,  3,  3,  3,  3,  3,  3,  3,  3,  3,  3,  3,
+			 3,  3,  3,  3,  3,  3,  3,  3,  3,  3,  3,  3,  3,  3,  3,  3,  3,  3,  4,
+			 4,  4,  4,  4,  4,  4,  4,  4,  4,  4,  4,  4,  4,  4,  4,  4,  4,  4,  4,
+			 4,  4,  4,  4,  4,  4,  4,  4,  4,  4,  4,  4,  5,  5,  5,  5,  5,  5,  5,
+			 5,  5,  5,  5,  5,  5,  5,  5,  5,  5,  5,  5,  5,  5,  5,  5,  5,  5,  5,
+			 5,  5,  5,  5,  5,  5,  6,  6,  6,  6,  6,  6,  6,  6,  6,  6,  6,  6,  6,
+			 6,  6,  6,  6,  6,  6,  6,  6,  6,  6,  6,  6,  6,  6,  6,  6,  6,  6,  6,
+			 7,  7,  7,  7,  7,  7,  7,  7,  7,  7,  7,  7,  7,  7,  7,  7,  7,  7,  7,
+			 7,  7,  7,  7,  7,  7,  7,  7,  7,  7,  7,  7,  7,  8,  8,  8,  8,  8,  8,
+			 8,  8,  8,  8,  8,  8,  8,  8,  8,  8,  8,  8,  8,  8,  8,  8,  8,  8,  8,
+			 8,  8,  8,  8,  8,  8,  8,  9,  9,  9,  9,  9,  9,  9,  9,  9,  9,  9,  9,
+			 9,  9,  9,  9,  9,  9,  9,  9,  9,  9,  9,  9,  9,  9,  9,  9,  9,  9,  9,
+			 9,  10, 10, 10, 10, 10, 10, 10, 10, 10, 10, 10, 10, 10, 10, 10, 10, 10, 10,
+			 10, 10, 10, 10, 10, 10, 10, 10, 10, 10, 10, 10, 10, 10, 11, 11, 11, 11, 11,
+			 11, 11, 11, 11, 11, 11, 11, 11, 11, 11, 11, 11, 11, 11, 11, 11, 11, 11, 11,
+			 11, 11, 11, 11, 11, 11, 11, 11, 12, 12, 12, 12, 12, 12, 12, 12, 12, 12, 12,
+			 12, 12, 12, 12, 12, 12, 12, 12, 12, 12, 12, 12, 12, 12, 12, 12, 12, 12, 12,
+			 12, 12, 13, 13, 13, 13, 13, 13, 13, 13, 13, 13, 13, 13, 13, 13, 13, 13, 13,
+			 13, 13, 13, 13, 13, 13, 13, 13, 13, 13, 13, 13, 13, 13, 13, 14, 14, 14, 14,
+			 14, 14, 14, 14, 14, 14, 14, 14, 14, 14, 14, 14, 14, 14, 14, 14, 14, 14, 14,
+			 14, 14, 14, 14, 14, 14, 14, 14, 14, 15, 15, 15, 15, 15, 15, 15, 15, 15, 15,
+			 15, 15, 15, 15, 15, 15, 15, 15, 15, 15, 15, 15, 15, 15, 15, 15, 15, 15, 15,
+			 15, 15, 15, 15, 15, 15, 15, 15, 15, 15, 15, 15, 15, 15, 15, 15, 15, 15, 15,
+			 16, 16, 16, 16, 16, 16, 16, 16, 16, 16, 16, 16, 16, 16, 16, 16, 16, 16, 16,
+			 16, 16, 16, 16, 16, 16, 16, 16, 16, 16, 16, 16, 16, 16, 16, 16, 16, 16, 16,
+			 16, 16, 16, 16, 16, 16, 16, 16, 16, 16, 17, 17, 17, 17, 17, 17, 17, 17, 17,
+			 17, 17, 17, 17, 17, 17, 17, 17, 17, 17, 17, 17, 17, 17, 17, 17, 17, 17, 17,
+			 17, 17, 17, 17, 18, 18, 18, 18, 18, 18, 18, 18, 18, 18, 18, 18, 18, 18, 18,
+			 18, 18, 18, 18, 18, 18, 18, 18, 18, 18, 18, 18, 18, 18, 18, 18, 18, 19, 19,
+			 19, 19, 19, 19, 19, 19, 19, 19, 19, 19, 19, 19, 19, 19, 19, 19, 19, 19, 19,
+			 19, 19, 19, 19, 19, 19, 19, 19, 19, 19, 19, 20, 20, 20, 20, 20, 20, 20, 20,
+			 20, 20, 20, 20, 20, 20, 20, 20, 20, 20, 20, 20, 20, 20, 20, 20, 20, 20, 20,
+			 20, 20, 20, 20, 20, 21, 21, 21, 21, 21, 21, 21, 21, 21, 21, 21, 21, 21, 21,
+			 21, 21, 21, 21, 21, 21, 21, 21, 21, 21, 21, 21, 21, 21, 21, 21, 21, 21, 22,
+			 22, 22, 22, 22, 22, 22, 22, 22, 22, 22, 22, 22, 22, 22, 22, 22, 22, 22, 22,
+			 22, 22, 22, 22, 22, 22, 22, 22, 22, 22, 22, 22, 23, 23, 23, 23, 23, 23, 23,
+			 23, 23, 23, 23, 23, 23, 23, 23, 23, 23, 23, 23, 23, 23, 23, 23, 23, 23, 23,
+			 23, 23, 23, 23, 23, 23, 24, 24, 24, 24, 24, 24, 24, 24, 24, 24, 24, 24, 24,
+			 24, 24, 24, 24, 24, 24, 24, 24, 24, 24, 24, 24, 24, 24, 24, 24, 24, 24, 24,
+			 25, 25, 25, 25, 25, 25, 25, 25, 25, 25, 25, 25, 25, 25, 25, 25, 25, 25, 25,
+			 25, 25, 25, 25, 25, 25, 25, 25, 25, 25, 25, 25, 25, 26, 26, 26, 26, 26, 26,
+			 26, 26, 26, 26, 26, 26, 26, 26, 26, 26, 26, 26, 26, 26, 26, 26, 26, 26, 26,
+			 26, 26, 26, 26, 26, 26, 26, 27, 27, 27, 27, 27, 27, 27, 27, 27, 27, 27, 27,
+			 27, 27, 27, 27, 27, 27, 27, 27, 27, 27, 27, 27, 27, 27, 27, 27, 27, 27, 27,
+			 27, 28, 28, 28, 28, 28, 28, 28, 28, 28, 28, 28, 28, 28, 28, 28, 28, 28, 28,
+			 28, 28, 28, 28, 28, 28, 28, 28, 28, 28, 28, 28, 28, 28, 29, 29, 29, 29, 29,
+			 29, 29, 29, 29, 29, 29, 29, 29, 29, 29, 29, 29, 29, 29, 29, 29, 29, 29, 29,
+			 29, 29, 29, 29, 29, 29, 29, 29, 30, 30, 30, 30, 30, 30, 30, 30, 30, 30, 30,
+			 30, 30, 30, 30, 30, 30, 30, 30, 30, 30, 30, 30, 30, 30, 30, 30, 30, 30, 30,
+			 30, 30, 31, 31, 31, 31, 31, 31, 31, 31, 31, 31, 31, 31, 31, 31, 31, 31} };
+
+
+
+		/**
+		 * Table that describes the number of trits or quints along with bits required
+		 * for storing each range.
+		 */
+		const uint8 bits_trits_quints_table[RANGE_MAX][3] = {
+			{1, 0, 0},  // RANGE_2
+			{0, 1, 0},  // RANGE_3
+			{2, 0, 0},  // RANGE_4
+			{0, 0, 1},  // RANGE_5
+			{1, 1, 0},  // RANGE_6
+			{3, 0, 0},  // RANGE_8
+			{1, 0, 1},  // RANGE_10
+			{2, 1, 0},  // RANGE_12
+			{4, 0, 0},  // RANGE_16
+			{2, 0, 1},  // RANGE_20
+			{3, 1, 0},  // RANGE_24
+			{5, 0, 0},  // RANGE_32
+			{3, 0, 1},  // RANGE_40
+			{4, 1, 0},  // RANGE_48
+			{6, 0, 0},  // RANGE_64
+			{4, 0, 1},  // RANGE_80
+			{5, 1, 0},  // RANGE_96
+			{7, 0, 0},  // RANGE_128
+			{5, 0, 1},  // RANGE_160
+			{6, 1, 0},  // RANGE_192
+			{8, 0, 0}   // RANGE_256
+		};
+
+		/**
+		 * Encode a group of 5 numbers using trits and bits.
+		 */
+		inline void encode_trits(size_t bits,
+			uint8 b0,
+			uint8 b1,
+			uint8 b2,
+			uint8 b3,
+			uint8 b4,
+			bitwriter& writer)
+		{
+			uint8 t0, t1, t2, t3, t4;
+			uint8 m0, m1, m2, m3, m4;
+
+			split_high_low(b0, bits, t0, m0);
+			split_high_low(b1, bits, t1, m1);
+			split_high_low(b2, bits, t2, m2);
+			split_high_low(b3, bits, t3, m3);
+			split_high_low(b4, bits, t4, m4);
+
+			DCHECK(t0 < 3);
+			DCHECK(t1 < 3);
+			DCHECK(t2 < 3);
+			DCHECK(t3 < 3);
+			DCHECK(t4 < 3);
+
+			if (t0 < 3 && t1 < 3 && t2 < 3 && t3 < 3 && t4 < 3)
+			{
+				uint8 packed = integer_from_trits[t4][t3][t2][t1][t0];
+
+				writer.write8(m0, bits);
+				writer.write8(getbits(packed, 1, 0), 2);
+				writer.write8(m1, bits);
+				writer.write8(getbits(packed, 3, 2), 2);
+				writer.write8(m2, bits);
+				writer.write8(getbits(packed, 4, 4), 1);
+				writer.write8(m3, bits);
+				writer.write8(getbits(packed, 6, 5), 2);
+				writer.write8(m4, bits);
+				writer.write8(getbits(packed, 7, 7), 1);
+			}
+		}
+
+		/**
+		 * Encode a group of 3 numbers using quints and bits.
+		 */
+		inline void encode_quints(size_t bits,
+			uint8 b0,
+			uint8 b1,
+			uint8 b2,
+			bitwriter& writer) {
+			uint8 q0, q1, q2;
+			uint8 m0, m1, m2;
+
+			split_high_low(b0, bits, q0, m0);
+			split_high_low(b1, bits, q1, m1);
+			split_high_low(b2, bits, q2, m2);
+
+			DCHECK(q0 < 5);
+			DCHECK(q1 < 5);
+			DCHECK(q2 < 5);
+
+			if (q0 < 5 && q1 < 5 && q2 < 5)
+			{
+				uint8 packed = integer_from_quints[q2][q1][q0];
+
+				writer.write8(m0, bits);
+				writer.write8(getbits(packed, 2, 0), 3);
+				writer.write8(m1, bits);
+				writer.write8(getbits(packed, 4, 3), 2);
+				writer.write8(m2, bits);
+				writer.write8(getbits(packed, 6, 5), 2);
+			}
+		}
+
+		/**
+		 * Encode a sequence of numbers using using one trit and a custom number of
+		 * bits per number.
+		 */
+		inline void encode_trits(const uint8* numbers,
+			size_t count,
+			bitwriter& writer,
+			size_t bits) {
+			for (size_t i = 0; i < count; i += 5) {
+				uint8 b0 = numbers[i + 0];
+				uint8 b1 = i + 1 >= count ? 0 : numbers[i + 1];
+				uint8 b2 = i + 2 >= count ? 0 : numbers[i + 2];
+				uint8 b3 = i + 3 >= count ? 0 : numbers[i + 3];
+				uint8 b4 = i + 4 >= count ? 0 : numbers[i + 4];
+
+				encode_trits(bits, b0, b1, b2, b3, b4, writer);
+			}
+		}
+
+		/**
+		 * Encode a sequence of numbers using one quint and the custom number of bits
+		 * per number.
+		 */
+		inline void encode_quints(const uint8* numbers,
+			size_t count,
+			bitwriter& writer,
+			size_t bits) {
+			for (size_t i = 0; i < count; i += 3) {
+				uint8 b0 = numbers[i + 0];
+				uint8 b1 = i + 1 >= count ? 0 : numbers[i + 1];
+				uint8 b2 = i + 2 >= count ? 0 : numbers[i + 2];
+				encode_quints(bits, b0, b1, b2, writer);
+			}
+		}
+
+		/**
+		 * Encode a sequence of numbers using binary representation with the selected
+		 * bit count.
+		 */
+		inline void encode_binary(const uint8* numbers,
+			size_t count,
+			bitwriter& writer,
+			size_t bits) {
+			DCHECK(count > 0);
+			for (size_t i = 0; i < count; ++i) {
+				writer.write8(numbers[i], bits);
+			}
+		}
+
+		/**
+		 * Encode a sequence of numbers in a specific range using the binary integer
+		 * sequence encoding. The numbers are assumed to be in the correct range and
+		 * the memory we are writing to is assumed to be zero-initialized.
+		 */
+		inline void integer_sequence_encode(const uint8* numbers,
+			size_t count,
+			range_t range,
+			bitwriter writer) {
+#ifndef NDEBUG
+			for (size_t i = 0; i < count; ++i) {
+				DCHECK(numbers[i] <= range_max_table[range]);
+			}
+#endif
+
+			size_t bits = bits_trits_quints_table[range][0];
+			size_t trits = bits_trits_quints_table[range][1];
+			size_t quints = bits_trits_quints_table[range][2];
+
+			if (trits == 1) {
+				encode_trits(numbers, count, writer, bits);
+			}
+			else if (quints == 1) {
+				encode_quints(numbers, count, writer, bits);
+			}
+			else {
+				encode_binary(numbers, count, writer, bits);
+			}
+		}
+
+		inline void integer_sequence_encode(const uint8* numbers,
+			size_t count,
+			range_t range,
+			uint8* output) {
+			integer_sequence_encode(numbers, count, range, bitwriter(output));
+		}
+
+		/**
+		 * Compute the number of bits required to store a number of items in a specific
+		 * range using the binary integer sequence encoding.
+		 */
+		inline size_t compute_ise_bitcount(size_t items, range_t range) {
+			size_t bits = bits_trits_quints_table[range][0];
+			size_t trits = bits_trits_quints_table[range][1];
+			size_t quints = bits_trits_quints_table[range][2];
+
+			if (trits) {
+				return ((8 + 5 * bits) * items + 4) / 5;
+			}
+
+			if (quints) {
+				return ((7 + 3 * bits) * items + 2) / 3;
+			}
+
+			return items * bits;
+		}
+
+
+		inline uint8 quantize_color(range_t quant, int c)
+		{
+			DCHECK(c >= 0 && c <= 255);
+			if (c >= 0 && c <= 255)
+			{
+				return color_quantize_table[quant][c];
+			}
+			return 0;
+		}
+
+		vec3i_t quantize_color(range_t quant, vec3i_t c)
+		{
+			vec3i_t result;
+			result.r = color_quantize_table[quant][c.r];
+			result.g = color_quantize_table[quant][c.g];
+			result.b = color_quantize_table[quant][c.b];
+			return result;
+		}
+
+		vec4i_t quantize_color(range_t quant, vec4i_t c)
+		{
+			vec4i_t result;
+			result.r = color_quantize_table[quant][c.r];
+			result.g = color_quantize_table[quant][c.g];
+			result.b = color_quantize_table[quant][c.b];
+			result.a = color_quantize_table[quant][c.a];
+			return result;
+		}
+
+		uint8 unquantize_color(range_t quant, int c)
+		{
+			DCHECK(c >= 0 && c <= 255);
+			if (c >= 0 && c <= 255)
+			{
+				return color_unquantize_table[quant][c];
+			}
+			return 0;
+		}
+
+		vec3i_t unquantize_color(range_t quant, vec3i_t c)
+		{
+			vec3i_t result;
+			result.r = color_unquantize_table[quant][c.r];
+			result.g = color_unquantize_table[quant][c.g];
+			result.b = color_unquantize_table[quant][c.b];
+			return result;
+		}
+
+		vec4i_t unquantize_color(range_t quant, vec4i_t c)
+		{
+			vec4i_t result;
+			result.r = color_unquantize_table[quant][c.r];
+			result.g = color_unquantize_table[quant][c.g];
+			result.b = color_unquantize_table[quant][c.b];
+			result.a = color_unquantize_table[quant][c.a];
+			return result;
+		}
+
+
+
+
+
+		int color_channel_sum(vec3i_t color)
+		{
+			return color.r + color.g + color.b;
+		}
+
+		int color_channel_sum(vec4i_t color)
+		{
+			return color.r + color.g + color.b + color.a;
+		}
+
+		void encode_luminance_direct(range_t endpoint_quant,
+			int v0,
+			int v1,
+			uint8 endpoint_unquantized[2],
+			uint8 endpoint_quantized[2])
+		{
+			endpoint_quantized[0] = quantize_color(endpoint_quant, v0);
+			endpoint_quantized[1] = quantize_color(endpoint_quant, v1);
+			endpoint_unquantized[0] = unquantize_color(endpoint_quant, endpoint_quantized[0]);
+			endpoint_unquantized[1] = unquantize_color(endpoint_quant, endpoint_quantized[1]);
+		}
+
+		void encode_rgb_direct(range_t endpoint_quant,
+			vec3i_t e0,
+			vec3i_t e1,
+			uint8 endpoint_quantized[6],
+			vec3i_t endpoint_unquantized[2])
+		{
+			vec3i_t e0q = quantize_color(endpoint_quant, e0);
+			vec3i_t e1q = quantize_color(endpoint_quant, e1);
+			vec3i_t e0u = unquantize_color(endpoint_quant, e0q);
+			vec3i_t e1u = unquantize_color(endpoint_quant, e1q);
+
+			// ASTC uses a different blue contraction encoding when the sum of values for
+			// the first endpoint is larger than the sum of values in the second
+			// endpoint. Sort the endpoints to ensure that the normal encoding is used.
+			if (color_channel_sum(e0u) > color_channel_sum(e1u))
+			{
+				endpoint_quantized[0] = static_cast<uint8>(e1q.r);
+				endpoint_quantized[1] = static_cast<uint8>(e0q.r);
+				endpoint_quantized[2] = static_cast<uint8>(e1q.g);
+				endpoint_quantized[3] = static_cast<uint8>(e0q.g);
+				endpoint_quantized[4] = static_cast<uint8>(e1q.b);
+				endpoint_quantized[5] = static_cast<uint8>(e0q.b);
+
+				endpoint_unquantized[0] = e1u;
+				endpoint_unquantized[1] = e0u;
+			}
+			else
+			{
+				endpoint_quantized[0] = static_cast<uint8>(e0q.r);
+				endpoint_quantized[1] = static_cast<uint8>(e1q.r);
+				endpoint_quantized[2] = static_cast<uint8>(e0q.g);
+				endpoint_quantized[3] = static_cast<uint8>(e1q.g);
+				endpoint_quantized[4] = static_cast<uint8>(e0q.b);
+				endpoint_quantized[5] = static_cast<uint8>(e1q.b);
+
+				endpoint_unquantized[0] = e0u;
+				endpoint_unquantized[1] = e1u;
+			}
+		}
+
+
+		void encode_rgba_direct(range_t endpoint_quant,
+			vec4i_t e0,
+			vec4i_t e1,
+			uint8 endpoint_quantized[8],
+			vec4i_t endpoint_unquantized[2])
+		{
+			vec4i_t e0q = quantize_color(endpoint_quant, e0);
+			vec4i_t e1q = quantize_color(endpoint_quant, e1);
+			vec4i_t e0u = unquantize_color(endpoint_quant, e0q);
+			vec4i_t e1u = unquantize_color(endpoint_quant, e1q);
+
+			// ASTC uses a different blue contraction encoding when the sum of values for
+			// the first endpoint is larger than the sum of values in the second
+			// endpoint. Sort the endpoints to ensure that the normal encoding is used.
+			if (color_channel_sum(e0u.rgb()) > color_channel_sum(e1u.rgb()))
+			{
+				endpoint_quantized[0] = static_cast<uint8>(e1q.r);
+				endpoint_quantized[1] = static_cast<uint8>(e0q.r);
+				endpoint_quantized[2] = static_cast<uint8>(e1q.g);
+				endpoint_quantized[3] = static_cast<uint8>(e0q.g);
+				endpoint_quantized[4] = static_cast<uint8>(e1q.b);
+				endpoint_quantized[5] = static_cast<uint8>(e0q.b);
+
+				endpoint_unquantized[0] = e1u;
+				endpoint_unquantized[1] = e0u;
+			}
+			else
+			{
+				endpoint_quantized[0] = static_cast<uint8>(e0q.r);
+				endpoint_quantized[1] = static_cast<uint8>(e1q.r);
+				endpoint_quantized[2] = static_cast<uint8>(e0q.g);
+				endpoint_quantized[3] = static_cast<uint8>(e1q.g);
+				endpoint_quantized[4] = static_cast<uint8>(e0q.b);
+				endpoint_quantized[5] = static_cast<uint8>(e1q.b);
+
+				endpoint_unquantized[0] = e0u;
+				endpoint_unquantized[1] = e1u;
+			}
+
+			// Sort alpha endpoints
+			if (e0.a > e1.a)
+			{
+				endpoint_quantized[6] = static_cast<uint8>(e1q.a);
+				endpoint_quantized[7] = static_cast<uint8>(e0q.a);
+				endpoint_unquantized[0].a = e1u.a;
+				endpoint_unquantized[1].a = e0u.a;
+			}
+			else
+			{
+				endpoint_quantized[6] = static_cast<uint8>(e0q.a);
+				endpoint_quantized[7] = static_cast<uint8>(e1q.a);
+				endpoint_unquantized[0].a = e0u.a;
+				endpoint_unquantized[1].a = e1u.a;
+			}
+		}
+
+
+		inline void symbolic_to_physical
+		(
+			color_endpoint_mode_t color_endpoint_mode,
+			range_t endpoint_quant,
+			range_t weight_quant,
+
+			size_t partition_count,
+
+			const uint8 endpoint_ise[MAXIMUM_ENCODED_COLOR_ENDPOINT_BYTES],
+
+			// FIXME: +1 needed here because orbits_8ptr breaks when the offset reaches
+			// the last byte which always happens if the weight mode is RANGE_32.
+			const uint8 weights_ise[MAXIMUM_ENCODED_WEIGHT_BYTES + 1],
+
+			PhysicalBlock* pb,
+			bool dual_plane = false
+		)
+		{
+			DCHECK(weight_quant <= RANGE_32);
+			DCHECK(endpoint_quant < RANGE_MAX);
+			DCHECK(color_endpoint_mode < CEM_MAX);
+			DCHECK(partition_count == 1);
+			DCHECK(compute_ise_bitcount(BLOCK_TEXEL_COUNT, weight_quant) < MAXIMUM_ENCODED_WEIGHT_BITS);
+
+			//        if (s_debuglog)
+			//        {
+			//            UE_LOG(LogMutableCore,Warning," physicalising interesting block" );
+			//        }
+
+
+			size_t n = BLOCK_WIDTH;
+			size_t m = BLOCK_HEIGHT;
+
+			static const bool h_table[RANGE_32 + 1] = { 0, 0, 0, 0, 0, 0,
+														1, 1, 1, 1, 1, 1 };
+
+			static const uint8 r_table[RANGE_32 + 1] = { 0x2, 0x3, 0x4, 0x5, 0x6, 0x7,
+														   0x2, 0x3, 0x4, 0x5, 0x6, 0x7 };
+
+			bool h = false;
+			if (weight_quant < RANGE_32 + 1) h = h_table[weight_quant];
+			size_t r = 0;
+			if (weight_quant < RANGE_32 + 1) r = r_table[weight_quant];
+
+			// Use the first row of Table 11 in the ASTC specification. Beware that
+			// this has to be changed if another block-size is used.
+			size_t a = m - 2;
+			size_t b = n - 4;
+
+			bool d = dual_plane;
+
+			size_t part_value = partition_count - 1;
+
+			size_t cem_offset = 13;
+			size_t ced_offset = 17;
+
+			size_t cem_bits = 4;
+			size_t cem = color_endpoint_mode;
+
+			// Block mode
+			orbits8_ptr(pb->data, 0, getbit(r, 1), 1);
+			orbits8_ptr(pb->data, 1, getbit(r, 2), 1);
+			orbits8_ptr(pb->data, 2, 0, 1);
+			orbits8_ptr(pb->data, 3, 0, 1);
+			orbits8_ptr(pb->data, 4, getbit(r, 0), 1);
+			orbits8_ptr(pb->data, 5, a, 2);
+			orbits8_ptr(pb->data, 7, b, 2);
+			orbits8_ptr(pb->data, 9, h, 1);
+			orbits8_ptr(pb->data, 10, d, 1);
+
+			// Partitions
+			orbits8_ptr(pb->data, 11, part_value, 2);
+
+			// CEM
+			orbits8_ptr(pb->data, cem_offset, cem, cem_bits);
+
+			copy_bytes(endpoint_ise, MAXIMUM_ENCODED_COLOR_ENDPOINT_BYTES, pb->data, ced_offset);
+
+			reverse_bytes(weights_ise, MAXIMUM_ENCODED_WEIGHT_BYTES, pb->data + 15);
+
+			if (dual_plane)
+			{
+				size_t bits_for_weights = compute_ise_bitcount(2 * BLOCK_TEXEL_COUNT, weight_quant);
+
+				size_t secondPlaneChannel = 3;
+				orbits8_ptr(pb->data, 128 - bits_for_weights - 2, secondPlaneChannel, 2);
+			}
+		}
+
+
+		uint8 quantize_weight(range_t weight_quant, size_t weight)
+		{
+			DCHECK(weight_quant <= RANGE_32);
+			// anticto: this may happen because of rounding in some ranges.
+			// apparently, it is ok to clamp it, based on what the arm reference implementation does.
+			if (weight > 1024) weight = 1024;
+			//DCHECK(weight <= 1024);
+			if ((weight_quant <= RANGE_32) && (weight <= 1024))
+			{
+				return weight_quantize_table[weight_quant][weight];
+			}
+			return 0;
+		}
+
+		/**
+		 * Project a texel to a line and quantize the result in 1 dimension.
+		 *
+		 * The line is defined by t=k*x + m. This function calculates and quantizes x
+		 * by projecting n=t-m onto k, x=|n|/|k|. Since k and m is derived from the
+		 * minimum and maximum of all texel values the result will be in the range [0,
+		 * 1].
+		 *
+		 * To quantize the result using the weight_quantize_table the value needs to
+		 * be extended to the range [0, 1024].
+		 *
+		 * @param k the derivative of the line
+		 * @param m the minimum endpoint
+		 * @param t the texel value
+		 */
+		size_t project(size_t k, size_t m, size_t t)
+		{
+			DCHECK(k > 0);
+			// anticto fix: underflow is possible because we use the unquantized limit, which may be
+			// bigger than the value. so i think we need to clamp.
+			// return size_t((t - m) * 1024) / k;
+			return size_t(FMath::Max(0, FMath::Min(int(k), int(t) - int(m))) * 1024) / k;
+		}
+
+		template<typename T> T clamp(T a, T b, T x)
+		{
+			if (x < a)
+			{
+				return a;
+			}
+
+			if (x > b)
+			{
+				return b;
+			}
+
+			return x;
+		}
+
+
+		/**
+		 * Project a texel to a line and quantize the result in 3 dimensions.
+		 */
+		size_t project(vec3i_t k, int kk, vec3i_t m, vec3i_t t)
+		{
+			DCHECK(kk > 0);
+			return static_cast<size_t>(clamp(0, 1024, dot(t - m, k) * 1024 / kk));
+		}
+
+		size_t project(vec4i_t k, int kk, vec4i_t m, vec4i_t t)
+		{
+			DCHECK(kk > 0);
+			return static_cast<size_t>(clamp(0, 1024, dot(t - m, k) * 1024 / kk));
+		}
+
+		void calculate_quantized_weights_luminance(const uint8 texels[BLOCK_TEXEL_COUNT],
+			range_t quant,
+			uint8 l0,
+			uint8 l1,
+			uint8 weights[BLOCK_TEXEL_COUNT])
+		{
+			DCHECK(l0 <= l1);
+			if (l0 < l1)
+			{
+				size_t k = l1 - l0;
+				size_t m = l0;
+
+				for (size_t i = 0; i < BLOCK_TEXEL_COUNT; ++i)
+				{
+					size_t t = static_cast<size_t>(texels[i]);
+					weights[i] = quantize_weight(quant, project(k, m, t));
+				}
+			}
+			else
+			{
+				for (size_t i = 0; i < BLOCK_TEXEL_COUNT; ++i)
+				{
+					DCHECK(static_cast<size_t>(texels[i]) == l0);
+					weights[i] = quantize_weight(quant, 0);
+				}
+			}
+		}
+
+
+		void calculate_quantized_weights_channel(const unorm8_t texels[BLOCK_TEXEL_COUNT],
+			range_t quant,
+			uint8 l0,
+			uint8 l1,
+			uint8 weights[BLOCK_TEXEL_COUNT],
+			int channel)
+		{
+			if (l0 == l1)
+			{
+				for (size_t i = 0; i < BLOCK_TEXEL_COUNT; ++i)
+				{
+					weights[i] = 0;  // quantize_weight(quant, 0) is always 0
+				}
+			}
+			else
+			{
+				DCHECK(l0 < l1);
+
+				size_t k = l1 - l0;
+				size_t m = l0;
+
+				for (size_t i = 0; i < BLOCK_TEXEL_COUNT; ++i)
+				{
+					size_t t = static_cast<size_t>(texels[i].components[channel]);
+					weights[i] = quantize_weight(quant, project(k, m, t));
+				}
+			}
+		}
+
+
+		void calculate_quantized_weights_rgb(const unorm8_t texels[BLOCK_TEXEL_COUNT],
+			range_t quant,
+			vec3i_t e0,
+			vec3i_t e1,
+			uint8 weights[BLOCK_TEXEL_COUNT])
+		{
+			if (e0 == e1)
+			{
+				for (size_t i = 0; i < BLOCK_TEXEL_COUNT; ++i)
+				{
+					weights[i] = 0;  // quantize_weight(quant, 0) is always 0
+				}
+			}
+			else
+			{
+				vec3i_t k = e1 - e0;
+				vec3i_t m = e0;
+
+				int kk = dot(k, k);
+				for (size_t i = 0; i < BLOCK_TEXEL_COUNT; ++i)
+				{
+					weights[i] = quantize_weight(quant, project(k, kk, m, to_vec3i(texels[i])));
+				}
+			}
+		}
+
+		void calculate_quantized_weights_rgba(const unorm8_t texels[BLOCK_TEXEL_COUNT],
+			range_t quant,
+			vec4i_t e0,
+			vec4i_t e1,
+			uint8 weights[BLOCK_TEXEL_COUNT])
+		{
+			if (e0 == e1)
+			{
+				for (size_t i = 0; i < BLOCK_TEXEL_COUNT; ++i)
+				{
+					weights[i] = 0;  // quantize_weight(quant, 0) is always 0
+				}
+			}
+			else
+			{
+				vec4i_t k = e1 - e0;
+				vec4i_t m = e0;
+
+				int kk = dot(k, k);
+				for (size_t i = 0; i < BLOCK_TEXEL_COUNT; ++i)
+				{
+					weights[i] = quantize_weight(quant, project(k, kk, m, to_vec4i(texels[i])));
+				}
+			}
+		}
+
+
+		range_t endpoint_quantization(size_t partitions,
+			range_t weight_quant,
+			color_endpoint_mode_t endpoint_mode)
+		{
+			int8_t ce_range = color_endpoint_range_table[partitions - 1][weight_quant][endpoint_mode];
+			DCHECK(ce_range >= 0 && ce_range <= RANGE_MAX);
+			return static_cast<range_t>(ce_range);
+		}
+
+
+
+		/**
+		 * Write void extent block bits for LDR mode and unused extent coordinates.
+		 */
+		void encode_void_extent(vec3i_t color, PhysicalBlock* physical_block)
+		{
+			void_extent_to_physical(unorm8_to_unorm16(to_unorm8(color)), physical_block);
+		}
+
+
+		void encode_void_extent(vec4i_t color, PhysicalBlock* physical_block)
+		{
+			void_extent_to_physical(unorm8_to_unorm16(to_unorm8(color)), physical_block);
+		}
+
+
+		void encode_luminance(const uint8 texels[BLOCK_TEXEL_COUNT],
+			PhysicalBlock* physical_block)
+		{
+			size_t partition_count = 1;
+
+			color_endpoint_mode_t color_endpoint_mode = CEM_LDR_LUMINANCE_DIRECT;
+			range_t weight_quant = RANGE_32;
+			range_t endpoint_quant = endpoint_quantization(partition_count, weight_quant, color_endpoint_mode);
+
+			uint8 l0 = 255;
+			uint8 l1 = 0;
+			for (size_t i = 0; i < BLOCK_TEXEL_COUNT; ++i)
+			{
+				l0 = FMath::Min(l0, texels[i]);
+				l1 = FMath::Max(l1, texels[i]);
+			}
+
+			uint8 endpoint_unquantized[2];
+			uint8 endpoint_quantized[2];
+			encode_luminance_direct(endpoint_quant, l0, l1, endpoint_quantized, endpoint_unquantized);
+
+			uint8 weights_quantized[BLOCK_TEXEL_COUNT];
+			calculate_quantized_weights_luminance(
+				texels, weight_quant, endpoint_unquantized[0], endpoint_unquantized[1],
+				weights_quantized);
+
+			uint8 endpoint_ise[MAXIMUM_ENCODED_COLOR_ENDPOINT_BYTES] = { 0 };
+			integer_sequence_encode(endpoint_quantized, 2, RANGE_256, endpoint_ise);
+
+			uint8 weights_ise[MAXIMUM_ENCODED_WEIGHT_BYTES + 1] = { 0 };
+			integer_sequence_encode(weights_quantized, BLOCK_TEXEL_COUNT, RANGE_32, weights_ise);
+
+			symbolic_to_physical(color_endpoint_mode, endpoint_quant, weight_quant,
+				partition_count, endpoint_ise,
+				weights_ise, physical_block);
+		}
+
+
+		void encode_rgb_single_partition(const unorm8_t texels[BLOCK_TEXEL_COUNT],
+			vec3f_t e0,
+			vec3f_t e1,
+			PhysicalBlock* physical_block)
+		{
+			size_t partition_count = 1;
+
+			color_endpoint_mode_t color_endpoint_mode = CEM_LDR_RGB_DIRECT;
+			range_t weight_quant = RANGE_12;
+			range_t endpoint_quant = endpoint_quantization(partition_count, weight_quant, color_endpoint_mode);
+
+			vec3i_t endpoint_unquantized[2];
+			uint8 endpoint_quantized[6];
+			encode_rgb_direct(endpoint_quant, round(e0), round(e1), endpoint_quantized,
+				endpoint_unquantized);
+
+			uint8 weights_quantized[BLOCK_TEXEL_COUNT];
+			calculate_quantized_weights_rgb(texels, weight_quant, endpoint_unquantized[0],
+				endpoint_unquantized[1], weights_quantized);
+
+			uint8 endpoint_ise[MAXIMUM_ENCODED_COLOR_ENDPOINT_BYTES] = { 0 };
+			integer_sequence_encode(endpoint_quantized, 6, endpoint_quant, endpoint_ise);
+
+			uint8 weights_ise[MAXIMUM_ENCODED_WEIGHT_BYTES + 1] = { 0 };
+			integer_sequence_encode(weights_quantized, BLOCK_TEXEL_COUNT, weight_quant,
+				weights_ise);
+
+			symbolic_to_physical(color_endpoint_mode, endpoint_quant, weight_quant,
+				partition_count, endpoint_ise,
+				weights_ise, physical_block);
+		}
+
+
+
+		void encode_rgba_single_partition_dual_plane(const unorm8_t texels[BLOCK_TEXEL_COUNT],
+			vec4f_t e0,
+			vec4f_t e1,
+			PhysicalBlock* physical_block)
+		{
+			size_t partition_count = 1;
+
+			color_endpoint_mode_t color_endpoint_mode = CEM_LDR_RGBA_DIRECT;
+			range_t weight_quant = RANGE_5;
+			range_t endpoint_quant = RANGE_16;
+			//      range_t weight_quant = RANGE_4;
+			//      range_t endpoint_quant = RANGE_48;
+
+
+			//      if (s_debuglog)
+			//      {
+			//          size_t endpointBits = compute_ise_bitcount(8, endpoint_quant);
+			//          size_t weightsBits  = compute_ise_bitcount(BLOCK_TEXEL_COUNT*2, weight_quant);
+			//          UE_LOG(LogMutableCore,Warning," interesting block. %d endpoint bits, %d weight bits", endpointBits, weightsBits );
+			//      }
+
+			vec4i_t endpoint_unquantized[2];
+			uint8 endpoint_quantized[8];
+			encode_rgba_direct(endpoint_quant, round(e0), round(e1), endpoint_quantized,
+				endpoint_unquantized);
+
+			//      if (s_debuglog)
+			//      {
+			//          UE_LOG(LogMutableCore,Warning," unquantized endpoints: %.3d %.3d %.3d %.3d \t %.3d %.3d %.3d %.3d",
+			//                  endpoint_unquantized[0].x, endpoint_unquantized[0].y, endpoint_unquantized[0].z, endpoint_unquantized[0].w,
+			//                  endpoint_unquantized[1].x, endpoint_unquantized[1].y, endpoint_unquantized[1].z, endpoint_unquantized[1].w );
+
+			//          UE_LOG(LogMutableCore,Warning," quantized endpoints:   %.3d %.3d %.3d %.3d \t %.3d %.3d %.3d %.3d",
+			//                  endpoint_quantized[0], endpoint_quantized[1], endpoint_quantized[2], endpoint_quantized[3],
+			//                  endpoint_quantized[4], endpoint_quantized[5], endpoint_quantized[6], endpoint_quantized[7] );
+			//      }
+
+			uint8 weights_quantized[BLOCK_TEXEL_COUNT];
+			calculate_quantized_weights_rgb(texels, weight_quant,
+				endpoint_unquantized[0].rgb(),
+				endpoint_unquantized[1].rgb(),
+				weights_quantized);
+
+
+			uint8 alpha_weights_quantized[BLOCK_TEXEL_COUNT];
+			calculate_quantized_weights_channel(texels, weight_quant,
+				(uint8)endpoint_unquantized[0].a,
+				(uint8)endpoint_unquantized[1].a,
+				alpha_weights_quantized,
+				3);
+
+			//      if (s_debuglog)
+			//      {
+			//          int i=0;
+			//          UE_LOG(LogMutableCore,Warning," quantized alpha weights:");
+			//          UE_LOG(LogMutableCore,Warning," %.3d %.3d %.3d %.3d ",
+			//                  alpha_weights_quantized[i+0], alpha_weights_quantized[i+1], alpha_weights_quantized[i+2], alpha_weights_quantized[i+3] );
+			//          i+=4;
+			//          UE_LOG(LogMutableCore,Warning," %.3d %.3d %.3d %.3d ",
+			//                  alpha_weights_quantized[i+0], alpha_weights_quantized[i+1], alpha_weights_quantized[i+2], alpha_weights_quantized[i+3] );
+			//          i+=4;
+			//          UE_LOG(LogMutableCore,Warning," %.3d %.3d %.3d %.3d ",
+			//                  alpha_weights_quantized[i+0], alpha_weights_quantized[i+1], alpha_weights_quantized[i+2], alpha_weights_quantized[i+3] );
+			//          i+=4;
+			//          UE_LOG(LogMutableCore,Warning," %.3d %.3d %.3d %.3d ",
+			//                  alpha_weights_quantized[i+0], alpha_weights_quantized[i+1], alpha_weights_quantized[i+2], alpha_weights_quantized[i+3] );
+			//          i+=4;
+			//      }
+
+			uint8 final_weights_quantized[BLOCK_TEXEL_COUNT * 2];
+			for (size_t i = 0; i < BLOCK_TEXEL_COUNT; i++)
+			{
+				final_weights_quantized[2 * i] = weights_quantized[i];
+				final_weights_quantized[2 * i + 1] = alpha_weights_quantized[i];
+			}
+
+			uint8 endpoint_ise[MAXIMUM_ENCODED_COLOR_ENDPOINT_BYTES] = { 0 };
+			integer_sequence_encode(endpoint_quantized, 8, endpoint_quant, endpoint_ise);
+
+			uint8 weights_ise[MAXIMUM_ENCODED_WEIGHT_BYTES + 1] = { 0 };
+			integer_sequence_encode(final_weights_quantized, BLOCK_TEXEL_COUNT * 2, weight_quant, weights_ise);
+
+			symbolic_to_physical(color_endpoint_mode, endpoint_quant, weight_quant,
+				partition_count, endpoint_ise,
+				weights_ise, physical_block, true);
+		}
+
+
+
+		void encode_rgba_single_partition(const unorm8_t texels[BLOCK_TEXEL_COUNT],
+			vec4f_t e0,
+			vec4f_t e1,
+			PhysicalBlock* physical_block)
+		{
+			size_t partition_count = 1;
+
+			color_endpoint_mode_t color_endpoint_mode = CEM_LDR_RGBA_DIRECT;
+			range_t weight_quant = RANGE_12;
+			range_t endpoint_quant = endpoint_quantization(partition_count, weight_quant, color_endpoint_mode);
+
+			vec4i_t endpoint_unquantized[2];
+			uint8 endpoint_quantized[8];
+			encode_rgba_direct(endpoint_quant, round(e0), round(e1), endpoint_quantized,
+				endpoint_unquantized);
+
+			// Calculate weights ignoring alpha (TODO)
+			uint8 weights_quantized[BLOCK_TEXEL_COUNT];
+			calculate_quantized_weights_rgba(texels, weight_quant,
+				endpoint_unquantized[0],
+				endpoint_unquantized[1],
+				weights_quantized);
+
+
+			uint8 endpoint_ise[MAXIMUM_ENCODED_COLOR_ENDPOINT_BYTES] = { 0 };
+			integer_sequence_encode(endpoint_quantized, 8, endpoint_quant, endpoint_ise);
+
+			uint8 weights_ise[MAXIMUM_ENCODED_WEIGHT_BYTES + 1] = { 0 };
+			integer_sequence_encode(weights_quantized, BLOCK_TEXEL_COUNT, weight_quant,
+				weights_ise);
+
+			symbolic_to_physical(color_endpoint_mode, endpoint_quant, weight_quant,
+				partition_count, endpoint_ise,
+				weights_ise, physical_block, false);
+		}
+
+
+
+		void encode_rg_single_partition_dual_plane(const unorm8_t texels[BLOCK_TEXEL_COUNT],
+			vec4f_t e0,
+			vec4f_t e1,
+			PhysicalBlock* physical_block)
+		{
+			size_t partition_count = 1;
+
+			color_endpoint_mode_t color_endpoint_mode = CEM_LDR_LUMINANCE_ALPHA_DIRECT;
+			range_t weight_quant = RANGE_6;
+			range_t endpoint_quant = RANGE_64;
+			//      range_t weight_quant = RANGE_5;
+			//      range_t endpoint_quant = RANGE_128;
+			//      range_t weight_quant = RANGE_8;
+			//      range_t endpoint_quant = RANGE_8;
+
+	//        if (s_debuglog)
+	//        {
+	//          size_t endpointBits = compute_ise_bitcount(4, endpoint_quant);
+	//          size_t weightsBits  = compute_ise_bitcount(BLOCK_TEXEL_COUNT*2, weight_quant);
+	//          UE_LOG(LogMutableCore,Warning," interesting block. %d endpoint bits, %d weight bits", endpointBits, weightsBits );
+	//        }
+
+			vec4i_t endpoint_unquantized[2];
+			uint8 endpoint_quantized[8];
+			encode_rgba_direct(endpoint_quant, round(e0), round(e1), endpoint_quantized,
+				endpoint_unquantized);
+
+			//        if (s_debuglog)
+			//        {
+			//            UE_LOG(LogMutableCore,Warning," unquantized endpoints: %.3d %.3d %.3d %.3d \t %.3d %.3d %.3d %.3d",
+			//                  endpoint_unquantized[0].r, endpoint_unquantized[0].g, endpoint_unquantized[0].b, endpoint_unquantized[0].a,
+			//                  endpoint_unquantized[1].r, endpoint_unquantized[1].g, endpoint_unquantized[1].b, endpoint_unquantized[1].a );
+
+			//            UE_LOG(LogMutableCore,Warning," quantized endpoints:   %.3d %.3d %.3d %.3d \t %.3d %.3d %.3d %.3d",
+			//                  endpoint_quantized[0], endpoint_quantized[1], endpoint_quantized[2], endpoint_quantized[3],
+			//                  endpoint_quantized[4], endpoint_quantized[5], endpoint_quantized[6], endpoint_quantized[7] );
+			//        }
+
+			uint8 r_weights_quantized[BLOCK_TEXEL_COUNT];
+			calculate_quantized_weights_channel(texels, weight_quant,
+				(uint8)endpoint_unquantized[0].r,
+				(uint8)endpoint_unquantized[1].r,
+				r_weights_quantized,
+				2);
+
+
+			uint8 g_weights_quantized[BLOCK_TEXEL_COUNT];
+			calculate_quantized_weights_channel(texels, weight_quant,
+				(uint8)endpoint_unquantized[0].g,
+				(uint8)endpoint_unquantized[1].g,
+				g_weights_quantized,
+				1);
+
+			//        if (s_debuglog)
+			//        {
+			//            int i=0;
+			//            UE_LOG(LogMutableCore,Warning," r values:");
+			//            for (int j=0;j<4;++j)
+			//            {
+			//                UE_LOG(LogMutableCore,Warning," %.3d %.3d %.3d %.3d ",
+			//                        texels[i].components[2], texels[i+1].components[2], texels[i+2].components[2], texels[i+3].components[2] );
+			//                i+=4;
+			//            }
+
+			//            i=0;
+			//            UE_LOG(LogMutableCore,Warning," quantized r weights:");
+			//            for (int j=0;j<4;++j)
+			//            {
+			//                UE_LOG(LogMutableCore,Warning," %.3d %.3d %.3d %.3d ",
+			//                        r_weights_quantized[i+0], r_weights_quantized[i+1], r_weights_quantized[i+2], r_weights_quantized[i+3] );
+			//                i+=4;
+			//            }
+
+			////            i=0;
+			////            UE_LOG(LogMutableCore,Warning," quantized g weights:");
+			////            for (int j=0;j<4;++j)
+			////            {
+			////                UE_LOG(LogMutableCore,Warning," %.3d %.3d %.3d %.3d ",
+			////                        g_weights_quantized[i+0], g_weights_quantized[i+1], g_weights_quantized[i+2], g_weights_quantized[i+3] );
+			////                i+=4;
+			////            }
+			//        }
+
+			uint8 final_weights_quantized[BLOCK_TEXEL_COUNT * 2];
+			for (size_t i = 0; i < BLOCK_TEXEL_COUNT; i++)
+			{
+				final_weights_quantized[2 * i] = r_weights_quantized[i];
+				final_weights_quantized[2 * i + 1] = g_weights_quantized[i];
+			}
+
+			uint8 endpoint_ise[MAXIMUM_ENCODED_COLOR_ENDPOINT_BYTES] = { 0 };
+			integer_sequence_encode(endpoint_quantized, 4, endpoint_quant, endpoint_ise);
+
+			uint8 weights_ise[MAXIMUM_ENCODED_WEIGHT_BYTES + 1] = { 0 };
+			integer_sequence_encode(final_weights_quantized, BLOCK_TEXEL_COUNT * 2, weight_quant, weights_ise);
+
+			symbolic_to_physical(color_endpoint_mode, endpoint_quant, weight_quant,
+				partition_count, endpoint_ise,
+				weights_ise, physical_block, true);
+		}
+
+
+		bool is_solid(const unorm8_t texels[BLOCK_TEXEL_COUNT],
+			size_t count,
+			unorm8_t* color)
+		{
+			for (size_t i = 0; i < count; ++i)
+			{
+				if (!approx_equal(to_vec3i(texels[i]), to_vec3i(texels[0])))
+				{
+					return false;
+				}
+			}
+
+			// TODO: Calculate average color?
+			*color = texels[0];
+			return true;
+		}
+
+
+		bool is_solid_rgba(const unorm8_t texels[BLOCK_TEXEL_COUNT],
+			size_t count,
+			unorm8_t* color)
+		{
+			for (size_t i = 0; i < count; ++i)
+			{
+				if (!approx_equal(to_vec4i(texels[i]), to_vec4i(texels[0])))
+				{
+					return false;
+				}
+			}
+
+			// TODO: Calculate average color?
+			*color = texels[0];
+			return true;
+		}
+
+
+		bool is_greyscale(const unorm8_t texels[BLOCK_TEXEL_COUNT],
+			size_t count,
+			uint8 luminances[BLOCK_TEXEL_COUNT])
+		{
+			for (size_t i = 0; i < count; ++i)
+			{
+				vec3i_t color = to_vec3i(texels[i]);
+				luminances[i] = static_cast<uint8>(luminance(color));
+				vec3i_t lum(luminances[i], luminances[i], luminances[i]);
+				if (!approx_equal(color, lum))
+				{
+					return false;
+				}
+			}
+
+			return true;
+		}
+
+
+
+		struct mat3x3f_t
+		{
+		public:
+			mat3x3f_t() {}
+
+			mat3x3f_t(float m00,
+				float m01,
+				float m02,
+				float m10,
+				float m11,
+				float m12,
+				float m20,
+				float m21,
+				float m22)
+			{
+				m[0] = vec3f_t(m00, m01, m02);
+				m[1] = vec3f_t(m10, m11, m12);
+				m[2] = vec3f_t(m20, m21, m22);
+			}
+
+			const vec3f_t& row(size_t i) const { return m[i]; }
+
+			float& at(size_t i, size_t j) { return m[i].components(j); }
+			const float& at(size_t i, size_t j) const { return m[i].components(j); }
+
+		private:
+			vec3f_t m[3];
+		};
+
+		inline vec3f_t operator*(const mat3x3f_t& a, vec3f_t b)
+		{
+			vec3f_t tmp;
+			tmp.r = dot(a.row(0), b);
+			tmp.g = dot(a.row(1), b);
+			tmp.b = dot(a.row(2), b);
+			return tmp;
+		}
+
+
+		vec3f_t mean(const unorm8_t texels[BLOCK_TEXEL_COUNT], size_t count)
+		{
+			vec3i_t sum(0, 0, 0);
+			for (size_t i = 0; i < count; ++i)
+			{
+				sum = sum + to_vec3i(texels[i]);
+			}
+
+			return to_vec3f(sum) / static_cast<float>(count);
+		}
+
+		void subtract(const unorm8_t texels[BLOCK_TEXEL_COUNT],
+			size_t count,
+			vec3f_t v,
+			vec3f_t output[BLOCK_TEXEL_COUNT])
+		{
+			for (size_t i = 0; i < count; ++i)
+			{
+				output[i] = to_vec3f(texels[i]) - v;
+			}
+		}
+
+		mat3x3f_t covariance(const vec3f_t m[BLOCK_TEXEL_COUNT], size_t count)
+		{
+			mat3x3f_t cov;
+			for (size_t i = 0; i < 3; ++i)
+			{
+				for (size_t j = 0; j < 3; ++j)
+				{
+					float s = 0;
+					for (size_t k = 0; k < count; ++k)
+					{
+						s += m[k].components(i) * m[k].components(j);
+					}
+					cov.at(i, j) = s / static_cast<float>(count - 1);
+				}
+			}
+
+			return cov;
+		}
+
+		void eigen_vector(const mat3x3f_t& a, vec3f_t& eig)
+		{
+			vec3f_t b = signorm(vec3f_t(1, 3, 2));  // FIXME: Magic number
+			for (size_t i = 0; i < 8; ++i)
+			{
+				b = signorm(a * b);
+			}
+
+			eig = b;
+		}
+
+
+
+		inline bool approx_equal(float x, float y, float epsilon)
+		{
+			return fabs(x - y) < epsilon;
+		}
+
+
+		void find_min_max(const unorm8_t texels[BLOCK_TEXEL_COUNT],
+			size_t count,
+			vec3f_t line_k,
+			vec3f_t line_m,
+			vec3f_t& e0,
+			vec3f_t& e1)
+		{
+			DCHECK(count <= BLOCK_TEXEL_COUNT);
+			DCHECK(approx_equal(quadrance(line_k), 1.0, 0.0001f));
+
+			float a, b;
+			{
+				float t = dot(to_vec3f(texels[0]) - line_m, line_k);
+				a = t;
+				b = t;
+			}
+
+			for (size_t i = 1; i < count; ++i) {
+				float t = dot(to_vec3f(texels[i]) - line_m, line_k);
+				a = FMath::Min(a, t);
+				b = FMath::Max(b, t);
+			}
+
+			e0 = clamp_rgb(line_k * a + line_m);
+			e1 = clamp_rgb(line_k * b + line_m);
+		}
+
+		void find_min_max_block(const unorm8_t texels[BLOCK_TEXEL_COUNT],
+			vec3f_t line_k,
+			vec3f_t line_m,
+			vec3f_t& e0,
+			vec3f_t& e1)
+		{
+			find_min_max(texels, BLOCK_TEXEL_COUNT, line_k, line_m, e0, e1);
+		}
+
+
+		void principal_component_analysis(const unorm8_t texels[BLOCK_TEXEL_COUNT],
+			size_t count,
+			vec3f_t& line_k,
+			vec3f_t& line_m) {
+			// Since we are working with fixed sized blocks count we can cap count. This
+			// avoids dynamic allocation.
+			DCHECK(count <= BLOCK_TEXEL_COUNT);
+
+			line_m = mean(texels, count);
+
+			vec3f_t n[BLOCK_TEXEL_COUNT];
+			subtract(texels, count, line_m, n);
+
+			mat3x3f_t w = covariance(n, count);
+
+			eigen_vector(w, line_k);
+		}
+
+
+		inline void principal_component_analysis_block(
+			const unorm8_t texels[BLOCK_TEXEL_COUNT],
+			vec3f_t& line_k,
+			vec3f_t& line_m)
+		{
+			principal_component_analysis(texels, BLOCK_TEXEL_COUNT, line_k, line_m);
+		}
+
+
+
+		void compress_block(const unorm8_t texels[BLOCK_TEXEL_COUNT],
+			PhysicalBlock* physical_block)
+		{
+			{
+				unorm8_t color;
+				if (is_solid(texels, BLOCK_TEXEL_COUNT, &color))
+				{
+					encode_void_extent(to_vec3i(color), physical_block);
+					return;
+				}
+			}
+
+			{
+				uint8 luminances[BLOCK_TEXEL_COUNT];
+				if (is_greyscale(texels, BLOCK_TEXEL_COUNT, luminances))
+				{
+					encode_luminance(luminances, physical_block);
+					return;
+				}
+			}
+
+			vec3f_t k, m;
+			principal_component_analysis_block(texels, k, m);
+			vec3f_t e0, e1;
+			find_min_max_block(texels, k, m, e0, e1);
+			encode_rgb_single_partition(texels, e0, e1, physical_block);
+		}
+
+
+		void compress_block_rgba(const unorm8_t texels[BLOCK_TEXEL_COUNT],
+			PhysicalBlock* physical_block)
+		{
+			unorm8_t color;
+			bool isSolidRGBA = is_solid_rgba(texels, BLOCK_TEXEL_COUNT, &color);
+			if (isSolidRGBA)
+			{
+				encode_void_extent(to_vec4i(color), physical_block);
+				return;
+			}
+
+			vec3f_t e0, e1;
+			bool isSolid = is_solid(texels, BLOCK_TEXEL_COUNT, &color);
+			if (isSolid)
+			{
+				e0 = e1 = to_vec3f(texels[0]);
+			}
+			else
+			{
+				vec3f_t k, m;
+				principal_component_analysis_block(texels, k, m);
+				find_min_max_block(texels, k, m, e0, e1);
+			}
+
+
+			uint8 l0 = 255;
+			uint8 l1 = 0;
+			for (size_t i = 0; i < BLOCK_TEXEL_COUNT; ++i)
+			{
+				l0 = FMath::Min(l0, texels[i].components[3]);
+				l1 = FMath::Max(l1, texels[i].components[3]);
+			}
+
+			vec4f_t re0(e0.r, e0.g, e0.b, float(l0));
+			vec4f_t re1(e1.r, e1.g, e1.b, float(l1));
+
+			//        if (s_debuglog)
+			//        {
+			//            UE_LOG(LogMutableCore,Warning," limits: %.3f %.3f %.3f %.3f \t %.3f %.3f %.3f %.3f",
+			//                    re0.x, re0.y, re0.z, re0.w,
+			//                    re1.x, re1.y, re1.z, re1.w );
+			//        }
+
+					//encode_rgba_single_partition(texels, re0, re1, physical_block);
+			encode_rgba_single_partition_dual_plane(texels, re0, re1, physical_block);
+		}
+
+
+		void compress_block_rg(const unorm8_t texels[BLOCK_TEXEL_COUNT],
+			PhysicalBlock* physical_block)
+		{
+			unorm8_t color;
+			bool isSolidRGBA = is_solid_rgba(texels, BLOCK_TEXEL_COUNT, &color);
+			if (isSolidRGBA)
+			{
+				encode_void_extent(vec4i_t(color.components[2], color.components[2], color.components[2], color.components[1]), physical_block);
+				return;
+			}
+
+			uint8 r0 = 255;
+			uint8 r1 = 0;
+			uint8 g0 = 255;
+			uint8 g1 = 0;
+			for (size_t i = 0; i < BLOCK_TEXEL_COUNT; ++i)
+			{
+				r0 = FMath::Min(r0, texels[i].components[2]);
+				r1 = FMath::Max(r1, texels[i].components[2]);
+				g0 = FMath::Min(g0, texels[i].components[1]);
+				g1 = FMath::Max(g1, texels[i].components[1]);
+			}
+
+			vec4f_t re0(r0, g0, r0, g0);
+			vec4f_t re1(r1, g1, r1, g1);
+
+			encode_rg_single_partition_dual_plane(texels, re0, re1, physical_block);
+		}
+
+	}
+
+
+
+
+
+	//-------------------------------------------------------------------------------------------------
+	// Based on ARMs reference implementation, under a special license
+	// See: [REDACTED_PROFILE_URL]/astc-encoder
+	//-------------------------------------------------------------------------------------------------
+	namespace arm
+	{
+		// Macro to silence warnings on ignored parameters.
+		// The presence of this macro should be a signal to look at refactoring.
+		//#define IGNORE(param) ((void)&param)
+
+#define astc_isnan(p) ((p)!=(p))
+
+// ASTC parameters
+#define MAX_TEXELS_PER_BLOCK 216
+#define MAX_WEIGHTS_PER_BLOCK 64
+#define MIN_WEIGHT_BITS_PER_BLOCK 24
+#define MAX_WEIGHT_BITS_PER_BLOCK 96
+#define PARTITION_BITS 10
+#define PARTITION_COUNT (1 << PARTITION_BITS)
+
+// the sum of weights for one texel.
+#define TEXEL_WEIGHT_SUM 16
+#define MAX_DECIMATION_MODES 87
+#define MAX_WEIGHT_MODES 2048
+
+// error reporting for codec internal errors.
+#define ASTC_CODEC_INTERNAL_ERROR check(false)
+
+
+/*
+   In ASTC, we don't necessarily provide a weight for every texel.
+   As such, for each block size, there are a number of patterns where some texels
+   have their weights computed as a weighted average of more than 1 weight.
+   As such, the codec uses a data structure that tells us: for each texel, which
+   weights it is a combination of for each weight, which texels it contributes to.
+   The decimation_table is this data structure.
+*/
+		struct decimation_table
+		{
+			int num_texels;
+			int num_weights;
+			uint8 texel_num_weights[MAX_TEXELS_PER_BLOCK];	// number of indices that go into the calculation for a texel
+			uint8 texel_weights_int[MAX_TEXELS_PER_BLOCK][4];	// the weight to assign to each weight
+			float texel_weights_float[MAX_TEXELS_PER_BLOCK][4];	// the weight to assign to each weight
+			uint8 texel_weights[MAX_TEXELS_PER_BLOCK][4];	// the weights that go into a texel calculation
+			uint8 weight_num_texels[MAX_WEIGHTS_PER_BLOCK];	// the number of texels that a given weight contributes to
+			uint8 weight_texel[MAX_WEIGHTS_PER_BLOCK][MAX_TEXELS_PER_BLOCK];	// the texels that the weight contributes to
+			uint8 weights_int[MAX_WEIGHTS_PER_BLOCK][MAX_TEXELS_PER_BLOCK];	// the weights that the weight contributes to a texel.
+			float weights_flt[MAX_WEIGHTS_PER_BLOCK][MAX_TEXELS_PER_BLOCK];	// the weights that the weight contributes to a texel.
+		};
+
+
+		struct block_mode
+		{
+			int8_t decimation_mode;
+			int8_t quantization_mode;
+			int8_t is_dual_plane;
+			int8_t permit_encode;
+			int8_t permit_decode;
+			float percentile;
+		};
+
+
+		struct block_size_descriptor
+		{
+			int decimation_mode_count;
+			int decimation_mode_samples[MAX_DECIMATION_MODES];
+			int decimation_mode_maxprec_1plane[MAX_DECIMATION_MODES];
+			int decimation_mode_maxprec_2planes[MAX_DECIMATION_MODES];
+			float decimation_mode_percentile[MAX_DECIMATION_MODES];
+			int permit_encode[MAX_DECIMATION_MODES];
+			const decimation_table* decimation_tables[MAX_DECIMATION_MODES + 1];
+			block_mode block_modes[MAX_WEIGHT_MODES];
+		};
+
+
+
+		struct symbolic_compressed_block
+		{
+			int error_block;			// 1 marks error block, 0 marks non-error-block.
+			int block_mode;				// 0 to 2047. Negative value marks constant-color block (-1: FP16, -2:UINT16)
+			int partition_count;		// 1 to 4; Zero marks a constant-color block.
+			int partition_index;		// 0 to 1023
+			int color_formats[4];		// color format for each endpoint color pair.
+			int color_formats_matched;	// color format for all endpoint pairs are matched.
+			int color_values[4][12];	// quantized endpoint color pairs.
+			int color_quantization_level;
+			uint8 plane1_weights[MAX_WEIGHTS_PER_BLOCK];	// quantized and decimated weights
+			uint8 plane2_weights[MAX_WEIGHTS_PER_BLOCK];
+			int plane2_color_component;	// color component for the secondary plane of weights
+			int constant_color[4];		// constant-color, as FP16 or UINT16. Used for constant-color blocks only.
+		};
+
+		struct physical_compressed_block
+		{
+			uint8 data[16];
+		};
+
+
+		static void initialize_decimation_table_2d(decimation_table* dt)
+		{
+			int i, j;
+			int x, y;
+
+			int texels_per_block = 16;
+			int weights_per_block = 16;
+
+			int weightcount_of_texel[MAX_TEXELS_PER_BLOCK];
+			int grid_weights_of_texel[MAX_TEXELS_PER_BLOCK][4];
+			int weights_of_texel[MAX_TEXELS_PER_BLOCK][4];
+
+			int texelcount_of_weight[MAX_WEIGHTS_PER_BLOCK];
+			// stack is to big
+			//        int texels_of_weight[MAX_WEIGHTS_PER_BLOCK][MAX_TEXELS_PER_BLOCK];
+			//        int texelweights_of_weight[MAX_WEIGHTS_PER_BLOCK][MAX_TEXELS_PER_BLOCK];
+			std::unique_ptr<int[]> texels_of_weight(new int[MAX_WEIGHTS_PER_BLOCK * MAX_TEXELS_PER_BLOCK]);
+			std::unique_ptr<int[]> texelweights_of_weight(new int[MAX_WEIGHTS_PER_BLOCK * MAX_TEXELS_PER_BLOCK]);
+
+			for (i = 0; i < MAX_WEIGHTS_PER_BLOCK; i++)
+			{
+				texelcount_of_weight[i] = 0;
+				weightcount_of_texel[i] = 0;
+			}
+
+			for (y = 0; y < 4; y++)
+				for (x = 0; x < 4; x++)
+				{
+					int texel = y * 4 + x;
+
+					int x_weight = (((1024 + 4 / 2) / (4 - 1)) * x * (4 - 1) + 32) >> 6;
+					int y_weight = (((1024 + 4 / 2) / (4 - 1)) * y * (4 - 1) + 32) >> 6;
+
+					int x_weight_frac = x_weight & 0xF;
+					int y_weight_frac = y_weight & 0xF;
+					int x_weight_int = x_weight >> 4;
+					int y_weight_int = y_weight >> 4;
+					int qweight[4];
+					int weight[4];
+					qweight[0] = x_weight_int + y_weight_int * 4;
+					qweight[1] = qweight[0] + 1;
+					qweight[2] = qweight[0] + 4;
+					qweight[3] = qweight[2] + 1;
+
+					// truncated-precision bilinear interpolation.
+					int prod = x_weight_frac * y_weight_frac;
+
+					weight[3] = (prod + 8) >> 4;
+					weight[1] = x_weight_frac - weight[3];
+					weight[2] = y_weight_frac - weight[3];
+					weight[0] = 16 - x_weight_frac - y_weight_frac + weight[3];
+
+					for (i = 0; i < 4; i++)
+						if (weight[i] != 0)
+						{
+							grid_weights_of_texel[texel][weightcount_of_texel[texel]] = qweight[i];
+							weights_of_texel[texel][weightcount_of_texel[texel]] = weight[i];
+							weightcount_of_texel[texel]++;
+							texels_of_weight[qweight[i] * MAX_TEXELS_PER_BLOCK + texelcount_of_weight[qweight[i]]] = texel;
+							texelweights_of_weight[qweight[i] * MAX_TEXELS_PER_BLOCK + texelcount_of_weight[qweight[i]]] = weight[i];
+							texelcount_of_weight[qweight[i]]++;
+						}
+				}
+
+			for (i = 0; i < texels_per_block; i++)
+			{
+				dt->texel_num_weights[i] = (uint8)weightcount_of_texel[i];
+
+				// ensure that all 4 entries are actually initialized.
+				// This allows a branch-free implementation of compute_value_of_texel_flt()
+				for (j = 0; j < 4; j++)
+				{
+					dt->texel_weights_int[i][j] = 0;
+					dt->texel_weights_float[i][j] = 0.0f;
+					dt->texel_weights[i][j] = 0;
+				}
+
+				for (j = 0; j < weightcount_of_texel[i]; j++)
+				{
+					dt->texel_weights_int[i][j] = (uint8)weights_of_texel[i][j];
+					dt->texel_weights_float[i][j] = static_cast <float>(weights_of_texel[i][j]) * (1.0f / TEXEL_WEIGHT_SUM);
+					dt->texel_weights[i][j] = (uint8)grid_weights_of_texel[i][j];
+				}
+			}
+
+			for (i = 0; i < weights_per_block; i++)
+			{
+				dt->weight_num_texels[i] = (uint8)texelcount_of_weight[i];
+
+
+				for (j = 0; j < texelcount_of_weight[i]; j++)
+				{
+					dt->weight_texel[i][j] = (uint8)texels_of_weight[i * MAX_TEXELS_PER_BLOCK + j];
+					dt->weights_int[i][j] = (uint8)texelweights_of_weight[i * MAX_TEXELS_PER_BLOCK + j];
+					dt->weights_flt[i][j] = static_cast <float>(texelweights_of_weight[i * MAX_TEXELS_PER_BLOCK + j]);
+				}
+			}
+
+			dt->num_texels = texels_per_block;
+			dt->num_weights = weights_per_block;
+		}
+
+
+		// enumeration of all the quantization methods we support under this format.
+		enum quantization_method
+		{
+			QUANT_2 = 0,
+			QUANT_3 = 1,
+			QUANT_4 = 2,
+			QUANT_5 = 3,
+			QUANT_6 = 4,
+			QUANT_8 = 5,
+			QUANT_10 = 6,
+			QUANT_12 = 7,
+			QUANT_16 = 8,
+			QUANT_20 = 9,
+			QUANT_24 = 10,
+			QUANT_32 = 11,
+			QUANT_40 = 12,
+			QUANT_48 = 13,
+			QUANT_64 = 14,
+			QUANT_80 = 15,
+			QUANT_96 = 16,
+			QUANT_128 = 17,
+			QUANT_160 = 18,
+			QUANT_192 = 19,
+			QUANT_256 = 20
+		};
+
+
+		int compute_ise_bitcount(int items, quantization_method quant)
+		{
+			switch (quant)
+			{
+			case QUANT_2:
+				return items;
+			case QUANT_3:
+				return (8 * items + 4) / 5;
+			case QUANT_4:
+				return 2 * items;
+			case QUANT_5:
+				return (7 * items + 2) / 3;
+			case QUANT_6:
+				return (13 * items + 4) / 5;
+			case QUANT_8:
+				return 3 * items;
+			case QUANT_10:
+				return (10 * items + 2) / 3;
+			case QUANT_12:
+				return (18 * items + 4) / 5;
+			case QUANT_16:
+				return items * 4;
+			case QUANT_20:
+				return (13 * items + 2) / 3;
+			case QUANT_24:
+				return (23 * items + 4) / 5;
+			case QUANT_32:
+				return 5 * items;
+			case QUANT_40:
+				return (16 * items + 2) / 3;
+			case QUANT_48:
+				return (28 * items + 4) / 5;
+			case QUANT_64:
+				return 6 * items;
+			case QUANT_80:
+				return (19 * items + 2) / 3;
+			case QUANT_96:
+				return (33 * items + 4) / 5;
+			case QUANT_128:
+				return 7 * items;
+			case QUANT_160:
+				return (22 * items + 2) / 3;
+			case QUANT_192:
+				return (38 * items + 4) / 5;
+			case QUANT_256:
+				return 8 * items;
+			default:
+				return 100000;
+			}
+		}
+
+
+		const float percentile_table_4x4[2048] =
+		{
+			1.0000f, 1.0000f, 1.0000f, 1.0000f, 1.0000f, 1.0000f, 1.0000f, 1.0000f, 1.0000f, 1.0000f, 1.0000f, 1.0000f, 1.0000f, 1.0000f, 1.0000f, 1.0000f,
+			1.0000f, 1.0000f, 1.0000f, 0.8661f, 1.0000f, 1.0000f, 1.0000f, 1.0000f, 1.0000f, 1.0000f, 1.0000f, 1.0000f, 1.0000f, 1.0000f, 1.0000f, 1.0000f,
+			1.0000f, 1.0000f, 0.7732f, 0.8567f, 1.0000f, 1.0000f, 1.0000f, 1.0000f, 1.0000f, 1.0000f, 1.0000f, 1.0000f, 1.0000f, 1.0000f, 1.0000f, 1.0000f,
+			1.0000f, 1.0000f, 0.7818f, 0.8914f, 1.0000f, 1.0000f, 1.0000f, 1.0000f, 1.0000f, 1.0000f, 1.0000f, 1.0000f, 1.0000f, 1.0000f, 1.0000f, 1.0000f,
+			1.0000f, 1.0000f, 0.4578f, 0.5679f, 1.0000f, 1.0000f, 1.0000f, 1.0000f, 1.0000f, 1.0000f, 1.0000f, 1.0000f, 1.0000f, 1.0000f, 1.0000f, 1.0000f,
+			1.0000f, 0.4183f, 0.4961f, 0.5321f, 1.0000f, 1.0000f, 1.0000f, 1.0000f, 1.0000f, 1.0000f, 1.0000f, 1.0000f, 1.0000f, 1.0000f, 1.0000f, 1.0000f,
+			1.0000f, 1.0000f, 1.0000f, 1.0000f, 1.0000f, 1.0000f, 1.0000f, 1.0000f, 1.0000f, 1.0000f, 1.0000f, 1.0000f, 1.0000f, 1.0000f, 1.0000f, 1.0000f,
+			1.0000f, 1.0000f, 1.0000f, 1.0000f, 1.0000f, 1.0000f, 1.0000f, 1.0000f, 1.0000f, 1.0000f, 1.0000f, 1.0000f, 1.0000f, 1.0000f, 1.0000f, 1.0000f,
+			1.0000f, 1.0000f, 1.0000f, 1.0000f, 1.0000f, 1.0000f, 1.0000f, 1.0000f, 1.0000f, 1.0000f, 1.0000f, 1.0000f, 1.0000f, 1.0000f, 1.0000f, 1.0000f,
+			1.0000f, 1.0000f, 1.0000f, 1.0000f, 1.0000f, 1.0000f, 1.0000f, 1.0000f, 1.0000f, 1.0000f, 1.0000f, 1.0000f, 1.0000f, 1.0000f, 1.0000f, 1.0000f,
+			1.0000f, 1.0000f, 1.0000f, 1.0000f, 1.0000f, 1.0000f, 1.0000f, 1.0000f, 1.0000f, 1.0000f, 1.0000f, 1.0000f, 1.0000f, 1.0000f, 1.0000f, 1.0000f,
+			1.0000f, 1.0000f, 1.0000f, 1.0000f, 1.0000f, 1.0000f, 1.0000f, 1.0000f, 1.0000f, 1.0000f, 1.0000f, 1.0000f, 1.0000f, 1.0000f, 1.0000f, 1.0000f,
+			1.0000f, 1.0000f, 1.0000f, 1.0000f, 1.0000f, 1.0000f, 1.0000f, 1.0000f, 1.0000f, 1.0000f, 1.0000f, 1.0000f, 1.0000f, 1.0000f, 1.0000f, 1.0000f,
+			1.0000f, 1.0000f, 1.0000f, 1.0000f, 1.0000f, 1.0000f, 1.0000f, 1.0000f, 1.0000f, 1.0000f, 1.0000f, 1.0000f, 1.0000f, 1.0000f, 1.0000f, 1.0000f,
+			1.0000f, 1.0000f, 1.0000f, 1.0000f, 1.0000f, 1.0000f, 1.0000f, 1.0000f, 1.0000f, 1.0000f, 1.0000f, 1.0000f, 1.0000f, 1.0000f, 1.0000f, 1.0000f,
+			1.0000f, 1.0000f, 1.0000f, 1.0000f, 1.0000f, 1.0000f, 1.0000f, 1.0000f, 1.0000f, 1.0000f, 1.0000f, 1.0000f, 1.0000f, 1.0000f, 1.0000f, 1.0000f,
+			1.0000f, 1.0000f, 1.0000f, 1.0000f, 1.0000f, 1.0000f, 1.0000f, 1.0000f, 1.0000f, 1.0000f, 1.0000f, 1.0000f, 1.0000f, 1.0000f, 1.0000f, 1.0000f,
+			1.0000f, 1.0000f, 1.0000f, 1.0000f, 1.0000f, 1.0000f, 1.0000f, 1.0000f, 1.0000f, 1.0000f, 1.0000f, 1.0000f, 1.0000f, 1.0000f, 1.0000f, 1.0000f,
+			1.0000f, 1.0000f, 1.0000f, 1.0000f, 1.0000f, 1.0000f, 1.0000f, 1.0000f, 1.0000f, 1.0000f, 1.0000f, 1.0000f, 1.0000f, 1.0000f, 1.0000f, 1.0000f,
+			1.0000f, 1.0000f, 1.0000f, 1.0000f, 1.0000f, 1.0000f, 1.0000f, 1.0000f, 1.0000f, 1.0000f, 1.0000f, 1.0000f, 1.0000f, 1.0000f, 1.0000f, 1.0000f,
+			1.0000f, 1.0000f, 1.0000f, 1.0000f, 1.0000f, 1.0000f, 1.0000f, 1.0000f, 1.0000f, 1.0000f, 1.0000f, 1.0000f, 1.0000f, 1.0000f, 1.0000f, 1.0000f,
+			1.0000f, 1.0000f, 1.0000f, 1.0000f, 1.0000f, 1.0000f, 1.0000f, 1.0000f, 1.0000f, 1.0000f, 1.0000f, 1.0000f, 1.0000f, 1.0000f, 1.0000f, 0.9151f,
+			1.0000f, 1.0000f, 1.0000f, 1.0000f, 1.0000f, 1.0000f, 1.0000f, 1.0000f, 1.0000f, 1.0000f, 1.0000f, 1.0000f, 1.0000f, 1.0000f, 1.0000f, 1.0000f,
+			1.0000f, 1.0000f, 1.0000f, 1.0000f, 1.0000f, 1.0000f, 1.0000f, 1.0000f, 1.0000f, 1.0000f, 1.0000f, 1.0000f, 1.0000f, 1.0000f, 1.0000f, 1.0000f,
+			1.0000f, 1.0000f, 1.0000f, 1.0000f, 1.0000f, 1.0000f, 1.0000f, 1.0000f, 1.0000f, 1.0000f, 1.0000f, 1.0000f, 1.0000f, 1.0000f, 1.0000f, 1.0000f,
+			1.0000f, 1.0000f, 1.0000f, 1.0000f, 1.0000f, 1.0000f, 1.0000f, 1.0000f, 1.0000f, 1.0000f, 1.0000f, 1.0000f, 1.0000f, 1.0000f, 1.0000f, 1.0000f,
+			1.0000f, 1.0000f, 1.0000f, 1.0000f, 1.0000f, 1.0000f, 1.0000f, 1.0000f, 1.0000f, 1.0000f, 1.0000f, 1.0000f, 1.0000f, 1.0000f, 1.0000f, 0.9400f,
+			1.0000f, 1.0000f, 1.0000f, 1.0000f, 1.0000f, 1.0000f, 1.0000f, 1.0000f, 1.0000f, 1.0000f, 1.0000f, 1.0000f, 1.0000f, 1.0000f, 1.0000f, 0.9678f,
+			1.0000f, 1.0000f, 1.0000f, 1.0000f, 1.0000f, 1.0000f, 1.0000f, 1.0000f, 1.0000f, 1.0000f, 1.0000f, 1.0000f, 1.0000f, 1.0000f, 0.8111f, 0.8833f,
+			1.0000f, 1.0000f, 1.0000f, 1.0000f, 1.0000f, 1.0000f, 1.0000f, 1.0000f, 1.0000f, 1.0000f, 1.0000f, 1.0000f, 1.0000f, 1.0000f, 0.8299f, 0.8988f,
+			1.0000f, 1.0000f, 1.0000f, 1.0000f, 1.0000f, 1.0000f, 1.0000f, 1.0000f, 1.0000f, 1.0000f, 1.0000f, 1.0000f, 1.0000f, 1.0000f, 1.0000f, 1.0000f,
+			1.0000f, 1.0000f, 1.0000f, 1.0000f, 1.0000f, 1.0000f, 1.0000f, 1.0000f, 1.0000f, 1.0000f, 1.0000f, 1.0000f, 1.0000f, 1.0000f, 1.0000f, 1.0000f,
+			1.0000f, 0.9182f, 0.9692f, 0.9820f, 1.0000f, 1.0000f, 1.0000f, 1.0000f, 1.0000f, 1.0000f, 1.0000f, 1.0000f, 1.0000f, 1.0000f, 1.0000f, 1.0000f,
+			1.0000f, 0.9663f, 0.9911f, 0.8707f, 1.0000f, 1.0000f, 1.0000f, 1.0000f, 1.0000f, 1.0000f, 1.0000f, 1.0000f, 1.0000f, 1.0000f, 1.0000f, 1.0000f,
+			1.0000f, 0.9088f, 0.9374f, 0.8793f, 1.0000f, 1.0000f, 1.0000f, 1.0000f, 1.0000f, 1.0000f, 1.0000f, 1.0000f, 1.0000f, 1.0000f, 1.0000f, 1.0000f,
+			1.0000f, 0.8750f, 0.8952f, 0.7356f, 1.0000f, 1.0000f, 1.0000f, 1.0000f, 1.0000f, 1.0000f, 1.0000f, 1.0000f, 1.0000f, 1.0000f, 1.0000f, 1.0000f,
+			1.0000f, 0.2746f, 0.0000f, 0.0772f, 1.0000f, 1.0000f, 1.0000f, 1.0000f, 1.0000f, 1.0000f, 1.0000f, 1.0000f, 1.0000f, 1.0000f, 1.0000f, 1.0000f,
+			1.0000f, 0.1487f, 0.2193f, 0.3263f, 1.0000f, 1.0000f, 1.0000f, 1.0000f, 1.0000f, 1.0000f, 1.0000f, 1.0000f, 1.0000f, 1.0000f, 1.0000f, 1.0000f,
+			1.0000f, 1.0000f, 1.0000f, 1.0000f, 1.0000f, 1.0000f, 1.0000f, 1.0000f, 1.0000f, 1.0000f, 1.0000f, 1.0000f, 1.0000f, 1.0000f, 1.0000f, 1.0000f,
+			1.0000f, 1.0000f, 1.0000f, 1.0000f, 1.0000f, 1.0000f, 1.0000f, 1.0000f, 1.0000f, 1.0000f, 1.0000f, 1.0000f, 1.0000f, 1.0000f, 1.0000f, 1.0000f,
+			1.0000f, 1.0000f, 1.0000f, 1.0000f, 1.0000f, 1.0000f, 1.0000f, 1.0000f, 1.0000f, 1.0000f, 1.0000f, 1.0000f, 1.0000f, 1.0000f, 1.0000f, 1.0000f,
+			1.0000f, 1.0000f, 1.0000f, 1.0000f, 1.0000f, 1.0000f, 1.0000f, 1.0000f, 1.0000f, 1.0000f, 1.0000f, 1.0000f, 1.0000f, 1.0000f, 1.0000f, 1.0000f,
+			1.0000f, 1.0000f, 1.0000f, 1.0000f, 1.0000f, 1.0000f, 1.0000f, 1.0000f, 1.0000f, 1.0000f, 1.0000f, 1.0000f, 1.0000f, 1.0000f, 1.0000f, 1.0000f,
+			1.0000f, 1.0000f, 1.0000f, 1.0000f, 1.0000f, 1.0000f, 1.0000f, 1.0000f, 1.0000f, 1.0000f, 1.0000f, 1.0000f, 1.0000f, 1.0000f, 1.0000f, 1.0000f,
+			1.0000f, 1.0000f, 1.0000f, 1.0000f, 1.0000f, 1.0000f, 1.0000f, 1.0000f, 1.0000f, 1.0000f, 1.0000f, 1.0000f, 1.0000f, 1.0000f, 1.0000f, 1.0000f,
+			1.0000f, 1.0000f, 1.0000f, 1.0000f, 1.0000f, 1.0000f, 1.0000f, 1.0000f, 1.0000f, 1.0000f, 1.0000f, 1.0000f, 1.0000f, 1.0000f, 1.0000f, 1.0000f,
+			1.0000f, 1.0000f, 1.0000f, 1.0000f, 1.0000f, 1.0000f, 1.0000f, 1.0000f, 1.0000f, 1.0000f, 1.0000f, 1.0000f, 1.0000f, 1.0000f, 1.0000f, 1.0000f,
+			1.0000f, 1.0000f, 1.0000f, 1.0000f, 1.0000f, 1.0000f, 1.0000f, 1.0000f, 1.0000f, 1.0000f, 1.0000f, 1.0000f, 1.0000f, 1.0000f, 1.0000f, 1.0000f,
+			1.0000f, 1.0000f, 1.0000f, 1.0000f, 1.0000f, 1.0000f, 1.0000f, 1.0000f, 1.0000f, 1.0000f, 1.0000f, 1.0000f, 1.0000f, 1.0000f, 1.0000f, 1.0000f,
+			1.0000f, 1.0000f, 1.0000f, 1.0000f, 1.0000f, 1.0000f, 1.0000f, 1.0000f, 1.0000f, 1.0000f, 1.0000f, 1.0000f, 1.0000f, 1.0000f, 1.0000f, 1.0000f,
+			1.0000f, 1.0000f, 1.0000f, 1.0000f, 1.0000f, 1.0000f, 1.0000f, 1.0000f, 1.0000f, 1.0000f, 1.0000f, 1.0000f, 1.0000f, 1.0000f, 0.9917f, 1.0000f,
+			1.0000f, 1.0000f, 1.0000f, 1.0000f, 1.0000f, 1.0000f, 1.0000f, 1.0000f, 1.0000f, 1.0000f, 1.0000f, 1.0000f, 1.0000f, 1.0000f, 0.9995f, 0.9996f,
+			1.0000f, 1.0000f, 1.0000f, 1.0000f, 1.0000f, 1.0000f, 1.0000f, 1.0000f, 1.0000f, 1.0000f, 1.0000f, 1.0000f, 1.0000f, 0.9516f, 0.9838f, 0.9927f,
+			1.0000f, 1.0000f, 1.0000f, 1.0000f, 1.0000f, 1.0000f, 1.0000f, 1.0000f, 1.0000f, 1.0000f, 1.0000f, 1.0000f, 1.0000f, 0.9731f, 0.9906f, 0.9448f,
+			1.0000f, 1.0000f, 1.0000f, 1.0000f, 1.0000f, 1.0000f, 1.0000f, 1.0000f, 1.0000f, 1.0000f, 1.0000f, 1.0000f, 1.0000f, 1.0000f, 1.0000f, 1.0000f,
+			1.0000f, 1.0000f, 1.0000f, 1.0000f, 1.0000f, 1.0000f, 1.0000f, 1.0000f, 1.0000f, 1.0000f, 1.0000f, 1.0000f, 1.0000f, 1.0000f, 1.0000f, 1.0000f,
+			1.0000f, 1.0000f, 1.0000f, 1.0000f, 1.0000f, 1.0000f, 1.0000f, 1.0000f, 1.0000f, 1.0000f, 1.0000f, 1.0000f, 1.0000f, 1.0000f, 0.9900f, 0.9999f,
+			1.0000f, 1.0000f, 1.0000f, 1.0000f, 1.0000f, 1.0000f, 1.0000f, 1.0000f, 1.0000f, 1.0000f, 1.0000f, 1.0000f, 1.0000f, 1.0000f, 0.9992f, 0.9991f,
+			1.0000f, 1.0000f, 1.0000f, 1.0000f, 1.0000f, 1.0000f, 1.0000f, 1.0000f, 1.0000f, 1.0000f, 1.0000f, 1.0000f, 1.0000f, 0.9949f, 0.9987f, 0.9936f,
+			1.0000f, 1.0000f, 1.0000f, 1.0000f, 1.0000f, 1.0000f, 1.0000f, 1.0000f, 1.0000f, 1.0000f, 1.0000f, 1.0000f, 1.0000f, 0.9966f, 0.9985f, 0.9615f,
+			1.0000f, 1.0000f, 1.0000f, 1.0000f, 1.0000f, 1.0000f, 1.0000f, 1.0000f, 1.0000f, 1.0000f, 1.0000f, 1.0000f, 1.0000f, 0.9269f, 0.9577f, 0.9057f,
+			1.0000f, 1.0000f, 1.0000f, 1.0000f, 1.0000f, 1.0000f, 1.0000f, 1.0000f, 1.0000f, 1.0000f, 1.0000f, 1.0000f, 1.0000f, 0.9023f, 0.9241f, 0.7499f,
+			1.0000f, 1.0000f, 1.0000f, 1.0000f, 1.0000f, 1.0000f, 1.0000f, 1.0000f, 1.0000f, 1.0000f, 1.0000f, 1.0000f, 1.0000f, 1.0000f, 1.0000f, 1.0000f,
+			1.0000f, 1.0000f, 1.0000f, 1.0000f, 1.0000f, 1.0000f, 1.0000f, 1.0000f, 1.0000f, 1.0000f, 1.0000f, 1.0000f, 1.0000f, 1.0000f, 1.0000f, 1.0000f,
+			1.0000f, 1.0000f, 0.9757f, 0.9846f, 1.0000f, 1.0000f, 1.0000f, 1.0000f, 1.0000f, 1.0000f, 1.0000f, 1.0000f, 1.0000f, 1.0000f, 1.0000f, 1.0000f,
+			1.0000f, 0.9922f, 0.9932f, 0.9792f, 1.0000f, 1.0000f, 1.0000f, 1.0000f, 1.0000f, 1.0000f, 1.0000f, 1.0000f, 1.0000f, 1.0000f, 1.0000f, 1.0000f,
+			1.0000f, 0.9881f, 0.9494f, 0.8178f, 1.0000f, 1.0000f, 1.0000f, 1.0000f, 1.0000f, 1.0000f, 1.0000f, 1.0000f, 1.0000f, 1.0000f, 1.0000f, 1.0000f,
+			1.0000f, 0.9780f, 0.8518f, 0.6206f, 1.0000f, 1.0000f, 1.0000f, 1.0000f, 1.0000f, 1.0000f, 1.0000f, 1.0000f, 1.0000f, 1.0000f, 1.0000f, 1.0000f,
+			1.0000f, 0.9472f, 0.7191f, 0.7003f, 1.0000f, 1.0000f, 1.0000f, 1.0000f, 1.0000f, 1.0000f, 1.0000f, 1.0000f, 1.0000f, 1.0000f, 1.0000f, 1.0000f,
+			1.0000f, 0.8356f, 0.3772f, 0.9971f, 1.0000f, 1.0000f, 1.0000f, 1.0000f, 1.0000f, 1.0000f, 1.0000f, 1.0000f, 1.0000f, 1.0000f, 1.0000f, 1.0000f,
+			1.0000f, 1.0000f, 1.0000f, 1.0000f, 1.0000f, 1.0000f, 1.0000f, 1.0000f, 1.0000f, 1.0000f, 1.0000f, 1.0000f, 1.0000f, 1.0000f, 1.0000f, 1.0000f,
+			1.0000f, 1.0000f, 1.0000f, 1.0000f, 1.0000f, 1.0000f, 1.0000f, 1.0000f, 1.0000f, 1.0000f, 1.0000f, 1.0000f, 1.0000f, 1.0000f, 1.0000f, 1.0000f,
+			1.0000f, 1.0000f, 1.0000f, 1.0000f, 1.0000f, 1.0000f, 1.0000f, 1.0000f, 1.0000f, 1.0000f, 1.0000f, 1.0000f, 1.0000f, 1.0000f, 1.0000f, 1.0000f,
+			1.0000f, 1.0000f, 1.0000f, 1.0000f, 1.0000f, 1.0000f, 1.0000f, 1.0000f, 1.0000f, 1.0000f, 1.0000f, 1.0000f, 1.0000f, 1.0000f, 1.0000f, 1.0000f,
+			1.0000f, 1.0000f, 1.0000f, 1.0000f, 1.0000f, 1.0000f, 1.0000f, 1.0000f, 1.0000f, 1.0000f, 1.0000f, 1.0000f, 1.0000f, 1.0000f, 1.0000f, 1.0000f,
+			1.0000f, 1.0000f, 1.0000f, 1.0000f, 1.0000f, 1.0000f, 1.0000f, 1.0000f, 1.0000f, 1.0000f, 1.0000f, 1.0000f, 1.0000f, 1.0000f, 1.0000f, 1.0000f,
+			1.0000f, 1.0000f, 1.0000f, 1.0000f, 1.0000f, 1.0000f, 1.0000f, 1.0000f, 1.0000f, 1.0000f, 1.0000f, 1.0000f, 1.0000f, 1.0000f, 1.0000f, 1.0000f,
+			1.0000f, 1.0000f, 1.0000f, 1.0000f, 1.0000f, 1.0000f, 1.0000f, 1.0000f, 1.0000f, 1.0000f, 1.0000f, 1.0000f, 1.0000f, 1.0000f, 1.0000f, 1.0000f,
+			1.0000f, 1.0000f, 1.0000f, 1.0000f, 1.0000f, 1.0000f, 1.0000f, 1.0000f, 1.0000f, 1.0000f, 1.0000f, 1.0000f, 1.0000f, 1.0000f, 1.0000f, 1.0000f,
+			1.0000f, 1.0000f, 1.0000f, 1.0000f, 1.0000f, 1.0000f, 1.0000f, 1.0000f, 1.0000f, 1.0000f, 1.0000f, 1.0000f, 1.0000f, 1.0000f, 1.0000f, 1.0000f,
+			1.0000f, 1.0000f, 1.0000f, 1.0000f, 1.0000f, 1.0000f, 1.0000f, 1.0000f, 1.0000f, 1.0000f, 1.0000f, 1.0000f, 1.0000f, 1.0000f, 1.0000f, 1.0000f,
+			1.0000f, 1.0000f, 1.0000f, 1.0000f, 1.0000f, 1.0000f, 1.0000f, 1.0000f, 1.0000f, 1.0000f, 1.0000f, 1.0000f, 1.0000f, 1.0000f, 1.0000f, 0.9997f,
+			1.0000f, 1.0000f, 1.0000f, 1.0000f, 1.0000f, 1.0000f, 1.0000f, 1.0000f, 1.0000f, 1.0000f, 1.0000f, 1.0000f, 1.0000f, 1.0000f, 0.9968f, 0.9980f,
+			1.0000f, 1.0000f, 1.0000f, 1.0000f, 1.0000f, 1.0000f, 1.0000f, 1.0000f, 1.0000f, 1.0000f, 1.0000f, 1.0000f, 1.0000f, 1.0000f, 0.9975f, 0.9977f,
+			1.0000f, 1.0000f, 1.0000f, 1.0000f, 1.0000f, 1.0000f, 1.0000f, 1.0000f, 1.0000f, 1.0000f, 1.0000f, 1.0000f, 1.0000f, 1.0000f, 0.9769f, 0.9875f,
+			1.0000f, 1.0000f, 1.0000f, 1.0000f, 1.0000f, 1.0000f, 1.0000f, 1.0000f, 1.0000f, 1.0000f, 1.0000f, 1.0000f, 1.0000f, 0.9945f, 0.9941f, 0.9861f,
+			1.0000f, 1.0000f, 1.0000f, 1.0000f, 1.0000f, 1.0000f, 1.0000f, 1.0000f, 1.0000f, 1.0000f, 1.0000f, 1.0000f, 1.0000f, 1.0000f, 1.0000f, 1.0000f,
+			1.0000f, 1.0000f, 1.0000f, 1.0000f, 1.0000f, 1.0000f, 1.0000f, 1.0000f, 1.0000f, 1.0000f, 1.0000f, 1.0000f, 1.0000f, 1.0000f, 1.0000f, 1.0000f,
+			1.0000f, 1.0000f, 1.0000f, 1.0000f, 1.0000f, 1.0000f, 1.0000f, 1.0000f, 1.0000f, 1.0000f, 1.0000f, 1.0000f, 1.0000f, 1.0000f, 0.9956f, 0.9982f,
+			1.0000f, 1.0000f, 1.0000f, 1.0000f, 1.0000f, 1.0000f, 1.0000f, 1.0000f, 1.0000f, 1.0000f, 1.0000f, 1.0000f, 1.0000f, 1.0000f, 0.9979f, 0.9989f,
+			1.0000f, 1.0000f, 1.0000f, 1.0000f, 1.0000f, 1.0000f, 1.0000f, 1.0000f, 1.0000f, 1.0000f, 1.0000f, 1.0000f, 1.0000f, 1.0000f, 0.9868f, 0.9854f,
+			1.0000f, 1.0000f, 1.0000f, 1.0000f, 1.0000f, 1.0000f, 1.0000f, 1.0000f, 1.0000f, 1.0000f, 1.0000f, 1.0000f, 1.0000f, 0.9888f, 0.9811f, 0.9706f,
+			1.0000f, 1.0000f, 1.0000f, 1.0000f, 1.0000f, 1.0000f, 1.0000f, 1.0000f, 1.0000f, 1.0000f, 1.0000f, 1.0000f, 1.0000f, 0.9894f, 0.9425f, 0.8465f,
+			1.0000f, 1.0000f, 1.0000f, 1.0000f, 1.0000f, 1.0000f, 1.0000f, 1.0000f, 1.0000f, 1.0000f, 1.0000f, 1.0000f, 1.0000f, 0.9718f, 0.8614f, 0.6422f,
+			1.0000f, 1.0000f, 1.0000f, 1.0000f, 1.0000f, 1.0000f, 1.0000f, 1.0000f, 1.0000f, 1.0000f, 1.0000f, 1.0000f, 1.0000f, 1.0000f, 1.0000f, 1.0000f,
+			1.0000f, 1.0000f, 1.0000f, 1.0000f, 1.0000f, 1.0000f, 1.0000f, 1.0000f, 1.0000f, 1.0000f, 1.0000f, 1.0000f, 1.0000f, 1.0000f, 1.0000f, 1.0000f,
+			1.0000f, 0.9648f, 0.9212f, 0.8412f, 1.0000f, 1.0000f, 1.0000f, 1.0000f, 1.0000f, 1.0000f, 1.0000f, 1.0000f, 1.0000f, 1.0000f, 1.0000f, 1.0000f,
+			1.0000f, 0.9596f, 0.9557f, 0.9537f, 1.0000f, 1.0000f, 1.0000f, 1.0000f, 1.0000f, 1.0000f, 1.0000f, 1.0000f, 1.0000f, 1.0000f, 1.0000f, 1.0000f,
+			1.0000f, 0.6814f, 1.0000f, 1.0000f, 1.0000f, 1.0000f, 1.0000f, 1.0000f, 1.0000f, 1.0000f, 1.0000f, 1.0000f, 1.0000f, 1.0000f, 1.0000f, 1.0000f,
+			1.0000f, 0.7971f, 1.0000f, 1.0000f, 1.0000f, 1.0000f, 1.0000f, 1.0000f, 1.0000f, 1.0000f, 1.0000f, 1.0000f, 1.0000f, 1.0000f, 1.0000f, 1.0000f,
+			1.0000f, 1.0000f, 1.0000f, 1.0000f, 1.0000f, 1.0000f, 1.0000f, 1.0000f, 1.0000f, 1.0000f, 1.0000f, 1.0000f, 1.0000f, 1.0000f, 1.0000f, 1.0000f,
+			1.0000f, 1.0000f, 1.0000f, 1.0000f, 1.0000f, 1.0000f, 1.0000f, 1.0000f, 1.0000f, 1.0000f, 1.0000f, 1.0000f, 1.0000f, 1.0000f, 1.0000f, 1.0000f,
+			1.0000f, 1.0000f, 1.0000f, 1.0000f, 1.0000f, 1.0000f, 1.0000f, 1.0000f, 1.0000f, 1.0000f, 1.0000f, 1.0000f, 1.0000f, 1.0000f, 1.0000f, 1.0000f,
+			1.0000f, 1.0000f, 1.0000f, 1.0000f, 1.0000f, 1.0000f, 1.0000f, 1.0000f, 1.0000f, 1.0000f, 1.0000f, 1.0000f, 1.0000f, 1.0000f, 1.0000f, 1.0000f,
+			1.0000f, 1.0000f, 1.0000f, 1.0000f, 1.0000f, 1.0000f, 1.0000f, 1.0000f, 1.0000f, 1.0000f, 1.0000f, 1.0000f, 1.0000f, 1.0000f, 1.0000f, 1.0000f,
+			1.0000f, 1.0000f, 1.0000f, 1.0000f, 1.0000f, 1.0000f, 1.0000f, 1.0000f, 1.0000f, 1.0000f, 1.0000f, 1.0000f, 1.0000f, 1.0000f, 1.0000f, 1.0000f,
+			1.0000f, 1.0000f, 1.0000f, 1.0000f, 1.0000f, 1.0000f, 1.0000f, 1.0000f, 1.0000f, 1.0000f, 1.0000f, 1.0000f, 1.0000f, 1.0000f, 1.0000f, 1.0000f,
+			1.0000f, 1.0000f, 1.0000f, 1.0000f, 1.0000f, 1.0000f, 1.0000f, 1.0000f, 1.0000f, 1.0000f, 1.0000f, 1.0000f, 1.0000f, 1.0000f, 1.0000f, 1.0000f,
+			1.0000f, 1.0000f, 1.0000f, 1.0000f, 1.0000f, 1.0000f, 1.0000f, 1.0000f, 1.0000f, 1.0000f, 1.0000f, 1.0000f, 1.0000f, 1.0000f, 1.0000f, 1.0000f,
+			1.0000f, 1.0000f, 1.0000f, 1.0000f, 1.0000f, 1.0000f, 1.0000f, 1.0000f, 1.0000f, 1.0000f, 1.0000f, 1.0000f, 1.0000f, 1.0000f, 1.0000f, 1.0000f,
+			1.0000f, 1.0000f, 1.0000f, 1.0000f, 1.0000f, 1.0000f, 1.0000f, 1.0000f, 1.0000f, 1.0000f, 1.0000f, 1.0000f, 1.0000f, 1.0000f, 1.0000f, 1.0000f,
+			1.0000f, 1.0000f, 1.0000f, 1.0000f, 1.0000f, 1.0000f, 1.0000f, 1.0000f, 1.0000f, 1.0000f, 1.0000f, 1.0000f, 1.0000f, 1.0000f, 1.0000f, 1.0000f,
+			1.0000f, 1.0000f, 1.0000f, 1.0000f, 1.0000f, 1.0000f, 1.0000f, 1.0000f, 1.0000f, 1.0000f, 1.0000f, 1.0000f, 1.0000f, 0.9998f, 0.9999f, 1.0000f,
+			1.0000f, 1.0000f, 1.0000f, 1.0000f, 1.0000f, 1.0000f, 1.0000f, 1.0000f, 1.0000f, 1.0000f, 1.0000f, 1.0000f, 1.0000f, 0.9999f, 1.0000f, 0.9994f,
+			1.0000f, 1.0000f, 1.0000f, 1.0000f, 1.0000f, 1.0000f, 1.0000f, 1.0000f, 1.0000f, 1.0000f, 1.0000f, 1.0000f, 1.0000f, 0.9961f, 0.9984f, 0.9958f,
+			1.0000f, 1.0000f, 1.0000f, 1.0000f, 1.0000f, 1.0000f, 1.0000f, 1.0000f, 1.0000f, 1.0000f, 1.0000f, 1.0000f, 1.0000f, 0.9952f, 0.9994f, 0.9633f,
+			1.0000f, 1.0000f, 1.0000f, 1.0000f, 1.0000f, 1.0000f, 1.0000f, 1.0000f, 1.0000f, 1.0000f, 1.0000f, 1.0000f, 1.0000f, 0.9829f, 0.9120f, 0.8042f,
+			1.0000f, 1.0000f, 1.0000f, 1.0000f, 1.0000f, 1.0000f, 1.0000f, 1.0000f, 1.0000f, 1.0000f, 1.0000f, 1.0000f, 1.0000f, 0.9744f, 0.9296f, 0.9349f,
+			1.0000f, 1.0000f, 1.0000f, 1.0000f, 1.0000f, 1.0000f, 1.0000f, 1.0000f, 1.0000f, 1.0000f, 1.0000f, 1.0000f, 1.0000f, 1.0000f, 1.0000f, 1.0000f,
+			1.0000f, 1.0000f, 1.0000f, 1.0000f, 1.0000f, 1.0000f, 1.0000f, 1.0000f, 1.0000f, 1.0000f, 1.0000f, 1.0000f, 1.0000f, 1.0000f, 1.0000f, 1.0000f,
+			1.0000f, 1.0000f, 1.0000f, 1.0000f, 1.0000f, 1.0000f, 1.0000f, 1.0000f, 1.0000f, 1.0000f, 1.0000f, 1.0000f, 1.0000f, 0.9973f, 0.9993f, 0.9988f,
+			1.0000f, 1.0000f, 1.0000f, 1.0000f, 1.0000f, 1.0000f, 1.0000f, 1.0000f, 1.0000f, 1.0000f, 1.0000f, 1.0000f, 1.0000f, 0.9964f, 0.9998f, 0.9802f,
+			1.0000f, 1.0000f, 1.0000f, 1.0000f, 1.0000f, 1.0000f, 1.0000f, 1.0000f, 1.0000f, 1.0000f, 1.0000f, 1.0000f, 1.0000f, 0.8874f, 0.5963f, 0.9323f,
+			1.0000f, 1.0000f, 1.0000f, 1.0000f, 1.0000f, 1.0000f, 1.0000f, 1.0000f, 1.0000f, 1.0000f, 1.0000f, 1.0000f, 1.0000f, 0.8239f, 0.7625f, 1.0000f,
+			1.0000f, 1.0000f, 1.0000f, 1.0000f, 1.0000f, 1.0000f, 1.0000f, 1.0000f, 1.0000f, 1.0000f, 1.0000f, 1.0000f, 1.0000f, 0.6625f, 1.0000f, 1.0000f,
+			1.0000f, 1.0000f, 1.0000f, 1.0000f, 1.0000f, 1.0000f, 1.0000f, 1.0000f, 1.0000f, 1.0000f, 1.0000f, 1.0000f, 1.0000f, 0.7895f, 1.0000f, 1.0000f,
+			1.0000f, 1.0000f, 1.0000f, 1.0000f, 1.0000f, 1.0000f, 1.0000f, 1.0000f, 1.0000f, 1.0000f, 1.0000f, 1.0000f, 1.0000f, 1.0000f, 1.0000f, 1.0000f,
+			1.0000f, 1.0000f, 1.0000f, 1.0000f, 1.0000f, 1.0000f, 1.0000f, 1.0000f, 1.0000f, 1.0000f, 1.0000f, 1.0000f, 1.0000f, 1.0000f, 1.0000f, 1.0000f,
+		};
+
+
+
+		const float* get_2d_percentile_table()
+		{
+			return percentile_table_4x4;
+		}
+
+
+		// return 0 on invalid mode, 1 on valid mode.
+		static int decode_block_mode_2d(int blockmode, int* Nval, int* Mval, int* dual_weight_plane, int* quant_mode)
+		{
+			int base_quant_mode = (blockmode >> 4) & 1;
+			int H = (blockmode >> 9) & 1;
+			int D = (blockmode >> 10) & 1;
+
+			int A = (blockmode >> 5) & 0x3;
+
+			int N = 0, M = 0;
+
+			if ((blockmode & 3) != 0)
+			{
+				base_quant_mode |= (blockmode & 3) << 1;
+				int B = (blockmode >> 7) & 3;
+				switch ((blockmode >> 2) & 3)
+				{
+				case 0:
+					N = B + 4;
+					M = A + 2;
+					break;
+				case 1:
+					N = B + 8;
+					M = A + 2;
+					break;
+				case 2:
+					N = A + 2;
+					M = B + 8;
+					break;
+				case 3:
+					B &= 1;
+					if (blockmode & 0x100)
+					{
+						N = B + 2;
+						M = A + 2;
+					}
+					else
+					{
+						N = A + 2;
+						M = B + 6;
+					}
+					break;
+				}
+			}
+			else
+			{
+				base_quant_mode |= ((blockmode >> 2) & 3) << 1;
+				if (((blockmode >> 2) & 3) == 0)
+					return 0;
+				int B = (blockmode >> 9) & 3;
+				switch ((blockmode >> 7) & 3)
+				{
+				case 0:
+					N = 12;
+					M = A + 2;
+					break;
+				case 1:
+					N = A + 2;
+					M = 12;
+					break;
+				case 2:
+					N = A + 6;
+					M = B + 6;
+					D = 0;
+					H = 0;
+					break;
+				case 3:
+					switch ((blockmode >> 5) & 3)
+					{
+					case 0:
+						N = 6;
+						M = 10;
+						break;
+					case 1:
+						N = 10;
+						M = 6;
+						break;
+					case 2:
+					case 3:
+						return 0;
+					}
+					break;
+				}
+			}
+
+			int weight_count = N * M * (D + 1);
+			int qmode = (base_quant_mode - 2) + 6 * H;
+
+			int weightbits = compute_ise_bitcount(weight_count, (quantization_method)qmode);
+			if (weight_count > MAX_WEIGHTS_PER_BLOCK || weightbits < MIN_WEIGHT_BITS_PER_BLOCK || weightbits > MAX_WEIGHT_BITS_PER_BLOCK)
+				return 0;
+
+			*Nval = N;
+			*Mval = M;
+			*dual_weight_plane = D;
+			*quant_mode = qmode;
+			return 1;
+		}
+
+
+		void construct_block_size_descriptor_2d(block_size_descriptor* bsd)
+		{
+			int decimation_mode_index = 0;
+			int decimation_mode_count = 0;
+
+			int i;
+			int x_weights = 4;
+			int y_weights = 4;
+
+			// gather all the infill-modes that can be used with the current block size
+			//for (x_weights = 2; x_weights <= 12; x_weights++)
+			//    for (y_weights = 2; y_weights <= 12; y_weights++)
+			{
+				//if (x_weights * y_weights > MAX_WEIGHTS_PER_BLOCK)
+				//    continue;
+				decimation_table* dt = new decimation_table;
+				decimation_mode_index = decimation_mode_count;
+				initialize_decimation_table_2d(dt);
+
+				int weight_count = x_weights * y_weights;
+
+				int maxprec_1plane = -1;
+				int maxprec_2planes = -1;
+				for (i = 0; i < 12; i++)
+				{
+					int bits_1plane = compute_ise_bitcount(weight_count, (quantization_method)i);
+					int bits_2planes = compute_ise_bitcount(2 * weight_count, (quantization_method)i);
+					if (bits_1plane >= MIN_WEIGHT_BITS_PER_BLOCK && bits_1plane <= MAX_WEIGHT_BITS_PER_BLOCK)
+						maxprec_1plane = i;
+					if (bits_2planes >= MIN_WEIGHT_BITS_PER_BLOCK && bits_2planes <= MAX_WEIGHT_BITS_PER_BLOCK)
+						maxprec_2planes = i;
+				}
+
+				// always false
+				//if (2 * x_weights * y_weights > MAX_WEIGHTS_PER_BLOCK)
+				//    maxprec_2planes = -1;
+
+				bsd->permit_encode[decimation_mode_count] = (x_weights <= 4 && y_weights <= 4);
+
+				bsd->decimation_mode_samples[decimation_mode_count] = weight_count;
+				bsd->decimation_mode_maxprec_1plane[decimation_mode_count] = maxprec_1plane;
+				bsd->decimation_mode_maxprec_2planes[decimation_mode_count] = maxprec_2planes;
+				bsd->decimation_tables[decimation_mode_count] = dt;
+
+				decimation_mode_count++;
+			}
+
+			for (i = 0; i < MAX_DECIMATION_MODES; i++)
+			{
+				bsd->decimation_mode_percentile[i] = 1.0f;
+			}
+
+			for (i = decimation_mode_count; i < MAX_DECIMATION_MODES; i++)
+			{
+				bsd->permit_encode[i] = 0;
+				bsd->decimation_mode_samples[i] = 0;
+				bsd->decimation_mode_maxprec_1plane[i] = -1;
+				bsd->decimation_mode_maxprec_2planes[i] = -1;
+			}
+
+			bsd->decimation_mode_count = decimation_mode_count;
+
+			const float* percentiles = get_2d_percentile_table();
+
+			// then construct the list of block formats
+			for (i = 0; i < 2048; i++)
+			{
+				//int x_weights, y_weights;
+				int is_dual_plane;
+				int quantization_mode;
+				int fail = 0;
+				int permit_encode = 1;
+
+				if (decode_block_mode_2d(i, &x_weights, &y_weights, &is_dual_plane, &quantization_mode))
+				{
+					if (x_weights > 4 || y_weights > 4)
+						permit_encode = 0;
+				}
+				else
+				{
+					fail = 1;
+					permit_encode = 0;
+				}
+
+				if (fail)
+				{
+					bsd->block_modes[i].decimation_mode = -1;
+					bsd->block_modes[i].quantization_mode = -1;
+					bsd->block_modes[i].is_dual_plane = -1;
+					bsd->block_modes[i].permit_encode = 0;
+					bsd->block_modes[i].permit_decode = 0;
+					bsd->block_modes[i].percentile = 1.0f;
+				}
+				else
+				{
+					int decimation_mode = decimation_mode_index;
+					bsd->block_modes[i].decimation_mode = (int8_t)decimation_mode;
+					bsd->block_modes[i].quantization_mode = (int8_t)quantization_mode;
+					bsd->block_modes[i].is_dual_plane = (int8_t)is_dual_plane;
+					bsd->block_modes[i].permit_encode = (int8_t)permit_encode;
+					bsd->block_modes[i].permit_decode = (int8_t)permit_encode;	// disallow decode of grid size larger than block size.
+					bsd->block_modes[i].percentile = percentiles[i];
+
+					if (bsd->decimation_mode_percentile[decimation_mode] > percentiles[i])
+						bsd->decimation_mode_percentile[decimation_mode] = percentiles[i];
+				}
+
+			}
+		}
+
+
+
+
+		static block_size_descriptor* s_bsd = nullptr;
+
+		// function to obtain a block size descriptor. If the descriptor does not exist,
+		// it is created as needed. Should not be called from within multi-threaded code.
+		const block_size_descriptor* get_block_size_descriptor()
+		{
+			if (s_bsd == nullptr)
+			{
+				s_bsd = new block_size_descriptor;
+				construct_block_size_descriptor_2d(s_bsd);
+			}
+			return s_bsd;
+		}
+
+
+
+		// routine to read up to 8 bits
+		static inline int read_bits(int bitcount, int bitoffset, const uint8* ptr)
+		{
+			int mask = (1 << bitcount) - 1;
+			ptr += bitoffset >> 3;
+			bitoffset &= 7;
+			int value = ptr[0] | (ptr[1] << 8);
+			value >>= bitoffset;
+			value &= mask;
+			return value;
+		}
+
+		int bitrev8(int p)
+		{
+			p = ((p & 0xF) << 4) | ((p >> 4) & 0xF);
+			p = ((p & 0x33) << 2) | ((p >> 2) & 0x33);
+			p = ((p & 0x55) << 1) | ((p >> 1) & 0x55);
+			return p;
+		}
+
+
+
+		void find_number_of_bits_trits_quints(int quantization_level, int* bits, int* trits, int* quints)
+		{
+			*bits = 0;
+			*trits = 0;
+			*quints = 0;
+			switch (quantization_level)
+			{
+			case QUANT_2:
+				*bits = 1;
+				break;
+			case QUANT_3:
+				*bits = 0;
+				*trits = 1;
+				break;
+			case QUANT_4:
+				*bits = 2;
+				break;
+			case QUANT_5:
+				*bits = 0;
+				*quints = 1;
+				break;
+			case QUANT_6:
+				*bits = 1;
+				*trits = 1;
+				break;
+			case QUANT_8:
+				*bits = 3;
+				break;
+			case QUANT_10:
+				*bits = 1;
+				*quints = 1;
+				break;
+			case QUANT_12:
+				*bits = 2;
+				*trits = 1;
+				break;
+			case QUANT_16:
+				*bits = 4;
+				break;
+			case QUANT_20:
+				*bits = 2;
+				*quints = 1;
+				break;
+			case QUANT_24:
+				*bits = 3;
+				*trits = 1;
+				break;
+			case QUANT_32:
+				*bits = 5;
+				break;
+			case QUANT_40:
+				*bits = 3;
+				*quints = 1;
+				break;
+			case QUANT_48:
+				*bits = 4;
+				*trits = 1;
+				break;
+			case QUANT_64:
+				*bits = 6;
+				break;
+			case QUANT_80:
+				*bits = 4;
+				*quints = 1;
+				break;
+			case QUANT_96:
+				*bits = 5;
+				*trits = 1;
+				break;
+			case QUANT_128:
+				*bits = 7;
+				break;
+			case QUANT_160:
+				*bits = 5;
+				*quints = 1;
+				break;
+			case QUANT_192:
+				*bits = 6;
+				*trits = 1;
+				break;
+			case QUANT_256:
+				*bits = 8;
+				break;
+			}
+		}
+
+
+		// unpacked trit quintuplets <low,_,_,_,high> for each packed-quint value
+		static const uint8 trits_of_integer[256][5] = {
+			{0, 0, 0, 0, 0},	{1, 0, 0, 0, 0},	{2, 0, 0, 0, 0},	{0, 0, 2, 0, 0},
+			{0, 1, 0, 0, 0},	{1, 1, 0, 0, 0},	{2, 1, 0, 0, 0},	{1, 0, 2, 0, 0},
+			{0, 2, 0, 0, 0},	{1, 2, 0, 0, 0},	{2, 2, 0, 0, 0},	{2, 0, 2, 0, 0},
+			{0, 2, 2, 0, 0},	{1, 2, 2, 0, 0},	{2, 2, 2, 0, 0},	{2, 0, 2, 0, 0},
+			{0, 0, 1, 0, 0},	{1, 0, 1, 0, 0},	{2, 0, 1, 0, 0},	{0, 1, 2, 0, 0},
+			{0, 1, 1, 0, 0},	{1, 1, 1, 0, 0},	{2, 1, 1, 0, 0},	{1, 1, 2, 0, 0},
+			{0, 2, 1, 0, 0},	{1, 2, 1, 0, 0},	{2, 2, 1, 0, 0},	{2, 1, 2, 0, 0},
+			{0, 0, 0, 2, 2},	{1, 0, 0, 2, 2},	{2, 0, 0, 2, 2},	{0, 0, 2, 2, 2},
+			{0, 0, 0, 1, 0},	{1, 0, 0, 1, 0},	{2, 0, 0, 1, 0},	{0, 0, 2, 1, 0},
+			{0, 1, 0, 1, 0},	{1, 1, 0, 1, 0},	{2, 1, 0, 1, 0},	{1, 0, 2, 1, 0},
+			{0, 2, 0, 1, 0},	{1, 2, 0, 1, 0},	{2, 2, 0, 1, 0},	{2, 0, 2, 1, 0},
+			{0, 2, 2, 1, 0},	{1, 2, 2, 1, 0},	{2, 2, 2, 1, 0},	{2, 0, 2, 1, 0},
+			{0, 0, 1, 1, 0},	{1, 0, 1, 1, 0},	{2, 0, 1, 1, 0},	{0, 1, 2, 1, 0},
+			{0, 1, 1, 1, 0},	{1, 1, 1, 1, 0},	{2, 1, 1, 1, 0},	{1, 1, 2, 1, 0},
+			{0, 2, 1, 1, 0},	{1, 2, 1, 1, 0},	{2, 2, 1, 1, 0},	{2, 1, 2, 1, 0},
+			{0, 1, 0, 2, 2},	{1, 1, 0, 2, 2},	{2, 1, 0, 2, 2},	{1, 0, 2, 2, 2},
+			{0, 0, 0, 2, 0},	{1, 0, 0, 2, 0},	{2, 0, 0, 2, 0},	{0, 0, 2, 2, 0},
+			{0, 1, 0, 2, 0},	{1, 1, 0, 2, 0},	{2, 1, 0, 2, 0},	{1, 0, 2, 2, 0},
+			{0, 2, 0, 2, 0},	{1, 2, 0, 2, 0},	{2, 2, 0, 2, 0},	{2, 0, 2, 2, 0},
+			{0, 2, 2, 2, 0},	{1, 2, 2, 2, 0},	{2, 2, 2, 2, 0},	{2, 0, 2, 2, 0},
+			{0, 0, 1, 2, 0},	{1, 0, 1, 2, 0},	{2, 0, 1, 2, 0},	{0, 1, 2, 2, 0},
+			{0, 1, 1, 2, 0},	{1, 1, 1, 2, 0},	{2, 1, 1, 2, 0},	{1, 1, 2, 2, 0},
+			{0, 2, 1, 2, 0},	{1, 2, 1, 2, 0},	{2, 2, 1, 2, 0},	{2, 1, 2, 2, 0},
+			{0, 2, 0, 2, 2},	{1, 2, 0, 2, 2},	{2, 2, 0, 2, 2},	{2, 0, 2, 2, 2},
+			{0, 0, 0, 0, 2},	{1, 0, 0, 0, 2},	{2, 0, 0, 0, 2},	{0, 0, 2, 0, 2},
+			{0, 1, 0, 0, 2},	{1, 1, 0, 0, 2},	{2, 1, 0, 0, 2},	{1, 0, 2, 0, 2},
+			{0, 2, 0, 0, 2},	{1, 2, 0, 0, 2},	{2, 2, 0, 0, 2},	{2, 0, 2, 0, 2},
+			{0, 2, 2, 0, 2},	{1, 2, 2, 0, 2},	{2, 2, 2, 0, 2},	{2, 0, 2, 0, 2},
+			{0, 0, 1, 0, 2},	{1, 0, 1, 0, 2},	{2, 0, 1, 0, 2},	{0, 1, 2, 0, 2},
+			{0, 1, 1, 0, 2},	{1, 1, 1, 0, 2},	{2, 1, 1, 0, 2},	{1, 1, 2, 0, 2},
+			{0, 2, 1, 0, 2},	{1, 2, 1, 0, 2},	{2, 2, 1, 0, 2},	{2, 1, 2, 0, 2},
+			{0, 2, 2, 2, 2},	{1, 2, 2, 2, 2},	{2, 2, 2, 2, 2},	{2, 0, 2, 2, 2},
+			{0, 0, 0, 0, 1},	{1, 0, 0, 0, 1},	{2, 0, 0, 0, 1},	{0, 0, 2, 0, 1},
+			{0, 1, 0, 0, 1},	{1, 1, 0, 0, 1},	{2, 1, 0, 0, 1},	{1, 0, 2, 0, 1},
+			{0, 2, 0, 0, 1},	{1, 2, 0, 0, 1},	{2, 2, 0, 0, 1},	{2, 0, 2, 0, 1},
+			{0, 2, 2, 0, 1},	{1, 2, 2, 0, 1},	{2, 2, 2, 0, 1},	{2, 0, 2, 0, 1},
+			{0, 0, 1, 0, 1},	{1, 0, 1, 0, 1},	{2, 0, 1, 0, 1},	{0, 1, 2, 0, 1},
+			{0, 1, 1, 0, 1},	{1, 1, 1, 0, 1},	{2, 1, 1, 0, 1},	{1, 1, 2, 0, 1},
+			{0, 2, 1, 0, 1},	{1, 2, 1, 0, 1},	{2, 2, 1, 0, 1},	{2, 1, 2, 0, 1},
+			{0, 0, 1, 2, 2},	{1, 0, 1, 2, 2},	{2, 0, 1, 2, 2},	{0, 1, 2, 2, 2},
+			{0, 0, 0, 1, 1},	{1, 0, 0, 1, 1},	{2, 0, 0, 1, 1},	{0, 0, 2, 1, 1},
+			{0, 1, 0, 1, 1},	{1, 1, 0, 1, 1},	{2, 1, 0, 1, 1},	{1, 0, 2, 1, 1},
+			{0, 2, 0, 1, 1},	{1, 2, 0, 1, 1},	{2, 2, 0, 1, 1},	{2, 0, 2, 1, 1},
+			{0, 2, 2, 1, 1},	{1, 2, 2, 1, 1},	{2, 2, 2, 1, 1},	{2, 0, 2, 1, 1},
+			{0, 0, 1, 1, 1},	{1, 0, 1, 1, 1},	{2, 0, 1, 1, 1},	{0, 1, 2, 1, 1},
+			{0, 1, 1, 1, 1},	{1, 1, 1, 1, 1},	{2, 1, 1, 1, 1},	{1, 1, 2, 1, 1},
+			{0, 2, 1, 1, 1},	{1, 2, 1, 1, 1},	{2, 2, 1, 1, 1},	{2, 1, 2, 1, 1},
+			{0, 1, 1, 2, 2},	{1, 1, 1, 2, 2},	{2, 1, 1, 2, 2},	{1, 1, 2, 2, 2},
+			{0, 0, 0, 2, 1},	{1, 0, 0, 2, 1},	{2, 0, 0, 2, 1},	{0, 0, 2, 2, 1},
+			{0, 1, 0, 2, 1},	{1, 1, 0, 2, 1},	{2, 1, 0, 2, 1},	{1, 0, 2, 2, 1},
+			{0, 2, 0, 2, 1},	{1, 2, 0, 2, 1},	{2, 2, 0, 2, 1},	{2, 0, 2, 2, 1},
+			{0, 2, 2, 2, 1},	{1, 2, 2, 2, 1},	{2, 2, 2, 2, 1},	{2, 0, 2, 2, 1},
+			{0, 0, 1, 2, 1},	{1, 0, 1, 2, 1},	{2, 0, 1, 2, 1},	{0, 1, 2, 2, 1},
+			{0, 1, 1, 2, 1},	{1, 1, 1, 2, 1},	{2, 1, 1, 2, 1},	{1, 1, 2, 2, 1},
+			{0, 2, 1, 2, 1},	{1, 2, 1, 2, 1},	{2, 2, 1, 2, 1},	{2, 1, 2, 2, 1},
+			{0, 2, 1, 2, 2},	{1, 2, 1, 2, 2},	{2, 2, 1, 2, 2},	{2, 1, 2, 2, 2},
+			{0, 0, 0, 1, 2},	{1, 0, 0, 1, 2},	{2, 0, 0, 1, 2},	{0, 0, 2, 1, 2},
+			{0, 1, 0, 1, 2},	{1, 1, 0, 1, 2},	{2, 1, 0, 1, 2},	{1, 0, 2, 1, 2},
+			{0, 2, 0, 1, 2},	{1, 2, 0, 1, 2},	{2, 2, 0, 1, 2},	{2, 0, 2, 1, 2},
+			{0, 2, 2, 1, 2},	{1, 2, 2, 1, 2},	{2, 2, 2, 1, 2},	{2, 0, 2, 1, 2},
+			{0, 0, 1, 1, 2},	{1, 0, 1, 1, 2},	{2, 0, 1, 1, 2},	{0, 1, 2, 1, 2},
+			{0, 1, 1, 1, 2},	{1, 1, 1, 1, 2},	{2, 1, 1, 1, 2},	{1, 1, 2, 1, 2},
+			{0, 2, 1, 1, 2},	{1, 2, 1, 1, 2},	{2, 2, 1, 1, 2},	{2, 1, 2, 1, 2},
+			{0, 2, 2, 2, 2},	{1, 2, 2, 2, 2},	{2, 2, 2, 2, 2},	{2, 1, 2, 2, 2},
+		};
+
+
+		static const uint8 quints_of_integer[128][3] = {
+			{0, 0, 0},	{1, 0, 0},	{2, 0, 0},	{3, 0, 0},
+			{4, 0, 0},	{0, 4, 0},	{4, 4, 0},	{4, 4, 4},
+			{0, 1, 0},	{1, 1, 0},	{2, 1, 0},	{3, 1, 0},
+			{4, 1, 0},	{1, 4, 0},	{4, 4, 1},	{4, 4, 4},
+			{0, 2, 0},	{1, 2, 0},	{2, 2, 0},	{3, 2, 0},
+			{4, 2, 0},	{2, 4, 0},	{4, 4, 2},	{4, 4, 4},
+			{0, 3, 0},	{1, 3, 0},	{2, 3, 0},	{3, 3, 0},
+			{4, 3, 0},	{3, 4, 0},	{4, 4, 3},	{4, 4, 4},
+			{0, 0, 1},	{1, 0, 1},	{2, 0, 1},	{3, 0, 1},
+			{4, 0, 1},	{0, 4, 1},	{4, 0, 4},	{0, 4, 4},
+			{0, 1, 1},	{1, 1, 1},	{2, 1, 1},	{3, 1, 1},
+			{4, 1, 1},	{1, 4, 1},	{4, 1, 4},	{1, 4, 4},
+			{0, 2, 1},	{1, 2, 1},	{2, 2, 1},	{3, 2, 1},
+			{4, 2, 1},	{2, 4, 1},	{4, 2, 4},	{2, 4, 4},
+			{0, 3, 1},	{1, 3, 1},	{2, 3, 1},	{3, 3, 1},
+			{4, 3, 1},	{3, 4, 1},	{4, 3, 4},	{3, 4, 4},
+			{0, 0, 2},	{1, 0, 2},	{2, 0, 2},	{3, 0, 2},
+			{4, 0, 2},	{0, 4, 2},	{2, 0, 4},	{3, 0, 4},
+			{0, 1, 2},	{1, 1, 2},	{2, 1, 2},	{3, 1, 2},
+			{4, 1, 2},	{1, 4, 2},	{2, 1, 4},	{3, 1, 4},
+			{0, 2, 2},	{1, 2, 2},	{2, 2, 2},	{3, 2, 2},
+			{4, 2, 2},	{2, 4, 2},	{2, 2, 4},	{3, 2, 4},
+			{0, 3, 2},	{1, 3, 2},	{2, 3, 2},	{3, 3, 2},
+			{4, 3, 2},	{3, 4, 2},	{2, 3, 4},	{3, 3, 4},
+			{0, 0, 3},	{1, 0, 3},	{2, 0, 3},	{3, 0, 3},
+			{4, 0, 3},	{0, 4, 3},	{0, 0, 4},	{1, 0, 4},
+			{0, 1, 3},	{1, 1, 3},	{2, 1, 3},	{3, 1, 3},
+			{4, 1, 3},	{1, 4, 3},	{0, 1, 4},	{1, 1, 4},
+			{0, 2, 3},	{1, 2, 3},	{2, 2, 3},	{3, 2, 3},
+			{4, 2, 3},	{2, 4, 3},	{0, 2, 4},	{1, 2, 4},
+			{0, 3, 3},	{1, 3, 3},	{2, 3, 3},	{3, 3, 3},
+			{4, 3, 3},	{3, 4, 3},	{0, 3, 4},	{1, 3, 4},
+		};
+
+		// quantization_mode_table[integercount/2][bits] gives
+		// us the quantization level for a given integer count and number of bits that
+		// the integer may fit into. This is needed for color decoding,
+		// and for the color encoding.
+		int8_t quantization_mode_table[17][128];
+
+		void build_quantization_mode_table(void)
+		{
+			int8_t i;
+			int j;
+			for (i = 0; i <= 16; i++)
+				for (j = 0; j < 128; j++)
+					quantization_mode_table[i][j] = -1;
+
+			for (i = 0; i < 21; i++)
+				for (j = 1; j <= 16; j++)
+				{
+					int p = compute_ise_bitcount(2 * j, (quantization_method)i);
+					if (p < 128)
+						quantization_mode_table[j][p] = i;
+				}
+			for (i = 0; i <= 16; i++)
+			{
+				int8_t largest_value_so_far = -1;
+				for (j = 0; j < 128; j++)
+				{
+					if (quantization_mode_table[i][j] > largest_value_so_far)
+						largest_value_so_far = quantization_mode_table[i][j];
+					else
+						quantization_mode_table[i][j] = largest_value_so_far;
+				}
+			}
+		}
+
+		void decode_ise(int quantization_level, int elements, const uint8* input_data, uint8* output_data, int bit_offset)
+		{
+			int i;
+			// note: due to how the trit/quint-block unpacking is done in this function,
+			// we may write more temporary results than the number of outputs
+			// The maximum actual number of results is 64 bit, but we keep 4 additional elements
+			// of padding.
+			uint8 results[68];
+			uint8 tq_blocks[22];		// trit-blocks or quint-blocks
+
+			int bits, trits, quints;
+			find_number_of_bits_trits_quints(quantization_level, &bits, &trits, &quints);
+
+			int lcounter = 0;
+			int hcounter = 0;
+
+			// trit-blocks or quint-blocks must be zeroed out before we collect them in the loop below.
+			for (i = 0; i < 22; i++)
+				tq_blocks[i] = 0;
+
+			// collect bits for each element, as well as bits for any trit-blocks and quint-blocks.
+			for (i = 0; i < elements; i++)
+			{
+				results[i] = (uint8)read_bits(bits, bit_offset, input_data);
+				bit_offset += bits;
+				if (trits)
+				{
+					static const int bits_to_read[5] = { 2, 2, 1, 2, 1 };
+					static const int block_shift[5] = { 0, 2, 4, 5, 7 };
+					static const int next_lcounter[5] = { 1, 2, 3, 4, 0 };
+					static const int hcounter_incr[5] = { 0, 0, 0, 0, 1 };
+					DCHECK(lcounter < 5);
+					int tdata = read_bits(bits_to_read[lcounter], bit_offset, input_data);
+					bit_offset += bits_to_read[lcounter];
+					tq_blocks[hcounter] |= tdata << block_shift[lcounter];
+					hcounter += hcounter_incr[lcounter];
+					lcounter = next_lcounter[lcounter];
+				}
+				if (quints)
+				{
+					static const int bits_to_read[3] = { 3, 2, 2 };
+					static const int block_shift[3] = { 0, 3, 5 };
+					static const int next_lcounter[3] = { 1, 2, 0 };
+					static const int hcounter_incr[3] = { 0, 0, 1 };
+					DCHECK(lcounter < 3);
+					if (lcounter < 3)
+					{
+						int tdata = read_bits(bits_to_read[lcounter], bit_offset, input_data);
+						bit_offset += bits_to_read[lcounter];
+						tq_blocks[hcounter] |= tdata << block_shift[lcounter];
+						hcounter += hcounter_incr[lcounter];
+						lcounter = next_lcounter[lcounter];
+					}
+				}
+			}
+
+
+			// unpack trit-blocks or quint-blocks as needed
+			if (trits)
+			{
+				int trit_blocks = (elements + 4) / 5;
+				for (i = 0; i < trit_blocks; i++)
+				{
+					const uint8* tritptr = trits_of_integer[tq_blocks[i]];
+					results[5 * i] |= tritptr[0] << bits;
+					results[5 * i + 1] |= tritptr[1] << bits;
+					results[5 * i + 2] |= tritptr[2] << bits;
+					results[5 * i + 3] |= tritptr[3] << bits;
+					results[5 * i + 4] |= tritptr[4] << bits;
+				}
+			}
+
+			if (quints)
+			{
+				int quint_blocks = (elements + 2) / 3;
+				for (i = 0; i < quint_blocks; i++)
+				{
+					const uint8* quintptr = quints_of_integer[tq_blocks[i]];
+					results[3 * i] |= quintptr[0] << bits;
+					results[3 * i + 1] |= quintptr[1] << bits;
+					results[3 * i + 2] |= quintptr[2] << bits;
+				}
+			}
+
+			for (i = 0; i < elements; i++)
+				output_data[i] = results[i];
+		}
+
+
+		inline void physical_to_symbolic(physical_compressed_block pb, symbolic_compressed_block* res)
+		{
+			uint8 bswapped[16];
+			int i, j;
+
+			res->error_block = 0;
+
+			// get hold of the block-size descriptor and the decimation tables.
+			const block_size_descriptor* bsd = get_block_size_descriptor();
+			const decimation_table* const* ixtab2 = bsd->decimation_tables;
+
+			// extract header fields
+			int block_mode = read_bits(11, 0, pb.data);
+
+
+			if ((block_mode & 0x1FF) == 0x1FC)
+			{
+				// void-extent block!
+
+				// check what format the data has
+				if (block_mode & 0x200)
+					res->block_mode = -1;	// floating-point
+				else
+					res->block_mode = -2;	// unorm16.
+
+				res->partition_count = 0;
+				for (i = 0; i < 4; i++)
+				{
+					res->constant_color[i] = pb.data[2 * i + 8] | (pb.data[2 * i + 9] << 8);
+				}
+
+				// additionally, check that the void-extent
+				{
+					// 2D void-extent
+					int rsvbits = read_bits(2, 10, pb.data);
+					if (rsvbits != 3)
+						res->error_block = 1;
+
+					int vx_low_s = read_bits(8, 12, pb.data) | (read_bits(5, 12 + 8, pb.data) << 8);
+					int vx_high_s = read_bits(8, 25, pb.data) | (read_bits(5, 25 + 8, pb.data) << 8);
+					int vx_low_t = read_bits(8, 38, pb.data) | (read_bits(5, 38 + 8, pb.data) << 8);
+					int vx_high_t = read_bits(8, 51, pb.data) | (read_bits(5, 51 + 8, pb.data) << 8);
+
+					int all_ones = vx_low_s == 0x1FFF && vx_high_s == 0x1FFF && vx_low_t == 0x1FFF && vx_high_t == 0x1FFF;
+
+					if ((vx_low_s >= vx_high_s || vx_low_t >= vx_high_t) && !all_ones)
+						res->error_block = 1;
+				}
+
+				return;
+			}
+
+			if (bsd->block_modes[block_mode].permit_decode == 0)
+			{
+				res->error_block = 1;
+				return;
+			}
+
+			int weight_count = ixtab2[bsd->block_modes[block_mode].decimation_mode]->num_weights;
+			int weight_quantization_method = bsd->block_modes[block_mode].quantization_mode;
+			int is_dual_plane = bsd->block_modes[block_mode].is_dual_plane;
+
+			int real_weight_count = is_dual_plane ? 2 * weight_count : weight_count;
+
+			int partition_count = read_bits(2, 11, pb.data) + 1;
+
+			res->block_mode = block_mode;
+			res->partition_count = partition_count;
+
+			for (i = 0; i < 16; i++)
+				bswapped[i] = (uint8)bitrev8(pb.data[15 - i]);
+
+			int bits_for_weights = compute_ise_bitcount(real_weight_count,
+				(quantization_method)weight_quantization_method);
+
+			int below_weights_pos = 128 - bits_for_weights;
+
+			if (is_dual_plane)
+			{
+				uint8 indices[64];
+				decode_ise(weight_quantization_method, real_weight_count, bswapped, indices, 0);
+				for (i = 0; i < weight_count; i++)
+				{
+					res->plane1_weights[i] = indices[2 * i];
+					res->plane2_weights[i] = indices[2 * i + 1];
+				}
+
+				//            if (s_debuglog)
+				//            {
+				//                i=0;
+				//                UE_LOG(LogMutableCore,Warning," decompression quantized r weights:");
+				//                for (int j=0;j<4;++j)
+				//                {
+				//                    UE_LOG(LogMutableCore,Warning," %.3d %.3d %.3d %.3d ",
+				//                            res->plane1_weights[i+0], res->plane1_weights[i+1], res->plane1_weights[i+2], res->plane1_weights[i+3] );
+				//                    i+=4;
+				//                }
+				////                i=0;
+				////                UE_LOG(LogMutableCore,Warning," decompression quantized g weights:");
+				////                for (int j=0;j<4;++j)
+				////                {
+				////                    UE_LOG(LogMutableCore,Warning," %.3d %.3d %.3d %.3d ",
+				////                            res->plane2_weights[i+0], res->plane2_weights[i+1], res->plane2_weights[i+2], res->plane2_weights[i+3] );
+				////                    i+=4;
+				////                }
+				//            }
+			}
+			else
+			{
+				decode_ise(weight_quantization_method, weight_count, bswapped, res->plane1_weights, 0);
+			}
+
+			if (is_dual_plane && partition_count == 4)
+				res->error_block = 1;
+
+
+
+			res->color_formats_matched = 0;
+
+			// then, determine the format of each endpoint pair
+			int color_formats[4];
+			int encoded_type_highpart_size = 0;
+			if (partition_count == 1)
+			{
+				color_formats[0] = read_bits(4, 13, pb.data);
+				res->partition_index = 0;
+			}
+			else
+			{
+				encoded_type_highpart_size = (3 * partition_count) - 4;
+				below_weights_pos -= encoded_type_highpart_size;
+				int encoded_type = read_bits(6, 13 + PARTITION_BITS, pb.data) | (read_bits(encoded_type_highpart_size, below_weights_pos, pb.data) << 6);
+				int baseclass = encoded_type & 0x3;
+				if (baseclass == 0)
+				{
+					for (i = 0; i < partition_count; i++)
+					{
+						color_formats[i] = (encoded_type >> 2) & 0xF;
+					}
+					below_weights_pos += encoded_type_highpart_size;
+					res->color_formats_matched = 1;
+					encoded_type_highpart_size = 0;
+				}
+				else
+				{
+					int bitpos = 2;
+					baseclass--;
+					for (i = 0; i < partition_count; i++)
+					{
+						color_formats[i] = (((encoded_type >> bitpos) & 1) + baseclass) << 2;
+						bitpos++;
+					}
+					for (i = 0; i < partition_count; i++)
+					{
+						color_formats[i] |= (encoded_type >> bitpos) & 3;
+						bitpos += 2;
+					}
+				}
+				res->partition_index = read_bits(6, 13, pb.data) | (read_bits(PARTITION_BITS - 6, 19, pb.data) << 6);
+
+			}
+			for (i = 0; i < partition_count; i++)
+				res->color_formats[i] = color_formats[i];
+
+
+			// then, determine the number of integers we need to unpack for the endpoint pairs
+			int color_integer_count = 0;
+			for (i = 0; i < partition_count; i++)
+			{
+				int endpoint_class = color_formats[i] >> 2;
+				color_integer_count += (endpoint_class + 1) * 2;
+			}
+
+			if (color_integer_count > 18)
+				res->error_block = 1;
+
+			// then, determine the color endpoint format to use for these integers
+			static const int color_bits_arr[5] = { -1, 115 - 4, 113 - 4 - PARTITION_BITS, 113 - 4 - PARTITION_BITS, 113 - 4 - PARTITION_BITS };
+			int color_bits = color_bits_arr[partition_count] - bits_for_weights - encoded_type_highpart_size;
+			if (is_dual_plane)
+				color_bits -= 2;
+			if (color_bits < 0)
+				color_bits = 0;
+
+			int color_quantization_level = quantization_mode_table[color_integer_count >> 1][color_bits];
+			res->color_quantization_level = color_quantization_level;
+			if (color_quantization_level < 4)
+				res->error_block = 1;
+
+
+			// then unpack the integer-bits
+			uint8 values_to_decode[32];
+			decode_ise(color_quantization_level, color_integer_count, pb.data, values_to_decode, (partition_count == 1 ? 17 : 19 + PARTITION_BITS));
+
+			// and distribute them over the endpoint types
+			int valuecount_to_decode = 0;
+
+			for (i = 0; i < partition_count; i++)
+			{
+				int vals = 2 * (color_formats[i] >> 2) + 2;
+				for (j = 0; j < vals; j++)
+					res->color_values[i][j] = values_to_decode[j + valuecount_to_decode];
+				valuecount_to_decode += vals;
+			}
+
+			// get hold of color component for second-plane in the case of dual plane of weights.
+			if (is_dual_plane)
+				res->plane2_color_component = read_bits(2, below_weights_pos - 2, pb.data);
+
+		}
+
+
+
+		// data structure representing one block of an image.
+		// it is expanded to float prior to processing to save some computation time
+		// on conversions to/from uint8 (this also allows us to handle HDR textures easily)
+		struct imageblock
+		{
+			float orig_data[MAX_TEXELS_PER_BLOCK * 4];  // original input data
+			float work_data[MAX_TEXELS_PER_BLOCK * 4];  // the data that we will compress, either linear or LNS (0..65535 in both cases)
+
+			uint8 rgb_lns[MAX_TEXELS_PER_BLOCK];      // 1 if RGB data are being treated as LNS
+			uint8 alpha_lns[MAX_TEXELS_PER_BLOCK];    // 1 if Alpha data are being treated as LNS
+		};
+
+		enum astc_decode_mode
+		{
+			DECODE_LDR_SRGB,
+			DECODE_LDR,
+			DECODE_HDR
+		};
+
+
+		template < typename vtype > class vtype4
+		{
+		public:
+			vtype x, y, z, w;
+			vtype4() {}
+			vtype4(vtype p, vtype q, vtype r, vtype s) :x(p), y(q), z(r), w(s) {}
+			vtype4(const vtype4& p) :x(p.x), y(p.y), z(p.z), w(p.w) {}
+		};
+		typedef unsigned short ushort;
+		typedef vtype4 < ushort > ushort4;
+		typedef vtype4 < int >int4;
+		static inline int4 operator-(int4 p, int4 q)
+		{
+			return int4(p.x - q.x, p.y - q.y, p.z - q.z, p.w - q.w);
+		}
+		static inline int4 operator>>(int4 p, int q)
+		{
+			return int4(p.x >> q, p.y >> q, p.z >> q, p.w >> q);
+		}
+		static inline int4 operator*(int4 p, int4 q)
+		{
+			return int4(p.x * q.x, p.y * q.y, p.z * q.z, p.w * q.w);
+		}
+		static inline int4 operator+(int4 p, int4 q)
+		{
+			return int4(p.x + q.x, p.y + q.y, p.z + q.z, p.w + q.w);
+		}
+		static inline int4 operator<<(int4 p, int q)
+		{
+			return int4(p.x << q, p.y << q, p.z << q, p.w << q);
+		}
+		static inline int4 operator|(int4 p, int4 q)
+		{
+			return int4(p.x | q.x, p.y | q.y, p.z | q.z, p.w | q.w);
+		}
+
+
+
+		int compute_value_of_texel_int(int texel_to_get, const decimation_table* it, const int* weights)
+		{
+			int i;
+			int summed_value = 8;
+			int weights_to_evaluate = it->texel_num_weights[texel_to_get];
+			for (i = 0; i < weights_to_evaluate; i++)
+			{
+				summed_value += weights[it->texel_weights[texel_to_get][i]] * it->texel_weights_int[texel_to_get][i];
+			}
+			return summed_value >> 4;
+		}
+
+
+		ushort4 lerp_color_int(astc_decode_mode decode_mode, ushort4 color0, ushort4 color1, int weight, int plane2_weight, int plane2_color_component	// -1 in 1-plane mode
+		)
+		{
+			int4 ecolor0 = int4(color0.x, color0.y, color0.z, color0.w);
+			int4 ecolor1 = int4(color1.x, color1.y, color1.z, color1.w);
+
+			int4 eweight1 = int4(weight, weight, weight, weight);
+			switch (plane2_color_component)
+			{
+			case 0:
+				eweight1.x = plane2_weight;
+				break;
+			case 1:
+				eweight1.y = plane2_weight;
+				break;
+			case 2:
+				eweight1.z = plane2_weight;
+				break;
+			case 3:
+				eweight1.w = plane2_weight;
+				break;
+			default:
+				break;
+			}
+
+			int4 eweight0 = int4(64, 64, 64, 64) - eweight1;
+
+			if (decode_mode == DECODE_LDR_SRGB)
+			{
+				ecolor0 = ecolor0 >> 8;
+				ecolor1 = ecolor1 >> 8;
+			}
+			int4 color = (ecolor0 * eweight0) + (ecolor1 * eweight1) + int4(32, 32, 32, 32);
+			color = color >> 6;
+			if (decode_mode == DECODE_LDR_SRGB)
+				color = color | (color << 8);
+
+			ushort4 rcolor = ushort4((ushort)color.x, (ushort)color.y, (ushort)color.z, (ushort)color.w);
+			return rcolor;
+		}
+
+		// conversion functions between the LNS representation and the FP16 representation.
+
+		float float_to_lns(float p)
+		{
+
+			if (astc_isnan(p) || p <= 1.0f / 67108864.0f)
+			{
+				// underflow or NaN value, return 0.
+				// We count underflow if the input value is smaller than 2^-26.
+				return 0;
+			}
+
+			if (fabs(p) >= 65536.0f)
+			{
+				// overflow, return a +INF value
+				return 65535;
+			}
+
+			int expo;
+			float normfrac = frexp(p, &expo);
+			float p1;
+			if (expo < -13)
+			{
+				// input number is smaller than 2^-14. In this case, multiply by 2^25.
+				p1 = p * 33554432.0f;
+				expo = 0;
+			}
+			else
+			{
+				expo += 14;
+				p1 = (normfrac - 0.5f) * 4096.0f;
+			}
+
+			if (p1 < 384.0f)
+				p1 *= 4.0f / 3.0f;
+			else if (p1 <= 1408.0f)
+				p1 += 128.0f;
+			else
+				p1 = (p1 + 512.0f) * (4.0f / 5.0f);
+
+			p1 += expo * 2048.0f;
+			return p1 + 1.0f;
+		}
+
+
+		// helper function to initialize the work-data from the orig-data
+		void imageblock_initialize_work_from_orig(imageblock* pb, int pixelcount)
+		{
+			int i;
+			float* fptr = pb->orig_data;
+			float* wptr = pb->work_data;
+
+			for (i = 0; i < pixelcount; i++)
+			{
+				if (pb->rgb_lns[i])
+				{
+					wptr[0] = float_to_lns(fptr[0]);
+					wptr[1] = float_to_lns(fptr[1]);
+					wptr[2] = float_to_lns(fptr[2]);
+				}
+				else
+				{
+					wptr[0] = fptr[0] * 65535.0f;
+					wptr[1] = fptr[1] * 65535.0f;
+					wptr[2] = fptr[2] * 65535.0f;
+				}
+
+				if (pb->alpha_lns[i])
+				{
+					wptr[3] = float_to_lns(fptr[3]);
+				}
+				else
+				{
+					wptr[3] = fptr[3] * 65535.0f;
+				}
+				fptr += 4;
+				wptr += 4;
+			}
+		}
+
+
+		/*	sized soft-float types. These are mapped to the sized integer types of C99, instead of C's
+			floating-point types; this is because the library needs to maintain exact, bit-level control on all
+			operations on these data types. */
+		typedef uint16 sf16;
+		typedef uint32 sf32;
+
+		typedef union if32_
+		{
+			uint32 u;
+			int32 s;
+			float f;
+		} if32;
+
+		/******************************************
+		  helper functions and their lookup tables
+		 ******************************************/
+		 /* count leading zeros functions. Only used when the input is nonzero. */
+
+#if defined(__GNUC__) && (defined(__i386) || defined(__amd64))
+#elif defined(__arm__) && defined(__ARMCC_VERSION)
+#elif defined(__arm__) && defined(__GNUC__)
+#else
+	/* table used for the slow default versions. */
+		static const uint8 clz_table[256] =
+		{
+			8, 7, 6, 6, 5, 5, 5, 5, 4, 4, 4, 4, 4, 4, 4, 4,
+			3, 3, 3, 3, 3, 3, 3, 3, 3, 3, 3, 3, 3, 3, 3, 3,
+			2, 2, 2, 2, 2, 2, 2, 2, 2, 2, 2, 2, 2, 2, 2, 2,
+			2, 2, 2, 2, 2, 2, 2, 2, 2, 2, 2, 2, 2, 2, 2, 2,
+			1, 1, 1, 1, 1, 1, 1, 1, 1, 1, 1, 1, 1, 1, 1, 1,
+			1, 1, 1, 1, 1, 1, 1, 1, 1, 1, 1, 1, 1, 1, 1, 1,
+			1, 1, 1, 1, 1, 1, 1, 1, 1, 1, 1, 1, 1, 1, 1, 1,
+			1, 1, 1, 1, 1, 1, 1, 1, 1, 1, 1, 1, 1, 1, 1, 1,
+			0, 0, 0, 0, 0, 0, 0, 0, 0, 0, 0, 0, 0, 0, 0, 0,
+			0, 0, 0, 0, 0, 0, 0, 0, 0, 0, 0, 0, 0, 0, 0, 0,
+			0, 0, 0, 0, 0, 0, 0, 0, 0, 0, 0, 0, 0, 0, 0, 0,
+			0, 0, 0, 0, 0, 0, 0, 0, 0, 0, 0, 0, 0, 0, 0, 0,
+			0, 0, 0, 0, 0, 0, 0, 0, 0, 0, 0, 0, 0, 0, 0, 0,
+			0, 0, 0, 0, 0, 0, 0, 0, 0, 0, 0, 0, 0, 0, 0, 0,
+			0, 0, 0, 0, 0, 0, 0, 0, 0, 0, 0, 0, 0, 0, 0, 0,
+			0, 0, 0, 0, 0, 0, 0, 0, 0, 0, 0, 0, 0, 0, 0, 0
+		};
+#endif
+
+		/*
+		   32-bit count-leading-zeros function: use the Assembly instruction whenever possible. */
+		inline uint32 clz32(uint32 inp)
+		{
+#if defined(__GNUC__) && (defined(__i386) || defined(__amd64))
+			uint32 bsr;
+			__asm__("bsrl %1, %0": "=r"(bsr) : "r"(inp | 1));
+			return 31 - bsr;
+#else
+#if defined(__arm__) && defined(__ARMCC_VERSION)
+			return __clz(inp);			/* armcc builtin */
+#else
+#if defined(__arm__) && defined(__GNUC__)
+			uint32 lz;
+			__asm__("clz %0, %1": "=r"(lz) : "r"(inp));
+			return lz;
+#else
+			/* slow default version */
+			uint32 summa = 24;
+			if (inp >= UINT32_C(0x10000))
+			{
+				inp >>= 16;
+				summa -= 16;
+			}
+			if (inp >= UINT32_C(0x100))
+			{
+				inp >>= 8;
+				summa -= 8;
+			}
+			return summa + clz_table[inp];
+#endif
+#endif
+#endif
+		}
+
+
+		/* convert from FP16 to FP32. */
+		sf32 sf16_to_sf32(sf16 inp)
+		{
+			uint32 inpx = inp;
+
+			/*
+				This table contains, for every FP16 sign/exponent value combination,
+				the difference between the input FP16 value and the value obtained
+				by shifting the correct FP32 result right by 13 bits.
+				This table allows us to handle every case except denormals and NaN
+				with just 1 table lookup, 2 shifts and 1 add.
+			*/
+
+#define WITH_MB(a) INT32_C((a) | (1 << 31))
+			static const int32 tbl[64] =
+			{
+				WITH_MB(0x00000), INT32_C(0x1C000), INT32_C(0x1C000), INT32_C(0x1C000), INT32_C(0x1C000), INT32_C(0x1C000), INT32_C(0x1C000), INT32_C(0x1C000),
+				INT32_C(0x1C000), INT32_C(0x1C000), INT32_C(0x1C000), INT32_C(0x1C000), INT32_C(0x1C000), INT32_C(0x1C000), INT32_C(0x1C000), INT32_C(0x1C000),
+				INT32_C(0x1C000), INT32_C(0x1C000), INT32_C(0x1C000), INT32_C(0x1C000), INT32_C(0x1C000), INT32_C(0x1C000), INT32_C(0x1C000), INT32_C(0x1C000),
+				INT32_C(0x1C000), INT32_C(0x1C000), INT32_C(0x1C000), INT32_C(0x1C000), INT32_C(0x1C000), INT32_C(0x1C000), INT32_C(0x1C000), WITH_MB(0x38000),
+				WITH_MB(0x38000), INT32_C(0x54000), INT32_C(0x54000), INT32_C(0x54000), INT32_C(0x54000), INT32_C(0x54000), INT32_C(0x54000), INT32_C(0x54000),
+				INT32_C(0x54000), INT32_C(0x54000), INT32_C(0x54000), INT32_C(0x54000), INT32_C(0x54000), INT32_C(0x54000), INT32_C(0x54000), INT32_C(0x54000),
+				INT32_C(0x54000), INT32_C(0x54000), INT32_C(0x54000), INT32_C(0x54000), INT32_C(0x54000), INT32_C(0x54000), INT32_C(0x54000), INT32_C(0x54000),
+				INT32_C(0x54000), INT32_C(0x54000), INT32_C(0x54000), INT32_C(0x54000), INT32_C(0x54000), INT32_C(0x54000), INT32_C(0x54000), WITH_MB(0x70000)
+			};
+
+			int32 res = tbl[inpx >> 10];
+			res += inpx;
+
+			/* the normal cases: the MSB of 'res' is not set. */
+			if (res >= 0)				/* signed compare */
+				return res << 13;
+
+			/* Infinity and Zero: the bottom 10 bits of 'res' are clear. */
+			if ((res & UINT32_C(0x3FF)) == 0)
+				return res << 13;
+
+			/* NaN: the exponent field of 'inp' is not zero; NaNs must be quietened. */
+			if ((inpx & 0x7C00) != 0)
+				return (res << 13) | UINT32_C(0x400000);
+
+			/* the remaining cases are Denormals. */
+			{
+				uint32 sign = (inpx & UINT32_C(0x8000)) << 16;
+				uint32 mskval = inpx & UINT32_C(0x7FFF);
+				uint32 leadingzeroes = clz32(mskval);
+				mskval <<= leadingzeroes;
+				return (mskval >> 8) + ((0x85 - leadingzeroes) << 23) + sign;
+			}
+		}
+
+
+		float sf16_to_float(sf16 p)
+		{
+			if32 i;
+			i.u = sf16_to_sf32(p);
+			return i.f;
+		}
+
+		// conversion function from 16-bit LDR value to FP16.
+		// note: for LDR interpolation, it is impossible to get a denormal result;
+		// this simplifies the conversion.
+		// FALSE; we can receive a very small UNORM16 through the constant-block.
+		uint16 unorm16_to_sf16(uint16 p)
+		{
+			if (p == 0xFFFF)
+				return 0x3C00;			// value of 1.0 .
+			if (p < 4)
+				return p << 8;
+
+			int lz = clz32(p) - 16;
+			p <<= (lz + 1);
+			p >>= 6;
+			p |= (14 - lz) << 10;
+			return p;
+		}
+
+
+		/*
+			Partition table representation:
+			For each block size, we have 3 tables, each with 1024 partitionings;
+			these three tables correspond to 2, 3 and 4 partitions respectively.
+			For each partitioning, we have:
+			* a 4-entry table indicating how many texels there are in each of the 4 partitions.
+			  This may be from 0 to a very large value.
+			* a table indicating the partition index of each of the texels in the block.
+			  Each index may be 0, 1, 2 or 3.
+			* Each element in the table is an uint8 indicating partition index (0, 1, 2 or 3)
+		*/
+
+		struct partition_info
+		{
+			int partition_count;
+			uint8 texels_per_partition[4];
+			uint8 partition_of_texel[MAX_TEXELS_PER_BLOCK];
+			uint8 texels_of_partition[4][MAX_TEXELS_PER_BLOCK];
+
+			uint64 coverage_bitmaps[4];	// used for the purposes of k-means partition search.
+		};
+
+
+		static partition_info* s_partition_table;
+
+
+		uint32 hash52(uint32 inp)
+		{
+			inp ^= inp >> 15;
+
+			inp *= 0xEEDE0891;			// (2^4+1)*(2^7+1)*(2^17-1)
+			inp ^= inp >> 5;
+			inp += inp << 16;
+			inp ^= inp >> 7;
+			inp ^= inp >> 3;
+			inp ^= inp << 6;
+			inp ^= inp >> 17;
+			return inp;
+		}
+
+		int select_partition(int seed, int x, int y, int z, int partitioncount, int small_block)
+		{
+			if (small_block)
+			{
+				x <<= 1;
+				y <<= 1;
+				z <<= 1;
+			}
+
+			seed += (partitioncount - 1) * 1024;
+
+			uint32 rnum = hash52(seed);
+
+			uint8 seed1 = rnum & 0xF;
+			uint8 seed2 = (rnum >> 4) & 0xF;
+			uint8 seed3 = (rnum >> 8) & 0xF;
+			uint8 seed4 = (rnum >> 12) & 0xF;
+			uint8 seed5 = (rnum >> 16) & 0xF;
+			uint8 seed6 = (rnum >> 20) & 0xF;
+			uint8 seed7 = (rnum >> 24) & 0xF;
+			uint8 seed8 = (rnum >> 28) & 0xF;
+			uint8 seed9 = (rnum >> 18) & 0xF;
+			uint8 seed10 = (rnum >> 22) & 0xF;
+			uint8 seed11 = (rnum >> 26) & 0xF;
+			uint8 seed12 = ((rnum >> 30) | (rnum << 2)) & 0xF;
+
+			// squaring all the seeds in order to bias their distribution
+			// towards lower values.
+			seed1 *= seed1;
+			seed2 *= seed2;
+			seed3 *= seed3;
+			seed4 *= seed4;
+			seed5 *= seed5;
+			seed6 *= seed6;
+			seed7 *= seed7;
+			seed8 *= seed8;
+			seed9 *= seed9;
+			seed10 *= seed10;
+			seed11 *= seed11;
+			seed12 *= seed12;
+
+
+			int sh1, sh2, sh3;
+			if (seed & 1)
+			{
+				sh1 = (seed & 2 ? 4 : 5);
+				sh2 = (partitioncount == 3 ? 6 : 5);
+			}
+			else
+			{
+				sh1 = (partitioncount == 3 ? 6 : 5);
+				sh2 = (seed & 2 ? 4 : 5);
+			}
+			sh3 = (seed & 0x10) ? sh1 : sh2;
+
+			seed1 >>= sh1;
+			seed2 >>= sh2;
+			seed3 >>= sh1;
+			seed4 >>= sh2;
+			seed5 >>= sh1;
+			seed6 >>= sh2;
+			seed7 >>= sh1;
+			seed8 >>= sh2;
+
+			seed9 >>= sh3;
+			seed10 >>= sh3;
+			seed11 >>= sh3;
+			seed12 >>= sh3;
+
+
+
+			int a = seed1 * x + seed2 * y + seed11 * z + (rnum >> 14);
+			int b = seed3 * x + seed4 * y + seed12 * z + (rnum >> 10);
+			int c = seed5 * x + seed6 * y + seed9 * z + (rnum >> 6);
+			int d = seed7 * x + seed8 * y + seed10 * z + (rnum >> 2);
+
+
+			// apply the saw
+			a &= 0x3F;
+			b &= 0x3F;
+			c &= 0x3F;
+			d &= 0x3F;
+
+			// remove some of the components if we are to output < 4 partitions.
+			if (partitioncount <= 3)
+				d = 0;
+			if (partitioncount <= 2)
+				c = 0;
+			if (partitioncount <= 1)
+				b = 0;
+
+			int partition;
+			if (a >= b && a >= c && a >= d)
+				partition = 0;
+			else if (b >= c && b >= d)
+				partition = 1;
+			else if (c >= d)
+				partition = 2;
+			else
+				partition = 3;
+			return partition;
+		}
+
+
+
+		void generate_one_partition_table(int partition_count, int partition_index, partition_info* pt)
+		{
+			int small_block = 1;
+
+			uint8* partition_of_texel = pt->partition_of_texel;
+			int x, y, z, i;
+
+
+			for (z = 0; z < 1; z++)
+				for (y = 0; y < 4; y++)
+					for (x = 0; x < 4; x++)
+					{
+						uint8 part = (uint8)select_partition(partition_index, x, y, z, partition_count, small_block);
+						*partition_of_texel++ = part;
+					}
+
+
+			int texels_per_block = 16;
+
+			int counts[4];
+			for (i = 0; i < 4; i++)
+				counts[i] = 0;
+
+			for (i = 0; i < texels_per_block; i++)
+			{
+				int partition = pt->partition_of_texel[i];
+				pt->texels_of_partition[partition][counts[partition]++] = (uint8)i;
+			}
+
+			for (i = 0; i < 4; i++)
+				pt->texels_per_partition[i] = (uint8)counts[i];
+
+			if (counts[0] == 0)
+				pt->partition_count = 0;
+			else if (counts[1] == 0)
+				pt->partition_count = 1;
+			else if (counts[2] == 0)
+				pt->partition_count = 2;
+			else if (counts[3] == 0)
+				pt->partition_count = 3;
+			else
+				pt->partition_count = 4;
+
+
+			for (i = 0; i < 4; i++)
+				pt->coverage_bitmaps[i] = 0ULL;
+
+			int texels_to_process = 16;
+			for (i = 0; i < texels_to_process; i++)
+			{
+				pt->coverage_bitmaps[pt->partition_of_texel[i]] |= 1ULL << i;
+			}
+		}
+
+
+		static void generate_partition_tables()
+		{
+			partition_info* one_partition = new partition_info;
+
+			generate_one_partition_table(1, 0, one_partition);
+
+			s_partition_table = one_partition;
+		}
+
+
+		const partition_info* get_partition_table()
+		{
+			if (s_partition_table == nullptr)
+				generate_partition_tables();
+
+			return s_partition_table;
+		}
+
+
+
+		const uint8 color_unquantization_tables[21][256] = {
+			{
+			 0, 255,
+			 },
+			{
+			 0, 128, 255,
+			 },
+			{
+			 0, 85, 170, 255,
+			 },
+			{
+			 0, 64, 128, 192, 255,
+			 },
+			{
+			 0, 255, 51, 204, 102, 153,
+			 },
+			{
+			 0, 36, 73, 109, 146, 182, 219, 255,
+			 },
+			{
+			 0, 255, 28, 227, 56, 199, 84, 171, 113, 142,
+			 },
+			{
+			 0, 255, 69, 186, 23, 232, 92, 163, 46, 209, 116, 139,
+			 },
+			{
+			 0, 17, 34, 51, 68, 85, 102, 119, 136, 153, 170, 187, 204, 221, 238, 255,
+			 },
+			{
+			 0, 255, 67, 188, 13, 242, 80, 175, 27, 228, 94, 161, 40, 215, 107, 148,
+			 54, 201, 121, 134,
+			 },
+			{
+			 0, 255, 33, 222, 66, 189, 99, 156, 11, 244, 44, 211, 77, 178, 110, 145,
+			 22, 233, 55, 200, 88, 167, 121, 134,
+			 },
+			{
+			 0, 8, 16, 24, 33, 41, 49, 57, 66, 74, 82, 90, 99, 107, 115, 123,
+			 132, 140, 148, 156, 165, 173, 181, 189, 198, 206, 214, 222, 231, 239, 247, 255,
+			 },
+			{
+			 0, 255, 32, 223, 65, 190, 97, 158, 6, 249, 39, 216, 71, 184, 104, 151,
+			 13, 242, 45, 210, 78, 177, 110, 145, 19, 236, 52, 203, 84, 171, 117, 138,
+			 26, 229, 58, 197, 91, 164, 123, 132,
+			 },
+			{
+			 0, 255, 16, 239, 32, 223, 48, 207, 65, 190, 81, 174, 97, 158, 113, 142,
+			 5, 250, 21, 234, 38, 217, 54, 201, 70, 185, 86, 169, 103, 152, 119, 136,
+			 11, 244, 27, 228, 43, 212, 59, 196, 76, 179, 92, 163, 108, 147, 124, 131,
+			 },
+			{
+			 0, 4, 8, 12, 16, 20, 24, 28, 32, 36, 40, 44, 48, 52, 56, 60,
+			 65, 69, 73, 77, 81, 85, 89, 93, 97, 101, 105, 109, 113, 117, 121, 125,
+			 130, 134, 138, 142, 146, 150, 154, 158, 162, 166, 170, 174, 178, 182, 186, 190,
+			 195, 199, 203, 207, 211, 215, 219, 223, 227, 231, 235, 239, 243, 247, 251, 255,
+			 },
+			{
+			 0, 255, 16, 239, 32, 223, 48, 207, 64, 191, 80, 175, 96, 159, 112, 143,
+			 3, 252, 19, 236, 35, 220, 51, 204, 67, 188, 83, 172, 100, 155, 116, 139,
+			 6, 249, 22, 233, 38, 217, 54, 201, 71, 184, 87, 168, 103, 152, 119, 136,
+			 9, 246, 25, 230, 42, 213, 58, 197, 74, 181, 90, 165, 106, 149, 122, 133,
+			 13, 242, 29, 226, 45, 210, 61, 194, 77, 178, 93, 162, 109, 146, 125, 130,
+			 },
+			{
+			 0, 255, 8, 247, 16, 239, 24, 231, 32, 223, 40, 215, 48, 207, 56, 199,
+			 64, 191, 72, 183, 80, 175, 88, 167, 96, 159, 104, 151, 112, 143, 120, 135,
+			 2, 253, 10, 245, 18, 237, 26, 229, 35, 220, 43, 212, 51, 204, 59, 196,
+			 67, 188, 75, 180, 83, 172, 91, 164, 99, 156, 107, 148, 115, 140, 123, 132,
+			 5, 250, 13, 242, 21, 234, 29, 226, 37, 218, 45, 210, 53, 202, 61, 194,
+			 70, 185, 78, 177, 86, 169, 94, 161, 102, 153, 110, 145, 118, 137, 126, 129,
+			 },
+			{
+			 0, 2, 4, 6, 8, 10, 12, 14, 16, 18, 20, 22, 24, 26, 28, 30,
+			 32, 34, 36, 38, 40, 42, 44, 46, 48, 50, 52, 54, 56, 58, 60, 62,
+			 64, 66, 68, 70, 72, 74, 76, 78, 80, 82, 84, 86, 88, 90, 92, 94,
+			 96, 98, 100, 102, 104, 106, 108, 110, 112, 114, 116, 118, 120, 122, 124, 126,
+			 129, 131, 133, 135, 137, 139, 141, 143, 145, 147, 149, 151, 153, 155, 157, 159,
+			 161, 163, 165, 167, 169, 171, 173, 175, 177, 179, 181, 183, 185, 187, 189, 191,
+			 193, 195, 197, 199, 201, 203, 205, 207, 209, 211, 213, 215, 217, 219, 221, 223,
+			 225, 227, 229, 231, 233, 235, 237, 239, 241, 243, 245, 247, 249, 251, 253, 255,
+			 },
+			{
+			 0, 255, 8, 247, 16, 239, 24, 231, 32, 223, 40, 215, 48, 207, 56, 199,
+			 64, 191, 72, 183, 80, 175, 88, 167, 96, 159, 104, 151, 112, 143, 120, 135,
+			 1, 254, 9, 246, 17, 238, 25, 230, 33, 222, 41, 214, 49, 206, 57, 198,
+			 65, 190, 73, 182, 81, 174, 89, 166, 97, 158, 105, 150, 113, 142, 121, 134,
+			 3, 252, 11, 244, 19, 236, 27, 228, 35, 220, 43, 212, 51, 204, 59, 196,
+			 67, 188, 75, 180, 83, 172, 91, 164, 99, 156, 107, 148, 115, 140, 123, 132,
+			 4, 251, 12, 243, 20, 235, 28, 227, 36, 219, 44, 211, 52, 203, 60, 195,
+			 68, 187, 76, 179, 84, 171, 92, 163, 100, 155, 108, 147, 116, 139, 124, 131,
+			 6, 249, 14, 241, 22, 233, 30, 225, 38, 217, 46, 209, 54, 201, 62, 193,
+			 70, 185, 78, 177, 86, 169, 94, 161, 102, 153, 110, 145, 118, 137, 126, 129,
+			 },
+			{
+			 0, 255, 4, 251, 8, 247, 12, 243, 16, 239, 20, 235, 24, 231, 28, 227,
+			 32, 223, 36, 219, 40, 215, 44, 211, 48, 207, 52, 203, 56, 199, 60, 195,
+			 64, 191, 68, 187, 72, 183, 76, 179, 80, 175, 84, 171, 88, 167, 92, 163,
+			 96, 159, 100, 155, 104, 151, 108, 147, 112, 143, 116, 139, 120, 135, 124, 131,
+			 1, 254, 5, 250, 9, 246, 13, 242, 17, 238, 21, 234, 25, 230, 29, 226,
+			 33, 222, 37, 218, 41, 214, 45, 210, 49, 206, 53, 202, 57, 198, 61, 194,
+			 65, 190, 69, 186, 73, 182, 77, 178, 81, 174, 85, 170, 89, 166, 93, 162,
+			 97, 158, 101, 154, 105, 150, 109, 146, 113, 142, 117, 138, 121, 134, 125, 130,
+			 2, 253, 6, 249, 10, 245, 14, 241, 18, 237, 22, 233, 26, 229, 30, 225,
+			 34, 221, 38, 217, 42, 213, 46, 209, 50, 205, 54, 201, 58, 197, 62, 193,
+			 66, 189, 70, 185, 74, 181, 78, 177, 82, 173, 86, 169, 90, 165, 94, 161,
+			 98, 157, 102, 153, 106, 149, 110, 145, 114, 141, 118, 137, 122, 133, 126, 129,
+			 },
+			{
+			 0, 1, 2, 3, 4, 5, 6, 7, 8, 9, 10, 11, 12, 13, 14, 15,
+			 16, 17, 18, 19, 20, 21, 22, 23, 24, 25, 26, 27, 28, 29, 30, 31,
+			 32, 33, 34, 35, 36, 37, 38, 39, 40, 41, 42, 43, 44, 45, 46, 47,
+			 48, 49, 50, 51, 52, 53, 54, 55, 56, 57, 58, 59, 60, 61, 62, 63,
+			 64, 65, 66, 67, 68, 69, 70, 71, 72, 73, 74, 75, 76, 77, 78, 79,
+			 80, 81, 82, 83, 84, 85, 86, 87, 88, 89, 90, 91, 92, 93, 94, 95,
+			 96, 97, 98, 99, 100, 101, 102, 103, 104, 105, 106, 107, 108, 109, 110, 111,
+			 112, 113, 114, 115, 116, 117, 118, 119, 120, 121, 122, 123, 124, 125, 126, 127,
+			 128, 129, 130, 131, 132, 133, 134, 135, 136, 137, 138, 139, 140, 141, 142, 143,
+			 144, 145, 146, 147, 148, 149, 150, 151, 152, 153, 154, 155, 156, 157, 158, 159,
+			 160, 161, 162, 163, 164, 165, 166, 167, 168, 169, 170, 171, 172, 173, 174, 175,
+			 176, 177, 178, 179, 180, 181, 182, 183, 184, 185, 186, 187, 188, 189, 190, 191,
+			 192, 193, 194, 195, 196, 197, 198, 199, 200, 201, 202, 203, 204, 205, 206, 207,
+			 208, 209, 210, 211, 212, 213, 214, 215, 216, 217, 218, 219, 220, 221, 222, 223,
+			 224, 225, 226, 227, 228, 229, 230, 231, 232, 233, 234, 235, 236, 237, 238, 239,
+			 240, 241, 242, 243, 244, 245, 246, 247, 248, 249, 250, 251, 252, 253, 254, 255,
+			 },
+		};
+
+
+		enum endpoint_formats
+		{
+			FMT_LUMINANCE = 0,
+			FMT_LUMINANCE_DELTA = 1,
+			FMT_HDR_LUMINANCE_LARGE_RANGE = 2,
+			FMT_HDR_LUMINANCE_SMALL_RANGE = 3,
+			FMT_LUMINANCE_ALPHA = 4,
+			FMT_LUMINANCE_ALPHA_DELTA = 5,
+			FMT_RGB_SCALE = 6,
+			FMT_HDR_RGB_SCALE = 7,
+			FMT_RGB = 8,
+			FMT_RGB_DELTA = 9,
+			FMT_RGB_SCALE_ALPHA = 10,
+			FMT_HDR_RGB = 11,
+			FMT_RGBA = 12,
+			FMT_RGBA_DELTA = 13,
+			FMT_HDR_RGB_LDR_ALPHA = 14,
+			FMT_HDR_RGBA = 15,
+		};
+
+
+		int rgb_delta_unpack(const int input[6], int quantization_level, ushort4* output0, ushort4* output1)
+		{
+			// unquantize the color endpoints
+			int r0 = color_unquantization_tables[quantization_level][input[0]];
+			int g0 = color_unquantization_tables[quantization_level][input[2]];
+			int b0 = color_unquantization_tables[quantization_level][input[4]];
+
+			int r1 = color_unquantization_tables[quantization_level][input[1]];
+			int g1 = color_unquantization_tables[quantization_level][input[3]];
+			int b1 = color_unquantization_tables[quantization_level][input[5]];
+
+			// perform the bit-transfer procedure
+			r0 |= (r1 & 0x80) << 1;
+			g0 |= (g1 & 0x80) << 1;
+			b0 |= (b1 & 0x80) << 1;
+			r1 &= 0x7F;
+			g1 &= 0x7F;
+			b1 &= 0x7F;
+			if (r1 & 0x40)
+				r1 -= 0x80;
+			if (g1 & 0x40)
+				g1 -= 0x80;
+			if (b1 & 0x40)
+				b1 -= 0x80;
+
+			r0 >>= 1;
+			g0 >>= 1;
+			b0 >>= 1;
+			r1 >>= 1;
+			g1 >>= 1;
+			b1 >>= 1;
+
+			int rgbsum = r1 + g1 + b1;
+
+			r1 += r0;
+			g1 += g0;
+			b1 += b0;
+
+
+			int retval;
+
+			int r0e, g0e, b0e;
+			int r1e, g1e, b1e;
+
+			if (rgbsum >= 0)
+			{
+				r0e = r0;
+				g0e = g0;
+				b0e = b0;
+
+				r1e = r1;
+				g1e = g1;
+				b1e = b1;
+
+				retval = 0;
+			}
+			else
+			{
+				r0e = (r1 + b1) >> 1;
+				g0e = (g1 + b1) >> 1;
+				b0e = b1;
+
+				r1e = (r0 + b0) >> 1;
+				g1e = (g0 + b0) >> 1;
+				b1e = b0;
+
+				retval = 1;
+			}
+
+			if (r0e < 0)
+				r0e = 0;
+			else if (r0e > 255)
+				r0e = 255;
+
+			if (g0e < 0)
+				g0e = 0;
+			else if (g0e > 255)
+				g0e = 255;
+
+			if (b0e < 0)
+				b0e = 0;
+			else if (b0e > 255)
+				b0e = 255;
+
+			if (r1e < 0)
+				r1e = 0;
+			else if (r1e > 255)
+				r1e = 255;
+
+			if (g1e < 0)
+				g1e = 0;
+			else if (g1e > 255)
+				g1e = 255;
+
+			if (b1e < 0)
+				b1e = 0;
+			else if (b1e > 255)
+				b1e = 255;
+
+			output0->x = (ushort)r0e;
+			output0->y = (ushort)g0e;
+			output0->z = (ushort)b0e;
+			output0->w = 0xFF;
+
+			output1->x = (ushort)r1e;
+			output1->y = (ushort)g1e;
+			output1->z = (ushort)b1e;
+			output1->w = 0xFF;
+
+			return retval;
+		}
+
+
+		int rgb_unpack(const int input[6], int quantization_level, ushort4* output0, ushort4* output1)
+		{
+
+			int ri0b = color_unquantization_tables[quantization_level][input[0]];
+			int ri1b = color_unquantization_tables[quantization_level][input[1]];
+			int gi0b = color_unquantization_tables[quantization_level][input[2]];
+			int gi1b = color_unquantization_tables[quantization_level][input[3]];
+			int bi0b = color_unquantization_tables[quantization_level][input[4]];
+			int bi1b = color_unquantization_tables[quantization_level][input[5]];
+
+			if (ri0b + gi0b + bi0b > ri1b + gi1b + bi1b)
+			{
+				// blue-contraction
+				ri0b = (ri0b + bi0b) >> 1;
+				gi0b = (gi0b + bi0b) >> 1;
+				ri1b = (ri1b + bi1b) >> 1;
+				gi1b = (gi1b + bi1b) >> 1;
+
+				output0->x = (ushort)ri1b;
+				output0->y = (ushort)gi1b;
+				output0->z = (ushort)bi1b;
+				output0->w = 255;
+
+				output1->x = (ushort)ri0b;
+				output1->y = (ushort)gi0b;
+				output1->z = (ushort)bi0b;
+				output1->w = 255;
+				return 1;
+			}
+			else
+			{
+				output0->x = (ushort)ri0b;
+				output0->y = (ushort)gi0b;
+				output0->z = (ushort)bi0b;
+				output0->w = 255;
+
+				output1->x = (ushort)ri1b;
+				output1->y = (ushort)gi1b;
+				output1->z = (ushort)bi1b;
+				output1->w = 255;
+				return 0;
+			}
+		}
+
+
+		void rgba_unpack(const int input[8], int quantization_level, ushort4* output0, ushort4* output1)
+		{
+			int order = rgb_unpack(input, quantization_level, output0, output1);
+			if (order == 0)
+			{
+				output0->w = color_unquantization_tables[quantization_level][input[6]];
+				output1->w = color_unquantization_tables[quantization_level][input[7]];
+			}
+			else
+			{
+				output0->w = color_unquantization_tables[quantization_level][input[7]];
+				output1->w = color_unquantization_tables[quantization_level][input[6]];
+			}
+		}
+
+
+		void rgba_delta_unpack(const int input[8], int quantization_level, ushort4* output0, ushort4* output1)
+		{
+			int a0 = color_unquantization_tables[quantization_level][input[6]];
+			int a1 = color_unquantization_tables[quantization_level][input[7]];
+			a0 |= (a1 & 0x80) << 1;
+			a1 &= 0x7F;
+			if (a1 & 0x40)
+				a1 -= 0x80;
+			a0 >>= 1;
+			a1 >>= 1;
+			a1 += a0;
+
+			if (a1 < 0)
+				a1 = 0;
+			else if (a1 > 255)
+				a1 = 255;
+
+			int order = rgb_delta_unpack(input, quantization_level, output0, output1);
+			if (order == 0)
+			{
+				output0->w = (ushort)a0;
+				output1->w = (ushort)a1;
+			}
+			else
+			{
+				output0->w = (ushort)a1;
+				output1->w = (ushort)a0;
+			}
+		}
+
+
+		void rgb_scale_unpack(const int input[4], int quantization_level, ushort4* output0, ushort4* output1)
+		{
+			int ir = color_unquantization_tables[quantization_level][input[0]];
+			int ig = color_unquantization_tables[quantization_level][input[1]];
+			int ib = color_unquantization_tables[quantization_level][input[2]];
+
+			int iscale = color_unquantization_tables[quantization_level][input[3]];
+
+			*output1 = ushort4((ushort)ir, (ushort)ig, (ushort)ib, 255);
+			*output0 = ushort4((ushort)((ir * iscale) >> 8), (ushort)((ig * iscale) >> 8), (ushort)((ib * iscale) >> 8), 255);
+		}
+
+
+		void rgb_scale_alpha_unpack(const int input[6], int quantization_level, ushort4* output0, ushort4* output1)
+		{
+			rgb_scale_unpack(input, quantization_level, output0, output1);
+			output0->w = color_unquantization_tables[quantization_level][input[4]];
+			output1->w = color_unquantization_tables[quantization_level][input[5]];
+
+		}
+
+
+		void luminance_unpack(const int input[2], int quantization_level, ushort4* output0, ushort4* output1)
+		{
+			uint8 lum0 = color_unquantization_tables[quantization_level][input[0]];
+			uint8 lum1 = color_unquantization_tables[quantization_level][input[1]];
+			*output0 = ushort4(lum0, lum0, lum0, 255);
+			*output1 = ushort4(lum1, lum1, lum1, 255);
+		}
+
+
+		void luminance_delta_unpack(const int input[2], int quantization_level, ushort4* output0, ushort4* output1)
+		{
+			ushort v0 = color_unquantization_tables[quantization_level][input[0]];
+			ushort v1 = color_unquantization_tables[quantization_level][input[1]];
+			ushort l0 = (v0 >> 2) | (v1 & 0xC0);
+			ushort l1 = l0 + (v1 & 0x3F);
+
+			if (l1 > 255)
+				l1 = 255;
+
+			*output0 = ushort4(l0, l0, l0, 255);
+			*output1 = ushort4(l1, l1, l1, 255);
+		}
+
+
+		void luminance_alpha_unpack(const int input[4], int quantization_level, ushort4* output0, ushort4* output1)
+		{
+			ushort lum0 = color_unquantization_tables[quantization_level][input[0]];
+			ushort lum1 = color_unquantization_tables[quantization_level][input[1]];
+			ushort alpha0 = color_unquantization_tables[quantization_level][input[2]];
+			ushort alpha1 = color_unquantization_tables[quantization_level][input[3]];
+			*output0 = ushort4(lum0, lum0, lum0, alpha0);
+			*output1 = ushort4(lum1, lum1, lum1, alpha1);
+		}
+
+
+		void luminance_alpha_delta_unpack(const int input[4], int quantization_level, ushort4* output0, ushort4* output1)
+		{
+			ushort lum0 = color_unquantization_tables[quantization_level][input[0]];
+			ushort lum1 = color_unquantization_tables[quantization_level][input[1]];
+			ushort alpha0 = color_unquantization_tables[quantization_level][input[2]];
+			ushort alpha1 = color_unquantization_tables[quantization_level][input[3]];
+
+			lum0 |= (lum1 & 0x80) << 1;
+			alpha0 |= (alpha1 & 0x80) << 1;
+			lum1 &= 0x7F;
+			alpha1 &= 0x7F;
+			if (lum1 & 0x40)
+				lum1 -= 0x80;
+			if (alpha1 & 0x40)
+				alpha1 -= 0x80;
+
+			lum0 >>= 1;
+			lum1 >>= 1;
+			alpha0 >>= 1;
+			alpha1 >>= 1;
+			lum1 += lum0;
+			alpha1 += alpha0;
+
+
+			if (lum1 > 255)
+				lum1 = 255;
+
+
+			if (alpha1 > 255)
+				alpha1 = 255;
+
+			*output0 = ushort4(lum0, lum0, lum0, alpha0);
+			*output1 = ushort4(lum1, lum1, lum1, alpha1);
+		}
+
+
+		// RGB-offset format
+		void hdr_rgbo_unpack3(const int input[4], int quantization_level, ushort4* output0, ushort4* output1)
+		{
+			int v0 = color_unquantization_tables[quantization_level][input[0]];
+			int v1 = color_unquantization_tables[quantization_level][input[1]];
+			int v2 = color_unquantization_tables[quantization_level][input[2]];
+			int v3 = color_unquantization_tables[quantization_level][input[3]];
+
+			int modeval = ((v0 & 0xC0) >> 6) | (((v1 & 0x80) >> 7) << 2) | (((v2 & 0x80) >> 7) << 3);
+
+			int majcomp;
+			int mode;
+			if ((modeval & 0xC) != 0xC)
+			{
+				majcomp = modeval >> 2;
+				mode = modeval & 3;
+			}
+			else if (modeval != 0xF)
+			{
+				majcomp = modeval & 3;
+				mode = 4;
+			}
+			else
+			{
+				majcomp = 0;
+				mode = 5;
+			}
+
+			int red = v0 & 0x3F;
+			int green = v1 & 0x1F;
+			int blue = v2 & 0x1F;
+			int scale = v3 & 0x1F;
+
+			int bit0 = (v1 >> 6) & 1;
+			int bit1 = (v1 >> 5) & 1;
+			int bit2 = (v2 >> 6) & 1;
+			int bit3 = (v2 >> 5) & 1;
+			int bit4 = (v3 >> 7) & 1;
+			int bit5 = (v3 >> 6) & 1;
+			int bit6 = (v3 >> 5) & 1;
+
+			int ohcomp = 1 << mode;
+
+			if (ohcomp & 0x30)
+				green |= bit0 << 6;
+			if (ohcomp & 0x3A)
+				green |= bit1 << 5;
+			if (ohcomp & 0x30)
+				blue |= bit2 << 6;
+			if (ohcomp & 0x3A)
+				blue |= bit3 << 5;
+
+			if (ohcomp & 0x3D)
+				scale |= bit6 << 5;
+			if (ohcomp & 0x2D)
+				scale |= bit5 << 6;
+			if (ohcomp & 0x04)
+				scale |= bit4 << 7;
+
+			if (ohcomp & 0x3B)
+				red |= bit4 << 6;
+			if (ohcomp & 0x04)
+				red |= bit3 << 6;
+
+			if (ohcomp & 0x10)
+				red |= bit5 << 7;
+			if (ohcomp & 0x0F)
+				red |= bit2 << 7;
+
+			if (ohcomp & 0x05)
+				red |= bit1 << 8;
+			if (ohcomp & 0x0A)
+				red |= bit0 << 8;
+
+			if (ohcomp & 0x05)
+				red |= bit0 << 9;
+			if (ohcomp & 0x02)
+				red |= bit6 << 9;
+
+			if (ohcomp & 0x01)
+				red |= bit3 << 10;
+			if (ohcomp & 0x02)
+				red |= bit5 << 10;
+
+
+			// expand to 12 bits.
+			static const int shamts[6] = { 1, 1, 2, 3, 4, 5 };
+			int shamt = shamts[mode];
+			red <<= shamt;
+			green <<= shamt;
+			blue <<= shamt;
+			scale <<= shamt;
+
+			// on modes 0 to 4, the values stored for "green" and "blue" are differentials,
+			// not absolute values.
+			if (mode != 5)
+			{
+				green = red - green;
+				blue = red - blue;
+			}
+
+			// switch around components.
+			int temp;
+			switch (majcomp)
+			{
+			case 1:
+				temp = red;
+				red = green;
+				green = temp;
+				break;
+			case 2:
+				temp = red;
+				red = blue;
+				blue = temp;
+				break;
+			default:
+				break;
+			}
+
+
+			int red0 = red - scale;
+			int green0 = green - scale;
+			int blue0 = blue - scale;
+
+			// clamp to [0,0xFFF].
+			if (red < 0)
+				red = 0;
+			if (green < 0)
+				green = 0;
+			if (blue < 0)
+				blue = 0;
+
+			if (red0 < 0)
+				red0 = 0;
+			if (green0 < 0)
+				green0 = 0;
+			if (blue0 < 0)
+				blue0 = 0;
+
+			*output0 = ushort4(ushort(red0 << 4), ushort(green0 << 4), ushort(blue0 << 4), 0x7800);
+			*output1 = ushort4(ushort(red << 4), ushort(green << 4), ushort(blue << 4), 0x7800);
+		}
+
+
+		void hdr_rgb_unpack3(const int input[6], int quantization_level, ushort4* output0, ushort4* output1)
+		{
+
+			int v0 = color_unquantization_tables[quantization_level][input[0]];
+			int v1 = color_unquantization_tables[quantization_level][input[1]];
+			int v2 = color_unquantization_tables[quantization_level][input[2]];
+			int v3 = color_unquantization_tables[quantization_level][input[3]];
+			int v4 = color_unquantization_tables[quantization_level][input[4]];
+			int v5 = color_unquantization_tables[quantization_level][input[5]];
+
+			// extract all the fixed-placement bitfields
+			int modeval = ((v1 & 0x80) >> 7) | (((v2 & 0x80) >> 7) << 1) | (((v3 & 0x80) >> 7) << 2);
+
+			int majcomp = ((v4 & 0x80) >> 7) | (((v5 & 0x80) >> 7) << 1);
+
+			if (majcomp == 3)
+			{
+				*output0 = ushort4(ushort(v0 << 8), ushort(v2 << 8), ushort(v4 & 0x7F) << 9, 0x7800);
+				*output1 = ushort4(ushort(v1 << 8), ushort(v3 << 8), ushort(v5 & 0x7F) << 9, 0x7800);
+				return;
+			}
+
+			int a = v0 | ((v1 & 0x40) << 2);
+			int b0 = v2 & 0x3f;
+			int b1 = v3 & 0x3f;
+			int c = v1 & 0x3f;
+			int d0 = v4 & 0x7f;
+			int d1 = v5 & 0x7f;
+
+			// get hold of the number of bits in 'd0' and 'd1'
+			static const int dbits_tab[8] = { 7, 6, 7, 6, 5, 6, 5, 6 };
+			int dbits = dbits_tab[modeval];
+
+			// extract six variable-placement bits
+			int bit0 = (v2 >> 6) & 1;
+			int bit1 = (v3 >> 6) & 1;
+
+			int bit2 = (v4 >> 6) & 1;
+			int bit3 = (v5 >> 6) & 1;
+			int bit4 = (v4 >> 5) & 1;
+			int bit5 = (v5 >> 5) & 1;
+
+
+			// and prepend the variable-placement bits depending on mode.
+			int ohmod = 1 << modeval;	// one-hot-mode
+			if (ohmod & 0xA4)
+				a |= bit0 << 9;
+			if (ohmod & 0x8)
+				a |= bit2 << 9;
+			if (ohmod & 0x50)
+				a |= bit4 << 9;
+
+			if (ohmod & 0x50)
+				a |= bit5 << 10;
+			if (ohmod & 0xA0)
+				a |= bit1 << 10;
+
+			if (ohmod & 0xC0)
+				a |= bit2 << 11;
+
+			if (ohmod & 0x4)
+				c |= bit1 << 6;
+			if (ohmod & 0xE8)
+				c |= bit3 << 6;
+
+			if (ohmod & 0x20)
+				c |= bit2 << 7;
+
+
+			if (ohmod & 0x5B)
+				b0 |= bit0 << 6;
+			if (ohmod & 0x5B)
+				b1 |= bit1 << 6;
+
+			if (ohmod & 0x12)
+				b0 |= bit2 << 7;
+			if (ohmod & 0x12)
+				b1 |= bit3 << 7;
+
+			if (ohmod & 0xAF)
+				d0 |= bit4 << 5;
+			if (ohmod & 0xAF)
+				d1 |= bit5 << 5;
+			if (ohmod & 0x5)
+				d0 |= bit2 << 6;
+			if (ohmod & 0x5)
+				d1 |= bit3 << 6;
+
+			// sign-extend 'd0' and 'd1'
+			// note: this code assumes that signed right-shift actually sign-fills, not zero-fills.
+			int32 d0x = d0;
+			int32 d1x = d1;
+			int sx_shamt = 32 - dbits;
+			d0x <<= sx_shamt;
+			d0x >>= sx_shamt;
+			d1x <<= sx_shamt;
+			d1x >>= sx_shamt;
+			d0 = d0x;
+			d1 = d1x;
+
+			// expand all values to 12 bits, with left-shift as needed.
+			int val_shamt = (modeval >> 1) ^ 3;
+			a <<= val_shamt;
+			b0 <<= val_shamt;
+			b1 <<= val_shamt;
+			c <<= val_shamt;
+			d0 <<= val_shamt;
+			d1 <<= val_shamt;
+
+			// then compute the actual color values.
+			int red1 = a;
+			int green1 = a - b0;
+			int blue1 = a - b1;
+			int red0 = a - c;
+			int green0 = a - b0 - c - d0;
+			int blue0 = a - b1 - c - d1;
+
+			// clamp the color components to [0,2^12 - 1]
+			if (red0 < 0)
+				red0 = 0;
+			else if (red0 > 0xFFF)
+				red0 = 0xFFF;
+
+			if (green0 < 0)
+				green0 = 0;
+			else if (green0 > 0xFFF)
+				green0 = 0xFFF;
+
+			if (blue0 < 0)
+				blue0 = 0;
+			else if (blue0 > 0xFFF)
+				blue0 = 0xFFF;
+
+			if (red1 < 0)
+				red1 = 0;
+			else if (red1 > 0xFFF)
+				red1 = 0xFFF;
+
+			if (green1 < 0)
+				green1 = 0;
+			else if (green1 > 0xFFF)
+				green1 = 0xFFF;
+
+			if (blue1 < 0)
+				blue1 = 0;
+			else if (blue1 > 0xFFF)
+				blue1 = 0xFFF;
+
+
+			// switch around the color components
+			int temp0, temp1;
+			switch (majcomp)
+			{
+			case 1:					// switch around red and green
+				temp0 = red0;
+				temp1 = red1;
+				red0 = green0;
+				red1 = green1;
+				green0 = temp0;
+				green1 = temp1;
+				break;
+			case 2:					// switch around red and blue
+				temp0 = red0;
+				temp1 = red1;
+				red0 = blue0;
+				red1 = blue1;
+				blue0 = temp0;
+				blue1 = temp1;
+				break;
+			case 0:					// no switch
+				break;
+			}
+
+			*output0 = ushort4(ushort(red0 << 4), ushort(green0 << 4), ushort(blue0 << 4), 0x7800);
+			*output1 = ushort4(ushort(red1 << 4), ushort(green1 << 4), ushort(blue1 << 4), 0x7800);
+		}
+
+
+		void hdr_rgb_ldr_alpha_unpack3(const int input[8], int quantization_level, ushort4* output0, ushort4* output1)
+		{
+			hdr_rgb_unpack3(input, quantization_level, output0, output1);
+
+			uint8 v6 = color_unquantization_tables[quantization_level][input[6]];
+			uint8 v7 = color_unquantization_tables[quantization_level][input[7]];
+			output0->w = v6;
+			output1->w = v7;
+		}
+
+
+		void hdr_luminance_small_range_unpack(const int input[2], int quantization_level, ushort4* output0, ushort4* output1)
+		{
+			int v0 = color_unquantization_tables[quantization_level][input[0]];
+			int v1 = color_unquantization_tables[quantization_level][input[1]];
+
+			int y0, y1;
+			if (v0 & 0x80)
+			{
+				y0 = ((v1 & 0xE0) << 4) | ((v0 & 0x7F) << 2);
+				y1 = (v1 & 0x1F) << 2;
+			}
+			else
+			{
+				y0 = ((v1 & 0xF0) << 4) | ((v0 & 0x7F) << 1);
+				y1 = (v1 & 0xF) << 1;
+			}
+
+			y1 += y0;
+			if (y1 > 0xFFF)
+				y1 = 0xFFF;
+
+			*output0 = ushort4(ushort(y0 << 4), ushort(y0 << 4), ushort(y0 << 4), 0x7800);
+			*output1 = ushort4(ushort(y1 << 4), ushort(y1 << 4), ushort(y1 << 4), 0x7800);
+		}
+
+
+		void hdr_luminance_large_range_unpack(const int input[2], int quantization_level, ushort4* output0, ushort4* output1)
+		{
+			int v0 = color_unquantization_tables[quantization_level][input[0]];
+			int v1 = color_unquantization_tables[quantization_level][input[1]];
+
+			int y0, y1;
+			if (v1 >= v0)
+			{
+				y0 = v0 << 4;
+				y1 = v1 << 4;
+			}
+			else
+			{
+				y0 = (v1 << 4) + 8;
+				y1 = (v0 << 4) - 8;
+			}
+			*output0 = ushort4(ushort(y0 << 4), ushort(y0 << 4), ushort(y0 << 4), 0x7800);
+			*output1 = ushort4(ushort(y1 << 4), ushort(y1 << 4), ushort(y1 << 4), 0x7800);
+		}
+
+
+		void hdr_alpha_unpack(const int input[2], int quantization_level, int* a0, int* a1)
+		{
+
+			int v6 = color_unquantization_tables[quantization_level][input[0]];
+			int v7 = color_unquantization_tables[quantization_level][input[1]];
+
+			int selector = ((v6 >> 7) & 1) | ((v7 >> 6) & 2);
+			v6 &= 0x7F;
+			v7 &= 0x7F;
+			if (selector == 3)
+			{
+				*a0 = v6 << 5;
+				*a1 = v7 << 5;
+			}
+			else
+			{
+				v6 |= (v7 << (selector + 1)) & 0x780;
+				v7 &= (0x3f >> selector);
+				v7 ^= 32 >> selector;
+				v7 -= 32 >> selector;
+				v6 <<= (4 - selector);
+				v7 <<= (4 - selector);
+				v7 += v6;
+
+				if (v7 < 0)
+					v7 = 0;
+				else if (v7 > 0xFFF)
+					v7 = 0xFFF;
+
+				*a0 = v6;
+				*a1 = v7;
+			}
+
+			*a0 <<= 4;
+			*a1 <<= 4;
+		}
+
+
+		void hdr_rgb_hdr_alpha_unpack3(const int input[8], int quantization_level, ushort4* output0, ushort4* output1)
+		{
+			hdr_rgb_unpack3(input, quantization_level, output0, output1);
+
+			int alpha0, alpha1;
+			hdr_alpha_unpack(input + 6, quantization_level, &alpha0, &alpha1);
+
+			output0->w = ushort(alpha0);
+			output1->w = ushort(alpha1);
+		}
+
+
+		void unpack_color_endpoints(astc_decode_mode decode_mode, int format, int quantization_level, const int* input, int* rgb_hdr, int* alpha_hdr, int* nan_endpoint, ushort4* output0, ushort4* output1)
+		{
+			*nan_endpoint = 0;
+
+			switch (format)
+			{
+			case FMT_LUMINANCE:
+				*rgb_hdr = 0;
+				*alpha_hdr = 0;
+				luminance_unpack(input, quantization_level, output0, output1);
+				break;
+
+			case FMT_LUMINANCE_DELTA:
+				*rgb_hdr = 0;
+				*alpha_hdr = 0;
+				luminance_delta_unpack(input, quantization_level, output0, output1);
+				break;
+
+			case FMT_HDR_LUMINANCE_SMALL_RANGE:
+				*rgb_hdr = 1;
+				*alpha_hdr = -1;
+				hdr_luminance_small_range_unpack(input, quantization_level, output0, output1);
+				break;
+
+			case FMT_HDR_LUMINANCE_LARGE_RANGE:
+				*rgb_hdr = 1;
+				*alpha_hdr = -1;
+				hdr_luminance_large_range_unpack(input, quantization_level, output0, output1);
+				break;
+
+			case FMT_LUMINANCE_ALPHA:
+				*rgb_hdr = 0;
+				*alpha_hdr = 0;
+				luminance_alpha_unpack(input, quantization_level, output0, output1);
+				break;
+
+			case FMT_LUMINANCE_ALPHA_DELTA:
+				*rgb_hdr = 0;
+				*alpha_hdr = 0;
+				luminance_alpha_delta_unpack(input, quantization_level, output0, output1);
+				break;
+
+			case FMT_RGB_SCALE:
+				*rgb_hdr = 0;
+				*alpha_hdr = 0;
+				rgb_scale_unpack(input, quantization_level, output0, output1);
+				break;
+
+			case FMT_RGB_SCALE_ALPHA:
+				*rgb_hdr = 0;
+				*alpha_hdr = 0;
+				rgb_scale_alpha_unpack(input, quantization_level, output0, output1);
+				break;
+
+			case FMT_HDR_RGB_SCALE:
+				*rgb_hdr = 1;
+				*alpha_hdr = -1;
+				hdr_rgbo_unpack3(input, quantization_level, output0, output1);
+				break;
+
+			case FMT_RGB:
+				*rgb_hdr = 0;
+				*alpha_hdr = 0;
+				rgb_unpack(input, quantization_level, output0, output1);
+				break;
+
+			case FMT_RGB_DELTA:
+				*rgb_hdr = 0;
+				*alpha_hdr = 0;
+				rgb_delta_unpack(input, quantization_level, output0, output1);
+				break;
+
+			case FMT_HDR_RGB:
+				*rgb_hdr = 1;
+				*alpha_hdr = -1;
+				hdr_rgb_unpack3(input, quantization_level, output0, output1);
+				break;
+
+			case FMT_RGBA:
+				*rgb_hdr = 0;
+				*alpha_hdr = 0;
+				rgba_unpack(input, quantization_level, output0, output1);
+				break;
+
+			case FMT_RGBA_DELTA:
+				*rgb_hdr = 0;
+				*alpha_hdr = 0;
+				rgba_delta_unpack(input, quantization_level, output0, output1);
+				break;
+
+			case FMT_HDR_RGB_LDR_ALPHA:
+				*rgb_hdr = 1;
+				*alpha_hdr = 0;
+				hdr_rgb_ldr_alpha_unpack3(input, quantization_level, output0, output1);
+				break;
+
+			case FMT_HDR_RGBA:
+				*rgb_hdr = 1;
+				*alpha_hdr = 1;
+				hdr_rgb_hdr_alpha_unpack3(input, quantization_level, output0, output1);
+				break;
+
+			default:
+				ASTC_CODEC_INTERNAL_ERROR;
+			}
+
+
+			if (*alpha_hdr == -1)
+			{
+				output0->w = 0x00FF;
+				output1->w = 0x00FF;
+				*alpha_hdr = 0;
+			}
+
+
+			switch (decode_mode)
+			{
+			case DECODE_LDR_SRGB:
+				if (*rgb_hdr == 1)
+				{
+					output0->x = 0xFF00;
+					output0->y = 0x0000;
+					output0->z = 0xFF00;
+					output0->w = 0xFF00;
+					output1->x = 0xFF00;
+					output1->y = 0x0000;
+					output1->z = 0xFF00;
+					output1->w = 0xFF00;
+				}
+				else
+				{
+					output0->x *= 257;
+					output0->y *= 257;
+					output0->z *= 257;
+					output0->w *= 257;
+					output1->x *= 257;
+					output1->y *= 257;
+					output1->z *= 257;
+					output1->w *= 257;
+				}
+				*rgb_hdr = 0;
+				*alpha_hdr = 0;
+				break;
+
+			case DECODE_LDR:
+				if (*rgb_hdr == 1)
+				{
+					output0->x = 0xFFFF;
+					output0->y = 0xFFFF;
+					output0->z = 0xFFFF;
+					output0->w = 0xFFFF;
+					output1->x = 0xFFFF;
+					output1->y = 0xFFFF;
+					output1->z = 0xFFFF;
+					output1->w = 0xFFFF;
+					*nan_endpoint = 1;
+				}
+				else
+				{
+					output0->x *= 257;
+					output0->y *= 257;
+					output0->z *= 257;
+					output0->w *= 257;
+					output1->x *= 257;
+					output1->y *= 257;
+					output1->z *= 257;
+					output1->w *= 257;
+				}
+				*rgb_hdr = 0;
+				*alpha_hdr = 0;
+				break;
+
+			case DECODE_HDR:
+
+				if (*rgb_hdr == 0)
+				{
+					output0->x *= 257;
+					output0->y *= 257;
+					output0->z *= 257;
+					output1->x *= 257;
+					output1->y *= 257;
+					output1->z *= 257;
+				}
+				if (*alpha_hdr == 0)
+				{
+					output0->w *= 257;
+					output1->w *= 257;
+				}
+				break;
+			}
+		}
+
+
+		struct quantization_and_transfer_table
+		{
+			quantization_method method;
+			uint8 unquantized_value[32];	// 0..64
+			float unquantized_value_flt[32];	// 0..1
+			uint8 prev_quantized_value[32];
+			uint8 next_quantized_value[32];
+			uint8 closest_quantized_weight[1025];
+		};
+
+		const quantization_and_transfer_table quant_and_xfer_tables[12] = {
+			// quantization method 0, range 0..1
+			{
+			 QUANT_2,
+			 {0, 64,},
+			 {
+			  0.000000, 1.000000,},
+			 {0, 0,},
+			 {1, 1,},
+			 {
+			  0, 0, 0, 0, 0, 0, 0, 0, 0, 0, 0, 0, 0, 0, 0, 0,
+			  0, 0, 0, 0, 0, 0, 0, 0, 0, 0, 0, 0, 0, 0, 0, 0,
+			  0, 0, 0, 0, 0, 0, 0, 0, 0, 0, 0, 0, 0, 0, 0, 0,
+			  0, 0, 0, 0, 0, 0, 0, 0, 0, 0, 0, 0, 0, 0, 0, 0,
+			  0, 0, 0, 0, 0, 0, 0, 0, 0, 0, 0, 0, 0, 0, 0, 0,
+			  0, 0, 0, 0, 0, 0, 0, 0, 0, 0, 0, 0, 0, 0, 0, 0,
+			  0, 0, 0, 0, 0, 0, 0, 0, 0, 0, 0, 0, 0, 0, 0, 0,
+			  0, 0, 0, 0, 0, 0, 0, 0, 0, 0, 0, 0, 0, 0, 0, 0,
+			  0, 0, 0, 0, 0, 0, 0, 0, 0, 0, 0, 0, 0, 0, 0, 0,
+			  0, 0, 0, 0, 0, 0, 0, 0, 0, 0, 0, 0, 0, 0, 0, 0,
+			  0, 0, 0, 0, 0, 0, 0, 0, 0, 0, 0, 0, 0, 0, 0, 0,
+			  0, 0, 0, 0, 0, 0, 0, 0, 0, 0, 0, 0, 0, 0, 0, 0,
+			  0, 0, 0, 0, 0, 0, 0, 0, 0, 0, 0, 0, 0, 0, 0, 0,
+			  0, 0, 0, 0, 0, 0, 0, 0, 0, 0, 0, 0, 0, 0, 0, 0,
+			  0, 0, 0, 0, 0, 0, 0, 0, 0, 0, 0, 0, 0, 0, 0, 0,
+			  0, 0, 0, 0, 0, 0, 0, 0, 0, 0, 0, 0, 0, 0, 0, 0,
+			  0, 0, 0, 0, 0, 0, 0, 0, 0, 0, 0, 0, 0, 0, 0, 0,
+			  0, 0, 0, 0, 0, 0, 0, 0, 0, 0, 0, 0, 0, 0, 0, 0,
+			  0, 0, 0, 0, 0, 0, 0, 0, 0, 0, 0, 0, 0, 0, 0, 0,
+			  0, 0, 0, 0, 0, 0, 0, 0, 0, 0, 0, 0, 0, 0, 0, 0,
+			  0, 0, 0, 0, 0, 0, 0, 0, 0, 0, 0, 0, 0, 0, 0, 0,
+			  0, 0, 0, 0, 0, 0, 0, 0, 0, 0, 0, 0, 0, 0, 0, 0,
+			  0, 0, 0, 0, 0, 0, 0, 0, 0, 0, 0, 0, 0, 0, 0, 0,
+			  0, 0, 0, 0, 0, 0, 0, 0, 0, 0, 0, 0, 0, 0, 0, 0,
+			  0, 0, 0, 0, 0, 0, 0, 0, 0, 0, 0, 0, 0, 0, 0, 0,
+			  0, 0, 0, 0, 0, 0, 0, 0, 0, 0, 0, 0, 0, 0, 0, 0,
+			  0, 0, 0, 0, 0, 0, 0, 0, 0, 0, 0, 0, 0, 0, 0, 0,
+			  0, 0, 0, 0, 0, 0, 0, 0, 0, 0, 0, 0, 0, 0, 0, 0,
+			  0, 0, 0, 0, 0, 0, 0, 0, 0, 0, 0, 0, 0, 0, 0, 0,
+			  0, 0, 0, 0, 0, 0, 0, 0, 0, 0, 0, 0, 0, 0, 0, 0,
+			  0, 0, 0, 0, 0, 0, 0, 0, 0, 0, 0, 0, 0, 0, 0, 0,
+			  0, 0, 0, 0, 0, 0, 0, 0, 0, 0, 0, 0, 0, 0, 0, 0,
+			  0, 1, 1, 1, 1, 1, 1, 1, 1, 1, 1, 1, 1, 1, 1, 1,
+			  1, 1, 1, 1, 1, 1, 1, 1, 1, 1, 1, 1, 1, 1, 1, 1,
+			  1, 1, 1, 1, 1, 1, 1, 1, 1, 1, 1, 1, 1, 1, 1, 1,
+			  1, 1, 1, 1, 1, 1, 1, 1, 1, 1, 1, 1, 1, 1, 1, 1,
+			  1, 1, 1, 1, 1, 1, 1, 1, 1, 1, 1, 1, 1, 1, 1, 1,
+			  1, 1, 1, 1, 1, 1, 1, 1, 1, 1, 1, 1, 1, 1, 1, 1,
+			  1, 1, 1, 1, 1, 1, 1, 1, 1, 1, 1, 1, 1, 1, 1, 1,
+			  1, 1, 1, 1, 1, 1, 1, 1, 1, 1, 1, 1, 1, 1, 1, 1,
+			  1, 1, 1, 1, 1, 1, 1, 1, 1, 1, 1, 1, 1, 1, 1, 1,
+			  1, 1, 1, 1, 1, 1, 1, 1, 1, 1, 1, 1, 1, 1, 1, 1,
+			  1, 1, 1, 1, 1, 1, 1, 1, 1, 1, 1, 1, 1, 1, 1, 1,
+			  1, 1, 1, 1, 1, 1, 1, 1, 1, 1, 1, 1, 1, 1, 1, 1,
+			  1, 1, 1, 1, 1, 1, 1, 1, 1, 1, 1, 1, 1, 1, 1, 1,
+			  1, 1, 1, 1, 1, 1, 1, 1, 1, 1, 1, 1, 1, 1, 1, 1,
+			  1, 1, 1, 1, 1, 1, 1, 1, 1, 1, 1, 1, 1, 1, 1, 1,
+			  1, 1, 1, 1, 1, 1, 1, 1, 1, 1, 1, 1, 1, 1, 1, 1,
+			  1, 1, 1, 1, 1, 1, 1, 1, 1, 1, 1, 1, 1, 1, 1, 1,
+			  1, 1, 1, 1, 1, 1, 1, 1, 1, 1, 1, 1, 1, 1, 1, 1,
+			  1, 1, 1, 1, 1, 1, 1, 1, 1, 1, 1, 1, 1, 1, 1, 1,
+			  1, 1, 1, 1, 1, 1, 1, 1, 1, 1, 1, 1, 1, 1, 1, 1,
+			  1, 1, 1, 1, 1, 1, 1, 1, 1, 1, 1, 1, 1, 1, 1, 1,
+			  1, 1, 1, 1, 1, 1, 1, 1, 1, 1, 1, 1, 1, 1, 1, 1,
+			  1, 1, 1, 1, 1, 1, 1, 1, 1, 1, 1, 1, 1, 1, 1, 1,
+			  1, 1, 1, 1, 1, 1, 1, 1, 1, 1, 1, 1, 1, 1, 1, 1,
+			  1, 1, 1, 1, 1, 1, 1, 1, 1, 1, 1, 1, 1, 1, 1, 1,
+			  1, 1, 1, 1, 1, 1, 1, 1, 1, 1, 1, 1, 1, 1, 1, 1,
+			  1, 1, 1, 1, 1, 1, 1, 1, 1, 1, 1, 1, 1, 1, 1, 1,
+			  1, 1, 1, 1, 1, 1, 1, 1, 1, 1, 1, 1, 1, 1, 1, 1,
+			  1, 1, 1, 1, 1, 1, 1, 1, 1, 1, 1, 1, 1, 1, 1, 1,
+			  1, 1, 1, 1, 1, 1, 1, 1, 1, 1, 1, 1, 1, 1, 1, 1,
+			  1, 1, 1, 1, 1, 1, 1, 1, 1, 1, 1, 1, 1, 1, 1, 1,
+			  1, 1, 1, 1, 1, 1, 1, 1, 1, 1, 1, 1, 1, 1, 1, 1,
+			  1,
+			  },
+			 },
+
+
+
+
+			 // quantization method 1, range 0..2
+			 {
+			  QUANT_3,
+			  {0, 32, 64,},
+			  {
+			   0.000000, 0.500000, 1.000000,},
+			  {0, 0, 1,},
+			  {1, 2, 2,},
+			  {
+			   0, 0, 0, 0, 0, 0, 0, 0, 0, 0, 0, 0, 0, 0, 0, 0,
+			   0, 0, 0, 0, 0, 0, 0, 0, 0, 0, 0, 0, 0, 0, 0, 0,
+			   0, 0, 0, 0, 0, 0, 0, 0, 0, 0, 0, 0, 0, 0, 0, 0,
+			   0, 0, 0, 0, 0, 0, 0, 0, 0, 0, 0, 0, 0, 0, 0, 0,
+			   0, 0, 0, 0, 0, 0, 0, 0, 0, 0, 0, 0, 0, 0, 0, 0,
+			   0, 0, 0, 0, 0, 0, 0, 0, 0, 0, 0, 0, 0, 0, 0, 0,
+			   0, 0, 0, 0, 0, 0, 0, 0, 0, 0, 0, 0, 0, 0, 0, 0,
+			   0, 0, 0, 0, 0, 0, 0, 0, 0, 0, 0, 0, 0, 0, 0, 0,
+			   0, 0, 0, 0, 0, 0, 0, 0, 0, 0, 0, 0, 0, 0, 0, 0,
+			   0, 0, 0, 0, 0, 0, 0, 0, 0, 0, 0, 0, 0, 0, 0, 0,
+			   0, 0, 0, 0, 0, 0, 0, 0, 0, 0, 0, 0, 0, 0, 0, 0,
+			   0, 0, 0, 0, 0, 0, 0, 0, 0, 0, 0, 0, 0, 0, 0, 0,
+			   0, 0, 0, 0, 0, 0, 0, 0, 0, 0, 0, 0, 0, 0, 0, 0,
+			   0, 0, 0, 0, 0, 0, 0, 0, 0, 0, 0, 0, 0, 0, 0, 0,
+			   0, 0, 0, 0, 0, 0, 0, 0, 0, 0, 0, 0, 0, 0, 0, 0,
+			   0, 0, 0, 0, 0, 0, 0, 0, 0, 0, 0, 0, 0, 0, 0, 0,
+			   0, 1, 1, 1, 1, 1, 1, 1, 1, 1, 1, 1, 1, 1, 1, 1,
+			   1, 1, 1, 1, 1, 1, 1, 1, 1, 1, 1, 1, 1, 1, 1, 1,
+			   1, 1, 1, 1, 1, 1, 1, 1, 1, 1, 1, 1, 1, 1, 1, 1,
+			   1, 1, 1, 1, 1, 1, 1, 1, 1, 1, 1, 1, 1, 1, 1, 1,
+			   1, 1, 1, 1, 1, 1, 1, 1, 1, 1, 1, 1, 1, 1, 1, 1,
+			   1, 1, 1, 1, 1, 1, 1, 1, 1, 1, 1, 1, 1, 1, 1, 1,
+			   1, 1, 1, 1, 1, 1, 1, 1, 1, 1, 1, 1, 1, 1, 1, 1,
+			   1, 1, 1, 1, 1, 1, 1, 1, 1, 1, 1, 1, 1, 1, 1, 1,
+			   1, 1, 1, 1, 1, 1, 1, 1, 1, 1, 1, 1, 1, 1, 1, 1,
+			   1, 1, 1, 1, 1, 1, 1, 1, 1, 1, 1, 1, 1, 1, 1, 1,
+			   1, 1, 1, 1, 1, 1, 1, 1, 1, 1, 1, 1, 1, 1, 1, 1,
+			   1, 1, 1, 1, 1, 1, 1, 1, 1, 1, 1, 1, 1, 1, 1, 1,
+			   1, 1, 1, 1, 1, 1, 1, 1, 1, 1, 1, 1, 1, 1, 1, 1,
+			   1, 1, 1, 1, 1, 1, 1, 1, 1, 1, 1, 1, 1, 1, 1, 1,
+			   1, 1, 1, 1, 1, 1, 1, 1, 1, 1, 1, 1, 1, 1, 1, 1,
+			   1, 1, 1, 1, 1, 1, 1, 1, 1, 1, 1, 1, 1, 1, 1, 1,
+			   1, 1, 1, 1, 1, 1, 1, 1, 1, 1, 1, 1, 1, 1, 1, 1,
+			   1, 1, 1, 1, 1, 1, 1, 1, 1, 1, 1, 1, 1, 1, 1, 1,
+			   1, 1, 1, 1, 1, 1, 1, 1, 1, 1, 1, 1, 1, 1, 1, 1,
+			   1, 1, 1, 1, 1, 1, 1, 1, 1, 1, 1, 1, 1, 1, 1, 1,
+			   1, 1, 1, 1, 1, 1, 1, 1, 1, 1, 1, 1, 1, 1, 1, 1,
+			   1, 1, 1, 1, 1, 1, 1, 1, 1, 1, 1, 1, 1, 1, 1, 1,
+			   1, 1, 1, 1, 1, 1, 1, 1, 1, 1, 1, 1, 1, 1, 1, 1,
+			   1, 1, 1, 1, 1, 1, 1, 1, 1, 1, 1, 1, 1, 1, 1, 1,
+			   1, 1, 1, 1, 1, 1, 1, 1, 1, 1, 1, 1, 1, 1, 1, 1,
+			   1, 1, 1, 1, 1, 1, 1, 1, 1, 1, 1, 1, 1, 1, 1, 1,
+			   1, 1, 1, 1, 1, 1, 1, 1, 1, 1, 1, 1, 1, 1, 1, 1,
+			   1, 1, 1, 1, 1, 1, 1, 1, 1, 1, 1, 1, 1, 1, 1, 1,
+			   1, 1, 1, 1, 1, 1, 1, 1, 1, 1, 1, 1, 1, 1, 1, 1,
+			   1, 1, 1, 1, 1, 1, 1, 1, 1, 1, 1, 1, 1, 1, 1, 1,
+			   1, 1, 1, 1, 1, 1, 1, 1, 1, 1, 1, 1, 1, 1, 1, 1,
+			   1, 1, 1, 1, 1, 1, 1, 1, 1, 1, 1, 1, 1, 1, 1, 1,
+			   1, 2, 2, 2, 2, 2, 2, 2, 2, 2, 2, 2, 2, 2, 2, 2,
+			   2, 2, 2, 2, 2, 2, 2, 2, 2, 2, 2, 2, 2, 2, 2, 2,
+			   2, 2, 2, 2, 2, 2, 2, 2, 2, 2, 2, 2, 2, 2, 2, 2,
+			   2, 2, 2, 2, 2, 2, 2, 2, 2, 2, 2, 2, 2, 2, 2, 2,
+			   2, 2, 2, 2, 2, 2, 2, 2, 2, 2, 2, 2, 2, 2, 2, 2,
+			   2, 2, 2, 2, 2, 2, 2, 2, 2, 2, 2, 2, 2, 2, 2, 2,
+			   2, 2, 2, 2, 2, 2, 2, 2, 2, 2, 2, 2, 2, 2, 2, 2,
+			   2, 2, 2, 2, 2, 2, 2, 2, 2, 2, 2, 2, 2, 2, 2, 2,
+			   2, 2, 2, 2, 2, 2, 2, 2, 2, 2, 2, 2, 2, 2, 2, 2,
+			   2, 2, 2, 2, 2, 2, 2, 2, 2, 2, 2, 2, 2, 2, 2, 2,
+			   2, 2, 2, 2, 2, 2, 2, 2, 2, 2, 2, 2, 2, 2, 2, 2,
+			   2, 2, 2, 2, 2, 2, 2, 2, 2, 2, 2, 2, 2, 2, 2, 2,
+			   2, 2, 2, 2, 2, 2, 2, 2, 2, 2, 2, 2, 2, 2, 2, 2,
+			   2, 2, 2, 2, 2, 2, 2, 2, 2, 2, 2, 2, 2, 2, 2, 2,
+			   2, 2, 2, 2, 2, 2, 2, 2, 2, 2, 2, 2, 2, 2, 2, 2,
+			   2, 2, 2, 2, 2, 2, 2, 2, 2, 2, 2, 2, 2, 2, 2, 2,
+			   2,
+			   },
+			  },
+
+
+
+
+			  // quantization method 2, range 0..3
+			  {
+			   QUANT_4,
+			   {0, 21, 43, 64,},
+			   {
+				0.000000, 0.328125, 0.671875, 1.000000,},
+			   {0, 0, 1, 2,},
+			   {1, 2, 3, 3,},
+			   {
+				0, 0, 0, 0, 0, 0, 0, 0, 0, 0, 0, 0, 0, 0, 0, 0,
+				0, 0, 0, 0, 0, 0, 0, 0, 0, 0, 0, 0, 0, 0, 0, 0,
+				0, 0, 0, 0, 0, 0, 0, 0, 0, 0, 0, 0, 0, 0, 0, 0,
+				0, 0, 0, 0, 0, 0, 0, 0, 0, 0, 0, 0, 0, 0, 0, 0,
+				0, 0, 0, 0, 0, 0, 0, 0, 0, 0, 0, 0, 0, 0, 0, 0,
+				0, 0, 0, 0, 0, 0, 0, 0, 0, 0, 0, 0, 0, 0, 0, 0,
+				0, 0, 0, 0, 0, 0, 0, 0, 0, 0, 0, 0, 0, 0, 0, 0,
+				0, 0, 0, 0, 0, 0, 0, 0, 0, 0, 0, 0, 0, 0, 0, 0,
+				0, 0, 0, 0, 0, 0, 0, 0, 0, 0, 0, 0, 0, 0, 0, 0,
+				0, 0, 0, 0, 0, 0, 0, 0, 0, 0, 0, 0, 0, 0, 0, 0,
+				0, 0, 0, 0, 0, 0, 0, 0, 0, 1, 1, 1, 1, 1, 1, 1,
+				1, 1, 1, 1, 1, 1, 1, 1, 1, 1, 1, 1, 1, 1, 1, 1,
+				1, 1, 1, 1, 1, 1, 1, 1, 1, 1, 1, 1, 1, 1, 1, 1,
+				1, 1, 1, 1, 1, 1, 1, 1, 1, 1, 1, 1, 1, 1, 1, 1,
+				1, 1, 1, 1, 1, 1, 1, 1, 1, 1, 1, 1, 1, 1, 1, 1,
+				1, 1, 1, 1, 1, 1, 1, 1, 1, 1, 1, 1, 1, 1, 1, 1,
+				1, 1, 1, 1, 1, 1, 1, 1, 1, 1, 1, 1, 1, 1, 1, 1,
+				1, 1, 1, 1, 1, 1, 1, 1, 1, 1, 1, 1, 1, 1, 1, 1,
+				1, 1, 1, 1, 1, 1, 1, 1, 1, 1, 1, 1, 1, 1, 1, 1,
+				1, 1, 1, 1, 1, 1, 1, 1, 1, 1, 1, 1, 1, 1, 1, 1,
+				1, 1, 1, 1, 1, 1, 1, 1, 1, 1, 1, 1, 1, 1, 1, 1,
+				1, 1, 1, 1, 1, 1, 1, 1, 1, 1, 1, 1, 1, 1, 1, 1,
+				1, 1, 1, 1, 1, 1, 1, 1, 1, 1, 1, 1, 1, 1, 1, 1,
+				1, 1, 1, 1, 1, 1, 1, 1, 1, 1, 1, 1, 1, 1, 1, 1,
+				1, 1, 1, 1, 1, 1, 1, 1, 1, 1, 1, 1, 1, 1, 1, 1,
+				1, 1, 1, 1, 1, 1, 1, 1, 1, 1, 1, 1, 1, 1, 1, 1,
+				1, 1, 1, 1, 1, 1, 1, 1, 1, 1, 1, 1, 1, 1, 1, 1,
+				1, 1, 1, 1, 1, 1, 1, 1, 1, 1, 1, 1, 1, 1, 1, 1,
+				1, 1, 1, 1, 1, 1, 1, 1, 1, 1, 1, 1, 1, 1, 1, 1,
+				1, 1, 1, 1, 1, 1, 1, 1, 1, 1, 1, 1, 1, 1, 1, 1,
+				1, 1, 1, 1, 1, 1, 1, 1, 1, 1, 1, 1, 1, 1, 1, 1,
+				1, 1, 1, 1, 1, 1, 1, 1, 1, 1, 1, 1, 1, 1, 1, 1,
+				1, 2, 2, 2, 2, 2, 2, 2, 2, 2, 2, 2, 2, 2, 2, 2,
+				2, 2, 2, 2, 2, 2, 2, 2, 2, 2, 2, 2, 2, 2, 2, 2,
+				2, 2, 2, 2, 2, 2, 2, 2, 2, 2, 2, 2, 2, 2, 2, 2,
+				2, 2, 2, 2, 2, 2, 2, 2, 2, 2, 2, 2, 2, 2, 2, 2,
+				2, 2, 2, 2, 2, 2, 2, 2, 2, 2, 2, 2, 2, 2, 2, 2,
+				2, 2, 2, 2, 2, 2, 2, 2, 2, 2, 2, 2, 2, 2, 2, 2,
+				2, 2, 2, 2, 2, 2, 2, 2, 2, 2, 2, 2, 2, 2, 2, 2,
+				2, 2, 2, 2, 2, 2, 2, 2, 2, 2, 2, 2, 2, 2, 2, 2,
+				2, 2, 2, 2, 2, 2, 2, 2, 2, 2, 2, 2, 2, 2, 2, 2,
+				2, 2, 2, 2, 2, 2, 2, 2, 2, 2, 2, 2, 2, 2, 2, 2,
+				2, 2, 2, 2, 2, 2, 2, 2, 2, 2, 2, 2, 2, 2, 2, 2,
+				2, 2, 2, 2, 2, 2, 2, 2, 2, 2, 2, 2, 2, 2, 2, 2,
+				2, 2, 2, 2, 2, 2, 2, 2, 2, 2, 2, 2, 2, 2, 2, 2,
+				2, 2, 2, 2, 2, 2, 2, 2, 2, 2, 2, 2, 2, 2, 2, 2,
+				2, 2, 2, 2, 2, 2, 2, 2, 2, 2, 2, 2, 2, 2, 2, 2,
+				2, 2, 2, 2, 2, 2, 2, 2, 2, 2, 2, 2, 2, 2, 2, 2,
+				2, 2, 2, 2, 2, 2, 2, 2, 2, 2, 2, 2, 2, 2, 2, 2,
+				2, 2, 2, 2, 2, 2, 2, 2, 2, 2, 2, 2, 2, 2, 2, 2,
+				2, 2, 2, 2, 2, 2, 2, 2, 2, 2, 2, 2, 2, 2, 2, 2,
+				2, 2, 2, 2, 2, 2, 2, 2, 2, 2, 2, 2, 2, 2, 2, 2,
+				2, 2, 2, 2, 2, 2, 2, 2, 2, 2, 2, 2, 2, 2, 2, 2,
+				2, 2, 2, 2, 2, 2, 2, 2, 2, 3, 3, 3, 3, 3, 3, 3,
+				3, 3, 3, 3, 3, 3, 3, 3, 3, 3, 3, 3, 3, 3, 3, 3,
+				3, 3, 3, 3, 3, 3, 3, 3, 3, 3, 3, 3, 3, 3, 3, 3,
+				3, 3, 3, 3, 3, 3, 3, 3, 3, 3, 3, 3, 3, 3, 3, 3,
+				3, 3, 3, 3, 3, 3, 3, 3, 3, 3, 3, 3, 3, 3, 3, 3,
+				3, 3, 3, 3, 3, 3, 3, 3, 3, 3, 3, 3, 3, 3, 3, 3,
+				3, 3, 3, 3, 3, 3, 3, 3, 3, 3, 3, 3, 3, 3, 3, 3,
+				3, 3, 3, 3, 3, 3, 3, 3, 3, 3, 3, 3, 3, 3, 3, 3,
+				3, 3, 3, 3, 3, 3, 3, 3, 3, 3, 3, 3, 3, 3, 3, 3,
+				3, 3, 3, 3, 3, 3, 3, 3, 3, 3, 3, 3, 3, 3, 3, 3,
+				3, 3, 3, 3, 3, 3, 3, 3, 3, 3, 3, 3, 3, 3, 3, 3,
+				3,
+				},
+			   },
+
+
+
+
+			   // quantization method 3, range 0..4
+			   {
+				QUANT_5,
+				{0, 16, 32, 48, 64,},
+				{
+				 0.000000, 0.250000, 0.500000, 0.750000,
+				 1.000000,},
+				{0, 0, 1, 2, 3,},
+				{1, 2, 3, 4, 4,},
+				{
+				 0, 0, 0, 0, 0, 0, 0, 0, 0, 0, 0, 0, 0, 0, 0, 0,
+				 0, 0, 0, 0, 0, 0, 0, 0, 0, 0, 0, 0, 0, 0, 0, 0,
+				 0, 0, 0, 0, 0, 0, 0, 0, 0, 0, 0, 0, 0, 0, 0, 0,
+				 0, 0, 0, 0, 0, 0, 0, 0, 0, 0, 0, 0, 0, 0, 0, 0,
+				 0, 0, 0, 0, 0, 0, 0, 0, 0, 0, 0, 0, 0, 0, 0, 0,
+				 0, 0, 0, 0, 0, 0, 0, 0, 0, 0, 0, 0, 0, 0, 0, 0,
+				 0, 0, 0, 0, 0, 0, 0, 0, 0, 0, 0, 0, 0, 0, 0, 0,
+				 0, 0, 0, 0, 0, 0, 0, 0, 0, 0, 0, 0, 0, 0, 0, 0,
+				 0, 1, 1, 1, 1, 1, 1, 1, 1, 1, 1, 1, 1, 1, 1, 1,
+				 1, 1, 1, 1, 1, 1, 1, 1, 1, 1, 1, 1, 1, 1, 1, 1,
+				 1, 1, 1, 1, 1, 1, 1, 1, 1, 1, 1, 1, 1, 1, 1, 1,
+				 1, 1, 1, 1, 1, 1, 1, 1, 1, 1, 1, 1, 1, 1, 1, 1,
+				 1, 1, 1, 1, 1, 1, 1, 1, 1, 1, 1, 1, 1, 1, 1, 1,
+				 1, 1, 1, 1, 1, 1, 1, 1, 1, 1, 1, 1, 1, 1, 1, 1,
+				 1, 1, 1, 1, 1, 1, 1, 1, 1, 1, 1, 1, 1, 1, 1, 1,
+				 1, 1, 1, 1, 1, 1, 1, 1, 1, 1, 1, 1, 1, 1, 1, 1,
+				 1, 1, 1, 1, 1, 1, 1, 1, 1, 1, 1, 1, 1, 1, 1, 1,
+				 1, 1, 1, 1, 1, 1, 1, 1, 1, 1, 1, 1, 1, 1, 1, 1,
+				 1, 1, 1, 1, 1, 1, 1, 1, 1, 1, 1, 1, 1, 1, 1, 1,
+				 1, 1, 1, 1, 1, 1, 1, 1, 1, 1, 1, 1, 1, 1, 1, 1,
+				 1, 1, 1, 1, 1, 1, 1, 1, 1, 1, 1, 1, 1, 1, 1, 1,
+				 1, 1, 1, 1, 1, 1, 1, 1, 1, 1, 1, 1, 1, 1, 1, 1,
+				 1, 1, 1, 1, 1, 1, 1, 1, 1, 1, 1, 1, 1, 1, 1, 1,
+				 1, 1, 1, 1, 1, 1, 1, 1, 1, 1, 1, 1, 1, 1, 1, 1,
+				 1, 2, 2, 2, 2, 2, 2, 2, 2, 2, 2, 2, 2, 2, 2, 2,
+				 2, 2, 2, 2, 2, 2, 2, 2, 2, 2, 2, 2, 2, 2, 2, 2,
+				 2, 2, 2, 2, 2, 2, 2, 2, 2, 2, 2, 2, 2, 2, 2, 2,
+				 2, 2, 2, 2, 2, 2, 2, 2, 2, 2, 2, 2, 2, 2, 2, 2,
+				 2, 2, 2, 2, 2, 2, 2, 2, 2, 2, 2, 2, 2, 2, 2, 2,
+				 2, 2, 2, 2, 2, 2, 2, 2, 2, 2, 2, 2, 2, 2, 2, 2,
+				 2, 2, 2, 2, 2, 2, 2, 2, 2, 2, 2, 2, 2, 2, 2, 2,
+				 2, 2, 2, 2, 2, 2, 2, 2, 2, 2, 2, 2, 2, 2, 2, 2,
+				 2, 2, 2, 2, 2, 2, 2, 2, 2, 2, 2, 2, 2, 2, 2, 2,
+				 2, 2, 2, 2, 2, 2, 2, 2, 2, 2, 2, 2, 2, 2, 2, 2,
+				 2, 2, 2, 2, 2, 2, 2, 2, 2, 2, 2, 2, 2, 2, 2, 2,
+				 2, 2, 2, 2, 2, 2, 2, 2, 2, 2, 2, 2, 2, 2, 2, 2,
+				 2, 2, 2, 2, 2, 2, 2, 2, 2, 2, 2, 2, 2, 2, 2, 2,
+				 2, 2, 2, 2, 2, 2, 2, 2, 2, 2, 2, 2, 2, 2, 2, 2,
+				 2, 2, 2, 2, 2, 2, 2, 2, 2, 2, 2, 2, 2, 2, 2, 2,
+				 2, 2, 2, 2, 2, 2, 2, 2, 2, 2, 2, 2, 2, 2, 2, 2,
+				 2, 3, 3, 3, 3, 3, 3, 3, 3, 3, 3, 3, 3, 3, 3, 3,
+				 3, 3, 3, 3, 3, 3, 3, 3, 3, 3, 3, 3, 3, 3, 3, 3,
+				 3, 3, 3, 3, 3, 3, 3, 3, 3, 3, 3, 3, 3, 3, 3, 3,
+				 3, 3, 3, 3, 3, 3, 3, 3, 3, 3, 3, 3, 3, 3, 3, 3,
+				 3, 3, 3, 3, 3, 3, 3, 3, 3, 3, 3, 3, 3, 3, 3, 3,
+				 3, 3, 3, 3, 3, 3, 3, 3, 3, 3, 3, 3, 3, 3, 3, 3,
+				 3, 3, 3, 3, 3, 3, 3, 3, 3, 3, 3, 3, 3, 3, 3, 3,
+				 3, 3, 3, 3, 3, 3, 3, 3, 3, 3, 3, 3, 3, 3, 3, 3,
+				 3, 3, 3, 3, 3, 3, 3, 3, 3, 3, 3, 3, 3, 3, 3, 3,
+				 3, 3, 3, 3, 3, 3, 3, 3, 3, 3, 3, 3, 3, 3, 3, 3,
+				 3, 3, 3, 3, 3, 3, 3, 3, 3, 3, 3, 3, 3, 3, 3, 3,
+				 3, 3, 3, 3, 3, 3, 3, 3, 3, 3, 3, 3, 3, 3, 3, 3,
+				 3, 3, 3, 3, 3, 3, 3, 3, 3, 3, 3, 3, 3, 3, 3, 3,
+				 3, 3, 3, 3, 3, 3, 3, 3, 3, 3, 3, 3, 3, 3, 3, 3,
+				 3, 3, 3, 3, 3, 3, 3, 3, 3, 3, 3, 3, 3, 3, 3, 3,
+				 3, 3, 3, 3, 3, 3, 3, 3, 3, 3, 3, 3, 3, 3, 3, 3,
+				 3, 4, 4, 4, 4, 4, 4, 4, 4, 4, 4, 4, 4, 4, 4, 4,
+				 4, 4, 4, 4, 4, 4, 4, 4, 4, 4, 4, 4, 4, 4, 4, 4,
+				 4, 4, 4, 4, 4, 4, 4, 4, 4, 4, 4, 4, 4, 4, 4, 4,
+				 4, 4, 4, 4, 4, 4, 4, 4, 4, 4, 4, 4, 4, 4, 4, 4,
+				 4, 4, 4, 4, 4, 4, 4, 4, 4, 4, 4, 4, 4, 4, 4, 4,
+				 4, 4, 4, 4, 4, 4, 4, 4, 4, 4, 4, 4, 4, 4, 4, 4,
+				 4, 4, 4, 4, 4, 4, 4, 4, 4, 4, 4, 4, 4, 4, 4, 4,
+				 4, 4, 4, 4, 4, 4, 4, 4, 4, 4, 4, 4, 4, 4, 4, 4,
+				 4,
+				 },
+				},
+
+
+
+
+				// quantization method 4, range 0..5
+				{
+				 QUANT_6,
+				 {0, 64, 12, 52, 25, 39,},
+				 {
+				  0.000000, 1.000000, 0.187500, 0.812500,
+				  0.390625, 0.609375,},
+				 {0, 3, 0, 5, 2, 4,},
+				 {2, 1, 4, 1, 5, 3,},
+				 {
+				  0, 0, 0, 0, 0, 0, 0, 0, 0, 0, 0, 0, 0, 0, 0, 0,
+				  0, 0, 0, 0, 0, 0, 0, 0, 0, 0, 0, 0, 0, 0, 0, 0,
+				  0, 0, 0, 0, 0, 0, 0, 0, 0, 0, 0, 0, 0, 0, 0, 0,
+				  0, 0, 0, 0, 0, 0, 0, 0, 0, 0, 0, 0, 0, 0, 0, 0,
+				  0, 0, 0, 0, 0, 0, 0, 0, 0, 0, 0, 0, 0, 0, 0, 0,
+				  0, 0, 0, 0, 0, 0, 0, 0, 0, 0, 0, 0, 0, 0, 0, 0,
+				  0, 2, 2, 2, 2, 2, 2, 2, 2, 2, 2, 2, 2, 2, 2, 2,
+				  2, 2, 2, 2, 2, 2, 2, 2, 2, 2, 2, 2, 2, 2, 2, 2,
+				  2, 2, 2, 2, 2, 2, 2, 2, 2, 2, 2, 2, 2, 2, 2, 2,
+				  2, 2, 2, 2, 2, 2, 2, 2, 2, 2, 2, 2, 2, 2, 2, 2,
+				  2, 2, 2, 2, 2, 2, 2, 2, 2, 2, 2, 2, 2, 2, 2, 2,
+				  2, 2, 2, 2, 2, 2, 2, 2, 2, 2, 2, 2, 2, 2, 2, 2,
+				  2, 2, 2, 2, 2, 2, 2, 2, 2, 2, 2, 2, 2, 2, 2, 2,
+				  2, 2, 2, 2, 2, 2, 2, 2, 2, 2, 2, 2, 2, 2, 2, 2,
+				  2, 2, 2, 2, 2, 2, 2, 2, 2, 2, 2, 2, 2, 2, 2, 2,
+				  2, 2, 2, 2, 2, 2, 2, 2, 2, 2, 2, 2, 2, 2, 2, 2,
+				  2, 2, 2, 2, 2, 2, 2, 2, 2, 2, 2, 2, 2, 2, 2, 2,
+				  2, 2, 2, 2, 2, 2, 2, 2, 2, 2, 2, 2, 2, 2, 2, 2,
+				  2, 2, 2, 2, 2, 2, 2, 2, 2, 4, 4, 4, 4, 4, 4, 4,
+				  4, 4, 4, 4, 4, 4, 4, 4, 4, 4, 4, 4, 4, 4, 4, 4,
+				  4, 4, 4, 4, 4, 4, 4, 4, 4, 4, 4, 4, 4, 4, 4, 4,
+				  4, 4, 4, 4, 4, 4, 4, 4, 4, 4, 4, 4, 4, 4, 4, 4,
+				  4, 4, 4, 4, 4, 4, 4, 4, 4, 4, 4, 4, 4, 4, 4, 4,
+				  4, 4, 4, 4, 4, 4, 4, 4, 4, 4, 4, 4, 4, 4, 4, 4,
+				  4, 4, 4, 4, 4, 4, 4, 4, 4, 4, 4, 4, 4, 4, 4, 4,
+				  4, 4, 4, 4, 4, 4, 4, 4, 4, 4, 4, 4, 4, 4, 4, 4,
+				  4, 4, 4, 4, 4, 4, 4, 4, 4, 4, 4, 4, 4, 4, 4, 4,
+				  4, 4, 4, 4, 4, 4, 4, 4, 4, 4, 4, 4, 4, 4, 4, 4,
+				  4, 4, 4, 4, 4, 4, 4, 4, 4, 4, 4, 4, 4, 4, 4, 4,
+				  4, 4, 4, 4, 4, 4, 4, 4, 4, 4, 4, 4, 4, 4, 4, 4,
+				  4, 4, 4, 4, 4, 4, 4, 4, 4, 4, 4, 4, 4, 4, 4, 4,
+				  4, 4, 4, 4, 4, 4, 4, 4, 4, 4, 4, 4, 4, 4, 4, 4,
+				  4, 5, 5, 5, 5, 5, 5, 5, 5, 5, 5, 5, 5, 5, 5, 5,
+				  5, 5, 5, 5, 5, 5, 5, 5, 5, 5, 5, 5, 5, 5, 5, 5,
+				  5, 5, 5, 5, 5, 5, 5, 5, 5, 5, 5, 5, 5, 5, 5, 5,
+				  5, 5, 5, 5, 5, 5, 5, 5, 5, 5, 5, 5, 5, 5, 5, 5,
+				  5, 5, 5, 5, 5, 5, 5, 5, 5, 5, 5, 5, 5, 5, 5, 5,
+				  5, 5, 5, 5, 5, 5, 5, 5, 5, 5, 5, 5, 5, 5, 5, 5,
+				  5, 5, 5, 5, 5, 5, 5, 5, 5, 5, 5, 5, 5, 5, 5, 5,
+				  5, 5, 5, 5, 5, 5, 5, 5, 5, 5, 5, 5, 5, 5, 5, 5,
+				  5, 5, 5, 5, 5, 5, 5, 5, 5, 5, 5, 5, 5, 5, 5, 5,
+				  5, 5, 5, 5, 5, 5, 5, 5, 5, 5, 5, 5, 5, 5, 5, 5,
+				  5, 5, 5, 5, 5, 5, 5, 5, 5, 5, 5, 5, 5, 5, 5, 5,
+				  5, 5, 5, 5, 5, 5, 5, 5, 5, 5, 5, 5, 5, 5, 5, 5,
+				  5, 5, 5, 5, 5, 5, 5, 5, 5, 5, 5, 5, 5, 5, 5, 5,
+				  5, 5, 5, 5, 5, 5, 5, 5, 5, 3, 3, 3, 3, 3, 3, 3,
+				  3, 3, 3, 3, 3, 3, 3, 3, 3, 3, 3, 3, 3, 3, 3, 3,
+				  3, 3, 3, 3, 3, 3, 3, 3, 3, 3, 3, 3, 3, 3, 3, 3,
+				  3, 3, 3, 3, 3, 3, 3, 3, 3, 3, 3, 3, 3, 3, 3, 3,
+				  3, 3, 3, 3, 3, 3, 3, 3, 3, 3, 3, 3, 3, 3, 3, 3,
+				  3, 3, 3, 3, 3, 3, 3, 3, 3, 3, 3, 3, 3, 3, 3, 3,
+				  3, 3, 3, 3, 3, 3, 3, 3, 3, 3, 3, 3, 3, 3, 3, 3,
+				  3, 3, 3, 3, 3, 3, 3, 3, 3, 3, 3, 3, 3, 3, 3, 3,
+				  3, 3, 3, 3, 3, 3, 3, 3, 3, 3, 3, 3, 3, 3, 3, 3,
+				  3, 3, 3, 3, 3, 3, 3, 3, 3, 3, 3, 3, 3, 3, 3, 3,
+				  3, 3, 3, 3, 3, 3, 3, 3, 3, 3, 3, 3, 3, 3, 3, 3,
+				  3, 3, 3, 3, 3, 3, 3, 3, 3, 3, 3, 3, 3, 3, 3, 3,
+				  3, 3, 3, 3, 3, 3, 3, 3, 3, 3, 3, 3, 3, 3, 3, 3,
+				  3, 1, 1, 1, 1, 1, 1, 1, 1, 1, 1, 1, 1, 1, 1, 1,
+				  1, 1, 1, 1, 1, 1, 1, 1, 1, 1, 1, 1, 1, 1, 1, 1,
+				  1, 1, 1, 1, 1, 1, 1, 1, 1, 1, 1, 1, 1, 1, 1, 1,
+				  1, 1, 1, 1, 1, 1, 1, 1, 1, 1, 1, 1, 1, 1, 1, 1,
+				  1, 1, 1, 1, 1, 1, 1, 1, 1, 1, 1, 1, 1, 1, 1, 1,
+				  1, 1, 1, 1, 1, 1, 1, 1, 1, 1, 1, 1, 1, 1, 1, 1,
+				  1,
+				  },
+				 },
+
+
+
+
+				 // quantization method 5, range 0..7
+				 {
+				  QUANT_8,
+				  {0, 9, 18, 27, 37, 46, 55, 64,},
+				  {
+				   0.000000, 0.140625, 0.281250, 0.421875,
+				   0.578125, 0.718750, 0.859375, 1.000000,},
+				  {0, 0, 1, 2, 3, 4, 5, 6,},
+				  {1, 2, 3, 4, 5, 6, 7, 7,},
+				  {
+				   0, 0, 0, 0, 0, 0, 0, 0, 0, 0, 0, 0, 0, 0, 0, 0,
+				   0, 0, 0, 0, 0, 0, 0, 0, 0, 0, 0, 0, 0, 0, 0, 0,
+				   0, 0, 0, 0, 0, 0, 0, 0, 0, 0, 0, 0, 0, 0, 0, 0,
+				   0, 0, 0, 0, 0, 0, 0, 0, 0, 0, 0, 0, 0, 0, 0, 0,
+				   0, 0, 0, 0, 0, 0, 0, 0, 0, 1, 1, 1, 1, 1, 1, 1,
+				   1, 1, 1, 1, 1, 1, 1, 1, 1, 1, 1, 1, 1, 1, 1, 1,
+				   1, 1, 1, 1, 1, 1, 1, 1, 1, 1, 1, 1, 1, 1, 1, 1,
+				   1, 1, 1, 1, 1, 1, 1, 1, 1, 1, 1, 1, 1, 1, 1, 1,
+				   1, 1, 1, 1, 1, 1, 1, 1, 1, 1, 1, 1, 1, 1, 1, 1,
+				   1, 1, 1, 1, 1, 1, 1, 1, 1, 1, 1, 1, 1, 1, 1, 1,
+				   1, 1, 1, 1, 1, 1, 1, 1, 1, 1, 1, 1, 1, 1, 1, 1,
+				   1, 1, 1, 1, 1, 1, 1, 1, 1, 1, 1, 1, 1, 1, 1, 1,
+				   1, 1, 1, 1, 1, 1, 1, 1, 1, 1, 1, 1, 1, 1, 1, 1,
+				   1, 1, 1, 1, 1, 1, 1, 1, 1, 2, 2, 2, 2, 2, 2, 2,
+				   2, 2, 2, 2, 2, 2, 2, 2, 2, 2, 2, 2, 2, 2, 2, 2,
+				   2, 2, 2, 2, 2, 2, 2, 2, 2, 2, 2, 2, 2, 2, 2, 2,
+				   2, 2, 2, 2, 2, 2, 2, 2, 2, 2, 2, 2, 2, 2, 2, 2,
+				   2, 2, 2, 2, 2, 2, 2, 2, 2, 2, 2, 2, 2, 2, 2, 2,
+				   2, 2, 2, 2, 2, 2, 2, 2, 2, 2, 2, 2, 2, 2, 2, 2,
+				   2, 2, 2, 2, 2, 2, 2, 2, 2, 2, 2, 2, 2, 2, 2, 2,
+				   2, 2, 2, 2, 2, 2, 2, 2, 2, 2, 2, 2, 2, 2, 2, 2,
+				   2, 2, 2, 2, 2, 2, 2, 2, 2, 2, 2, 2, 2, 2, 2, 2,
+				   2, 2, 2, 2, 2, 2, 2, 2, 2, 3, 3, 3, 3, 3, 3, 3,
+				   3, 3, 3, 3, 3, 3, 3, 3, 3, 3, 3, 3, 3, 3, 3, 3,
+				   3, 3, 3, 3, 3, 3, 3, 3, 3, 3, 3, 3, 3, 3, 3, 3,
+				   3, 3, 3, 3, 3, 3, 3, 3, 3, 3, 3, 3, 3, 3, 3, 3,
+				   3, 3, 3, 3, 3, 3, 3, 3, 3, 3, 3, 3, 3, 3, 3, 3,
+				   3, 3, 3, 3, 3, 3, 3, 3, 3, 3, 3, 3, 3, 3, 3, 3,
+				   3, 3, 3, 3, 3, 3, 3, 3, 3, 3, 3, 3, 3, 3, 3, 3,
+				   3, 3, 3, 3, 3, 3, 3, 3, 3, 3, 3, 3, 3, 3, 3, 3,
+				   3, 3, 3, 3, 3, 3, 3, 3, 3, 3, 3, 3, 3, 3, 3, 3,
+				   3, 3, 3, 3, 3, 3, 3, 3, 3, 3, 3, 3, 3, 3, 3, 3,
+				   3, 4, 4, 4, 4, 4, 4, 4, 4, 4, 4, 4, 4, 4, 4, 4,
+				   4, 4, 4, 4, 4, 4, 4, 4, 4, 4, 4, 4, 4, 4, 4, 4,
+				   4, 4, 4, 4, 4, 4, 4, 4, 4, 4, 4, 4, 4, 4, 4, 4,
+				   4, 4, 4, 4, 4, 4, 4, 4, 4, 4, 4, 4, 4, 4, 4, 4,
+				   4, 4, 4, 4, 4, 4, 4, 4, 4, 4, 4, 4, 4, 4, 4, 4,
+				   4, 4, 4, 4, 4, 4, 4, 4, 4, 4, 4, 4, 4, 4, 4, 4,
+				   4, 4, 4, 4, 4, 4, 4, 4, 4, 4, 4, 4, 4, 4, 4, 4,
+				   4, 4, 4, 4, 4, 4, 4, 4, 4, 4, 4, 4, 4, 4, 4, 4,
+				   4, 4, 4, 4, 4, 4, 4, 4, 4, 4, 4, 4, 4, 4, 4, 4,
+				   4, 4, 4, 4, 4, 4, 4, 4, 4, 5, 5, 5, 5, 5, 5, 5,
+				   5, 5, 5, 5, 5, 5, 5, 5, 5, 5, 5, 5, 5, 5, 5, 5,
+				   5, 5, 5, 5, 5, 5, 5, 5, 5, 5, 5, 5, 5, 5, 5, 5,
+				   5, 5, 5, 5, 5, 5, 5, 5, 5, 5, 5, 5, 5, 5, 5, 5,
+				   5, 5, 5, 5, 5, 5, 5, 5, 5, 5, 5, 5, 5, 5, 5, 5,
+				   5, 5, 5, 5, 5, 5, 5, 5, 5, 5, 5, 5, 5, 5, 5, 5,
+				   5, 5, 5, 5, 5, 5, 5, 5, 5, 5, 5, 5, 5, 5, 5, 5,
+				   5, 5, 5, 5, 5, 5, 5, 5, 5, 5, 5, 5, 5, 5, 5, 5,
+				   5, 5, 5, 5, 5, 5, 5, 5, 5, 5, 5, 5, 5, 5, 5, 5,
+				   5, 5, 5, 5, 5, 5, 5, 5, 5, 6, 6, 6, 6, 6, 6, 6,
+				   6, 6, 6, 6, 6, 6, 6, 6, 6, 6, 6, 6, 6, 6, 6, 6,
+				   6, 6, 6, 6, 6, 6, 6, 6, 6, 6, 6, 6, 6, 6, 6, 6,
+				   6, 6, 6, 6, 6, 6, 6, 6, 6, 6, 6, 6, 6, 6, 6, 6,
+				   6, 6, 6, 6, 6, 6, 6, 6, 6, 6, 6, 6, 6, 6, 6, 6,
+				   6, 6, 6, 6, 6, 6, 6, 6, 6, 6, 6, 6, 6, 6, 6, 6,
+				   6, 6, 6, 6, 6, 6, 6, 6, 6, 6, 6, 6, 6, 6, 6, 6,
+				   6, 6, 6, 6, 6, 6, 6, 6, 6, 6, 6, 6, 6, 6, 6, 6,
+				   6, 6, 6, 6, 6, 6, 6, 6, 6, 6, 6, 6, 6, 6, 6, 6,
+				   6, 6, 6, 6, 6, 6, 6, 6, 6, 7, 7, 7, 7, 7, 7, 7,
+				   7, 7, 7, 7, 7, 7, 7, 7, 7, 7, 7, 7, 7, 7, 7, 7,
+				   7, 7, 7, 7, 7, 7, 7, 7, 7, 7, 7, 7, 7, 7, 7, 7,
+				   7, 7, 7, 7, 7, 7, 7, 7, 7, 7, 7, 7, 7, 7, 7, 7,
+				   7, 7, 7, 7, 7, 7, 7, 7, 7, 7, 7, 7, 7, 7, 7, 7,
+				   7,
+				   },
+				  },
+
+
+
+
+				  // quantization method 6, range 0..9
+				  {
+				   QUANT_10,
+				   {0, 64, 7, 57, 14, 50, 21, 43, 28, 36,},
+				   {
+					0.000000, 1.000000, 0.109375, 0.890625,
+					0.218750, 0.781250, 0.328125, 0.671875,
+					0.437500, 0.562500,},
+				   {0, 3, 0, 5, 2, 7, 4, 9, 6, 8,},
+				   {2, 1, 4, 1, 6, 3, 8, 5, 9, 7,},
+				   {
+					0, 0, 0, 0, 0, 0, 0, 0, 0, 0, 0, 0, 0, 0, 0, 0,
+					0, 0, 0, 0, 0, 0, 0, 0, 0, 0, 0, 0, 0, 0, 0, 0,
+					0, 0, 0, 0, 0, 0, 0, 0, 0, 0, 0, 0, 0, 0, 0, 0,
+					0, 0, 0, 0, 0, 0, 0, 0, 0, 2, 2, 2, 2, 2, 2, 2,
+					2, 2, 2, 2, 2, 2, 2, 2, 2, 2, 2, 2, 2, 2, 2, 2,
+					2, 2, 2, 2, 2, 2, 2, 2, 2, 2, 2, 2, 2, 2, 2, 2,
+					2, 2, 2, 2, 2, 2, 2, 2, 2, 2, 2, 2, 2, 2, 2, 2,
+					2, 2, 2, 2, 2, 2, 2, 2, 2, 2, 2, 2, 2, 2, 2, 2,
+					2, 2, 2, 2, 2, 2, 2, 2, 2, 2, 2, 2, 2, 2, 2, 2,
+					2, 2, 2, 2, 2, 2, 2, 2, 2, 2, 2, 2, 2, 2, 2, 2,
+					2, 2, 2, 2, 2, 2, 2, 2, 2, 4, 4, 4, 4, 4, 4, 4,
+					4, 4, 4, 4, 4, 4, 4, 4, 4, 4, 4, 4, 4, 4, 4, 4,
+					4, 4, 4, 4, 4, 4, 4, 4, 4, 4, 4, 4, 4, 4, 4, 4,
+					4, 4, 4, 4, 4, 4, 4, 4, 4, 4, 4, 4, 4, 4, 4, 4,
+					4, 4, 4, 4, 4, 4, 4, 4, 4, 4, 4, 4, 4, 4, 4, 4,
+					4, 4, 4, 4, 4, 4, 4, 4, 4, 4, 4, 4, 4, 4, 4, 4,
+					4, 4, 4, 4, 4, 4, 4, 4, 4, 4, 4, 4, 4, 4, 4, 4,
+					4, 4, 4, 4, 4, 4, 4, 4, 4, 6, 6, 6, 6, 6, 6, 6,
+					6, 6, 6, 6, 6, 6, 6, 6, 6, 6, 6, 6, 6, 6, 6, 6,
+					6, 6, 6, 6, 6, 6, 6, 6, 6, 6, 6, 6, 6, 6, 6, 6,
+					6, 6, 6, 6, 6, 6, 6, 6, 6, 6, 6, 6, 6, 6, 6, 6,
+					6, 6, 6, 6, 6, 6, 6, 6, 6, 6, 6, 6, 6, 6, 6, 6,
+					6, 6, 6, 6, 6, 6, 6, 6, 6, 6, 6, 6, 6, 6, 6, 6,
+					6, 6, 6, 6, 6, 6, 6, 6, 6, 6, 6, 6, 6, 6, 6, 6,
+					6, 6, 6, 6, 6, 6, 6, 6, 6, 8, 8, 8, 8, 8, 8, 8,
+					8, 8, 8, 8, 8, 8, 8, 8, 8, 8, 8, 8, 8, 8, 8, 8,
+					8, 8, 8, 8, 8, 8, 8, 8, 8, 8, 8, 8, 8, 8, 8, 8,
+					8, 8, 8, 8, 8, 8, 8, 8, 8, 8, 8, 8, 8, 8, 8, 8,
+					8, 8, 8, 8, 8, 8, 8, 8, 8, 8, 8, 8, 8, 8, 8, 8,
+					8, 8, 8, 8, 8, 8, 8, 8, 8, 8, 8, 8, 8, 8, 8, 8,
+					8, 8, 8, 8, 8, 8, 8, 8, 8, 8, 8, 8, 8, 8, 8, 8,
+					8, 8, 8, 8, 8, 8, 8, 8, 8, 8, 8, 8, 8, 8, 8, 8,
+					8, 9, 9, 9, 9, 9, 9, 9, 9, 9, 9, 9, 9, 9, 9, 9,
+					9, 9, 9, 9, 9, 9, 9, 9, 9, 9, 9, 9, 9, 9, 9, 9,
+					9, 9, 9, 9, 9, 9, 9, 9, 9, 9, 9, 9, 9, 9, 9, 9,
+					9, 9, 9, 9, 9, 9, 9, 9, 9, 9, 9, 9, 9, 9, 9, 9,
+					9, 9, 9, 9, 9, 9, 9, 9, 9, 9, 9, 9, 9, 9, 9, 9,
+					9, 9, 9, 9, 9, 9, 9, 9, 9, 9, 9, 9, 9, 9, 9, 9,
+					9, 9, 9, 9, 9, 9, 9, 9, 9, 9, 9, 9, 9, 9, 9, 9,
+					9, 9, 9, 9, 9, 9, 9, 9, 9, 7, 7, 7, 7, 7, 7, 7,
+					7, 7, 7, 7, 7, 7, 7, 7, 7, 7, 7, 7, 7, 7, 7, 7,
+					7, 7, 7, 7, 7, 7, 7, 7, 7, 7, 7, 7, 7, 7, 7, 7,
+					7, 7, 7, 7, 7, 7, 7, 7, 7, 7, 7, 7, 7, 7, 7, 7,
+					7, 7, 7, 7, 7, 7, 7, 7, 7, 7, 7, 7, 7, 7, 7, 7,
+					7, 7, 7, 7, 7, 7, 7, 7, 7, 7, 7, 7, 7, 7, 7, 7,
+					7, 7, 7, 7, 7, 7, 7, 7, 7, 7, 7, 7, 7, 7, 7, 7,
+					7, 7, 7, 7, 7, 7, 7, 7, 7, 5, 5, 5, 5, 5, 5, 5,
+					5, 5, 5, 5, 5, 5, 5, 5, 5, 5, 5, 5, 5, 5, 5, 5,
+					5, 5, 5, 5, 5, 5, 5, 5, 5, 5, 5, 5, 5, 5, 5, 5,
+					5, 5, 5, 5, 5, 5, 5, 5, 5, 5, 5, 5, 5, 5, 5, 5,
+					5, 5, 5, 5, 5, 5, 5, 5, 5, 5, 5, 5, 5, 5, 5, 5,
+					5, 5, 5, 5, 5, 5, 5, 5, 5, 5, 5, 5, 5, 5, 5, 5,
+					5, 5, 5, 5, 5, 5, 5, 5, 5, 5, 5, 5, 5, 5, 5, 5,
+					5, 5, 5, 5, 5, 5, 5, 5, 5, 3, 3, 3, 3, 3, 3, 3,
+					3, 3, 3, 3, 3, 3, 3, 3, 3, 3, 3, 3, 3, 3, 3, 3,
+					3, 3, 3, 3, 3, 3, 3, 3, 3, 3, 3, 3, 3, 3, 3, 3,
+					3, 3, 3, 3, 3, 3, 3, 3, 3, 3, 3, 3, 3, 3, 3, 3,
+					3, 3, 3, 3, 3, 3, 3, 3, 3, 3, 3, 3, 3, 3, 3, 3,
+					3, 3, 3, 3, 3, 3, 3, 3, 3, 3, 3, 3, 3, 3, 3, 3,
+					3, 3, 3, 3, 3, 3, 3, 3, 3, 3, 3, 3, 3, 3, 3, 3,
+					3, 3, 3, 3, 3, 3, 3, 3, 3, 1, 1, 1, 1, 1, 1, 1,
+					1, 1, 1, 1, 1, 1, 1, 1, 1, 1, 1, 1, 1, 1, 1, 1,
+					1, 1, 1, 1, 1, 1, 1, 1, 1, 1, 1, 1, 1, 1, 1, 1,
+					1, 1, 1, 1, 1, 1, 1, 1, 1, 1, 1, 1, 1, 1, 1, 1,
+					1,
+					},
+				   },
+
+
+
+
+				   // quantization method 7, range 0..11
+				   {
+					QUANT_12,
+					{0, 64, 17, 47, 5, 59, 23, 41, 11, 53, 28, 36,},
+					{
+					 0.000000, 1.000000, 0.265625, 0.734375,
+					 0.078125, 0.921875, 0.359375, 0.640625,
+					 0.171875, 0.828125, 0.437500, 0.562500,},
+					{0, 5, 8, 7, 0, 9, 2, 11, 4, 3, 6, 10,},
+					{4, 1, 6, 9, 8, 1, 10, 3, 2, 5, 11, 7,},
+					{
+					 0, 0, 0, 0, 0, 0, 0, 0, 0, 0, 0, 0, 0, 0, 0, 0,
+					 0, 0, 0, 0, 0, 0, 0, 0, 0, 0, 0, 0, 0, 0, 0, 0,
+					 0, 0, 0, 0, 0, 0, 0, 0, 0, 4, 4, 4, 4, 4, 4, 4,
+					 4, 4, 4, 4, 4, 4, 4, 4, 4, 4, 4, 4, 4, 4, 4, 4,
+					 4, 4, 4, 4, 4, 4, 4, 4, 4, 4, 4, 4, 4, 4, 4, 4,
+					 4, 4, 4, 4, 4, 4, 4, 4, 4, 4, 4, 4, 4, 4, 4, 4,
+					 4, 4, 4, 4, 4, 4, 4, 4, 4, 4, 4, 4, 4, 4, 4, 4,
+					 4, 4, 4, 4, 4, 4, 4, 4, 4, 4, 4, 4, 4, 4, 4, 4,
+					 4, 8, 8, 8, 8, 8, 8, 8, 8, 8, 8, 8, 8, 8, 8, 8,
+					 8, 8, 8, 8, 8, 8, 8, 8, 8, 8, 8, 8, 8, 8, 8, 8,
+					 8, 8, 8, 8, 8, 8, 8, 8, 8, 8, 8, 8, 8, 8, 8, 8,
+					 8, 8, 8, 8, 8, 8, 8, 8, 8, 8, 8, 8, 8, 8, 8, 8,
+					 8, 8, 8, 8, 8, 8, 8, 8, 8, 8, 8, 8, 8, 8, 8, 8,
+					 8, 8, 8, 8, 8, 8, 8, 8, 8, 8, 8, 8, 8, 8, 8, 8,
+					 8, 2, 2, 2, 2, 2, 2, 2, 2, 2, 2, 2, 2, 2, 2, 2,
+					 2, 2, 2, 2, 2, 2, 2, 2, 2, 2, 2, 2, 2, 2, 2, 2,
+					 2, 2, 2, 2, 2, 2, 2, 2, 2, 2, 2, 2, 2, 2, 2, 2,
+					 2, 2, 2, 2, 2, 2, 2, 2, 2, 2, 2, 2, 2, 2, 2, 2,
+					 2, 2, 2, 2, 2, 2, 2, 2, 2, 2, 2, 2, 2, 2, 2, 2,
+					 2, 2, 2, 2, 2, 2, 2, 2, 2, 2, 2, 2, 2, 2, 2, 2,
+					 2, 6, 6, 6, 6, 6, 6, 6, 6, 6, 6, 6, 6, 6, 6, 6,
+					 6, 6, 6, 6, 6, 6, 6, 6, 6, 6, 6, 6, 6, 6, 6, 6,
+					 6, 6, 6, 6, 6, 6, 6, 6, 6, 6, 6, 6, 6, 6, 6, 6,
+					 6, 6, 6, 6, 6, 6, 6, 6, 6, 6, 6, 6, 6, 6, 6, 6,
+					 6, 6, 6, 6, 6, 6, 6, 6, 6, 6, 6, 6, 6, 6, 6, 6,
+					 6, 6, 6, 6, 6, 6, 6, 6, 6, 10, 10, 10, 10, 10, 10, 10,
+					 10, 10, 10, 10, 10, 10, 10, 10, 10, 10, 10, 10, 10, 10, 10, 10,
+					 10, 10, 10, 10, 10, 10, 10, 10, 10, 10, 10, 10, 10, 10, 10, 10,
+					 10, 10, 10, 10, 10, 10, 10, 10, 10, 10, 10, 10, 10, 10, 10, 10,
+					 10, 10, 10, 10, 10, 10, 10, 10, 10, 10, 10, 10, 10, 10, 10, 10,
+					 10, 10, 10, 10, 10, 10, 10, 10, 10, 10, 10, 10, 10, 10, 10, 10,
+					 10, 10, 10, 10, 10, 10, 10, 10, 10, 10, 10, 10, 10, 10, 10, 10,
+					 10, 11, 11, 11, 11, 11, 11, 11, 11, 11, 11, 11, 11, 11, 11, 11,
+					 11, 11, 11, 11, 11, 11, 11, 11, 11, 11, 11, 11, 11, 11, 11, 11,
+					 11, 11, 11, 11, 11, 11, 11, 11, 11, 11, 11, 11, 11, 11, 11, 11,
+					 11, 11, 11, 11, 11, 11, 11, 11, 11, 11, 11, 11, 11, 11, 11, 11,
+					 11, 11, 11, 11, 11, 11, 11, 11, 11, 11, 11, 11, 11, 11, 11, 11,
+					 11, 11, 11, 11, 11, 11, 11, 11, 11, 11, 11, 11, 11, 11, 11, 11,
+					 11, 11, 11, 11, 11, 11, 11, 11, 11, 7, 7, 7, 7, 7, 7, 7,
+					 7, 7, 7, 7, 7, 7, 7, 7, 7, 7, 7, 7, 7, 7, 7, 7,
+					 7, 7, 7, 7, 7, 7, 7, 7, 7, 7, 7, 7, 7, 7, 7, 7,
+					 7, 7, 7, 7, 7, 7, 7, 7, 7, 7, 7, 7, 7, 7, 7, 7,
+					 7, 7, 7, 7, 7, 7, 7, 7, 7, 7, 7, 7, 7, 7, 7, 7,
+					 7, 7, 7, 7, 7, 7, 7, 7, 7, 7, 7, 7, 7, 7, 7, 7,
+					 7, 3, 3, 3, 3, 3, 3, 3, 3, 3, 3, 3, 3, 3, 3, 3,
+					 3, 3, 3, 3, 3, 3, 3, 3, 3, 3, 3, 3, 3, 3, 3, 3,
+					 3, 3, 3, 3, 3, 3, 3, 3, 3, 3, 3, 3, 3, 3, 3, 3,
+					 3, 3, 3, 3, 3, 3, 3, 3, 3, 3, 3, 3, 3, 3, 3, 3,
+					 3, 3, 3, 3, 3, 3, 3, 3, 3, 3, 3, 3, 3, 3, 3, 3,
+					 3, 3, 3, 3, 3, 3, 3, 3, 3, 3, 3, 3, 3, 3, 3, 3,
+					 3, 9, 9, 9, 9, 9, 9, 9, 9, 9, 9, 9, 9, 9, 9, 9,
+					 9, 9, 9, 9, 9, 9, 9, 9, 9, 9, 9, 9, 9, 9, 9, 9,
+					 9, 9, 9, 9, 9, 9, 9, 9, 9, 9, 9, 9, 9, 9, 9, 9,
+					 9, 9, 9, 9, 9, 9, 9, 9, 9, 9, 9, 9, 9, 9, 9, 9,
+					 9, 9, 9, 9, 9, 9, 9, 9, 9, 9, 9, 9, 9, 9, 9, 9,
+					 9, 9, 9, 9, 9, 9, 9, 9, 9, 9, 9, 9, 9, 9, 9, 9,
+					 9, 5, 5, 5, 5, 5, 5, 5, 5, 5, 5, 5, 5, 5, 5, 5,
+					 5, 5, 5, 5, 5, 5, 5, 5, 5, 5, 5, 5, 5, 5, 5, 5,
+					 5, 5, 5, 5, 5, 5, 5, 5, 5, 5, 5, 5, 5, 5, 5, 5,
+					 5, 5, 5, 5, 5, 5, 5, 5, 5, 5, 5, 5, 5, 5, 5, 5,
+					 5, 5, 5, 5, 5, 5, 5, 5, 5, 5, 5, 5, 5, 5, 5, 5,
+					 5, 5, 5, 5, 5, 5, 5, 5, 5, 1, 1, 1, 1, 1, 1, 1,
+					 1, 1, 1, 1, 1, 1, 1, 1, 1, 1, 1, 1, 1, 1, 1, 1,
+					 1, 1, 1, 1, 1, 1, 1, 1, 1, 1, 1, 1, 1, 1, 1, 1,
+					 1,
+					 },
+					},
+
+
+
+
+					// quantization method 8, range 0..15
+					{
+					 QUANT_16,
+					 {0, 4, 8, 12, 17, 21, 25, 29, 35, 39, 43, 47, 52, 56, 60, 64,},
+					 {
+					  0.000000, 0.062500, 0.125000, 0.187500,
+					  0.265625, 0.328125, 0.390625, 0.453125,
+					  0.546875, 0.609375, 0.671875, 0.734375,
+					  0.812500, 0.875000, 0.937500, 1.000000,},
+					 {0, 0, 1, 2, 3, 4, 5, 6, 7, 8, 9, 10, 11, 12, 13, 14,},
+					 {1, 2, 3, 4, 5, 6, 7, 8, 9, 10, 11, 12, 13, 14, 15, 15,},
+					 {
+					  0, 0, 0, 0, 0, 0, 0, 0, 0, 0, 0, 0, 0, 0, 0, 0,
+					  0, 0, 0, 0, 0, 0, 0, 0, 0, 0, 0, 0, 0, 0, 0, 0,
+					  0, 1, 1, 1, 1, 1, 1, 1, 1, 1, 1, 1, 1, 1, 1, 1,
+					  1, 1, 1, 1, 1, 1, 1, 1, 1, 1, 1, 1, 1, 1, 1, 1,
+					  1, 1, 1, 1, 1, 1, 1, 1, 1, 1, 1, 1, 1, 1, 1, 1,
+					  1, 1, 1, 1, 1, 1, 1, 1, 1, 1, 1, 1, 1, 1, 1, 1,
+					  1, 2, 2, 2, 2, 2, 2, 2, 2, 2, 2, 2, 2, 2, 2, 2,
+					  2, 2, 2, 2, 2, 2, 2, 2, 2, 2, 2, 2, 2, 2, 2, 2,
+					  2, 2, 2, 2, 2, 2, 2, 2, 2, 2, 2, 2, 2, 2, 2, 2,
+					  2, 2, 2, 2, 2, 2, 2, 2, 2, 2, 2, 2, 2, 2, 2, 2,
+					  2, 3, 3, 3, 3, 3, 3, 3, 3, 3, 3, 3, 3, 3, 3, 3,
+					  3, 3, 3, 3, 3, 3, 3, 3, 3, 3, 3, 3, 3, 3, 3, 3,
+					  3, 3, 3, 3, 3, 3, 3, 3, 3, 3, 3, 3, 3, 3, 3, 3,
+					  3, 3, 3, 3, 3, 3, 3, 3, 3, 3, 3, 3, 3, 3, 3, 3,
+					  3, 3, 3, 3, 3, 3, 3, 3, 3, 4, 4, 4, 4, 4, 4, 4,
+					  4, 4, 4, 4, 4, 4, 4, 4, 4, 4, 4, 4, 4, 4, 4, 4,
+					  4, 4, 4, 4, 4, 4, 4, 4, 4, 4, 4, 4, 4, 4, 4, 4,
+					  4, 4, 4, 4, 4, 4, 4, 4, 4, 4, 4, 4, 4, 4, 4, 4,
+					  4, 4, 4, 4, 4, 4, 4, 4, 4, 4, 4, 4, 4, 4, 4, 4,
+					  4, 5, 5, 5, 5, 5, 5, 5, 5, 5, 5, 5, 5, 5, 5, 5,
+					  5, 5, 5, 5, 5, 5, 5, 5, 5, 5, 5, 5, 5, 5, 5, 5,
+					  5, 5, 5, 5, 5, 5, 5, 5, 5, 5, 5, 5, 5, 5, 5, 5,
+					  5, 5, 5, 5, 5, 5, 5, 5, 5, 5, 5, 5, 5, 5, 5, 5,
+					  5, 6, 6, 6, 6, 6, 6, 6, 6, 6, 6, 6, 6, 6, 6, 6,
+					  6, 6, 6, 6, 6, 6, 6, 6, 6, 6, 6, 6, 6, 6, 6, 6,
+					  6, 6, 6, 6, 6, 6, 6, 6, 6, 6, 6, 6, 6, 6, 6, 6,
+					  6, 6, 6, 6, 6, 6, 6, 6, 6, 6, 6, 6, 6, 6, 6, 6,
+					  6, 7, 7, 7, 7, 7, 7, 7, 7, 7, 7, 7, 7, 7, 7, 7,
+					  7, 7, 7, 7, 7, 7, 7, 7, 7, 7, 7, 7, 7, 7, 7, 7,
+					  7, 7, 7, 7, 7, 7, 7, 7, 7, 7, 7, 7, 7, 7, 7, 7,
+					  7, 7, 7, 7, 7, 7, 7, 7, 7, 7, 7, 7, 7, 7, 7, 7,
+					  7, 7, 7, 7, 7, 7, 7, 7, 7, 7, 7, 7, 7, 7, 7, 7,
+					  7, 8, 8, 8, 8, 8, 8, 8, 8, 8, 8, 8, 8, 8, 8, 8,
+					  8, 8, 8, 8, 8, 8, 8, 8, 8, 8, 8, 8, 8, 8, 8, 8,
+					  8, 8, 8, 8, 8, 8, 8, 8, 8, 8, 8, 8, 8, 8, 8, 8,
+					  8, 8, 8, 8, 8, 8, 8, 8, 8, 8, 8, 8, 8, 8, 8, 8,
+					  8, 8, 8, 8, 8, 8, 8, 8, 8, 8, 8, 8, 8, 8, 8, 8,
+					  8, 9, 9, 9, 9, 9, 9, 9, 9, 9, 9, 9, 9, 9, 9, 9,
+					  9, 9, 9, 9, 9, 9, 9, 9, 9, 9, 9, 9, 9, 9, 9, 9,
+					  9, 9, 9, 9, 9, 9, 9, 9, 9, 9, 9, 9, 9, 9, 9, 9,
+					  9, 9, 9, 9, 9, 9, 9, 9, 9, 9, 9, 9, 9, 9, 9, 9,
+					  9, 10, 10, 10, 10, 10, 10, 10, 10, 10, 10, 10, 10, 10, 10, 10,
+					  10, 10, 10, 10, 10, 10, 10, 10, 10, 10, 10, 10, 10, 10, 10, 10,
+					  10, 10, 10, 10, 10, 10, 10, 10, 10, 10, 10, 10, 10, 10, 10, 10,
+					  10, 10, 10, 10, 10, 10, 10, 10, 10, 10, 10, 10, 10, 10, 10, 10,
+					  10, 11, 11, 11, 11, 11, 11, 11, 11, 11, 11, 11, 11, 11, 11, 11,
+					  11, 11, 11, 11, 11, 11, 11, 11, 11, 11, 11, 11, 11, 11, 11, 11,
+					  11, 11, 11, 11, 11, 11, 11, 11, 11, 11, 11, 11, 11, 11, 11, 11,
+					  11, 11, 11, 11, 11, 11, 11, 11, 11, 11, 11, 11, 11, 11, 11, 11,
+					  11, 11, 11, 11, 11, 11, 11, 11, 11, 12, 12, 12, 12, 12, 12, 12,
+					  12, 12, 12, 12, 12, 12, 12, 12, 12, 12, 12, 12, 12, 12, 12, 12,
+					  12, 12, 12, 12, 12, 12, 12, 12, 12, 12, 12, 12, 12, 12, 12, 12,
+					  12, 12, 12, 12, 12, 12, 12, 12, 12, 12, 12, 12, 12, 12, 12, 12,
+					  12, 12, 12, 12, 12, 12, 12, 12, 12, 12, 12, 12, 12, 12, 12, 12,
+					  12, 13, 13, 13, 13, 13, 13, 13, 13, 13, 13, 13, 13, 13, 13, 13,
+					  13, 13, 13, 13, 13, 13, 13, 13, 13, 13, 13, 13, 13, 13, 13, 13,
+					  13, 13, 13, 13, 13, 13, 13, 13, 13, 13, 13, 13, 13, 13, 13, 13,
+					  13, 13, 13, 13, 13, 13, 13, 13, 13, 13, 13, 13, 13, 13, 13, 13,
+					  13, 14, 14, 14, 14, 14, 14, 14, 14, 14, 14, 14, 14, 14, 14, 14,
+					  14, 14, 14, 14, 14, 14, 14, 14, 14, 14, 14, 14, 14, 14, 14, 14,
+					  14, 14, 14, 14, 14, 14, 14, 14, 14, 14, 14, 14, 14, 14, 14, 14,
+					  14, 14, 14, 14, 14, 14, 14, 14, 14, 14, 14, 14, 14, 14, 14, 14,
+					  14, 15, 15, 15, 15, 15, 15, 15, 15, 15, 15, 15, 15, 15, 15, 15,
+					  15, 15, 15, 15, 15, 15, 15, 15, 15, 15, 15, 15, 15, 15, 15, 15,
+					  15,
+					  },
+					 },
+
+
+
+
+					 // quantization method 9, range 0..19
+					 {
+					  QUANT_20,
+					  {0, 64, 16, 48, 3, 61, 19, 45, 6, 58, 23, 41, 9, 55, 26, 38, 13, 51, 29, 35,},
+					  {
+					   0.000000, 1.000000, 0.250000, 0.750000,
+					   0.046875, 0.953125, 0.296875, 0.703125,
+					   0.093750, 0.906250, 0.359375, 0.640625,
+					   0.140625, 0.859375, 0.406250, 0.593750,
+					   0.203125, 0.796875, 0.453125, 0.546875,},
+					  {0, 5, 16, 7, 0, 9, 2, 11, 4, 13, 6, 15, 8, 17, 10, 19, 12, 3, 14, 18,},
+					  {4, 1, 6, 17, 8, 1, 10, 3, 12, 5, 14, 7, 16, 9, 18, 11, 2, 13, 19, 15,},
+					  {
+					   0, 0, 0, 0, 0, 0, 0, 0, 0, 0, 0, 0, 0, 0, 0, 0,
+					   0, 0, 0, 0, 0, 0, 0, 0, 0, 4, 4, 4, 4, 4, 4, 4,
+					   4, 4, 4, 4, 4, 4, 4, 4, 4, 4, 4, 4, 4, 4, 4, 4,
+					   4, 4, 4, 4, 4, 4, 4, 4, 4, 4, 4, 4, 4, 4, 4, 4,
+					   4, 4, 4, 4, 4, 4, 4, 4, 4, 8, 8, 8, 8, 8, 8, 8,
+					   8, 8, 8, 8, 8, 8, 8, 8, 8, 8, 8, 8, 8, 8, 8, 8,
+					   8, 8, 8, 8, 8, 8, 8, 8, 8, 8, 8, 8, 8, 8, 8, 8,
+					   8, 8, 8, 8, 8, 8, 8, 8, 8, 12, 12, 12, 12, 12, 12, 12,
+					   12, 12, 12, 12, 12, 12, 12, 12, 12, 12, 12, 12, 12, 12, 12, 12,
+					   12, 12, 12, 12, 12, 12, 12, 12, 12, 12, 12, 12, 12, 12, 12, 12,
+					   12, 12, 12, 12, 12, 12, 12, 12, 12, 12, 12, 12, 12, 12, 12, 12,
+					   12, 16, 16, 16, 16, 16, 16, 16, 16, 16, 16, 16, 16, 16, 16, 16,
+					   16, 16, 16, 16, 16, 16, 16, 16, 16, 16, 16, 16, 16, 16, 16, 16,
+					   16, 16, 16, 16, 16, 16, 16, 16, 16, 16, 16, 16, 16, 16, 16, 16,
+					   16, 16, 16, 16, 16, 16, 16, 16, 16, 2, 2, 2, 2, 2, 2, 2,
+					   2, 2, 2, 2, 2, 2, 2, 2, 2, 2, 2, 2, 2, 2, 2, 2,
+					   2, 2, 2, 2, 2, 2, 2, 2, 2, 2, 2, 2, 2, 2, 2, 2,
+					   2, 2, 2, 2, 2, 2, 2, 2, 2, 6, 6, 6, 6, 6, 6, 6,
+					   6, 6, 6, 6, 6, 6, 6, 6, 6, 6, 6, 6, 6, 6, 6, 6,
+					   6, 6, 6, 6, 6, 6, 6, 6, 6, 6, 6, 6, 6, 6, 6, 6,
+					   6, 6, 6, 6, 6, 6, 6, 6, 6, 6, 6, 6, 6, 6, 6, 6,
+					   6, 10, 10, 10, 10, 10, 10, 10, 10, 10, 10, 10, 10, 10, 10, 10,
+					   10, 10, 10, 10, 10, 10, 10, 10, 10, 10, 10, 10, 10, 10, 10, 10,
+					   10, 10, 10, 10, 10, 10, 10, 10, 10, 10, 10, 10, 10, 10, 10, 10,
+					   10, 10, 10, 10, 10, 10, 10, 10, 10, 14, 14, 14, 14, 14, 14, 14,
+					   14, 14, 14, 14, 14, 14, 14, 14, 14, 14, 14, 14, 14, 14, 14, 14,
+					   14, 14, 14, 14, 14, 14, 14, 14, 14, 14, 14, 14, 14, 14, 14, 14,
+					   14, 14, 14, 14, 14, 14, 14, 14, 14, 18, 18, 18, 18, 18, 18, 18,
+					   18, 18, 18, 18, 18, 18, 18, 18, 18, 18, 18, 18, 18, 18, 18, 18,
+					   18, 18, 18, 18, 18, 18, 18, 18, 18, 18, 18, 18, 18, 18, 18, 18,
+					   18, 18, 18, 18, 18, 18, 18, 18, 18, 18, 18, 18, 18, 18, 18, 18,
+					   18, 18, 18, 18, 18, 18, 18, 18, 18, 18, 18, 18, 18, 18, 18, 18,
+					   18, 19, 19, 19, 19, 19, 19, 19, 19, 19, 19, 19, 19, 19, 19, 19,
+					   19, 19, 19, 19, 19, 19, 19, 19, 19, 19, 19, 19, 19, 19, 19, 19,
+					   19, 19, 19, 19, 19, 19, 19, 19, 19, 19, 19, 19, 19, 19, 19, 19,
+					   19, 19, 19, 19, 19, 19, 19, 19, 19, 19, 19, 19, 19, 19, 19, 19,
+					   19, 19, 19, 19, 19, 19, 19, 19, 19, 15, 15, 15, 15, 15, 15, 15,
+					   15, 15, 15, 15, 15, 15, 15, 15, 15, 15, 15, 15, 15, 15, 15, 15,
+					   15, 15, 15, 15, 15, 15, 15, 15, 15, 15, 15, 15, 15, 15, 15, 15,
+					   15, 15, 15, 15, 15, 15, 15, 15, 15, 11, 11, 11, 11, 11, 11, 11,
+					   11, 11, 11, 11, 11, 11, 11, 11, 11, 11, 11, 11, 11, 11, 11, 11,
+					   11, 11, 11, 11, 11, 11, 11, 11, 11, 11, 11, 11, 11, 11, 11, 11,
+					   11, 11, 11, 11, 11, 11, 11, 11, 11, 11, 11, 11, 11, 11, 11, 11,
+					   11, 7, 7, 7, 7, 7, 7, 7, 7, 7, 7, 7, 7, 7, 7, 7,
+					   7, 7, 7, 7, 7, 7, 7, 7, 7, 7, 7, 7, 7, 7, 7, 7,
+					   7, 7, 7, 7, 7, 7, 7, 7, 7, 7, 7, 7, 7, 7, 7, 7,
+					   7, 7, 7, 7, 7, 7, 7, 7, 7, 3, 3, 3, 3, 3, 3, 3,
+					   3, 3, 3, 3, 3, 3, 3, 3, 3, 3, 3, 3, 3, 3, 3, 3,
+					   3, 3, 3, 3, 3, 3, 3, 3, 3, 3, 3, 3, 3, 3, 3, 3,
+					   3, 3, 3, 3, 3, 3, 3, 3, 3, 17, 17, 17, 17, 17, 17, 17,
+					   17, 17, 17, 17, 17, 17, 17, 17, 17, 17, 17, 17, 17, 17, 17, 17,
+					   17, 17, 17, 17, 17, 17, 17, 17, 17, 17, 17, 17, 17, 17, 17, 17,
+					   17, 17, 17, 17, 17, 17, 17, 17, 17, 17, 17, 17, 17, 17, 17, 17,
+					   17, 13, 13, 13, 13, 13, 13, 13, 13, 13, 13, 13, 13, 13, 13, 13,
+					   13, 13, 13, 13, 13, 13, 13, 13, 13, 13, 13, 13, 13, 13, 13, 13,
+					   13, 13, 13, 13, 13, 13, 13, 13, 13, 13, 13, 13, 13, 13, 13, 13,
+					   13, 13, 13, 13, 13, 13, 13, 13, 13, 9, 9, 9, 9, 9, 9, 9,
+					   9, 9, 9, 9, 9, 9, 9, 9, 9, 9, 9, 9, 9, 9, 9, 9,
+					   9, 9, 9, 9, 9, 9, 9, 9, 9, 9, 9, 9, 9, 9, 9, 9,
+					   9, 9, 9, 9, 9, 9, 9, 9, 9, 5, 5, 5, 5, 5, 5, 5,
+					   5, 5, 5, 5, 5, 5, 5, 5, 5, 5, 5, 5, 5, 5, 5, 5,
+					   5, 5, 5, 5, 5, 5, 5, 5, 5, 5, 5, 5, 5, 5, 5, 5,
+					   5, 5, 5, 5, 5, 5, 5, 5, 5, 1, 1, 1, 1, 1, 1, 1,
+					   1, 1, 1, 1, 1, 1, 1, 1, 1, 1, 1, 1, 1, 1, 1, 1,
+					   1,
+					   },
+					  },
+
+
+
+
+					  // quantization method 10, range 0..23
+					  {
+					   QUANT_24,
+					   {0, 64, 8, 56, 16, 48, 24, 40, 2, 62, 11, 53, 19, 45, 27, 37, 5, 59, 13, 51, 22, 42, 30, 34,},
+					   {
+						0.000000, 1.000000, 0.125000, 0.875000,
+						0.250000, 0.750000, 0.375000, 0.625000,
+						0.031250, 0.968750, 0.171875, 0.828125,
+						0.296875, 0.703125, 0.421875, 0.578125,
+						0.078125, 0.921875, 0.203125, 0.796875,
+						0.343750, 0.656250, 0.468750, 0.531250,},
+					   {0, 9, 16, 11, 18, 13, 20, 15, 0, 17, 2, 19, 4, 21, 6, 23, 8, 3, 10, 5, 12, 7, 14, 22,},
+					   {8, 1, 10, 17, 12, 19, 14, 21, 16, 1, 18, 3, 20, 5, 22, 7, 2, 9, 4, 11, 6, 13, 23, 15,},
+					   {
+						0, 0, 0, 0, 0, 0, 0, 0, 0, 0, 0, 0, 0, 0, 0, 0,
+						0, 8, 8, 8, 8, 8, 8, 8, 8, 8, 8, 8, 8, 8, 8, 8,
+						8, 8, 8, 8, 8, 8, 8, 8, 8, 8, 8, 8, 8, 8, 8, 8,
+						8, 8, 8, 8, 8, 8, 8, 8, 8, 16, 16, 16, 16, 16, 16, 16,
+						16, 16, 16, 16, 16, 16, 16, 16, 16, 16, 16, 16, 16, 16, 16, 16,
+						16, 16, 16, 16, 16, 16, 16, 16, 16, 16, 16, 16, 16, 16, 16, 16,
+						16, 16, 16, 16, 16, 16, 16, 16, 16, 2, 2, 2, 2, 2, 2, 2,
+						2, 2, 2, 2, 2, 2, 2, 2, 2, 2, 2, 2, 2, 2, 2, 2,
+						2, 2, 2, 2, 2, 2, 2, 2, 2, 2, 2, 2, 2, 2, 2, 2,
+						2, 2, 2, 2, 2, 2, 2, 2, 2, 10, 10, 10, 10, 10, 10, 10,
+						10, 10, 10, 10, 10, 10, 10, 10, 10, 10, 10, 10, 10, 10, 10, 10,
+						10, 10, 10, 10, 10, 10, 10, 10, 10, 10, 10, 10, 10, 10, 10, 10,
+						10, 18, 18, 18, 18, 18, 18, 18, 18, 18, 18, 18, 18, 18, 18, 18,
+						18, 18, 18, 18, 18, 18, 18, 18, 18, 18, 18, 18, 18, 18, 18, 18,
+						18, 18, 18, 18, 18, 18, 18, 18, 18, 4, 4, 4, 4, 4, 4, 4,
+						4, 4, 4, 4, 4, 4, 4, 4, 4, 4, 4, 4, 4, 4, 4, 4,
+						4, 4, 4, 4, 4, 4, 4, 4, 4, 4, 4, 4, 4, 4, 4, 4,
+						4, 4, 4, 4, 4, 4, 4, 4, 4, 12, 12, 12, 12, 12, 12, 12,
+						12, 12, 12, 12, 12, 12, 12, 12, 12, 12, 12, 12, 12, 12, 12, 12,
+						12, 12, 12, 12, 12, 12, 12, 12, 12, 12, 12, 12, 12, 12, 12, 12,
+						12, 12, 12, 12, 12, 12, 12, 12, 12, 20, 20, 20, 20, 20, 20, 20,
+						20, 20, 20, 20, 20, 20, 20, 20, 20, 20, 20, 20, 20, 20, 20, 20,
+						20, 20, 20, 20, 20, 20, 20, 20, 20, 20, 20, 20, 20, 20, 20, 20,
+						20, 6, 6, 6, 6, 6, 6, 6, 6, 6, 6, 6, 6, 6, 6, 6,
+						6, 6, 6, 6, 6, 6, 6, 6, 6, 6, 6, 6, 6, 6, 6, 6,
+						6, 6, 6, 6, 6, 6, 6, 6, 6, 14, 14, 14, 14, 14, 14, 14,
+						14, 14, 14, 14, 14, 14, 14, 14, 14, 14, 14, 14, 14, 14, 14, 14,
+						14, 14, 14, 14, 14, 14, 14, 14, 14, 14, 14, 14, 14, 14, 14, 14,
+						14, 14, 14, 14, 14, 14, 14, 14, 14, 22, 22, 22, 22, 22, 22, 22,
+						22, 22, 22, 22, 22, 22, 22, 22, 22, 22, 22, 22, 22, 22, 22, 22,
+						22, 22, 22, 22, 22, 22, 22, 22, 22, 22, 22, 22, 22, 22, 22, 22,
+						22, 22, 22, 22, 22, 22, 22, 22, 22, 22, 22, 22, 22, 22, 22, 22,
+						22, 23, 23, 23, 23, 23, 23, 23, 23, 23, 23, 23, 23, 23, 23, 23,
+						23, 23, 23, 23, 23, 23, 23, 23, 23, 23, 23, 23, 23, 23, 23, 23,
+						23, 23, 23, 23, 23, 23, 23, 23, 23, 23, 23, 23, 23, 23, 23, 23,
+						23, 23, 23, 23, 23, 23, 23, 23, 23, 15, 15, 15, 15, 15, 15, 15,
+						15, 15, 15, 15, 15, 15, 15, 15, 15, 15, 15, 15, 15, 15, 15, 15,
+						15, 15, 15, 15, 15, 15, 15, 15, 15, 15, 15, 15, 15, 15, 15, 15,
+						15, 15, 15, 15, 15, 15, 15, 15, 15, 7, 7, 7, 7, 7, 7, 7,
+						7, 7, 7, 7, 7, 7, 7, 7, 7, 7, 7, 7, 7, 7, 7, 7,
+						7, 7, 7, 7, 7, 7, 7, 7, 7, 7, 7, 7, 7, 7, 7, 7,
+						7, 21, 21, 21, 21, 21, 21, 21, 21, 21, 21, 21, 21, 21, 21, 21,
+						21, 21, 21, 21, 21, 21, 21, 21, 21, 21, 21, 21, 21, 21, 21, 21,
+						21, 21, 21, 21, 21, 21, 21, 21, 21, 13, 13, 13, 13, 13, 13, 13,
+						13, 13, 13, 13, 13, 13, 13, 13, 13, 13, 13, 13, 13, 13, 13, 13,
+						13, 13, 13, 13, 13, 13, 13, 13, 13, 13, 13, 13, 13, 13, 13, 13,
+						13, 13, 13, 13, 13, 13, 13, 13, 13, 5, 5, 5, 5, 5, 5, 5,
+						5, 5, 5, 5, 5, 5, 5, 5, 5, 5, 5, 5, 5, 5, 5, 5,
+						5, 5, 5, 5, 5, 5, 5, 5, 5, 5, 5, 5, 5, 5, 5, 5,
+						5, 5, 5, 5, 5, 5, 5, 5, 5, 19, 19, 19, 19, 19, 19, 19,
+						19, 19, 19, 19, 19, 19, 19, 19, 19, 19, 19, 19, 19, 19, 19, 19,
+						19, 19, 19, 19, 19, 19, 19, 19, 19, 19, 19, 19, 19, 19, 19, 19,
+						19, 11, 11, 11, 11, 11, 11, 11, 11, 11, 11, 11, 11, 11, 11, 11,
+						11, 11, 11, 11, 11, 11, 11, 11, 11, 11, 11, 11, 11, 11, 11, 11,
+						11, 11, 11, 11, 11, 11, 11, 11, 11, 3, 3, 3, 3, 3, 3, 3,
+						3, 3, 3, 3, 3, 3, 3, 3, 3, 3, 3, 3, 3, 3, 3, 3,
+						3, 3, 3, 3, 3, 3, 3, 3, 3, 3, 3, 3, 3, 3, 3, 3,
+						3, 3, 3, 3, 3, 3, 3, 3, 3, 17, 17, 17, 17, 17, 17, 17,
+						17, 17, 17, 17, 17, 17, 17, 17, 17, 17, 17, 17, 17, 17, 17, 17,
+						17, 17, 17, 17, 17, 17, 17, 17, 17, 17, 17, 17, 17, 17, 17, 17,
+						17, 17, 17, 17, 17, 17, 17, 17, 17, 9, 9, 9, 9, 9, 9, 9,
+						9, 9, 9, 9, 9, 9, 9, 9, 9, 9, 9, 9, 9, 9, 9, 9,
+						9, 9, 9, 9, 9, 9, 9, 9, 9, 9, 9, 9, 9, 9, 9, 9,
+						9, 1, 1, 1, 1, 1, 1, 1, 1, 1, 1, 1, 1, 1, 1, 1,
+						1,
+						},
+					   },
+
+
+
+
+					   // quantization method 11, range 0..31
+					   {
+						QUANT_32,
+						{0, 2, 4, 6, 8, 10, 12, 14, 16, 18, 20, 22, 24, 26, 28, 30, 34, 36, 38, 40, 42, 44, 46, 48, 50, 52, 54, 56, 58, 60, 62, 64,},
+						{
+						 0.000000, 0.031250, 0.062500, 0.093750,
+						 0.125000, 0.156250, 0.187500, 0.218750,
+						 0.250000, 0.281250, 0.312500, 0.343750,
+						 0.375000, 0.406250, 0.437500, 0.468750,
+						 0.531250, 0.562500, 0.593750, 0.625000,
+						 0.656250, 0.687500, 0.718750, 0.750000,
+						 0.781250, 0.812500, 0.843750, 0.875000,
+						 0.906250, 0.937500, 0.968750, 1.000000,},
+						{0, 0, 1, 2, 3, 4, 5, 6, 7, 8, 9, 10, 11, 12, 13, 14, 15, 16, 17, 18, 19, 20, 21, 22, 23, 24, 25, 26, 27, 28, 29, 30,},
+						{1, 2, 3, 4, 5, 6, 7, 8, 9, 10, 11, 12, 13, 14, 15, 16, 17, 18, 19, 20, 21, 22, 23, 24, 25, 26, 27, 28, 29, 30, 31, 31,},
+						{
+						 0, 0, 0, 0, 0, 0, 0, 0, 0, 0, 0, 0, 0, 0, 0, 0,
+						 0, 1, 1, 1, 1, 1, 1, 1, 1, 1, 1, 1, 1, 1, 1, 1,
+						 1, 1, 1, 1, 1, 1, 1, 1, 1, 1, 1, 1, 1, 1, 1, 1,
+						 1, 2, 2, 2, 2, 2, 2, 2, 2, 2, 2, 2, 2, 2, 2, 2,
+						 2, 2, 2, 2, 2, 2, 2, 2, 2, 2, 2, 2, 2, 2, 2, 2,
+						 2, 3, 3, 3, 3, 3, 3, 3, 3, 3, 3, 3, 3, 3, 3, 3,
+						 3, 3, 3, 3, 3, 3, 3, 3, 3, 3, 3, 3, 3, 3, 3, 3,
+						 3, 4, 4, 4, 4, 4, 4, 4, 4, 4, 4, 4, 4, 4, 4, 4,
+						 4, 4, 4, 4, 4, 4, 4, 4, 4, 4, 4, 4, 4, 4, 4, 4,
+						 4, 5, 5, 5, 5, 5, 5, 5, 5, 5, 5, 5, 5, 5, 5, 5,
+						 5, 5, 5, 5, 5, 5, 5, 5, 5, 5, 5, 5, 5, 5, 5, 5,
+						 5, 6, 6, 6, 6, 6, 6, 6, 6, 6, 6, 6, 6, 6, 6, 6,
+						 6, 6, 6, 6, 6, 6, 6, 6, 6, 6, 6, 6, 6, 6, 6, 6,
+						 6, 7, 7, 7, 7, 7, 7, 7, 7, 7, 7, 7, 7, 7, 7, 7,
+						 7, 7, 7, 7, 7, 7, 7, 7, 7, 7, 7, 7, 7, 7, 7, 7,
+						 7, 8, 8, 8, 8, 8, 8, 8, 8, 8, 8, 8, 8, 8, 8, 8,
+						 8, 8, 8, 8, 8, 8, 8, 8, 8, 8, 8, 8, 8, 8, 8, 8,
+						 8, 9, 9, 9, 9, 9, 9, 9, 9, 9, 9, 9, 9, 9, 9, 9,
+						 9, 9, 9, 9, 9, 9, 9, 9, 9, 9, 9, 9, 9, 9, 9, 9,
+						 9, 10, 10, 10, 10, 10, 10, 10, 10, 10, 10, 10, 10, 10, 10, 10,
+						 10, 10, 10, 10, 10, 10, 10, 10, 10, 10, 10, 10, 10, 10, 10, 10,
+						 10, 11, 11, 11, 11, 11, 11, 11, 11, 11, 11, 11, 11, 11, 11, 11,
+						 11, 11, 11, 11, 11, 11, 11, 11, 11, 11, 11, 11, 11, 11, 11, 11,
+						 11, 12, 12, 12, 12, 12, 12, 12, 12, 12, 12, 12, 12, 12, 12, 12,
+						 12, 12, 12, 12, 12, 12, 12, 12, 12, 12, 12, 12, 12, 12, 12, 12,
+						 12, 13, 13, 13, 13, 13, 13, 13, 13, 13, 13, 13, 13, 13, 13, 13,
+						 13, 13, 13, 13, 13, 13, 13, 13, 13, 13, 13, 13, 13, 13, 13, 13,
+						 13, 14, 14, 14, 14, 14, 14, 14, 14, 14, 14, 14, 14, 14, 14, 14,
+						 14, 14, 14, 14, 14, 14, 14, 14, 14, 14, 14, 14, 14, 14, 14, 14,
+						 14, 15, 15, 15, 15, 15, 15, 15, 15, 15, 15, 15, 15, 15, 15, 15,
+						 15, 15, 15, 15, 15, 15, 15, 15, 15, 15, 15, 15, 15, 15, 15, 15,
+						 15, 15, 15, 15, 15, 15, 15, 15, 15, 15, 15, 15, 15, 15, 15, 15,
+						 15, 16, 16, 16, 16, 16, 16, 16, 16, 16, 16, 16, 16, 16, 16, 16,
+						 16, 16, 16, 16, 16, 16, 16, 16, 16, 16, 16, 16, 16, 16, 16, 16,
+						 16, 16, 16, 16, 16, 16, 16, 16, 16, 16, 16, 16, 16, 16, 16, 16,
+						 16, 17, 17, 17, 17, 17, 17, 17, 17, 17, 17, 17, 17, 17, 17, 17,
+						 17, 17, 17, 17, 17, 17, 17, 17, 17, 17, 17, 17, 17, 17, 17, 17,
+						 17, 18, 18, 18, 18, 18, 18, 18, 18, 18, 18, 18, 18, 18, 18, 18,
+						 18, 18, 18, 18, 18, 18, 18, 18, 18, 18, 18, 18, 18, 18, 18, 18,
+						 18, 19, 19, 19, 19, 19, 19, 19, 19, 19, 19, 19, 19, 19, 19, 19,
+						 19, 19, 19, 19, 19, 19, 19, 19, 19, 19, 19, 19, 19, 19, 19, 19,
+						 19, 20, 20, 20, 20, 20, 20, 20, 20, 20, 20, 20, 20, 20, 20, 20,
+						 20, 20, 20, 20, 20, 20, 20, 20, 20, 20, 20, 20, 20, 20, 20, 20,
+						 20, 21, 21, 21, 21, 21, 21, 21, 21, 21, 21, 21, 21, 21, 21, 21,
+						 21, 21, 21, 21, 21, 21, 21, 21, 21, 21, 21, 21, 21, 21, 21, 21,
+						 21, 22, 22, 22, 22, 22, 22, 22, 22, 22, 22, 22, 22, 22, 22, 22,
+						 22, 22, 22, 22, 22, 22, 22, 22, 22, 22, 22, 22, 22, 22, 22, 22,
+						 22, 23, 23, 23, 23, 23, 23, 23, 23, 23, 23, 23, 23, 23, 23, 23,
+						 23, 23, 23, 23, 23, 23, 23, 23, 23, 23, 23, 23, 23, 23, 23, 23,
+						 23, 24, 24, 24, 24, 24, 24, 24, 24, 24, 24, 24, 24, 24, 24, 24,
+						 24, 24, 24, 24, 24, 24, 24, 24, 24, 24, 24, 24, 24, 24, 24, 24,
+						 24, 25, 25, 25, 25, 25, 25, 25, 25, 25, 25, 25, 25, 25, 25, 25,
+						 25, 25, 25, 25, 25, 25, 25, 25, 25, 25, 25, 25, 25, 25, 25, 25,
+						 25, 26, 26, 26, 26, 26, 26, 26, 26, 26, 26, 26, 26, 26, 26, 26,
+						 26, 26, 26, 26, 26, 26, 26, 26, 26, 26, 26, 26, 26, 26, 26, 26,
+						 26, 27, 27, 27, 27, 27, 27, 27, 27, 27, 27, 27, 27, 27, 27, 27,
+						 27, 27, 27, 27, 27, 27, 27, 27, 27, 27, 27, 27, 27, 27, 27, 27,
+						 27, 28, 28, 28, 28, 28, 28, 28, 28, 28, 28, 28, 28, 28, 28, 28,
+						 28, 28, 28, 28, 28, 28, 28, 28, 28, 28, 28, 28, 28, 28, 28, 28,
+						 28, 29, 29, 29, 29, 29, 29, 29, 29, 29, 29, 29, 29, 29, 29, 29,
+						 29, 29, 29, 29, 29, 29, 29, 29, 29, 29, 29, 29, 29, 29, 29, 29,
+						 29, 30, 30, 30, 30, 30, 30, 30, 30, 30, 30, 30, 30, 30, 30, 30,
+						 30, 30, 30, 30, 30, 30, 30, 30, 30, 30, 30, 30, 30, 30, 30, 30,
+						 30, 31, 31, 31, 31, 31, 31, 31, 31, 31, 31, 31, 31, 31, 31, 31,
+						 31,
+						 },
+						},
+
+		};
+
+
+
+
+		uint16 lns_to_sf16(uint16 p)
+		{
+
+			uint16 mc = p & 0x7FF;
+			uint16 ec = p >> 11;
+			uint16 mt;
+			if (mc < 512)
+				mt = 3 * mc;
+			else if (mc < 1536)
+				mt = 4 * mc - 512;
+			else
+				mt = 5 * mc - 2048;
+
+			uint16 res = (ec << 10) | (mt >> 3);
+			if (res >= 0x7BFF)
+				res = 0x7BFF;
+			return res;
+		}
+
+
+		// helper function to initialize the orig-data from the work-data
+		void imageblock_initialize_orig_from_work(imageblock* pb, int pixelcount)
+		{
+			int i;
+			float* fptr = pb->orig_data;
+			float* wptr = pb->work_data;
+
+			for (i = 0; i < pixelcount; i++)
+			{
+				if (pb->rgb_lns[i])
+				{
+					fptr[0] = sf16_to_float(lns_to_sf16((uint16)wptr[0]));
+					fptr[1] = sf16_to_float(lns_to_sf16((uint16)wptr[1]));
+					fptr[2] = sf16_to_float(lns_to_sf16((uint16)wptr[2]));
+				}
+				else
+				{
+					fptr[0] = sf16_to_float(unorm16_to_sf16((uint16)wptr[0]));
+					fptr[1] = sf16_to_float(unorm16_to_sf16((uint16)wptr[1]));
+					fptr[2] = sf16_to_float(unorm16_to_sf16((uint16)wptr[2]));
+				}
+
+				if (pb->alpha_lns[i])
+				{
+					fptr[3] = sf16_to_float(lns_to_sf16((uint16)wptr[3]));
+				}
+				else
+				{
+					fptr[3] = sf16_to_float(unorm16_to_sf16((uint16)wptr[3]));
+				}
+
+				fptr += 4;
+				wptr += 4;
+			}
+		}
+
+
+		inline void decompress_symbolic_block(astc_decode_mode decode_mode,
+			const symbolic_compressed_block* scb,
+			imageblock* blk)
+		{
+			int i;
+
+			//        if (s_debuglog)
+			//        {
+			//            UE_LOG(LogMutableCore,Warning," uncompressing interesting block" );
+			//        }
+
+					// if we detected an error-block, blow up immediately.
+			if (scb->error_block)
+			{
+				//check(false);
+				return;
+			}
+
+
+			if (scb->block_mode < 0)
+			{
+				float red = 0, green = 0, blue = 0, alpha = 0;
+				int use_lns = 0;
+
+				if (scb->block_mode == -2)
+				{
+					// For sRGB decoding, we should return only the top 8 bits.
+					int mask = (decode_mode == DECODE_LDR_SRGB) ? 0xFF00 : 0xFFFF;
+
+					red = sf16_to_float(unorm16_to_sf16(uint16(scb->constant_color[0] & mask)));
+					green = sf16_to_float(unorm16_to_sf16(uint16(scb->constant_color[1] & mask)));
+					blue = sf16_to_float(unorm16_to_sf16(uint16(scb->constant_color[2] & mask)));
+					alpha = sf16_to_float(unorm16_to_sf16(uint16(scb->constant_color[3] & mask)));
+					use_lns = 0;
+				}
+				else
+				{
+					switch (decode_mode)
+					{
+					case DECODE_LDR_SRGB:
+						red = 1.0f;
+						green = 0.0f;
+						blue = 1.0f;
+						alpha = 1.0f;
+						use_lns = 0;
+						break;
+					case DECODE_LDR:
+						red = 0.0f;
+						green = 0.0f;
+						blue = 0.0f;
+						alpha = 0.0f;
+						use_lns = 0;
+						break;
+					case DECODE_HDR:
+						// constant-color block; unpack from FP16 to FP32.
+						red = sf16_to_float((sf16)scb->constant_color[0]);
+						green = sf16_to_float((sf16)scb->constant_color[1]);
+						blue = sf16_to_float((sf16)scb->constant_color[2]);
+						alpha = sf16_to_float((sf16)scb->constant_color[3]);
+						use_lns = 1;
+						break;
+					}
+				}
+
+				for (i = 0; i < 16; i++)
+				{
+					blk->orig_data[4 * i] = red;
+					blk->orig_data[4 * i + 1] = green;
+					blk->orig_data[4 * i + 2] = blue;
+					blk->orig_data[4 * i + 3] = alpha;
+					blk->rgb_lns[i] = (uint8)use_lns;
+					blk->alpha_lns[i] = (uint8)use_lns;
+				}
+
+
+				imageblock_initialize_work_from_orig(blk, 16);
+				return;
+			}
+
+
+			// get the appropriate partition-table entry
+			int partition_count = scb->partition_count;
+			check(partition_count == 1);
+			const partition_info* pt = get_partition_table();
+			pt += scb->partition_index;
+
+			// get the appropriate block descriptor
+			const block_size_descriptor* bsd = get_block_size_descriptor();
+			const decimation_table* const* ixtab2 = bsd->decimation_tables;
+
+
+			const decimation_table* it = ixtab2[bsd->block_modes[scb->block_mode].decimation_mode];
+
+			int is_dual_plane = bsd->block_modes[scb->block_mode].is_dual_plane;
+
+			int weight_quantization_level = bsd->block_modes[scb->block_mode].quantization_mode;
+
+
+			// decode the color endpoints
+			ushort4 color_endpoint0[4];
+			ushort4 color_endpoint1[4];
+			int rgb_hdr_endpoint[4] = { 0 };
+			int alpha_hdr_endpoint[4] = { 0 };
+			int nan_endpoint[4] = { 0 };
+
+			for (i = 0; i < partition_count; i++)
+				unpack_color_endpoints(decode_mode,
+					scb->color_formats[i],
+					scb->color_quantization_level, scb->color_values[i], &(rgb_hdr_endpoint[i]), &(alpha_hdr_endpoint[i]), &(nan_endpoint[i]), &(color_endpoint0[i]), &(color_endpoint1[i]));
+
+
+
+
+
+			// first unquantize the weights
+			int uq_plane1_weights[MAX_WEIGHTS_PER_BLOCK];
+			int uq_plane2_weights[MAX_WEIGHTS_PER_BLOCK];
+			int weight_count = it->num_weights;
+
+
+			const quantization_and_transfer_table* qat = &(quant_and_xfer_tables[weight_quantization_level]);
+
+			for (i = 0; i < weight_count; i++)
+			{
+				uq_plane1_weights[i] = qat->unquantized_value[scb->plane1_weights[i]];
+			}
+			if (is_dual_plane)
+			{
+				for (i = 0; i < weight_count; i++)
+					uq_plane2_weights[i] = qat->unquantized_value[scb->plane2_weights[i]];
+			}
+
+
+			// then undecimate them.
+			int weights[MAX_TEXELS_PER_BLOCK] = { 0 };
+			int plane2_weights[MAX_TEXELS_PER_BLOCK] = { 0 };
+
+
+			int texels_per_block = 16;
+			for (i = 0; i < texels_per_block; i++)
+				weights[i] = compute_value_of_texel_int(i, it, uq_plane1_weights);
+
+			if (is_dual_plane)
+				for (i = 0; i < texels_per_block; i++)
+					plane2_weights[i] = compute_value_of_texel_int(i, it, uq_plane2_weights);
+
+
+			int plane2_color_component = scb->plane2_color_component;
+
+
+			// now that we have endpoint colors and weights, we can unpack actual colors for
+			// each texel.
+			for (i = 0; i < texels_per_block; i++)
+			{
+				int partition = pt->partition_of_texel[i];
+
+				ushort4 color = lerp_color_int(decode_mode,
+					color_endpoint0[partition],
+					color_endpoint1[partition],
+					weights[i],
+					plane2_weights[i],
+					is_dual_plane ? plane2_color_component : -1);
+
+				blk->rgb_lns[i] = (uint8)rgb_hdr_endpoint[partition];
+				blk->alpha_lns[i] = (uint8)alpha_hdr_endpoint[partition];
+
+				blk->work_data[4 * i] = color.x;
+				blk->work_data[4 * i + 1] = color.y;
+				blk->work_data[4 * i + 2] = color.z;
+				blk->work_data[4 * i + 3] = color.w;
+			}
+
+			imageblock_initialize_orig_from_work(blk, 16);
+		}
+
+	} // namespace arm
+
+
+#endif // MIRO_INCLUDE_ASTC
+
+
+} // namespace impl
+
+
+using namespace impl;
+
+
+//-------------------------------------------------------------------------------------------------
+// Implementation of the miro interface
+//-------------------------------------------------------------------------------------------------
+namespace miro
+{
+
+	//---------------------------------------------------------------------------------------------
+	void init_astc_decompress()
+	{
+
+#if MIRO_INCLUDE_ASTC
+
+		if (!s_initialised_astc_decompress)
+		{
+			arm::build_quantization_mode_table();
+
+			// pregenerate astc decompression tables for 4x4x1 blocks
+			arm::get_block_size_descriptor();
+			arm::generate_partition_tables();
+
+			s_initialised_astc_decompress = 1;
+		}
+
+#endif
+
+	}
+
+
+	//---------------------------------------------------------------------------------------------
+	void init_dxtc_decompress()
+	{
+
+#if MIRO_INCLUDE_BC
+
+		if (!s_initialised_dxtc_decompress)
+		{
+			impl::stb::stb__InitDXT();
+
+			s_initialised_dxtc_decompress = 1;
+		}
+
+#endif
+
+	}
+
+
+	//---------------------------------------------------------------------------------------------
+	void initialize()
+	{
+		// do it only if necessary
+		//init_astc_decompress();
+
+		init_dxtc_decompress();
+	}
+
+
+	//---------------------------------------------------------------------------------------------
+	void finalize()
+	{
+	}
+
+
+	//---------------------------------------------------------------------------------------------
+	const char* GetVersion()
+	{
+		return "0.0.2";
+	}
+
+#if MIRO_INCLUDE_BC
+
+	//---------------------------------------------------------------------------------------------
+	void RGB_to_BC1(uint32 sx, uint32 sy, const uint8* from, uint8* to, int q)
+	{
+		uint32 bx = FMath::DivideAndRoundUp(sx, 4u);
+		uint32 by = FMath::DivideAndRoundUp(sy, 4u);
+		//for ( uint32 y = 0; y < by; ++y )
+		ParallelFor(by,
+			[
+				bx, sx, sy, from, to, q
+			] (uint32 y)
+			{
+				uint8* rowTo = to + 8 * bx * y;
+				for (uint32 x = 0; x < bx; ++x)
+				{
+					miro_pixel_block block;
+
+					int b = 0;
+					for (uint32 j = 0; j < 4; ++j)
+					{
+						// Clamp y in case block goes beyond the image. Biases the result.
+						uint32 rj = FMath::Min(sy - 1, y * 4 + j);
+
+						for (int i = 0; i < 4; ++i)
+						{
+							// Clamp x in case block goes beyond the image. Biases the result.
+							int ri = FMath::Min(sx - 1, x * 4 + i);
+
+							const uint8* pPixelSource = from + 3 * (sx * rj + ri);
+
+							block.bytes[b++] = pPixelSource[0];
+							block.bytes[b++] = pPixelSource[1];
+							block.bytes[b++] = pPixelSource[2];
+							block.bytes[b++] = 255;
+						}
+					}
+
+					int quality = q < 2 ? STB_DXT_NORMAL : STB_DXT_HIGHQUAL;
+					stb::stb_compress_dxt_block(rowTo, block.bytes, 0, quality);
+
+					rowTo += 8;
+				}
+			});
+	}
+
+
+	//---------------------------------------------------------------------------------------------
+	void RGBA_to_BC1(uint32 sx, uint32 sy, const uint8* from, uint8* to, int q)
+	{
+		uint32 bx = FMath::DivideAndRoundUp(sx, 4u);
+		uint32 by = FMath::DivideAndRoundUp(sy, 4u);
+		//for ( uint32 y = 0; y < by; ++y )
+		ParallelFor(by,
+			[
+				bx, sx, sy, from, to, q
+			] (uint32 y)
+			{
+				uint8* rowTo = to + 8 * bx * y;
+				for (uint32 x = 0; x < bx; ++x)
+				{
+					uint8 block[4 * 4 * 4];
+
+					int b = 0;
+					for (uint32 j = 0; j < 4; ++j)
+					{
+						// Clamp y in case block goes beyond the image. Biases the result.
+						uint32 rj = FMath::Min(sy - 1, y * 4 + j);
+
+						for (int i = 0; i < 4; ++i)
+						{
+							// Clamp x in case block goes beyond the image. Biases the result.
+							int ri = FMath::Min(sx - 1, x * 4 + i);
+
+							const uint8* pPixelSource = from + 4 * (sx * rj + ri);
+
+							block[b++] = pPixelSource[0];
+							block[b++] = pPixelSource[1];
+							block[b++] = pPixelSource[2];
+							block[b++] = pPixelSource[3];
+						}
+					}
+
+					int quality = q < 2 ? STB_DXT_NORMAL : STB_DXT_HIGHQUAL;
+					stb::stb_compress_dxt_block(rowTo, block, 0, quality);
+
+					rowTo += 8;
+				}
+			});
+	}
+
+
+	//---------------------------------------------------------------------------------------------
+	void BC1_to_RGB(uint32 sx, uint32 sy, const uint8* from, uint8* to)
+	{
+		uint32 nx = FMath::DivideAndRoundUp(sx, 4u);
+		uint32 ny = FMath::DivideAndRoundUp(sy, 4u);
+
+		//for (uint32 by = 0; by < ny; ++by)
+		ParallelFor(ny,
+			[
+				nx, sx, sy, from, to
+			] (uint32 by)
+			{
+				const uint8* rowFrom = from + 8 * nx * by;
+
+				for (uint32 bx = 0; bx < nx; ++bx)
+				{
+					uint8 block[4 * 4 * 4];
+
+					bcdec::bcdec_bc1(rowFrom, block, 4 * 4);
+
+					int b = 0;
+					for (int j = 0; j < 4; ++j)
+					{
+						for (int i = 0; i < 4; ++i)
+						{
+							if (by * 4 + j < sy && bx * 4 + i < sx)
+							{
+								uint8* pPixelDest = to + 3 * (sx * (by * 4 + j) + (bx * 4 + i));
+
+								pPixelDest[0] = block[b + 0];
+								pPixelDest[1] = block[b + 1];
+								pPixelDest[2] = block[b + 2];
+							}
+
+							b += 4;
+						}
+					}
+
+					rowFrom += 8;
+				}
+			});
+	}
+
+
+	//---------------------------------------------------------------------------------------------
+	void BC1_to_RGBA(uint32 sx, uint32 sy, const uint8* from, uint8* to)
+	{
+		uint32 nx = FMath::DivideAndRoundUp(sx, 4u);
+		uint32 ny = FMath::DivideAndRoundUp(sy, 4u);
+
+		//for (uint32 by = 0; by < ny; ++by)
+		ParallelFor(ny,
+			[
+				nx, sx, sy, from, to
+			] (uint32 by)
+			{
+				const uint8* rowFrom = from + 8 * nx * by;
+
+				for (uint32 bx = 0; bx < nx; ++bx)
+				{
+					uint8 block[4 * 4 * 4];
+
+					bcdec::bcdec_bc1(rowFrom, block, 4 * 4);
+
+					int b = 0;
+					for (int j = 0; j < 4; ++j)
+					{
+						for (int i = 0; i < 4; ++i)
+						{
+							if (by * 4 + j < sy && bx * 4 + i < sx)
+							{
+								uint8* pPixelDest = to + 4 * (sx * (by * 4 + j) + (bx * 4 + i));
+
+								pPixelDest[0] = block[b + 0];
+								pPixelDest[1] = block[b + 1];
+								pPixelDest[2] = block[b + 2];
+								pPixelDest[3] = block[b + 3];
+							}
+
+							b += 4;
+						}
+					}
+
+					rowFrom += 8;
+				}
+			});
+	}
+
+
+	//---------------------------------------------------------------------------------------------
+	void RGB_to_BC2(uint32 sx, uint32 sy, const uint8* from, uint8* to, int q)
+	{
+		uint32 bx = FMath::DivideAndRoundUp(sx, 4u);
+		uint32 by = FMath::DivideAndRoundUp(sy, 4u);
+
+		//for ( uint32 y=0; y<by; ++y )
+		ParallelFor(by,
+			[
+				bx, sx, sy, from, to, q
+			] (uint32 y)
+			{
+				uint8* rowTo = to + 16 * bx * y;
+
+				for (uint32 x = 0; x < bx; ++x)
+				{
+					uint8 block[4 * 4 * 4];
+
+					int b = 0;
+					for (uint32 j = 0; j < 4; ++j)
+					{
+						// Clamp y in case block goes beyond the image. Biases the result.
+						uint32 rj = FMath::Min(sy - 1, y * 4 + j);
+
+						for (int i = 0; i < 4; ++i)
+						{
+							// Clamp x in case block goes beyond the image. Biases the result.
+							int ri = FMath::Min(sx - 1, x * 4 + i);
+
+							const uint8* pPixelSource = from + 3 * (sx * rj + ri);
+
+							block[b++] = pPixelSource[0];
+							block[b++] = pPixelSource[1];
+							block[b++] = pPixelSource[2];
+							block[b++] = 255;
+						}
+					}
+
+					int quality = q < 2 ? STB_DXT_NORMAL : STB_DXT_HIGHQUAL;
+					stb::stb_compress_dxt_block(rowTo, block, 2, quality);
+
+					rowTo += 16;
+				}
+			});
+	}
+
+
+	//---------------------------------------------------------------------------------------------
+	void RGBA_to_BC2(uint32 sx, uint32 sy, const uint8* from, uint8* to, int q)
+	{
+		uint32 bx = FMath::DivideAndRoundUp(sx, 4u);
+		uint32 by = FMath::DivideAndRoundUp(sy, 4u);
+
+		//for (uint32 y = 0; y < by; ++y)
+		ParallelFor(by,
+			[
+				bx, sx, sy, from, to, q
+			] (uint32 y)
+			{
+				uint8* rowTo = to + 16 * bx * y;
+				for (uint32 x = 0; x < bx; ++x)
+				{
+					uint8 block[4 * 4 * 4];
+
+					int b = 0;
+					for (uint32 j = 0; j < 4; ++j)
+					{
+						// Clamp y in case block goes beyond the image. Biases the result.
+						uint32 rj = FMath::Min(sy - 1, y * 4 + j);
+
+						for (int i = 0; i < 4; ++i)
+						{
+							// Clamp x in case block goes beyond the image. Biases the result.
+							int ri = FMath::Min(sx - 1, x * 4 + i);
+
+							const uint8* pPixelSource = from + 4 * (sx * rj + ri);
+
+							block[b++] = pPixelSource[0];
+							block[b++] = pPixelSource[1];
+							block[b++] = pPixelSource[2];
+							block[b++] = pPixelSource[3];
+						}
+					}
+
+					int quality = q < 2 ? STB_DXT_NORMAL : STB_DXT_HIGHQUAL;
+					stb::stb_compress_dxt_block(rowTo, block, 2, quality);
+
+					rowTo += 16;
+				}
+			});
+	}
+
+
+	//---------------------------------------------------------------------------------------------
+	void BC2_to_RGBA(uint32 sx, uint32 sy, const uint8* from, uint8* to)
+	{
+		uint32 nx = FMath::DivideAndRoundUp(sx, 4u);
+		uint32 ny = FMath::DivideAndRoundUp(sy, 4u);
+
+		//for (uint32 by = 0; by < ny; ++by)
+		ParallelFor(ny,
+			[
+				nx, sx, sy, from, to
+			] (uint32 by)
+			{
+				const uint8* rowFrom = from + 16 * nx * by;
+
+				for (uint32 bx = 0; bx < nx; ++bx)
+				{
+					uint8 block[4 * 4 * 4];
+
+					uint8 const* colourBlock = rowFrom + 8;
+					uint8 const* alphaBlock = rowFrom;
+
+					bcdec::bcdec_bc2(rowFrom, block, 4 * 4);
+
+					int b = 0;
+					for (int j = 0; j < 4; ++j)
+					{
+						for (int i = 0; i < 4; ++i)
+						{
+							if (by * 4 + j < sy && bx * 4 + i < sx)
+							{
+								uint8* pPixelDest = to + 4 * (sx * (by * 4 + j) + (bx * 4 + i));
+
+								pPixelDest[0] = block[b + 0];
+								pPixelDest[1] = block[b + 1];
+								pPixelDest[2] = block[b + 2];
+								pPixelDest[3] = block[b + 3];
+							}
+
+							b += 4;
+						}
+					}
+
+					rowFrom += 16;
+				}
+			});
+	}
+
+
+	//---------------------------------------------------------------------------------------------
+	void BC2_to_RGB(uint32 sx, uint32 sy, const uint8* from, uint8* to)
+	{
+		uint32 nx = FMath::DivideAndRoundUp(sx, 4u);
+		uint32 ny = FMath::DivideAndRoundUp(sy, 4u);
+
+		//for (uint32 by = 0; by < ny; ++by)
+		ParallelFor(ny,
+			[
+				nx, sx, sy, from, to
+			] (uint32 by)
+			{
+				const uint8* rowFrom = from + 16 * nx * by;
+
+				for (uint32 bx = 0; bx < nx; ++bx)
+				{
+					uint8 block[4 * 4 * 4];
+
+					//uint8 const* colourBlock = from + 8;
+
+					// TODO, can skip alpha part
+					bcdec::bcdec_bc2(rowFrom, block, 4 * 4);
+
+					int b = 0;
+					for (int j = 0; j < 4; ++j)
+					{
+						for (int i = 0; i < 4; ++i)
+						{
+							if (by * 4 + j < sy && bx * 4 + i < sx)
+							{
+								uint8* pPixelDest = to + 3 * (sx * (by * 4 + j) + (bx * 4 + i));
+
+								pPixelDest[0] = block[b + 0];
+								pPixelDest[1] = block[b + 1];
+								pPixelDest[2] = block[b + 2];
+							}
+
+							b += 4;
+						}
+					}
+
+					rowFrom += 16;
+				}
+			});
+	}
+
+
+	//---------------------------------------------------------------------------------------------
+	void RGB_to_BC3(uint32 sx, uint32 sy, const uint8* from, uint8* to, int q)
+	{
+		uint32 bx = FMath::DivideAndRoundUp(sx, 4u);
+		uint32 by = FMath::DivideAndRoundUp(sy, 4u);
+
+		//for ( uint32 y=0; y<by; ++y )
+		ParallelFor(by,
+			[
+				bx, sx, sy, from, to, q
+			] (uint32 y)
+			{
+				uint8* rowTo = to + 16 * bx * y;
+
+				for (uint32 x = 0; x < bx; ++x)
+				{
+					uint8 block[4 * 4 * 4];
+
+					int b = 0;
+					for (uint32 j = 0; j < 4; ++j)
+					{
+						// Clamp y in case block goes beyond the image. Biases the result.
+						uint32 rj = FMath::Min(sy - 1, y * 4 + j);
+
+						for (int i = 0; i < 4; ++i)
+						{
+							// Clamp x in case block goes beyond the image. Biases the result.
+							int ri = FMath::Min(sx - 1, x * 4 + i);
+
+							const uint8* pPixelSource = from + 3 * (sx * rj + ri);
+
+							block[b++] = pPixelSource[0];
+							block[b++] = pPixelSource[1];
+							block[b++] = pPixelSource[2];
+							block[b++] = 255;
+						}
+					}
+
+					int quality = q < 2 ? STB_DXT_NORMAL : STB_DXT_HIGHQUAL;
+					stb::stb_compress_dxt_block(rowTo, block, 1, quality);
+
+					rowTo += 16;
+				}
+			});
+	}
+
+
+	//---------------------------------------------------------------------------------------------
+	void RGBA_to_BC3(uint32 sx, uint32 sy, const uint8* from, uint8* to, int q)
+	{
+		uint32 bx = FMath::DivideAndRoundUp(sx, 4u);
+		uint32 by = FMath::DivideAndRoundUp(sy, 4u);
+
+		//for ( uint32 y=0; y<by; ++y )
+		ParallelFor(by,
+			[
+				bx, sx, sy, from, to, q
+			] (uint32 y)
+			{
+				uint8* rowTo = to + 16 * bx * y;
+
+				for (uint32 x = 0; x < bx; ++x)
+				{
+					uint8 block[4 * 4 * 4];
+
+					int b = 0;
+					for (uint32 j = 0; j < 4; ++j)
+					{
+						// Clamp y in case block goes beyond the image. Biases the result.
+						uint32 rj = FMath::Min(sy - 1, y * 4 + j);
+
+						for (int i = 0; i < 4; ++i)
+						{
+							// Clamp x in case block goes beyond the image. Biases the result.
+							int ri = FMath::Min(sx - 1, x * 4 + i);
+
+							const uint8* pPixelSource = from + 4 * (sx * rj + ri);
+
+							block[b++] = pPixelSource[0];
+							block[b++] = pPixelSource[1];
+							block[b++] = pPixelSource[2];
+							block[b++] = pPixelSource[3];
+						}
+					}
+
+					int quality = q < 2 ? STB_DXT_NORMAL : STB_DXT_HIGHQUAL;
+					stb::stb_compress_dxt_block(rowTo, block, 1, quality);
+
+					rowTo += 16;
+				}
+			});
+	}
+
+
+	//---------------------------------------------------------------------------------------------
+	void BC3_to_RGBA(uint32 sx, uint32 sy, const uint8* from, uint8* to)
+	{
+		uint32 nx = FMath::DivideAndRoundUp(sx, 4u);
+		uint32 ny = FMath::DivideAndRoundUp(sy, 4u);
+
+		//for (uint32 by = 0; by < ny; ++by)
+		ParallelFor(ny,
+			[
+				nx, sx, sy, from, to
+			] (uint32 by)
+			{
+				const uint8* rowFrom = from + 16 * nx * by;
+
+				for (uint32 bx = 0; bx < nx; ++bx)
+				{
+					miro_pixel_block block;
+
+					//uint8 const* colourBlock = from + 8;
+					//uint8 const* alphaBlock = from;
+
+					bcdec::bcdec_bc3(rowFrom, block.bytes, 4 * 4);
+
+					int b = 0;
+					for (int j = 0; j < 4; ++j)
+					{
+						for (int i = 0; i < 4; ++i)
+						{
+							if (by * 4 + j < sy && bx * 4 + i < sx)
+							{
+								uint8* pPixelDest = to + 4 * (sx * (by * 4 + j) + (bx * 4 + i));
+
+								pPixelDest[0] = block.bytes[b + 0];
+								pPixelDest[1] = block.bytes[b + 1];
+								pPixelDest[2] = block.bytes[b + 2];
+								pPixelDest[3] = block.bytes[b + 3];
+							}
+
+							b += 4;
+						}
+					}
+
+					rowFrom += 16;
+				}
+			});
+	}
+
+
+	//---------------------------------------------------------------------------------------------
+	void BC3_to_RGB(uint32 sx, uint32 sy, const uint8* from, uint8* to)
+	{
+		uint32 nx = FMath::DivideAndRoundUp(sx, 4u);
+		uint32 ny = FMath::DivideAndRoundUp(sy, 4u);
+
+		//for (uint32 by = 0; by < ny; ++by)
+		ParallelFor(ny,
+			[
+				nx, sx, sy, from, to
+			] (uint32 by)
+			{
+				const uint8* rowFrom = from + 16 * nx * by;
+
+				for (uint32 bx = 0; bx < nx; ++bx)
+				{
+					uint8 block[4 * 4 * 4];
+
+					bcdec::bcdec_bc3(rowFrom, block, 4 * 4);
+
+					int b = 0;
+					for (int j = 0; j < 4; ++j)
+					{
+						for (int i = 0; i < 4; ++i)
+						{
+							if (by * 4 + j < sy && bx * 4 + i < sx)
+							{
+								uint8* pPixelDest = to + 3 * (sx * (by * 4 + j) + (bx * 4 + i));
+
+								pPixelDest[0] = block[b + 0];
+								pPixelDest[1] = block[b + 1];
+								pPixelDest[2] = block[b + 2];
+							}
+
+							b += 4;
+						}
+					}
+
+					rowFrom += 16;
+				}
+			});
+	}
+
+
+	//---------------------------------------------------------------------------------------------
+	void BC1_to_BC3(uint32 sx, uint32 sy, const uint8* from, uint8* to, int q)
+	{
+		uint32 bx = FMath::DivideAndRoundUp(sx, 4u);
+		uint32 by = FMath::DivideAndRoundUp(sy, 4u);
+		uint32 bcount = bx * by;
+
+		const auto& ProcessBlock = [from, to, q](uint32 b)
+		{
+			const uint8* blockFrom = from + 8 * b;
+			uint8* blockTo = to + 16 * b;
+
+			bool hasAlpha = BC1BlockHasAlpha(blockFrom);
+
+			if (hasAlpha)
+			{
+				uint8 block[4 * 4 * 4];
+				bcdec::bcdec_bc1(blockFrom, block, 4 * 4);
+
+				int quality = q < 2 ? STB_DXT_NORMAL : STB_DXT_HIGHQUAL;
+				stb::stb_compress_dxt_block(blockTo, block, 1, quality);
+			}
+			else
+			{
+				// set the constant white alpha part
+				FMemory::Memset(blockTo, 0xff, 8);
+
+				// Copy the colour part
+				FMemory::Memcpy(blockTo + 8, blockFrom, 8);
+			}
+		};
+
+		// \TODO: No benefit from setting a threshold?
+		constexpr uint32 BlockCountConcurrencyThreshold = 0;
+		if (bcount > BlockCountConcurrencyThreshold)
+		{
+			ParallelFor(bcount, ProcessBlock);
+		}
+		//else
+		//{
+		//	for (uint32 p = 0; p < bcount; ++p)
+		//	{
+		//		ProcessBlock(p);
+		//	}
+		//}
+
+	}
+
+
+	//---------------------------------------------------------------------------------------------
+	void L_to_BC4(uint32 sx, uint32 sy, const uint8* from, uint8* to, int q)
+	{
+		// We don't support quality setting for this format yet
+		(void)q;
+
+		uint32 bx = FMath::DivideAndRoundUp(sx, 4u);
+		uint32 by = FMath::DivideAndRoundUp(sy, 4u);
+
+		//for ( uint32 y=0; y<by; ++y )
+		ParallelFor(by,
+			[
+				bx, sx, sy, from, to, q
+			] (uint32 y)
+			{
+				uint8* rowTo = to + 8 * bx * y;
+				for (uint32 x = 0; x < bx; ++x)
+				{
+					uint8 block[4 * 4];
+
+					int b = 0;
+					for (uint32 j = 0; j < 4; ++j)
+					{
+						// Clamp y in case block goes beyond the image. Biases the result.
+						uint32 rj = FMath::Min(sy - 1, y * 4 + j);
+
+						for (int i = 0; i < 4; ++i)
+						{
+							// Clamp x in case block goes beyond the image. Biases the result.
+							int ri = FMath::Min(sx - 1, x * 4 + i);
+
+							const uint8* pPixelSource = from + 1 * (sx * rj + ri);
+							block[b++] = pPixelSource[0];
+						}
+					}
+
+					stb::stb_compress_bc4_block(rowTo, block);
+					rowTo += 8;
+				}
+			});
+	}
+
+
+	//---------------------------------------------------------------------------------------------
+	void BC4_to_L(uint32 sx, uint32 sy, const uint8* from, uint8* to)
+	{
+		uint32 nx = FMath::DivideAndRoundUp(sx, 4u);
+		uint32 ny = FMath::DivideAndRoundUp(sy, 4u);
+
+		//for (uint32 by = 0; by < ny; ++by)
+		ParallelFor(ny,
+			[
+				nx, sx, sy, from, to
+			] (uint32 by)
+			{
+				const uint8* rowFrom = from + 8 * nx * by;
+
+				for (uint32 bx = 0; bx < nx; ++bx)
+				{
+					uint8 block[4 * 4];
+
+					bcdec::bcdec_bc4(rowFrom, block, 4);
+
+					int b = 0;
+					for (int j = 0; j < 4; ++j)
+					{
+						for (int i = 0; i < 4; ++i)
+						{
+							if (by * 4 + j < sy && bx * 4 + i < sx)
+							{
+								uint8* pPixelDest = to + 1 * (sx * (by * 4 + j) + (bx * 4 + i));
+
+								pPixelDest[0] = block[b];
+							}
+
+							++b;
+						}
+					}
+
+					rowFrom += 8;
+				}
+			});
+	}
+
+
+	//---------------------------------------------------------------------------------------------
+	void BC4_to_RGB(uint32 sx, uint32 sy, const uint8* from, uint8* to)
+	{
+		uint32 nx = FMath::DivideAndRoundUp(sx, 4u);
+		uint32 ny = FMath::DivideAndRoundUp(sy, 4u);
+
+		//for (uint32 by = 0; by < ny; ++by)
+		ParallelFor(ny,
+			[
+				nx, sx, sy, from, to
+			] (uint32 by)
+			{
+				const uint8* rowFrom = from + 8 * nx * by;
+
+				for (uint32 bx = 0; bx < nx; ++bx)
+				{
+					uint8 block[4 * 4];
+
+					bcdec::bcdec_bc4(rowFrom, block, 4);
+
+					int b = 0;
+					for (int j = 0; j < 4; ++j)
+					{
+						for (int i = 0; i < 4; ++i)
+						{
+							if (by * 4 + j < sy && bx * 4 + i < sx)
+							{
+								uint8* pPixelDest = to + 3 * (sx * (by * 4 + j) + (bx * 4 + i));
+
+								pPixelDest[0] = block[b];
+								pPixelDest[1] = block[b];
+								pPixelDest[2] = block[b];
+							}
+
+							++b;
+						}
+					}
+
+					rowFrom += 8;
+				}
+			});
+	}
+
+
+	//---------------------------------------------------------------------------------------------
+	void BC4_to_RGBA(uint32 sx, uint32 sy, const uint8* from, uint8* to)
+	{
+		uint32 nx = FMath::DivideAndRoundUp(sx, 4u);
+		uint32 ny = FMath::DivideAndRoundUp(sy, 4u);
+
+		//for (uint32 by = 0; by < ny; ++by)
+		ParallelFor(ny,
+			[
+				nx, sx, sy, from, to
+			] (uint32 by)
+			{
+				const uint8* rowFrom = from + 8 * nx * by;
+
+				for (uint32 bx = 0; bx < nx; ++bx)
+				{
+					uint8 block[4 * 4];
+
+					bcdec::bcdec_bc4(rowFrom, block, 4);
+
+					int b = 0;
+					for (int j = 0; j < 4; ++j)
+					{
+						for (int i = 0; i < 4; ++i)
+						{
+							if (by * 4 + j < sy && bx * 4 + i < sx)
+							{
+								uint8* pPixelDest = to + 4 * (sx * (by * 4 + j) + (bx * 4 + i));
+
+								pPixelDest[0] = block[b];
+								pPixelDest[1] = block[b];
+								pPixelDest[2] = block[b];
+								pPixelDest[3] = 255;
+							}
+
+							++b;
+						}
+					}
+
+					rowFrom += 8;
+				}
+			});
+	}
+
+
+	//---------------------------------------------------------------------------------------------
+	void RGBA_to_BC5(uint32 sx, uint32 sy, const uint8* from, uint8* to, int q)
+	{
+		// We don't support quality setting for this format yet
+		(void)q;
+
+		uint32 bx = FMath::DivideAndRoundUp(sx, 4u);
+		uint32 by = FMath::DivideAndRoundUp(sy, 4u);
+
+		//for ( uint32 y=0; y<by; ++y )
+		ParallelFor(by,
+			[
+				bx, sx, sy, from, to, q
+			] (uint32 y)
+			{
+				uint8* rowTo = to + 16 * bx * y;
+				for (uint32 x = 0; x < bx; ++x)
+				{
+					uint8 block[4 * 4 * 2];
+
+					int b = 0;
+					for (uint32 j = 0; j < 4; ++j)
+					{
+						// Clamp y in case block goes beyond the image. Biases the result.
+						uint32 rj = FMath::Min(sy - 1, y * 4 + j);
+
+						for (int i = 0; i < 4; ++i)
+						{
+							// Clamp x in case block goes beyond the image. Biases the result.
+							int ri = FMath::Min(sx - 1, x * 4 + i);
+
+							const uint8* pPixelSource = from + 4 * (sx * rj + ri);
+							block[b++] = pPixelSource[0];
+							block[b++] = pPixelSource[1];
+						}
+					}
+
+					stb::stb_compress_bc5_block(rowTo, block);
+					rowTo += 16;
+				}
+			});
+	}
+
+
+	//---------------------------------------------------------------------------------------------
+	void RGB_to_BC5(uint32 sx, uint32 sy, const uint8* from, uint8* to, int q)
+	{
+		// We don't support quality setting for this format yet
+		(void)q;
+
+		uint32 bx = FMath::DivideAndRoundUp(sx, 4u);
+		uint32 by = FMath::DivideAndRoundUp(sy, 4u);
+
+		//for ( uint32 y=0; y<by; ++y )
+		ParallelFor(by,
+			[
+				bx, sx, sy, from, to, q
+			] (uint32 y)
+			{
+				uint8* rowTo = to + 16 * bx * y;
+				for (uint32 x = 0; x < bx; ++x)
+				{
+					uint8 block[4 * 4 * 2];
+
+					int b = 0;
+					for (uint32 j = 0; j < 4; ++j)
+					{
+						// Clamp y in case block goes beyond the image. Biases the result.
+						uint32 rj = FMath::Min(sy - 1, y * 4 + j);
+
+						for (uint32 i = 0; i < 4; ++i)
+						{
+							// Clamp x in case block goes beyond the image. Biases the result.
+							uint32 ri = FMath::Min(sx - 1, x * 4 + i);
+
+							const uint8* pPixelSource = from + 3 * (sx * rj + ri);
+							block[b++] = pPixelSource[0];
+							block[b++] = pPixelSource[1];
+						}
+					}
+
+					stb::stb_compress_bc5_block(rowTo, block);
+					rowTo += 16;
+				}
+			});
+	}
+
+
+	//---------------------------------------------------------------------------------------------
+	void BC5_to_RGBA(uint32 sx, uint32 sy, const uint8* from, uint8* to)
+	{
+		uint32 nx = FMath::DivideAndRoundUp(sx, 4u);
+		uint32 ny = FMath::DivideAndRoundUp(sy, 4u);
+
+		//for (uint32 by = 0; by < ny; ++by)
+		ParallelFor(ny,
+			[
+				nx, sx, sy, from, to
+			] (uint32 by)
+			{
+				const uint8* rowFrom = from + 16 * nx * by;
+
+				for (uint32 bx = 0; bx < nx; ++bx)
+				{
+					uint8 block[4 * 4 * 2];
+					bcdec::bcdec_bc5(rowFrom, block, 4 * 2);
+
+					int b = 0;
+					for (int j = 0; j < 4; ++j)
+					{
+						for (int i = 0; i < 4; ++i)
+						{
+							if (by * 4 + j < sy && bx * 4 + i < sx)
+							{
+								uint8* pPixelDest = to + 4 * (sx * (by * 4 + j) + (bx * 4 + i));
+
+								pPixelDest[0] = block[b + 0];
+								pPixelDest[1] = block[b + 1];
+								pPixelDest[2] = 255;
+								pPixelDest[3] = 255;
+							}
+
+							b += 2;
+						}
+					}
+
+					rowFrom += 16;
+				}
+			});
+	}
+
+
+	//---------------------------------------------------------------------------------------------
+	void BC5_to_RGB(uint32 sx, uint32 sy, const uint8* from, uint8* to)
+	{
+		uint32 nx = FMath::DivideAndRoundUp(sx, 4u);
+		uint32 ny = FMath::DivideAndRoundUp(sy, 4u);
+
+		//for (uint32 by = 0; by < ny; ++by)
+		ParallelFor(ny,
+			[
+				nx, sx, sy, from, to
+			] (uint32 by)
+			{
+				const uint8* rowFrom = from + 16 * nx * by;
+
+				for (uint32 bx = 0; bx < nx; ++bx)
+				{
+					uint8 block[4 * 4 * 2];
+					bcdec::bcdec_bc5(rowFrom, block, 4 * 2);
+
+					int b = 0;
+					for (int j = 0; j < 4; ++j)
+					{
+						for (int i = 0; i < 4; ++i)
+						{
+							if (by * 4 + j < sy && bx * 4 + i < sx)
+							{
+								uint8* pPixelDest = to + 3 * (sx * (by * 4 + j) + (bx * 4 + i));
+
+								pPixelDest[0] = block[b + 0];
+								pPixelDest[1] = block[b + 1];
+								pPixelDest[2] = 255;
+							}
+
+							b += 2;
+						}
+					}
+
+					rowFrom += 16;
+				}
+			});
+	}
+
+#endif // MIRO_INCLUDE_BC
+
+
+	//---------------------------------------------------------------------------------------------
+	//---------------------------------------------------------------------------------------------
+	//---------------------------------------------------------------------------------------------
+#if MIRO_INCLUDE_ASTC
+
+	void RGB_to_ASTC4x4RGBL(uint32 sx, uint32 sy, const uint8* from, uint8* to, int q)
+	{
+		(void)q;
+		astcrt::PhysicalBlock physical_block_zero = { {0} };
+		astcrt::PhysicalBlock* dst_re = reinterpret_cast<astcrt::PhysicalBlock*>(to);
+
+		uint32 bx = FMath::DivideAndRoundUp(sx, 4u);
+		uint32 by = FMath::DivideAndRoundUp(sy, 4u);
+
+		//for ( uint32 y=0; y<by; ++y )
+		ParallelFor(by,
+			[
+				bx, sx, sy, from, to, q, &physical_block_zero
+			] (uint32 y)
+			{
+				astcrt::PhysicalBlock* rowTo = reinterpret_cast<astcrt::PhysicalBlock*>(to + sizeof(astcrt::PhysicalBlock) * bx * y);
+				for (uint32 x = 0; x < bx; ++x)
+				{
+					// Fetch the block
+					uint8 block[4 * 4 * 4];
+					int b = 0;
+					for (int j = 0; j < 4; ++j)
+					{
+						// Clamp y in case block goes beyond the image. Biases the result.
+						uint32 rj = FMath::Min(sy - 1, y * 4 + j);
+
+						for (int i = 0; i < 4; ++i)
+						{
+							// Clamp x in case block goes beyond the image. Biases the result.
+							uint32 ri = FMath::Min(sx - 1, x * 4 + i);
+
+							const uint8* pPixelSource = from + 3 * (sx * rj + ri);
+
+							block[b++] = pPixelSource[2];
+							block[b++] = pPixelSource[1];
+							block[b++] = pPixelSource[0];
+							block[b++] = 0;
+						}
+					}
+
+					*rowTo = physical_block_zero;
+					astcrt::compress_block((astcrt::unorm8_t*)block, rowTo);
+
+					++rowTo;
+				}
+			});
+	}
+
+
+	//---------------------------------------------------------------------------------------------
+	void RGBA_to_ASTC4x4RGBL(uint32 sx, uint32 sy, const uint8* from, uint8* to, int q)
+	{
+		(void)q;
+		astcrt::PhysicalBlock physical_block_zero = { {0} };
+		astcrt::PhysicalBlock* dst_re = reinterpret_cast<astcrt::PhysicalBlock*>(to);
+
+		uint32 bx = FMath::DivideAndRoundUp(sx, 4u);
+		uint32 by = FMath::DivideAndRoundUp(sy, 4u);
+
+		//for ( uint32 y=0; y<by; ++y )
+		ParallelFor(by,
+			[
+				bx, sx, sy, from, to, q, &physical_block_zero
+			] (uint32 y)
+			{
+				astcrt::PhysicalBlock* rowTo = reinterpret_cast<astcrt::PhysicalBlock*>(to + sizeof(astcrt::PhysicalBlock) * bx * y);
+
+				for (uint32 x = 0; x < bx; ++x)
+				{
+					// Fetch the block
+					uint8 block[4 * 4 * 4];
+					int b = 0;
+					for (int j = 0; j < 4; ++j)
+					{
+						// Clamp y in case block goes beyond the image. Biases the result.
+						uint32 rj = FMath::Min(sy - 1, y * 4 + j);
+
+						for (int i = 0; i < 4; ++i)
+						{
+							// Clamp x in case block goes beyond the image. Biases the result.
+							uint32 ri = FMath::Min(sx - 1, x * 4 + i);
+
+							const uint8* pPixelSource = from + 4 * (sx * rj + ri);
+
+							block[b++] = pPixelSource[2];
+							block[b++] = pPixelSource[1];
+							block[b++] = pPixelSource[0];
+							block[b++] = 0;
+						}
+					}
+
+					*rowTo = physical_block_zero;
+					astcrt::compress_block((astcrt::unorm8_t*)block, rowTo);
+
+					++rowTo;
+				}
+			});
+	}
+
+
+	//---------------------------------------------------------------------------------------------
+	void ASTC4x4RGBL_to_RGB(uint32 sx, uint32 sy, const uint8* from, uint8* to)
+	{
+		init_astc_decompress();
+
+		// actually use arm's decompression
+		arm::symbolic_compressed_block symbolic;
+
+		for (uint32 y = 0; y < sy; y += 4)
+		{
+			for (uint32 x = 0; x < sx; x += 4)
+			{
+				arm::physical_compressed_block physical;
+				FMemory::Memcpy(physical.data, from, 16);
+
+				arm::physical_to_symbolic(physical, &symbolic);
+
+				arm::imageblock block;
+				arm::decompress_symbolic_block(arm::DECODE_LDR, &symbolic, &block);
+
+				// Copy data from block
+				const float* fptr = block.orig_data;
+				float data[3];
+				for (uint32 py = 0; py < 4; py++)
+				{
+					for (uint32 px = 0; px < 4; px++)
+					{
+						uint32 xi = x + px;
+						uint32 yi = y + py;
+
+						if (xi < sx && yi < sy)
+						{
+							uint8* toPixel = to + (yi * sx + xi) * 3;
+
+							data[0] = FMath::Min(1.0f, fptr[0]);
+							data[1] = FMath::Min(1.0f, fptr[1]);
+							data[2] = FMath::Min(1.0f, fptr[2]);
+
+							// pack the data
+							int ri = static_cast <int>(floor(data[0] * 255.0f + 0.5f));
+							int gi = static_cast <int>(floor(data[1] * 255.0f + 0.5f));
+							int bi = static_cast <int>(floor(data[2] * 255.0f + 0.5f));
+
+							toPixel[0] = (uint8)ri;
+							toPixel[1] = (uint8)gi;
+							toPixel[2] = (uint8)bi;
+						}
+						fptr += 4;
+					}
+				}
+
+				from += 16;
+			}
+		}
+	}
+
+
+	//---------------------------------------------------------------------------------------------
+	void ASTC4x4RGBL_to_RGBA(uint32 sx, uint32 sy, const uint8* from, uint8* to)
+	{
+		init_astc_decompress();
+
+		// actually use arm's decompression
+		arm::symbolic_compressed_block symbolic;
+
+		for (uint32 y = 0; y < sy; y += 4)
+		{
+			for (uint32 x = 0; x < sx; x += 4)
+			{
+				arm::physical_compressed_block physical;
+				FMemory::Memcpy(physical.data, from, 16);
+
+				arm::physical_to_symbolic(physical, &symbolic);
+
+				arm::imageblock block;
+				arm::decompress_symbolic_block(arm::DECODE_LDR, &symbolic, &block);
+
+				// Copy data from block
+				const float* fptr = block.orig_data;
+				float data[3];
+				for (uint32 py = 0; py < 4; py++)
+				{
+					for (uint32 px = 0; px < 4; px++)
+					{
+						uint32 xi = x + px;
+						uint32 yi = y + py;
+
+						if (xi < sx && yi < sy)
+						{
+							uint8* toPixel = to + (yi * sx + xi) * 4;
+
+							data[0] = FMath::Min(1.0f, fptr[0]);
+							data[1] = FMath::Min(1.0f, fptr[1]);
+							data[2] = FMath::Min(1.0f, fptr[2]);
+
+							// pack the data
+							int ri = static_cast <int>(floor(data[0] * 255.0f + 0.5f));
+							int gi = static_cast <int>(floor(data[1] * 255.0f + 0.5f));
+							int bi = static_cast <int>(floor(data[2] * 255.0f + 0.5f));
+
+							toPixel[0] = (uint8)ri;
+							toPixel[1] = (uint8)gi;
+							toPixel[2] = (uint8)bi;
+							toPixel[3] = 0;
+						}
+						fptr += 4;
+					}
+				}
+
+				from += 16;
+			}
+		}
+	}
+
+
+	//---------------------------------------------------------------------------------------------
+	void RGBA_to_ASTC4x4RGBAL(uint32 sx, uint32 sy, const uint8* from, uint8* to, int q)
+	{
+		(void)q;
+		astcrt::PhysicalBlock physical_block_zero = { {0} };
+		astcrt::PhysicalBlock* dst_re = reinterpret_cast<astcrt::PhysicalBlock*>(to);
+
+		uint32 bx = FMath::DivideAndRoundUp(sx, 4u);
+		uint32 by = FMath::DivideAndRoundUp(sy, 4u);
+
+		//for ( uint32 y=0; y<by; ++y )
+		ParallelFor(by,
+			[
+				bx, sx, sy, from, to, q, &physical_block_zero
+			] (uint32 y)
+			{
+				astcrt::PhysicalBlock* rowTo = reinterpret_cast<astcrt::PhysicalBlock*>(to + sizeof(astcrt::PhysicalBlock) * bx * y);
+
+				for (uint32 x = 0; x < bx; ++x)
+				{
+					// Fetch the block
+					uint8 block[4 * 4 * 4];
+					int b = 0;
+					for (int j = 0; j < 4; ++j)
+					{
+						// Clamp y in case block goes beyond the image. Biases the result.
+						uint32 rj = FMath::Min(sy - 1, y * 4 + j);
+
+						for (int i = 0; i < 4; ++i)
+						{
+							// Clamp x in case block goes beyond the image. Biases the result.
+							uint32 ri = FMath::Min(sx - 1, x * 4 + i);
+
+							const uint8* pPixelSource = from + 4 * (sx * rj + ri);
+
+							block[b++] = (uint8)pPixelSource[2];
+							block[b++] = (uint8)pPixelSource[1];
+							block[b++] = (uint8)pPixelSource[0];
+							block[b++] = (uint8)pPixelSource[3];
+						}
+					}
+
+					*rowTo = physical_block_zero;
+					astcrt::compress_block_rgba((astcrt::unorm8_t*)block, rowTo);
+
+					++rowTo;
+				}
+			});
+	}
+
+
+	//---------------------------------------------------------------------------------------------
+	void RGB_to_ASTC4x4RGBAL(uint32 sx, uint32 sy, const uint8* from, uint8* to, int q)
+	{
+		(void)q;
+		astcrt::PhysicalBlock physical_block_zero = { {0} };
+		astcrt::PhysicalBlock* dst_re = reinterpret_cast<astcrt::PhysicalBlock*>(to);
+
+		uint32 bx = FMath::DivideAndRoundUp(sx, 4u);
+		uint32 by = FMath::DivideAndRoundUp(sy, 4u);
+
+		//for ( uint32 y=0; y<by; ++y )
+		ParallelFor(by,
+			[
+				bx, sx, sy, from, to, q, &physical_block_zero
+			] (uint32 y)
+			{
+				astcrt::PhysicalBlock* rowTo = reinterpret_cast<astcrt::PhysicalBlock*>(to + sizeof(astcrt::PhysicalBlock) * bx * y);
+
+				for (uint32 x = 0; x < bx; ++x)
+				{
+					// Fetch the block
+					uint8 block[4 * 4 * 4];
+					int b = 0;
+					for (int j = 0; j < 4; ++j)
+					{
+						// Clamp y in case block goes beyond the image. Biases the result.
+						uint32 rj = FMath::Min(sy - 1, y * 4 + j);
+
+						for (int i = 0; i < 4; ++i)
+						{
+							// Clamp x in case block goes beyond the image. Biases the result.
+							uint32 ri = FMath::Min(sx - 1, x * 4 + i);
+
+							const uint8* pPixelSource = from + 3 * (sx * rj + ri);
+
+							block[b++] = (uint8)pPixelSource[2];
+							block[b++] = (uint8)pPixelSource[1];
+							block[b++] = (uint8)pPixelSource[0];
+							block[b++] = 0;
+						}
+					}
+
+					*rowTo = physical_block_zero;
+					astcrt::compress_block_rgba((astcrt::unorm8_t*)block, rowTo);
+
+					++rowTo;
+				}
+			});
+	}
+
+
+	//---------------------------------------------------------------------------------------------
+	void ASTC4x4RGBAL_to_RGBA(uint32 sx, uint32 sy, const uint8* from, uint8* to)
+	{
+		init_astc_decompress();
+
+		// actually use arm's decompression
+		arm::symbolic_compressed_block symbolic;
+
+		for (uint32 y = 0; y < sy; y += 4)
+		{
+			for (uint32 x = 0; x < sx; x += 4)
+			{
+				arm::physical_compressed_block physical;
+				FMemory::Memcpy(physical.data, from, 16);
+
+				//                if (x==8 && y==8)
+				//                    s_debuglog = true;
+
+				physical_to_symbolic(physical, &symbolic);
+
+				arm::imageblock block;
+				decompress_symbolic_block(arm::DECODE_LDR, &symbolic, &block);
+
+				//                s_debuglog = false;
+
+				// Copy data from block
+				const float* fptr = block.orig_data;
+				float data[4];
+				for (uint32 py = 0; py < 4; py++)
+				{
+					for (uint32 px = 0; px < 4; px++)
+					{
+						uint32 xi = x + px;
+						uint32 yi = y + py;
+
+						if (xi < sx && yi < sy)
+						{
+							uint8* toPixel = to + (yi * sx + xi) * 4;
+
+							data[0] = FMath::Min(1.0f, fptr[0]);
+							data[1] = FMath::Min(1.0f, fptr[1]);
+							data[2] = FMath::Min(1.0f, fptr[2]);
+							data[3] = FMath::Min(1.0f, fptr[3]);
+
+							// pack the data
+							int ri = static_cast <int>(floor(data[0] * 255.0f + 0.5f));
+							int gi = static_cast <int>(floor(data[1] * 255.0f + 0.5f));
+							int bi = static_cast <int>(floor(data[2] * 255.0f + 0.5f));
+							int ai = static_cast <int>(floor(data[3] * 255.0f + 0.5f));
+
+							toPixel[0] = (uint8)ri;
+							toPixel[1] = (uint8)gi;
+							toPixel[2] = (uint8)bi;
+							toPixel[3] = (uint8)ai;
+						}
+						fptr += 4;
+					}
+				}
+
+				from += 16;
+			}
+		}
+	}
+
+
+
+	//---------------------------------------------------------------------------------------------
+	void ASTC4x4RGBAL_to_RGB(uint32 sx, uint32 sy, const uint8* from, uint8* to)
+	{
+		init_astc_decompress();
+
+		// actually use arm's decompression
+		arm::symbolic_compressed_block symbolic;
+
+		for (uint32 y = 0; y < sy; y += 4)
+		{
+			for (uint32 x = 0; x < sx; x += 4)
+			{
+				arm::physical_compressed_block physical;
+				FMemory::Memcpy(physical.data, from, 16);
+
+				physical_to_symbolic(physical, &symbolic);
+
+				arm::imageblock block;
+				decompress_symbolic_block(arm::DECODE_LDR, &symbolic, &block);
+
+				// Copy data from block
+				const float* fptr = block.orig_data;
+				float data[4];
+				for (uint32 py = 0; py < 4; py++)
+				{
+					for (uint32 px = 0; px < 4; px++)
+					{
+						uint32 xi = x + px;
+						uint32 yi = y + py;
+
+						if (xi < sx && yi < sy)
+						{
+							uint8* toPixel = to + (yi * sx + xi) * 3;
+
+							data[0] = FMath::Min(1.0f, fptr[0]);
+							data[1] = FMath::Min(1.0f, fptr[1]);
+							data[2] = FMath::Min(1.0f, fptr[2]);
+
+							// pack the data
+							int ri = static_cast <int>(floor(data[0] * 255.0f + 0.5f));
+							int gi = static_cast <int>(floor(data[1] * 255.0f + 0.5f));
+							int bi = static_cast <int>(floor(data[2] * 255.0f + 0.5f));
+
+							toPixel[0] = (uint8)ri;
+							toPixel[1] = (uint8)gi;
+							toPixel[2] = (uint8)bi;
+						}
+						fptr += 4;
+					}
+				}
+
+				from += 16;
+			}
+		}
+	}
+
+
+
+	//---------------------------------------------------------------------------------------------
+	void L_to_ASTC4x4RGBL(uint32_t sx, uint32_t sy, const uint8_t* from, uint8_t* to, int q)
+	{
+		// Almost the same as RGB_to_ASTC4x4RGBL
+		(void)q;
+		astcrt::PhysicalBlock physical_block_zero = { {0} };
+		astcrt::PhysicalBlock* dst_re = reinterpret_cast<astcrt::PhysicalBlock*>(to);
+
+		uint32 bx = FMath::DivideAndRoundUp(sx, 4u);
+		uint32 by = FMath::DivideAndRoundUp(sy, 4u);
+
+		//for ( uint32 y=0; y<by; ++y )
+		ParallelFor(by,
+			[
+				bx, sx, sy, from, to, q, &physical_block_zero
+			] (uint32 y)
+			{
+				astcrt::PhysicalBlock* rowTo = reinterpret_cast<astcrt::PhysicalBlock*>(to + sizeof(astcrt::PhysicalBlock) * bx * y);
+				for (uint32 x = 0; x < bx; ++x)
+				{
+					// Fetch the block
+					uint8 block[4 * 4 * 4];
+					int b = 0;
+					for (int j = 0; j < 4; ++j)
+					{
+						// Clamp y in case block goes beyond the image. Biases the result.
+						uint32 rj = FMath::Min(sy - 1, y * 4 + j);
+
+						for (int i = 0; i < 4; ++i)
+						{
+							// Clamp x in case block goes beyond the image. Biases the result.
+							uint32 ri = FMath::Min(sx - 1, x * 4 + i);
+
+							const uint8* pPixelSource = from + (sx * rj + ri);
+
+							block[b++] = pPixelSource[0];
+							block[b++] = pPixelSource[0];
+							block[b++] = pPixelSource[0];
+							block[b++] = 0;
+						}
+					}
+
+					*rowTo = physical_block_zero;
+					astcrt::compress_block((astcrt::unorm8_t*)block, rowTo);
+
+					++rowTo;
+				}
+			});
+	}
+
+	//---------------------------------------------------------------------------------------------
+	//---------------------------------------------------------------------------------------------
+	//---------------------------------------------------------------------------------------------
+	void RGB_to_ASTC4x4RGL(uint32 sx, uint32 sy, const uint8* from, uint8* to, int q)
+	{
+		(void)q;
+		astcrt::PhysicalBlock physical_block_zero = { {0} };
+		astcrt::PhysicalBlock* dst_re = reinterpret_cast<astcrt::PhysicalBlock*>(to);
+
+		uint32 bx = FMath::DivideAndRoundUp(sx, 4u);
+		uint32 by = FMath::DivideAndRoundUp(sy, 4u);
+
+		//for ( uint32 y=0; y<by; ++y )
+		ParallelFor(by,
+			[
+				bx, sx, sy, from, to, q, &physical_block_zero
+			] (uint32 y)
+			{
+				astcrt::PhysicalBlock* rowTo = reinterpret_cast<astcrt::PhysicalBlock*>(to + sizeof(astcrt::PhysicalBlock) * bx * y);
+
+				for (uint32 x = 0; x < bx; ++x)
+				{
+					// Fetch the block
+					uint8 block[4 * 4 * 4];
+					int b = 0;
+					for (int j = 0; j < 4; ++j)
+					{
+						// Clamp y in case block goes beyond the image. Biases the result.
+						uint32 rj = FMath::Min(sy - 1, y * 4 + j);
+
+						for (int i = 0; i < 4; ++i)
+						{
+							// Clamp x in case block goes beyond the image. Biases the result.
+							uint32 ri = FMath::Min(sx - 1, x * 4 + i);
+
+							const uint8* pPixelSource = from + 3 * (sx * rj + ri);
+
+							block[b++] = 0;
+							block[b++] = pPixelSource[1];
+							block[b++] = pPixelSource[0];
+							block[b++] = 0;
+						}
+					}
+
+					*rowTo = physical_block_zero;
+					astcrt::compress_block_rg((astcrt::unorm8_t*)block, rowTo);
+					++rowTo;
+				}
+			});
+	}
+
+
+	//---------------------------------------------------------------------------------------------
+	void ASTC4x4RGBL_to_ASTC4x4RGBAL(uint32 sx, uint32 sy, const uint8* from, uint8* to)
+	{
+		int blockCount = FMath::DivideAndRoundUp(sy, 4u) * FMath::DivideAndRoundUp(sx, 4u);
+		FMemory::Memcpy(to, from, blockCount * 16);
+	}
+
+
+	//---------------------------------------------------------------------------------------------
+	void ASTC4x4RGBAL_to_ASTC4x4RGBL(uint32 sx, uint32 sy, const uint8* from, uint8* to)
+	{
+		int blockCount = FMath::DivideAndRoundUp(sy, 4u) * FMath::DivideAndRoundUp(sx, 4u);
+		FMemory::Memcpy(to, from, blockCount * 16);
+	}
+
+
+	//---------------------------------------------------------------------------------------------
+	void RGBA_to_ASTC4x4RGL(uint32 sx, uint32 sy, const uint8* from, uint8* to, int q)
+	{
+		(void)q;
+		astcrt::PhysicalBlock physical_block_zero = { {0} };
+		astcrt::PhysicalBlock* dst_re = reinterpret_cast<astcrt::PhysicalBlock*>(to);
+
+		uint32 bx = FMath::DivideAndRoundUp(sx, 4u);
+		uint32 by = FMath::DivideAndRoundUp(sy, 4u);
+
+		//for ( uint32 y=0; y<by; ++y )
+		ParallelFor(by,
+			[
+				bx, sx, sy, from, to, q, &physical_block_zero
+			] (uint32 y)
+			{
+				astcrt::PhysicalBlock* rowTo = reinterpret_cast<astcrt::PhysicalBlock*>(to + sizeof(astcrt::PhysicalBlock) * bx * y);
+
+				for (uint32 x = 0; x < bx; ++x)
+				{
+					// Fetch the block
+					uint8 block[4 * 4 * 4];
+					int b = 0;
+					for (int j = 0; j < 4; ++j)
+					{
+						// Clamp y in case block goes beyond the image. Biases the result.
+						uint32 rj = FMath::Min(sy - 1, y * 4 + j);
+
+						for (int i = 0; i < 4; ++i)
+						{
+							// Clamp x in case block goes beyond the image. Biases the result.
+							uint32 ri = FMath::Min(sx - 1, x * 4 + i);
+
+							const uint8* pPixelSource = from + 4 * (sx * rj + ri);
+
+							block[b++] = 0;
+							block[b++] = pPixelSource[1];
+							block[b++] = pPixelSource[0];
+							block[b++] = 0;
+						}
+					}
+
+					*rowTo = physical_block_zero;
+					astcrt::compress_block_rg((astcrt::unorm8_t*)block, rowTo);
+
+					++rowTo;
+				}
+			});
+	}
+
+
+	//---------------------------------------------------------------------------------------------
+	void ASTC4x4RGL_to_RGB(uint32 sx, uint32 sy, const uint8* from, uint8* to)
+	{
+		init_astc_decompress();
+
+		// actually use arm's decompression
+		arm::symbolic_compressed_block symbolic;
+
+		for (uint32 y = 0; y < sy; y += 4)
+		{
+			for (uint32 x = 0; x < sx; x += 4)
+			{
+				arm::physical_compressed_block physical;
+				FMemory::Memcpy(physical.data, from, 16);
+
+				//                if (x==476 && y==256)
+				//                    s_debuglog = true;
+
+				arm::physical_to_symbolic(physical, &symbolic);
+
+				arm::imageblock block;
+				arm::decompress_symbolic_block(arm::DECODE_LDR,
+					&symbolic,
+					&block);
+
+				//                s_debuglog = false;
+
+								// Copy data from block
+				const float* fptr = block.orig_data;
+				float data[4];
+				for (uint32 py = 0; py < 4; py++)
+				{
+					for (uint32 px = 0; px < 4; px++)
+					{
+						uint32 xi = x + px;
+						uint32 yi = y + py;
+
+						if (xi < sx && yi < sy)
+						{
+							uint8* toPixel = to + (yi * sx + xi) * 3;
+
+							data[0] = FMath::Min(1.0f, fptr[0]);
+							data[3] = FMath::Min(1.0f, fptr[3]);
+
+							// pack the data
+							int ri = static_cast <int>(floor(data[0] * 255.0f + 0.5f));
+							int gi = static_cast <int>(floor(data[3] * 255.0f + 0.5f));
+
+							toPixel[0] = (uint8)ri;
+							toPixel[1] = (uint8)gi;
+							toPixel[2] = 0;
+						}
+						fptr += 4;
+					}
+				}
+
+				from += 16;
+			}
+		}
+	}
+
+
+	//---------------------------------------------------------------------------------------------
+	void ASTC4x4RGL_to_RGBA(uint32 sx, uint32 sy, const uint8* from, uint8* to)
+	{
+		init_astc_decompress();
+
+		// actually use arm's decompression
+		arm::symbolic_compressed_block symbolic;
+
+		for (uint32 y = 0; y < sy; y += 4)
+		{
+			for (uint32 x = 0; x < sx; x += 4)
+			{
+				arm::physical_compressed_block physical;
+				FMemory::Memcpy(physical.data, from, 16);
+
+				arm::physical_to_symbolic(physical, &symbolic);
+
+				arm::imageblock block;
+				arm::decompress_symbolic_block(arm::DECODE_LDR, &symbolic, &block);
+
+				// Copy data from block
+				const float* fptr = block.orig_data;
+				float data[4];
+				for (uint32 py = 0; py < 4; py++)
+				{
+					for (uint32 px = 0; px < 4; px++)
+					{
+						uint32 xi = x + px;
+						uint32 yi = y + py;
+
+						if (xi < sx && yi < sy)
+						{
+							uint8* toPixel = to + (yi * sx + xi) * 4;
+
+							data[0] = FMath::Min(1.0f, fptr[0]);
+							data[3] = FMath::Min(1.0f, fptr[3]);
+
+							// pack the data
+							int ri = static_cast <int>(floor(data[0] * 255.0f + 0.5f));
+							int gi = static_cast <int>(floor(data[3] * 255.0f + 0.5f));
+
+							toPixel[0] = (uint8)ri;
+							toPixel[1] = (uint8)gi;
+							toPixel[2] = 0;
+							toPixel[3] = 0;
+						}
+						fptr += 4;
+					}
+				}
+
+				from += 16;
+			}
+		}
+	}
+
+#endif // MIRO_INCLUDE_ASTC
+
+
+}
+
