[mutable] Move Mutable plugin from the restricted folder to the experimental plugins folder.

diff --git a/Engine/Plugins/Experimental/Mutable/Source/MutableRuntime/Private/BlockCompression/Miro/Miro.h b/Engine/Plugins/Experimental/Mutable/Source/MutableRuntime/Private/BlockCompression/Miro/Miro.h
new file mode 100644
--- /dev/null
+++ b/Engine/Plugins/Experimental/Mutable/Source/MutableRuntime/Private/BlockCompression/Miro/Miro.h
@@ -0,0 +1,107 @@
+// Copyright Epic Games, Inc. All Rights Reserved.
+
+#pragma once
+
+#include <cstdint>
+
+// Configure the compression library with these defines:
+
+// Support for BC1 to BC7, also known as DXTC1 to 5
+#ifndef MIRO_INCLUDE_BC
+#define MIRO_INCLUDE_BC                     1
+#endif
+
+// Support for ASTC
+#ifndef MIRO_INCLUDE_ASTC
+#define MIRO_INCLUDE_ASTC                   1
+#endif
+
+// Exclude high quality compression code. The index corresponds to the quality levels.
+#ifndef MIRO_EXCLUDE_QUALITY_ABOVE
+#define MIRO_EXCLUDE_QUALITY_ABOVE          4
+#endif
+
+
+namespace miro
+{
+    // MIRO is a runtime texture compression library for hardware-accelerated formats.
+
+    // Compression quality levels (q):
+    // 0 - Fastest for runtime
+    // 1 - Best for runtime
+    // 2 - Fast for tools
+    // 3 - Best for tools
+    // 4 - Maximum, with no time limits.
+
+
+    extern void initialize();
+    extern void finalize();
+
+    //!
+    extern const char* GetVersion();
+
+#if MIRO_INCLUDE_BC
+
+    //! BC1 support
+    extern void RGB_to_BC1( uint32_t sx, uint32_t sy, const uint8_t* from, uint8_t* to, int q );
+    extern void RGBA_to_BC1( uint32_t sx, uint32_t sy, const uint8_t* from, uint8_t* to, int q );
+    extern void BC1_to_RGB( uint32_t sx, uint32_t sy, const uint8_t* from, uint8_t* to );
+    extern void BC1_to_RGBA( uint32_t sx, uint32_t sy, const uint8_t* from, uint8_t* to );
+
+    //! BC2 support
+    extern void RGB_to_BC2( uint32_t sx, uint32_t sy, const uint8_t* from, uint8_t* to, int q );
+    extern void RGBA_to_BC2( uint32_t sx, uint32_t sy, const uint8_t* from, uint8_t* to, int q );
+    extern void BC2_to_RGB( uint32_t sx, uint32_t sy, const uint8_t* from, uint8_t* to );
+    extern void BC2_to_RGBA( uint32_t sx, uint32_t sy, const uint8_t* from, uint8_t* to );
+
+    //! BC3 support
+    extern void RGB_to_BC3( uint32_t sx, uint32_t sy, const uint8_t* from, uint8_t* to, int q );
+    extern void RGBA_to_BC3( uint32_t sx, uint32_t sy, const uint8_t* from, uint8_t* to, int q );
+    extern void BC3_to_RGB( uint32_t sx, uint32_t sy, const uint8_t* from, uint8_t* to );
+    extern void BC3_to_RGBA( uint32_t sx, uint32_t sy, const uint8_t* from, uint8_t* to );
+    extern void BC1_to_BC3( uint32_t sx, uint32_t sy, const uint8_t* from, uint8_t* to, int q );
+
+    //! BC4 support
+    extern void L_to_BC4( uint32_t sx, uint32_t sy, const uint8_t* from, uint8_t* to, int q );
+    extern void BC4_to_L( uint32_t sx, uint32_t sy, const uint8_t* from, uint8_t* to );
+    extern void BC4_to_RGB( uint32_t sx, uint32_t sy, const uint8_t* from, uint8_t* to );
+    extern void BC4_to_RGBA( uint32_t sx, uint32_t sy, const uint8_t* from, uint8_t* to );
+
+    //! BC5 support
+    extern void RGBA_to_BC5( uint32_t sx, uint32_t sy, const uint8_t* from, uint8_t* to, int q );
+    extern void RGB_to_BC5( uint32_t sx, uint32_t sy, const uint8_t* from, uint8_t* to, int q );
+    extern void BC5_to_RGBA( uint32_t sx, uint32_t sy, const uint8_t* from, uint8_t* to );
+    extern void BC5_to_RGB( uint32_t sx, uint32_t sy, const uint8_t* from, uint8_t* to );
+
+#endif
+
+
+#if MIRO_INCLUDE_ASTC
+
+    //! ASTC 4x4 RGBA low support
+    extern void RGBA_to_ASTC4x4RGBAL( uint32_t sx, uint32_t sy, const uint8_t* from, uint8_t* to, int q );
+    extern void RGB_to_ASTC4x4RGBAL( uint32_t sx, uint32_t sy, const uint8_t* from, uint8_t* to, int q );
+    extern void ASTC4x4RGBAL_to_RGBA( uint32_t sx, uint32_t sy, const uint8_t* from, uint8_t* to );
+    extern void ASTC4x4RGBAL_to_RGB( uint32_t sx, uint32_t sy, const uint8_t* from, uint8_t* to );
+
+    //! ASTC 4x4 RGB low support
+    extern void RGB_to_ASTC4x4RGBL( uint32_t sx, uint32_t sy, const uint8_t* from, uint8_t* to, int q );
+    extern void RGBA_to_ASTC4x4RGBL( uint32_t sx, uint32_t sy, const uint8_t* from, uint8_t* to, int q );
+    extern void ASTC4x4RGBL_to_RGB( uint32_t sx, uint32_t sy, const uint8_t* from, uint8_t* to );
+    extern void ASTC4x4RGBL_to_RGBA( uint32_t sx, uint32_t sy, const uint8_t* from, uint8_t* to );
+	extern void L_to_ASTC4x4RGBL(uint32_t sx, uint32_t sy, const uint8_t* from, uint8_t* to, int q);
+    
+	//! Convenience transforms
+    extern void ASTC4x4RGBL_to_ASTC4x4RGBAL( uint32_t sx, uint32_t sy, const uint8_t* from, uint8_t* to );
+    extern void ASTC4x4RGBAL_to_ASTC4x4RGBL( uint32_t sx, uint32_t sy, const uint8_t* from, uint8_t* to );
+
+    //! ASTC 4x4 RG low support
+    extern void RGB_to_ASTC4x4RGL( uint32_t sx, uint32_t sy, const uint8_t* from, uint8_t* to, int q );
+    extern void RGBA_to_ASTC4x4RGL( uint32_t sx, uint32_t sy, const uint8_t* from, uint8_t* to, int q );
+    extern void ASTC4x4RGL_to_RGB( uint32_t sx, uint32_t sy, const uint8_t* from, uint8_t* to );
+    extern void ASTC4x4RGL_to_RGBA( uint32_t sx, uint32_t sy, const uint8_t* from, uint8_t* to );
+
+#endif
+
+}
+
