[mutable] Move Mutable plugin from the restricted folder to the experimental plugins folder.

diff --git a/Engine/Plugins/Experimental/Mutable/Source/MutableRuntime/Private/CodeRunner.cpp b/Engine/Plugins/Experimental/Mutable/Source/MutableRuntime/Private/CodeRunner.cpp
new file mode 100644
--- /dev/null
+++ b/Engine/Plugins/Experimental/Mutable/Source/MutableRuntime/Private/CodeRunner.cpp
@@ -0,0 +1,5827 @@
+// Copyright Epic Games, Inc. All Rights Reserved.
+
+#include "CodeRunner.h"
+
+#include "Platform.h"
+
+#include "CodeRunner.h"
+
+#include "SystemPrivate.h"
+#include "SettingsPrivate.h"
+#include "ModelPrivate.h"
+#include "InstancePrivate.h"
+#include "ImagePrivate.h"
+#include "MeshPrivate.h"
+#include "MemoryPrivate.h"
+
+#include "Operations.h"
+
+#include "OpImageResize.h"
+#include "OpImagePixelFormat.h"
+#include "OpImageMipmap.h"
+#include "OpImageSoftLight.h"
+#include "OpImageHardLight.h"
+#include "OpImageBurn.h"
+#include "OpImageDodge.h"
+#include "OpImageScreen.h"
+#include "OpImageOverlay.h"
+#include "OpImageAlphaOverlay.h"
+#include "OpImageCompose.h"
+#include "OpImageBlend.h"
+#include "OpImageDifference.h"
+#include "OpImageInterpolate.h"
+#include "OpImageMultiply.h"
+#include "OpImageSaturate.h"
+#include "OpImageLuminance.h"
+#include "OpImageSwizzle.h"
+#include "OpImageSelectColour.h"
+#include "OpImageColourMap.h"
+#include "OpImageGradient.h"
+#include "OpImageBinarise.h"
+#include "OpImageCrop.h"
+#include "OpImageRasterMesh.h"
+#include "OpImageGrow.h"
+#include "OpImageProject.h"
+#include "OpImageDisplace.h"
+#include "OpImageInvert.h"
+#include "OpImageNormalCombine.h"
+#include "OpImageApplyComposite.h"
+#include "OpImageTransform.h"
+
+#include "OpMeshApplyLayout.h"
+#include "OpMeshDifference.h"
+#include "OpMeshMorph.h"
+#include "OpMeshInterpolate.h"
+#include "OpMeshChartDifference.h"
+#include "OpMeshSubtract.h"
+#include "OpMeshRemoveChart.h"
+#include "OpMeshFormat.h"
+#include "OpMeshMerge.h"
+#include "OpMeshExtractLayoutBlock.h"
+#include "OpMeshTransform.h"
+#include "OpMeshClipMorphPlane.h"
+#include "OpMeshClipWithMesh.h"
+#include "OpMeshApplyPose.h"
+#include "OpMeshRemapIndices.h"
+#include "OpMeshGeometryOperation.h"
+#include "OpMeshReshape.h"
+#include "OpMeshClipDeform.h"
+
+#include "OpLayoutPack.h"
+#include "OpLayoutRemoveBlocks.h"
+
+#include "OpEvalCurve.h"
+
+#include "Math/UnrealMathUtility.h"
+
+namespace mu
+{
+
+    //---------------------------------------------------------------------------------------------
+    CodeRunner::CodeRunner(const SettingsPtrConst& pSettings,
+		class System::Private* s,
+		const Model* pModel,
+		const Parameters* pParams,
+		OP::ADDRESS at,
+		uint32 lodMask, uint8 executionOptions, SCHEDULED_OP::EType Type )
+		: m_pSettings(pSettings), m_pSystem(s), m_pModel(pModel), m_pParams(pParams), m_lodMask(lodMask)
+	{
+		ScheduledStagePerOp.resize(m_pModel->GetPrivate()->m_program.m_opAddress.size());
+
+		// We will read this in the end, so make sure we keep it.
+		m_pSystem->m_memory->IncreaseHitCount(CACHE_ADDRESS(at, 0, executionOptions));
+
+		PROGRAM& program = m_pModel->GetPrivate()->m_program;
+		m_romPendingOps.resize(program.m_roms.Num(), 0);
+
+		// Push the root operation
+		SCHEDULED_OP rootOp;
+		rootOp.at = at;
+		rootOp.executionOptions = executionOptions;
+		rootOp.Type = Type;
+		AddOp(rootOp);
+	}
+
+
+    //---------------------------------------------------------------------------------------------
+	FProgramCache& CodeRunner::GetMemory()
+    {
+		return *m_pSystem->m_memory;
+	}
+
+
+    //---------------------------------------------------------------------------------------------
+    ImagePtr CodeRunner::LoadExternalImage( EXTERNAL_IMAGE_ID id )
+    {
+		MUTABLE_CPUPROFILER_SCOPE(LoadExternalImage);
+
+		check(m_pSystem);
+
+		ImagePtr pResult;
+
+		if (m_pSystem->m_pImageParameterGenerator)
+		{
+			pResult = m_pSystem->m_pImageParameterGenerator->GetImage(id);
+			// Don't cache for now. Need to figure out how to invalidate them.
+			//m_pSystem->m_externalImages.push_back( pair<EXTERNAL_IMAGE_ID,ImagePtr>(id,pResult) );
+		}
+		else
+		{
+			// Not found and there is no generator!
+		}
+
+		return pResult;
+	}
+
+
+    //---------------------------------------------------------------------------------------------
+    void CodeRunner::RunCode_Conditional( SCHEDULED_OP& item,
+                                          const Model* pModel
+                                          )
+    {
+        //MUTABLE_CPUPROFILER_SCOPE(RunCode_Conditional);
+
+		OP_TYPE type = pModel->GetPrivate()->m_program.GetOpType(item.at);
+        auto args = pModel->GetPrivate()->m_program.GetOpArgs<OP::ConditionalArgs>(item.at);
+
+        // Conditionals have the following execution stages:
+        // 0: we need to run the condition
+        // 1: we need to run the branch
+        // 2: we need to fetch the result and store it in this op
+        switch( item.stage )
+        {
+        case 0:
+        {
+            AddOp( SCHEDULED_OP( item.at,item,1 ),
+                   SCHEDULED_OP( args.condition, item ) );
+            break;
+        }
+
+        case 1:
+        {
+            // Get the condition result
+
+            // If there is no expression, we'll assume true.
+            bool value = true;
+            value = GetMemory().GetBool( CACHE_ADDRESS(args.condition, item.executionIndex, item.executionOptions) );
+
+            OP::ADDRESS resultAt = value ? args.yes : args.no;
+
+            // Schedule the end of this instruction if necessary
+            AddOp( SCHEDULED_OP( item.at, item, 2, (uint32_t)value),
+				SCHEDULED_OP( resultAt, item) );
+
+            break;
+        }
+
+        case 2:
+        {
+            OP::ADDRESS resultAt = item.customState ? args.yes : args.no;
+
+            // Store the final result
+            CACHE_ADDRESS cat( item );
+            CACHE_ADDRESS rat( resultAt, item );
+            switch (GetOpDataType(type))
+            {
+            case DT_BOOL:       GetMemory().SetBool( cat, GetMemory().GetBool(rat) ); break;
+            case DT_INT:        GetMemory().SetInt( cat, GetMemory().GetInt(rat) ); break;
+            case DT_SCALAR:     GetMemory().SetScalar( cat, GetMemory().GetScalar(rat) ); break;
+			case DT_STRING:		GetMemory().SetString( cat, GetMemory().GetString( rat ) ); break;
+            case DT_COLOUR:		GetMemory().SetColour( cat, GetMemory().GetColour( rat ) ); break;
+            case DT_PROJECTOR:  GetMemory().SetProjector( cat, GetMemory().GetProjector(rat) ); break;
+            case DT_MESH:       GetMemory().SetMesh( cat, GetMemory().GetMesh(rat) ); break;
+            case DT_IMAGE:      GetMemory().SetImage( cat, GetMemory().GetImage(rat) ); break;
+            case DT_LAYOUT:     GetMemory().SetLayout( cat, GetMemory().GetLayout(rat) ); break;
+            case DT_INSTANCE:   GetMemory().SetInstance( cat, GetMemory().GetInstance(rat) ); break;
+            default:
+                // Not implemented
+                check( false );
+            }
+
+            break;
+        }
+
+        default:
+            check(false);
+        }
+    }
+
+
+	//---------------------------------------------------------------------------------------------
+	void CodeRunner::RunCode_Switch(SCHEDULED_OP& item,
+		const Model* pModel
+	)
+	{
+		OP_TYPE type = pModel->GetPrivate()->m_program.GetOpType(item.at);
+
+		const uint8_t* data = pModel->GetPrivate()->m_program.GetOpArgsPointer(item.at);
+
+		OP::ADDRESS VarAddress;
+		FMemory::Memcpy(&VarAddress, data, sizeof(OP::ADDRESS));
+		data += sizeof(OP::ADDRESS);
+
+		OP::ADDRESS DefAddress;
+		FMemory::Memcpy(&DefAddress, data, sizeof(OP::ADDRESS));
+		data += sizeof(OP::ADDRESS);
+
+		uint32_t CaseCount;
+		FMemory::Memcpy(&CaseCount, data, sizeof(uint32_t));
+		data += sizeof(uint32_t);
+
+		switch (item.stage)
+		{
+		case 0:
+		{
+			if (VarAddress)
+			{
+				AddOp(SCHEDULED_OP(item.at, item, 1),
+					SCHEDULED_OP(VarAddress, item));
+			}
+			else
+			{
+				switch (GetOpDataType(type))
+				{
+				case DT_BOOL:       GetMemory().SetBool(item, false); break;
+				case DT_INT:        GetMemory().SetInt(item, 0); break;
+				case DT_SCALAR:		GetMemory().SetScalar(item, 0.0f); break;
+				case DT_STRING:		GetMemory().SetString(item, nullptr); break;
+				case DT_COLOUR:		GetMemory().SetColour(item, vec4f()); break;
+				case DT_PROJECTOR:  GetMemory().SetProjector(item, nullptr); break;
+				case DT_MESH:       GetMemory().SetMesh(item, nullptr); break;
+				case DT_IMAGE:      GetMemory().SetImage(item, nullptr); break;
+				case DT_LAYOUT:     GetMemory().SetLayout(item, nullptr); break;
+				case DT_INSTANCE:   GetMemory().SetInstance(item, nullptr); break;
+				default:
+					// Not implemented
+					check(false);
+				}
+			}
+			break;
+		}
+
+		case 1:
+		{
+			// Get the variable result
+			int var = GetMemory().GetInt(CACHE_ADDRESS(VarAddress, item));
+
+			OP::ADDRESS valueAt = DefAddress;
+			for (uint32_t C = 0; C < CaseCount; ++C)
+			{
+				int32_t Condition;
+				FMemory::Memcpy(&Condition, data, sizeof(int32_t));
+				data += sizeof(int32_t);
+
+				OP::ADDRESS At;
+				FMemory::Memcpy(&At, data, sizeof(OP::ADDRESS));
+				data += sizeof(OP::ADDRESS);
+
+				if (At && var == (int)Condition)
+				{
+					valueAt = At;
+					break; 
+				}
+			}
+
+            // Schedule the end of this instruction if necessary
+            AddOp( SCHEDULED_OP( item.at, item, 2, valueAt ),
+				   SCHEDULED_OP( valueAt, item ) );
+
+            break;
+        }
+
+        case 2:
+        {
+			OP::ADDRESS resultAt = OP::ADDRESS(item.customState);
+
+            // Store the final result
+            CACHE_ADDRESS cat( item );
+            CACHE_ADDRESS rat( resultAt, item );
+            switch (GetOpDataType(type))
+            {
+            case DT_BOOL:       GetMemory().SetBool( cat, GetMemory().GetBool(rat) ); break;
+            case DT_INT:        GetMemory().SetInt( cat, GetMemory().GetInt(rat) ); break;
+            case DT_SCALAR:     GetMemory().SetScalar( cat, GetMemory().GetScalar(rat) ); break;
+            case DT_STRING:		GetMemory().SetString( cat, GetMemory().GetString( rat ) ); break;
+            case DT_COLOUR:		GetMemory().SetColour( cat, GetMemory().GetColour( rat ) ); break;
+            case DT_PROJECTOR:  GetMemory().SetProjector( cat, GetMemory().GetProjector(rat) ); break;
+            case DT_MESH:       GetMemory().SetMesh( cat, GetMemory().GetMesh(rat) ); break;
+            case DT_IMAGE:      GetMemory().SetImage( cat, GetMemory().GetImage(rat) ); break;
+            case DT_LAYOUT:     GetMemory().SetLayout( cat, GetMemory().GetLayout(rat) ); break;
+            case DT_INSTANCE:   GetMemory().SetInstance( cat, GetMemory().GetInstance(rat) ); break;
+            default:
+                // Not implemented
+                check( false );
+            }
+
+            break;
+        }
+
+        default:
+            check(false);
+        }
+    }
+
+
+    //---------------------------------------------------------------------------------------------
+    void CodeRunner::RunCode_Instance( SCHEDULED_OP& item,
+                                  const Model* pModel,
+                                  uint32_t lodMask
+                                  )
+    {
+        //MUTABLE_CPUPROFILER_SCOPE(RunCode_Instance);
+
+		OP_TYPE type = pModel->GetPrivate()->m_program.GetOpType(item.at);
+        switch (type)
+        {
+
+        case OP_TYPE::IN_ADDVECTOR:
+        {
+            auto args = pModel->GetPrivate()->m_program.GetOpArgs<OP::InstanceAddArgs>(item.at);
+
+            switch (item.stage)
+            {
+            case 0:
+                    AddOp( SCHEDULED_OP( item.at, item, 1),
+                           SCHEDULED_OP( args.instance, item),
+                           SCHEDULED_OP( args.value, item) );
+                break;
+
+            case 1:
+            {
+                InstancePtrConst pBase = GetMemory().GetInstance( CACHE_ADDRESS(args.instance,item) );
+                InstancePtr pResult;
+                if (!pBase)
+                {
+                    pResult = new Instance();
+                }
+                else
+                {
+                    pResult = pBase->Clone();
+                }
+
+                if ( args.value )
+                {
+                    vec4f value = GetMemory().GetColour( CACHE_ADDRESS(args.value,item) );
+
+                    OP::ADDRESS nameAd = args.name;
+                    check(  nameAd < pModel->GetPrivate()->m_program.m_constantStrings.size() );
+                    const char* strName = pModel->GetPrivate()->m_program.m_constantStrings[ nameAd ].c_str();
+
+                    pResult->GetPrivate()->AddVector( 0, 0, 0, value, strName );
+                }
+                GetMemory().SetInstance( item, pResult );
+                break;
+            }
+
+            default:
+                check(false);
+            }
+
+            break;
+        }
+
+        case OP_TYPE::IN_ADDSCALAR:
+        {
+            auto args = pModel->GetPrivate()->m_program.GetOpArgs<OP::InstanceAddArgs>(item.at);
+            switch (item.stage)
+            {
+            case 0:
+                    AddOp( SCHEDULED_OP( item.at, item, 1),
+                           SCHEDULED_OP( args.instance, item),
+                           SCHEDULED_OP( args.value, item) );
+                break;
+
+            case 1:
+            {
+                InstancePtrConst pBase = GetMemory().GetInstance( CACHE_ADDRESS(args.instance,item) );
+                InstancePtr pResult;
+                if (!pBase)
+                {
+                    pResult = new Instance();
+                }
+                else
+                {
+                    pResult = pBase->Clone();
+                }
+
+                if ( args.value )
+                {
+                    float value = GetMemory().GetScalar( CACHE_ADDRESS(args.value,item) );
+
+                    OP::ADDRESS nameAd = args.name;
+                    check(  nameAd < pModel->GetPrivate()->m_program.m_constantStrings.size() );
+                    const char* strName = pModel->GetPrivate()->m_program.m_constantStrings[ nameAd ].c_str();
+
+                    pResult->GetPrivate()->AddScalar( 0, 0, 0, value, strName );
+                }
+                GetMemory().SetInstance( item, pResult );
+                break;
+            }
+
+            default:
+                check(false);
+            }
+
+            break;
+        }
+
+        case OP_TYPE::IN_ADDSTRING:
+        {
+            auto args = pModel->GetPrivate()->m_program.GetOpArgs<OP::InstanceAddArgs>( item.at );
+            switch ( item.stage )
+            {
+            case 0:
+                    AddOp( SCHEDULED_OP( item.at, item, 1 ), SCHEDULED_OP( args.instance, item ),
+                           SCHEDULED_OP( args.value, item ) );
+                break;
+
+            case 1:
+            {
+                InstancePtrConst pBase =
+                    GetMemory().GetInstance( CACHE_ADDRESS( args.instance, item ) );
+                InstancePtr pResult;
+                if ( !pBase )
+                {
+                    pResult = new Instance();
+                }
+                else
+                {
+                    pResult = pBase->Clone();
+                }
+
+                if ( args.value )
+                {
+                    Ptr<const String> value =
+                        GetMemory().GetString( CACHE_ADDRESS( args.value, item ) );
+
+                    OP::ADDRESS nameAd = args.name;
+                    check( nameAd <
+                                    pModel->GetPrivate()->m_program.m_constantStrings.size() );
+                    const char* strName =
+                        pModel->GetPrivate()->m_program.m_constantStrings[nameAd].c_str();
+
+                    pResult->GetPrivate()->AddString( 0, 0, 0, value->GetValue(), strName );
+                }
+                GetMemory().SetInstance( item, pResult );
+                break;
+            }
+
+            default:
+                check( false );
+            }
+
+            break;
+        }
+
+        case OP_TYPE::IN_ADDCOMPONENT:
+        {
+            auto args = pModel->GetPrivate()->m_program.GetOpArgs<OP::InstanceAddArgs>(item.at);
+            switch (item.stage)
+            {
+            case 0:
+                    AddOp( SCHEDULED_OP( item.at, item, 1),
+                           SCHEDULED_OP( args.instance, item),
+                           SCHEDULED_OP( args.value, item) );
+                break;
+
+            case 1:
+            {
+                InstancePtrConst pBase = GetMemory().GetInstance( CACHE_ADDRESS(args.instance,item) );
+                InstancePtr pResult;
+                if (!pBase)
+                {
+                    pResult = new Instance();
+                }
+                else
+                {
+                    pResult = pBase->Clone();
+                }
+
+                if ( args.value )
+                {
+                    InstancePtrConst pComp = GetMemory().GetInstance( CACHE_ADDRESS(args.value,item) );
+
+                    int cindex = pResult->GetPrivate()->AddComponent( 0 );
+
+                    if ( !pComp->GetPrivate()->m_lods.empty()
+                         &&
+                         !pResult->GetPrivate()->m_lods.empty()
+                         &&
+                         !pComp->GetPrivate()->m_lods[0].m_components.empty() )
+                    {
+                        pResult->GetPrivate()->m_lods[0].m_components[cindex] =
+                                pComp->GetPrivate()->m_lods[0].m_components[0];
+
+                    	pResult->GetPrivate()->m_lods[0].m_components[cindex].m_id = args.id;
+                    	
+                        // Name
+                        OP::ADDRESS nameAd = args.name;
+                        check( nameAd < pModel->GetPrivate()->m_program.m_constantStrings.size() );
+                        const char* strName = pModel->GetPrivate()->m_program.m_constantStrings[ nameAd ].c_str();
+                        pResult->GetPrivate()->SetComponentName( 0, cindex, strName );
+                    }
+                }
+                GetMemory().SetInstance( item, pResult );
+                break;
+            }
+
+            default:
+                check(false);
+            }
+
+            break;
+        }
+
+        case OP_TYPE::IN_ADDSURFACE:
+        {
+            auto args = pModel->GetPrivate()->m_program.GetOpArgs<OP::InstanceAddArgs>(item.at);
+            switch (item.stage)
+            {
+            case 0:
+                    AddOp( SCHEDULED_OP( item.at, item, 1),
+                           SCHEDULED_OP( args.instance, item),
+                           SCHEDULED_OP( args.value, item) );
+                break;
+
+            case 1:
+            {
+                InstancePtrConst pBase = GetMemory().GetInstance( CACHE_ADDRESS(args.instance,item) );
+
+                InstancePtr pResult;
+                if (pBase) pResult = pBase->Clone();
+                else pResult = new Instance();
+
+                // Empty surfaces are ok, they still need to be created, because they may contain
+                // additional information like internal or external IDs
+                //if ( args.value )
+                {
+                    InstancePtrConst pSurf = GetMemory().GetInstance( CACHE_ADDRESS(args.value,item) );
+
+                    int sindex = pResult->GetPrivate()->AddSurface( 0, 0 );
+
+                    // Surface data
+                    if (pSurf
+                            &&
+                            pSurf->GetPrivate()->m_lods.size()
+                            &&
+                            pSurf->GetPrivate()->m_lods[0].m_components.size()
+                            &&
+                            pSurf->GetPrivate()->m_lods[0].m_components[0].m_surfaces.size())
+                    {
+                        pResult->GetPrivate()->m_lods[0].m_components[0].m_surfaces[sindex] =
+                            pSurf->GetPrivate()->m_lods[0].m_components[0].m_surfaces[0];
+
+                        // Meshes must be added later.
+                        check(!pSurf->GetPrivate()->m_lods[0].m_components[0].m_meshes.size());
+                    }
+
+                    // Name
+                    OP::ADDRESS nameAd = args.name;
+                    check( nameAd < pModel->GetPrivate()->m_program.m_constantStrings.size() );
+                    const char* strName = pModel->GetPrivate()->m_program.m_constantStrings[ nameAd ].c_str();
+                    pResult->GetPrivate()->SetSurfaceName( 0, 0, sindex, strName );
+
+                    // IDs
+                    pResult->GetPrivate()->m_lods[0].m_components[0].m_surfaces[sindex].m_internalID = args.id;
+                    pResult->GetPrivate()->m_lods[0].m_components[0].m_surfaces[sindex].m_customID = args.externalId;
+                }
+                GetMemory().SetInstance( item, pResult );
+                break;
+            }
+
+            default:
+                check(false);
+            }
+
+            break;
+        }
+
+        case OP_TYPE::IN_ADDLOD:
+        {
+            auto args = pModel->GetPrivate()->m_program.GetOpArgs<OP::InstanceAddLODArgs>(item.at);
+            switch (item.stage)
+            {
+            case 0:
+            {                
+                    vector<SCHEDULED_OP> deps;
+                    for ( int i=0; i<MUTABLE_OP_MAX_ADD_COUNT; ++i )
+                    {
+                        if ( args.lod[i] )
+                        {
+                            bool selectedLod = ( (1<<i) & lodMask ) != 0;
+
+                            if ( selectedLod )
+                            {
+                                deps.emplace_back(args.lod[i], item);
+                            }
+                        }
+                    }
+
+                    AddOp( SCHEDULED_OP( item.at,item, 1), deps );
+
+                break;
+            }
+
+            case 1:
+            {
+                // Assemble result
+                InstancePtr pResult = new Instance();
+
+                for ( int i=0; i<MUTABLE_OP_MAX_ADD_COUNT; ++i )
+                {
+                    if ( args.lod[i] )
+                    {
+                        bool selectedLod = ( (1<<i) & lodMask ) != 0;
+
+                        if ( selectedLod )
+                        {
+                            InstancePtrConst pLOD = GetMemory().GetInstance( CACHE_ADDRESS(args.lod[i],item) );
+
+                            int lindex = pResult->GetPrivate()->AddLOD();
+
+                            // In a degenerated case, the returned pLOD may not have an LOD inside
+                            if ( pLOD && !pLOD->GetPrivate()->m_lods.empty() )
+                            {
+                                pResult->GetPrivate()->m_lods[lindex] = pLOD->GetPrivate()->m_lods[0];
+                            }
+                        }
+                        else
+                        {
+                            // LOD not selected. Add an empty one
+                            pResult->GetPrivate()->AddLOD();
+                        }
+                    }
+                }
+
+                GetMemory().SetInstance( item, pResult );
+                break;
+            }
+
+            default:
+                check(false);
+            }
+
+            break;
+        }
+
+        default:
+                check(false);
+        }
+    }
+
+
+    //---------------------------------------------------------------------------------------------
+    void CodeRunner::RunCode_InstanceAddResource( SCHEDULED_OP& item,
+                                  const Model* pModel,
+                                  const Parameters* pParams
+                                  )
+    {
+		MUTABLE_CPUPROFILER_SCOPE(RunCode_InstanceAddResource);
+
+		OP_TYPE type = pModel->GetPrivate()->m_program.GetOpType(item.at);
+        switch (type)
+        {
+
+        case OP_TYPE::IN_ADDMESH:
+        {
+            auto args = pModel->GetPrivate()->m_program.GetOpArgs<OP::InstanceAddArgs>(item.at);
+            switch (item.stage)
+            {
+            case 0:
+                    AddOp( SCHEDULED_OP( item.at, item, 1),
+                           SCHEDULED_OP( args.instance, item) );
+
+                // We don't build the resources when building instance: just store ids for them.
+                //PushIfNotVisited(args.value, item);
+                break;
+
+            case 1:
+            {
+                InstancePtrConst pBase = GetMemory().GetInstance( CACHE_ADDRESS(args.instance,item) );
+                InstancePtr pResult;
+                if (!pBase)
+                {
+                    pResult = new Instance();
+                }
+                else
+                {
+                    pResult = pBase->Clone();
+                }
+
+                if ( args.value )
+                {
+                    RESOURCE_ID meshId = pModel->GetPrivate()->GetResourceKey(
+                                args.relevantParametersListIndex,
+                                args.value,
+                                pParams);
+                    OP::ADDRESS nameAd = args.name;
+                    check(  nameAd < pModel->GetPrivate()->m_program.m_constantStrings.size() );
+                    const char* strName = pModel->GetPrivate()->m_program.m_constantStrings[ nameAd ].c_str();
+                    pResult->GetPrivate()->AddMesh( 0, 0, meshId, strName );
+                }
+                GetMemory().SetInstance( item, pResult );
+                break;
+            }
+
+            default:
+                check(false);
+            }
+            break;
+        }
+
+        case OP_TYPE::IN_ADDIMAGE:
+        {
+            auto args = pModel->GetPrivate()->m_program.GetOpArgs<OP::InstanceAddArgs>(item.at);
+            switch (item.stage)
+            {
+            case 0:
+				// We don't build the resources when building instance: just store ids for them.
+				AddOp( SCHEDULED_OP( item.at, item, 1), SCHEDULED_OP( args.instance, item) );
+                break;
+
+            case 1:
+            {
+                InstancePtrConst pBase = GetMemory().GetInstance( CACHE_ADDRESS(args.instance,item) );
+                InstancePtr pResult;
+                if (!pBase)
+                {
+                    pResult = new Instance();
+                }
+                else
+                {
+                    pResult = pBase->Clone();
+                }
+
+                if ( args.value )
+                {
+                    RESOURCE_ID imageId = pModel->GetPrivate()->GetResourceKey(
+                                args.relevantParametersListIndex,
+                                args.value,
+                                pParams);
+                    OP::ADDRESS nameAd = args.name;
+                    check(  nameAd < pModel->GetPrivate()->m_program.m_constantStrings.size() );
+                    const char* strName = pModel->GetPrivate()->m_program.m_constantStrings[ nameAd ].c_str();
+                    pResult->GetPrivate()->AddImage( 0, 0, 0, imageId, strName );
+                }
+                GetMemory().SetInstance( item, pResult );
+                break;
+            }
+
+            default:
+                check(false);
+            }
+
+            break;
+        }
+
+        default:
+			check(false);
+        }
+    }
+
+
+    //---------------------------------------------------------------------------------------------
+    void CodeRunner::RunCode_ConstantResource( SCHEDULED_OP& item, const Model* pModel )
+    {
+		MUTABLE_CPUPROFILER_SCOPE(RunCode_Constant);
+
+		OP_TYPE type = pModel->GetPrivate()->m_program.GetOpType(item.at);
+        switch (type)
+        {
+
+        case OP_TYPE::ME_CONSTANT:
+        {
+			OP::MeshConstantArgs args = pModel->GetPrivate()->m_program.GetOpArgs<OP::MeshConstantArgs>(item.at);
+
+            OP::ADDRESS cat = args.value;
+
+            PROGRAM& program = pModel->GetPrivate()->m_program;
+
+            // Assume the ROM has been loaded previously
+            check(program.m_constantMeshes[cat].Value )
+
+            MeshPtrConst pSourceConst;
+            program.GetConstant( cat, pSourceConst );
+            MeshPtr pSource = pSourceConst->Clone();
+
+            // Set the separate skeleton if necessary
+            if (args.skeleton>=0)
+            {
+                check( program.m_constantSkeletons.size()>size_t(args.skeleton)  );
+                Ptr<const Skeleton> pSkeleton = program.m_constantSkeletons[args.skeleton];
+                pSource->SetSkeleton(pSkeleton);
+            }
+
+			if (args.physicsBody >= 0)
+			{
+                check( program.m_constantPhysicsBodies.size()>size_t(args.physicsBody)  );
+                Ptr<const PhysicsBody> pPhysicsBody = program.m_constantPhysicsBodies[args.physicsBody];
+                pSource->SetPhysicsBody(pPhysicsBody);
+			}
+
+            GetMemory().SetMesh( item, pSource );
+			//UE_LOG(LogMutableCore, Log, TEXT("Set mesh constant %d."), item.at);
+            break;
+        }
+
+        case OP_TYPE::IM_CONSTANT:
+        {
+            auto args = pModel->GetPrivate()->m_program.GetOpArgs<OP::ResourceConstantArgs>(item.at);
+                OP::ADDRESS cat = args.value;
+
+            const PROGRAM& program = pModel->GetPrivate()->m_program;
+
+			int32 MipsToSkip = item.executionOptions;
+            ImagePtrConst pSource;
+            program.GetConstant( cat, pSource, MipsToSkip);
+
+			// Assume the ROM has been loaded previously in a task generated at IssueOp
+			check(pSource);
+
+            GetMemory().SetImage( item, pSource );
+			//UE_LOG(LogMutableCore, Log, TEXT("Set image constant %d."), item.at);
+			break;
+        }
+
+        default:
+            if (type!=OP_TYPE::NONE)
+            {
+                // Operation not implemented
+                check( false );
+            }
+            break;
+        }
+    }
+
+    //---------------------------------------------------------------------------------------------
+    void CodeRunner::RunCode_Mesh( SCHEDULED_OP& item, const Model* pModel )
+    {
+		MUTABLE_CPUPROFILER_SCOPE(RunCode_Mesh);
+
+		OP_TYPE type = pModel->GetPrivate()->m_program.GetOpType(item.at);
+        switch (type)
+        {
+
+        case OP_TYPE::ME_APPLYLAYOUT:
+        {
+            auto args = pModel->GetPrivate()->m_program.GetOpArgs<OP::MeshApplyLayoutArgs>(item.at);
+            switch (item.stage)
+            {
+            case 0:
+                    AddOp( SCHEDULED_OP( item.at, item, 1),
+                           SCHEDULED_OP( args.mesh, item),
+                           SCHEDULED_OP( args.layout, item) );
+                break;
+
+            case 1:
+            {
+                Ptr<const Mesh> pBase = GetMemory().GetMesh( CACHE_ADDRESS(args.mesh,item) );
+
+                MeshPtr pApplied;
+                if (pBase)
+                {
+                    pApplied = pBase->Clone();
+
+                    Ptr<const Layout> pLayout = GetMemory().GetLayout( CACHE_ADDRESS(args.layout,item) );
+                    int texCoordsSet = args.channel;
+                    MeshApplyLayout( pApplied.get(), pLayout.get(), texCoordsSet );
+                }
+
+                GetMemory().SetMesh( item, pApplied );
+                break;
+            }
+
+            default:
+                check(false);
+            }
+
+            break;
+        }
+
+        case OP_TYPE::ME_DIFFERENCE:
+        {
+            auto args = pModel->GetPrivate()->m_program.GetOpArgs<OP::MeshDifferenceArgs>(item.at);
+            switch (item.stage)
+            {
+            case 0:
+                if ( args.base && args.target )
+                {
+                    AddOp( SCHEDULED_OP( item.at, item, 1),
+                            SCHEDULED_OP( args.base, item),
+                            SCHEDULED_OP( args.target, item) );
+                }
+                else
+                {
+                    GetMemory().SetMesh( item, nullptr );
+                }
+                break;
+
+            case 1:
+            {
+                Ptr<const Mesh> pBase = GetMemory().GetMesh(CACHE_ADDRESS(args.base,item));
+                Ptr<const Mesh> pTarget = GetMemory().GetMesh(CACHE_ADDRESS(args.target,item));
+
+                MESH_BUFFER_SEMANTIC semantics[ MUTABLE_OP_MAX_MORPH_CHANNELS ];
+                int semanticIndices[ MUTABLE_OP_MAX_MORPH_CHANNELS ];
+
+                int numChannels = 0;
+                for ( int i=0;
+                      i<MUTABLE_OP_MAX_MORPH_CHANNELS && args.channelSemantic[i];
+                      ++i )
+                {
+					semantics[numChannels] = MESH_BUFFER_SEMANTIC(args.channelSemantic[i]);
+					semanticIndices[numChannels] = args.channelSemanticIndex[i];
+					++numChannels;
+                }
+
+                MeshPtr pResult = MeshDifference( pBase.get(), pTarget.get(),
+                                          numChannels, semantics, semanticIndices,
+                                          args.ignoreTextureCoords!=0 );
+
+                GetMemory().SetMesh( item, pResult );
+                break;
+            }
+
+            default:
+                check(false);
+            }
+
+            break;
+        }
+
+        case OP_TYPE::ME_MORPH2:
+        {
+            auto args = pModel->GetPrivate()->m_program.GetOpArgs<OP::MeshMorph2Args>(item.at);
+            switch (item.stage)
+            {
+            case 0:
+                if ( args.base )
+                {
+                    AddOp( SCHEDULED_OP( item.at, item, 1),
+                           SCHEDULED_OP( args.base, item),
+                           SCHEDULED_OP( args.factor, item) );
+                }
+                else
+                {
+                    GetMemory().SetMesh( item, nullptr );
+                }
+                break;
+
+            case 1:
+            {
+                bool baseValid = GetMemory().IsValid( CACHE_ADDRESS(args.base,item) );
+                float factor = GetMemory().GetScalar( CACHE_ADDRESS(args.factor,item) );
+
+                int count = 0;
+                for ( int i=0
+                    ; i<MUTABLE_OP_MAX_MORPH2_TARGETS && args.targets[i]
+                    ; ++i )
+                {
+                    count++;
+                }
+
+                // Factor goes from -1 to 1 across all targets. [0 - 1] represents positive morphs, while [-1, 0) represent negative morphs.
+				factor = FMath::Clamp(factor, -1.0f, 1.0f); // Is the factor not in range [-1, 1], it will index a non existing morph.
+
+				float absFactor = fabsf(factor);
+                float delta = 1.0f/(count-1);
+				int min = (int)floorf(absFactor/delta);
+				int max = (int)ceilf(absFactor/delta);
+
+                // Factor from 0 to 1 between the two targets
+				float bifactor = absFactor/delta - min;
+
+				// From [0,1] to [-1,0] for negative factors
+				if (factor < 0.0f)
+				{
+					float threshold = -1.0f + SMALL_NUMBER;
+
+					if (factor <= threshold)
+					{
+						// setting bifactor to -1.0
+						bifactor = -1.0f;
+					}
+					else
+					{
+						// we calculate the bifactor using positive values but is negative here
+						bifactor *= -1.0f;
+					}
+				}
+
+                if (baseValid)
+                {
+                    SCHEDULED_OP_DATA data;
+                    data.bifactor = bifactor;
+					data.min = FMath::Clamp(min, 0, count - 1);
+					data.max = FMath::Clamp(max, 0, count - 1);
+					uint32 dataAddress = uint32(m_heapData.Add(data));
+
+                    // Just the first of the targets
+					if ( bifactor < SMALL_NUMBER && bifactor > -SMALL_NUMBER )
+                    {                        
+                        AddOp( SCHEDULED_OP( item.at, item, 2, dataAddress),
+                                SCHEDULED_OP( args.base, item),
+                                SCHEDULED_OP( args.targets[min], item) );
+                    }
+                    // Just the second of the targets
+					else if ( bifactor > 1.0f - SMALL_NUMBER || bifactor <= -1.0f + SMALL_NUMBER )
+                    {
+                        check( max>0 );
+
+                        AddOp( SCHEDULED_OP( item.at, item, 2, dataAddress),
+                                SCHEDULED_OP( args.base, item),
+                                SCHEDULED_OP( args.targets[max], item) );
+                    }
+                    // Mix two targets on the base
+                    else
+                    {
+                        // We will need the base again
+                        AddOp( SCHEDULED_OP( item.at, item, 2, dataAddress),
+                                SCHEDULED_OP( args.base, item),
+                                SCHEDULED_OP( args.targets[min], item),
+                                SCHEDULED_OP( args.targets[max], item) );
+                    }
+                }
+
+                break;
+            }
+
+            case 2:
+            {
+                Ptr<const Mesh> pBase = GetMemory().GetMesh( CACHE_ADDRESS(args.base,item) );
+
+                // Factor from 0 to 1 between the two targets
+                const SCHEDULED_OP_DATA& data = m_heapData[ (size_t)item.customState ];
+                float bifactor = data.bifactor;
+                int min = data.min;
+                int max = data.max;
+
+                MeshPtrConst pResult;
+
+                if (pBase)
+                {
+                    // Just the first of the targets
+					if ( bifactor < SMALL_NUMBER && bifactor > -SMALL_NUMBER )
+                    {
+                        // Base with one full morph
+                        Ptr<const Mesh> pMorph = GetMemory().GetMesh( CACHE_ADDRESS(args.targets[min],item) );
+                        if (pMorph)
+                        {
+                            //pResult = MeshMorphReference( pBase.get(), pMorph.get() );
+							pResult = MeshMorph(pBase.get(), pMorph.get());
+                        }
+                    }
+                    // Just the second of the targets
+                    else if ( bifactor > 1.0f - SMALL_NUMBER )
+                    {
+                        check( max>0 );
+                        Ptr<const Mesh> pMorph = GetMemory().GetMesh( CACHE_ADDRESS(args.targets[max],item) );
+                        if (pMorph)
+                        {
+                            //pResult = MeshMorphReference( pBase.get(), pMorph.get() );
+							pResult = MeshMorph(pBase.get(), pMorph.get());
+                        }
+                    }
+					// Negative target
+					else if (bifactor <= -1.0f + SMALL_NUMBER)
+					{
+						check( max > 0 );
+						Ptr<const Mesh> pMorph = GetMemory().GetMesh(CACHE_ADDRESS(args.targets[max], item));
+						if (pMorph)
+						{
+							//pResult = MeshMorphReference(pBase.get(), pMorph.get(), bifactor);
+							pResult = MeshMorph(pBase.get(), pMorph.get(), bifactor);
+						}
+					}
+                    // Mix two targets on the base
+                    else
+                    {
+                        Ptr<const Mesh> pMin = GetMemory().GetMesh( CACHE_ADDRESS(args.targets[min],item) );
+                        Ptr<const Mesh> pMax = GetMemory().GetMesh( CACHE_ADDRESS(args.targets[max],item) );
+                        if (pMin && pMax)
+                        {
+						    //pResult = MeshMorph2Reference( pBase.get(), pMin.get(), pMax.get(), bifactor );
+                            pResult = MeshMorph2( pBase.get(), pMin.get(), pMax.get(), bifactor );
+                        }
+                    }
+					// Missing branch. With the current system we can never get apply a negative pMin morph, only a negative pMax morph.
+					// This is due to not having a bifactor value that represents a negative pMin morph.
+					// "-0" should represent a negative pMin morph, but it has the contradiction that 0 also represents positive pMax morph.
+					// In other words, we have a discontinuity arround 0.
+					// With the current implementation, the only solutions to get a negative pMin morph is to apply a really small negative value (lim x -> -0) (for example: -0.00001).
+                }
+                GetMemory().SetMesh( item, pResult );
+
+                break;
+            }
+
+            default:
+                check(false);
+            }
+
+            break;
+        }
+
+        case OP_TYPE::ME_MERGE:
+        {
+            auto args = pModel->GetPrivate()->m_program.GetOpArgs<OP::MeshMergeArgs>(item.at);
+            switch (item.stage)
+            {
+            case 0:
+                    AddOp( SCHEDULED_OP( item.at, item, 1),
+                           SCHEDULED_OP( args.base, item),
+                           SCHEDULED_OP( args.added, item) );
+                break;
+
+            case 1:
+            {
+                MeshPtrConst pA = GetMemory().GetMesh( CACHE_ADDRESS(args.base,item) );
+                MeshPtrConst pB = GetMemory().GetMesh( CACHE_ADDRESS(args.added,item) );
+
+                MeshPtr pResult;
+
+                if (pA && pB && pA->GetVertexCount() && pB->GetVertexCount())
+                {
+                    pResult = MeshMerge( pA.get(), pB.get() );
+
+                    if (!args.newSurfaceID)
+                    {
+                        // Make it a single surface.
+                        pResult->m_surfaces.Empty();
+                        pResult->EnsureSurfaceData();
+                    }
+                    else
+                    {
+                        check(pB->GetSurfaceCount()==1);
+                        pResult->m_surfaces.Last().m_id=args.newSurfaceID;
+                    }
+                }
+                else if (pA && pA->GetVertexCount())
+                {
+                    pResult = pA->Clone();
+                }
+                else if (pB && pB->GetVertexCount())
+                {
+                    pResult = pB->Clone();
+                    check(pResult->GetSurfaceCount()==1);
+                    if (pResult->GetSurfaceCount()>0 && args.newSurfaceID)
+                    {
+                        pResult->m_surfaces.Last().m_id=args.newSurfaceID;
+                    }
+                }
+                else
+                {
+                    pResult = new Mesh();
+                }
+
+                GetMemory().SetMesh( item, pResult );
+                break;
+            }
+
+            default:
+                check(false);
+            }
+
+            break;
+        }
+
+        case OP_TYPE::ME_INTERPOLATE:
+        {
+            auto args = pModel->GetPrivate()->m_program.GetOpArgs<OP::MeshInterpolateArgs>(item.at);
+            switch (item.stage)
+            {
+            case 0:
+                if ( args.base )
+                {
+                    AddOp( SCHEDULED_OP( item.at, item, 1),
+                           SCHEDULED_OP( args.base, item),
+                           SCHEDULED_OP( args.factor, item ) );
+                }
+                else
+                {
+                    GetMemory().SetMesh( item, nullptr );
+                }
+                break;
+
+            case 1:
+            {
+                int count = 1;
+                for ( int i=0
+                    ; i<MUTABLE_OP_MAX_INTERPOLATE_COUNT-1 && args.targets[i]
+                    ; ++i )
+                {
+                    count++;
+                }
+
+                // Factor from 0 to 1 across all targets
+                float factor = GetMemory().GetScalar( CACHE_ADDRESS(args.factor,item) );
+
+                float delta = 1.0f/(count-1);
+                int min = (int)floorf( factor/delta );
+                int max = (int)ceilf( factor/delta );
+
+                // Factor from 0 to 1 between the two targets
+                float bifactor = factor/delta - min;
+
+                SCHEDULED_OP_DATA data;
+                data.bifactor = bifactor;
+				data.min = FMath::Clamp(min, 0, count - 1);
+				data.max = FMath::Clamp(max, 0, count - 1);
+				uint32 dataAddress = uint32(m_heapData.Num());
+
+                // Just the first of the targets
+                if ( bifactor < SMALL_NUMBER )
+                {
+                    if( min==0 )
+                    {
+                        // Just the base
+                            Ptr<const Mesh> pBase = GetMemory().GetMesh( CACHE_ADDRESS(args.base,item) );
+                            MeshPtr pResult;
+                            if (pBase)
+                            {
+                                pResult = pBase->Clone();
+                            }
+                            GetMemory().SetMesh( item, pResult );
+                        }
+                    else
+                    {
+                        // Base with one full morph
+                        m_heapData.Add(data);
+                            AddOp( SCHEDULED_OP( item.at, item, 2, dataAddress),
+                                   SCHEDULED_OP( args.base, item),
+                                   SCHEDULED_OP( args.targets[min-1], item) );
+                        }
+                    }
+                // Just the second of the targets
+                else if ( bifactor > 1.0f-SMALL_NUMBER )
+                {
+                    m_heapData.Add(data);
+                        AddOp( SCHEDULED_OP( item.at, item, 2, dataAddress),
+                               SCHEDULED_OP( args.base, item),
+                               SCHEDULED_OP( args.targets[max-1], item) );
+                    }
+                // Mix the first target on the base
+                else if ( min==0 )
+                {
+                    m_heapData.Add(data);
+                        AddOp( SCHEDULED_OP( item.at, item, 2, dataAddress),
+                               SCHEDULED_OP( args.base, item),
+                               SCHEDULED_OP( args.targets[0], item)
+                               );
+                    }
+                // Mix two targets on the base
+                else
+                {
+                    m_heapData.Add(data);
+                        AddOp( SCHEDULED_OP( item.at, item, 2, dataAddress),
+                               SCHEDULED_OP( args.base, item),
+                               SCHEDULED_OP( args.targets[min-1], item),
+                               SCHEDULED_OP( args.targets[max-1], item) );
+                    }
+
+                break;
+            }
+
+            case 2:
+            {
+                int count = 1;
+                for ( int i=0
+                    ; i<MUTABLE_OP_MAX_INTERPOLATE_COUNT-1 && args.targets[i]
+                    ; ++i )
+                {
+                    count++;
+                }
+
+                const SCHEDULED_OP_DATA& data = m_heapData[ (size_t)item.customState ];
+
+                // Factor from 0 to 1 between the two targets
+                float bifactor = data.bifactor;
+                int min = data.min;
+                int max = data.max;
+
+                Ptr<const Mesh> pBase = GetMemory().GetMesh( CACHE_ADDRESS(args.base, item) );
+
+                MeshPtr pResult;
+
+                if (pBase)
+                {
+                    // Just the first of the targets
+                    if ( bifactor < SMALL_NUMBER )
+                    {
+                        if( min==0 )
+                        {
+                            // Just the base. It should have been dealt with in the previous stage.
+                            check( false );
+                        }
+                        else
+                        {
+                            // Base with one full morph
+                            Ptr<const Mesh> pMorph = GetMemory().GetMesh( CACHE_ADDRESS(args.targets[min-1],item) );
+                            pResult = MeshMorph( pBase.get(), pMorph.get() );
+                        }
+                    }
+                    // Just the second of the targets
+                    else if ( bifactor > 1.0f-SMALL_NUMBER )
+                    {
+                        check( max>0 );
+                        Ptr<const Mesh> pMorph = GetMemory().GetMesh( CACHE_ADDRESS(args.targets[max-1],item) );
+
+                        if (pMorph)
+                        {
+                            pResult = MeshMorph( pBase.get(), pMorph.get() );
+                        }
+                        else
+                        {
+                            pResult = pBase->Clone();
+                        }
+                    }
+                    // Mix the first target on the base
+                    else if ( min==0 )
+                    {
+                        Ptr<const Mesh> pMorph = GetMemory().GetMesh( CACHE_ADDRESS( args.targets[0], item ) );
+                        if (pMorph)
+                        {
+                            pResult = MeshMorph( pBase.get(), pMorph.get(), bifactor );
+                        }
+                        else
+                        {
+                            pResult = pBase->Clone();
+                        }
+
+                    }
+                    // Mix two targets on the base
+                    else
+                    {
+                        Ptr<const Mesh> pMin = GetMemory().GetMesh( CACHE_ADDRESS(args.targets[min-1],item) );
+                        Ptr<const Mesh> pMax = GetMemory().GetMesh( CACHE_ADDRESS(args.targets[max-1],item) );
+
+                        if (pMin && pMax)
+                        {
+                            pResult = MeshMorph2( pBase.get(), pMin.get(), pMax.get(), bifactor );
+                        }
+                        else
+                        {
+                            pResult = pBase->Clone();
+                        }
+                    }
+                }
+
+                GetMemory().SetMesh( item, pResult );
+                break;
+            }
+
+            default:
+                check(false);
+            }
+
+            break;
+        }
+
+        case OP_TYPE::ME_MASKCLIPMESH:
+        {
+            auto args = pModel->GetPrivate()->m_program.GetOpArgs<OP::MeshMaskClipMeshArgs>(item.at);
+            switch (item.stage)
+            {
+            case 0:
+                    AddOp( SCHEDULED_OP( item.at, item, 1),
+                           SCHEDULED_OP( args.source, item),
+                           SCHEDULED_OP( args.clip, item) );
+                break;
+
+            case 1:
+            {
+                Ptr<const Mesh> pSource = GetMemory().GetMesh( CACHE_ADDRESS(args.source,item) );
+                Ptr<const Mesh> pClip = GetMemory().GetMesh( CACHE_ADDRESS(args.clip,item) );
+
+                // Only if both are valid.
+                MeshPtr pResult;
+                if (pSource.get() && pClip.get())
+                {
+                    pResult = MeshMaskClipMesh(pSource.get(), pClip.get());
+                }
+                GetMemory().SetMesh( item, pResult );
+
+                break;
+            }
+
+            default:
+                check(false);
+            }
+
+            break;
+        }
+
+        case OP_TYPE::ME_MASKDIFF:
+        {
+            auto args = pModel->GetPrivate()->m_program.GetOpArgs<OP::MeshMaskDiffArgs>(item.at);
+            switch (item.stage)
+            {
+            case 0:
+                    AddOp( SCHEDULED_OP( item.at, item, 1),
+                           SCHEDULED_OP( args.source, item),
+                           SCHEDULED_OP( args.fragment, item) );
+                break;
+
+            case 1:
+            {
+                Ptr<const Mesh> pSource = GetMemory().GetMesh( CACHE_ADDRESS(args.source,item) );
+                Ptr<const Mesh> pClip = GetMemory().GetMesh( CACHE_ADDRESS(args.fragment,item) );
+
+                // Only if both are valid.
+                MeshPtr pResult;
+                if (pSource.get() && pClip.get())
+                {
+                    pResult = MeshMaskDiff(pSource.get(), pClip.get());
+                }
+                GetMemory().SetMesh( item, pResult );
+
+                break;
+            }
+
+            default:
+                check(false);
+            }
+
+            break;
+        }
+
+        case OP_TYPE::ME_SUBTRACT:
+        {
+            auto args = pModel->GetPrivate()->m_program.GetOpArgs<OP::MeshSubtractArgs>(item.at);
+            switch (item.stage)
+            {
+            case 0:
+                    AddOp( SCHEDULED_OP( item.at, item, 1),
+                           SCHEDULED_OP( args.a, item),
+                           SCHEDULED_OP( args.b, item) );
+                break;
+
+            case 1:
+            {
+                Ptr<const Mesh> pA = GetMemory().GetMesh( CACHE_ADDRESS(args.a,item) );
+                Ptr<const Mesh> pB = GetMemory().GetMesh( CACHE_ADDRESS(args.b,item) );
+
+                MeshPtr pResult = MeshSubtract( pA.get(), pB.get() );
+
+                GetMemory().SetMesh( item, pResult );
+                break;
+            }
+
+            default:
+                check(false);
+            }
+
+            break;
+        }
+
+        case OP_TYPE::ME_FORMAT:
+        {
+            auto args = pModel->GetPrivate()->m_program.GetOpArgs<OP::MeshFormatArgs>(item.at);
+            switch (item.stage)
+            {
+            case 0:
+                if (args.source && args.format)
+                {
+                         AddOp( SCHEDULED_OP( item.at, item, 1),
+                                SCHEDULED_OP( args.source, item),
+                                SCHEDULED_OP( args.format, item));
+                    }
+                    else
+                    {
+                    GetMemory().SetMesh( item, nullptr );
+                }
+                break;
+
+            case 1:
+            {
+                Ptr<const Mesh> pSource = GetMemory().GetMesh( CACHE_ADDRESS(args.source,item) );
+                Ptr<const Mesh> pFormat = GetMemory().GetMesh( CACHE_ADDRESS(args.format,item) );
+
+                uint8_t flags = args.buffers;
+                MeshPtr pResult;
+                pResult = MeshFormat( pSource.get(),
+                                      pFormat.get(),
+                                      true,
+                                      (flags & OP::MeshFormatArgs::BT_VERTEX) != 0,
+                                      (flags & OP::MeshFormatArgs::BT_INDEX) != 0,
+                                      (flags & OP::MeshFormatArgs::BT_FACE) != 0,
+                                      (flags & OP::MeshFormatArgs::BT_REBUILD_TANGENTS) != 0,
+                                      (flags & OP::MeshFormatArgs::BT_IGNORE_MISSING) != 0
+                                      );
+
+                if (flags & OP::MeshFormatArgs::BT_RESETBUFFERINDICES)
+                {
+                    pResult->ResetBufferIndices();
+                }
+
+                GetMemory().SetMesh( item, pResult );
+                break;
+            }
+
+            default:
+                check(false);
+            }
+
+            break;
+        }
+
+        case OP_TYPE::ME_EXTRACTLAYOUTBLOCK:
+        {
+            const uint8_t* data = pModel->GetPrivate()->m_program.GetOpArgsPointer(item.at);
+
+            mu::OP::ADDRESS source;
+            memcpy( &source, data, sizeof(OP::ADDRESS) );
+            data += sizeof(OP::ADDRESS);
+
+            uint16_t layout;
+            memcpy( &layout, data, sizeof(uint16_t) );
+            data += sizeof(uint16_t);
+
+            uint16_t blockCount;
+            memcpy( &blockCount, data, sizeof(uint16_t) );
+            data += sizeof(uint16_t);
+
+            switch (item.stage)
+            {
+            case 0:
+                if (source)
+                {
+                        AddOp( SCHEDULED_OP( item.at, item, 1),
+                               SCHEDULED_OP( source, item) );
+                    }
+                    else
+                    {
+                    GetMemory().SetMesh( item, nullptr );
+                }
+                break;
+
+            case 1:
+            {
+                Ptr<const Mesh> pSource = GetMemory().GetMesh( CACHE_ADDRESS(source,item) );
+
+                // Access with memcpy necessary for unaligned arm issues.
+                uint32_t blocks[1024];
+                memcpy( blocks, data, sizeof(uint32_t)*std::min(1024,int(blockCount)) );
+
+                MeshPtr pResult;
+                pResult = MeshExtractLayoutBlock( pSource.get(),
+                                                  layout,
+                                                  blockCount,
+                                                  blocks );
+
+                GetMemory().SetMesh( item, pResult );
+                break;
+            }
+
+            default:
+                check(false);
+            }
+
+            break;
+        }
+
+        case OP_TYPE::ME_EXTRACTFACEGROUP:
+        {
+            auto args = pModel->GetPrivate()->m_program.GetOpArgs<OP::MeshExtractFaceGroupArgs>(item.at);
+            switch (item.stage)
+            {
+            case 0:
+                if (args.source)
+                {
+                        AddOp( SCHEDULED_OP( item.at, item, 1),
+                               SCHEDULED_OP( args.source, item) );
+                    }
+                    else
+                    {
+                    GetMemory().SetMesh( item, nullptr );
+                }
+                break;
+
+            case 1:
+            {
+                Ptr<const Mesh> pSource = GetMemory().GetMesh( CACHE_ADDRESS(args.source,item) );
+
+                MeshPtr pResult;
+                pResult = MeshExtractFaceGroup( pSource.get(),
+                                                args.group );
+
+                GetMemory().SetMesh( item, pResult );
+                break;
+            }
+
+            default:
+                check(false);
+            }
+
+            break;
+        }
+
+        case OP_TYPE::ME_TRANSFORM:
+        {
+            auto args = pModel->GetPrivate()->m_program.GetOpArgs<OP::MeshTransformArgs>(item.at);
+            switch (item.stage)
+            {
+            case 0:
+                if (args.source)
+                {
+                        AddOp( SCHEDULED_OP( item.at, item, 1),
+                               SCHEDULED_OP( args.source, item) );
+                    }
+                    else
+                    {
+                    GetMemory().SetMesh( item, nullptr );
+                }
+                break;
+
+            case 1:
+            {
+                Ptr<const Mesh> pSource = GetMemory().GetMesh( CACHE_ADDRESS(args.source,item) );
+
+                const mat4f& mat = pModel->GetPrivate()->m_program.
+                    m_constantMatrices[args.matrix];
+
+                MeshPtr pResult = MeshTransform(pSource.get(), ToUnreal(mat) );
+
+                GetMemory().SetMesh( item, pResult );
+                break;
+            }
+
+            default:
+                check(false);
+            }
+
+            break;
+        }
+
+        case OP_TYPE::ME_CLIPMORPHPLANE:
+        {
+            auto args = pModel->GetPrivate()->m_program.GetOpArgs<OP::MeshClipMorphPlaneArgs>(item.at);
+            switch (item.stage)
+            {
+            case 0:
+                if (args.source)
+                {
+                        AddOp( SCHEDULED_OP( item.at, item, 1),
+                               SCHEDULED_OP( args.source, item) );
+                    }
+                    else
+                    {
+                    GetMemory().SetMesh( item, nullptr );
+                }
+                break;
+
+            case 1:
+            {
+                Ptr<const Mesh> pSource = GetMemory().GetMesh( CACHE_ADDRESS(args.source,item) );
+
+                MeshPtr pResult;
+
+                check( args.morphShape
+                                <
+                                pModel->GetPrivate()->m_program.m_constantShapes.size() );
+
+                // Should be an ellipse
+                const SHAPE& morphShape = pModel->GetPrivate()->m_program.
+                    m_constantShapes[args.morphShape];
+
+                const mu::vec3f& origin = morphShape.position;
+                const mu::vec3f& normal = morphShape.up;
+
+                if (args.vertexSelectionType == OP::MeshClipMorphPlaneArgs::VS_SHAPE)
+                {
+                    check( args.vertexSelectionShapeOrBone
+                                    <
+                                    pModel->GetPrivate()->m_program.m_constantShapes.size() );
+
+                    // Should be None or an axis aligned box
+                    const SHAPE& selectionShape = pModel->GetPrivate()->m_program.m_constantShapes[args.vertexSelectionShapeOrBone];
+                    pResult = MeshClipMorphPlane(pSource.get(), origin, normal, args.dist, args.factor, morphShape.size[0], morphShape.size[1], morphShape.size[2], selectionShape);
+                }
+
+                else if (args.vertexSelectionType == OP::MeshClipMorphPlaneArgs::VS_BONE_HIERARCHY)
+                {
+                    check( args.vertexSelectionShapeOrBone
+                                    <
+                                    pModel->GetPrivate()->m_program.m_constantStrings.size() );
+
+                    SHAPE selectionShape;
+                    selectionShape.type = (uint8_t)SHAPE::Type::None;
+                    const string& selectionBone = pModel->GetPrivate()->m_program.m_constantStrings[args.vertexSelectionShapeOrBone];
+					pResult = MeshClipMorphPlane(pSource.get(), origin, normal, args.dist, args.factor, morphShape.size[0], morphShape.size[1], morphShape.size[2], selectionShape, selectionBone, args.maxBoneRadius);
+                }
+                else
+                {
+                    // No vertex selection
+                    SHAPE selectionShape;
+                    selectionShape.type = (uint8_t)SHAPE::Type::None;
+                    pResult = MeshClipMorphPlane(pSource.get(), origin, normal, args.dist, args.factor, morphShape.size[0], morphShape.size[1], morphShape.size[2], selectionShape);
+                }
+
+                GetMemory().SetMesh( item, pResult );
+                break;
+            }
+
+            default:
+                check(false);
+            }
+
+            break;
+        }
+
+
+        case OP_TYPE::ME_CLIPWITHMESH:
+        {
+            auto args = pModel->GetPrivate()->m_program.GetOpArgs<OP::MeshClipWithMeshArgs>(item.at);
+            switch (item.stage)
+            {
+            case 0:
+                if (args.source)
+                {
+                        AddOp( SCHEDULED_OP( item.at, item, 1),
+                               SCHEDULED_OP( args.source, item),
+                               SCHEDULED_OP( args.clipMesh, item) );
+                    }
+                    else
+                    {
+                    GetMemory().SetMesh( item, nullptr );
+                }
+                break;
+
+            case 1:
+            {
+                Ptr<const Mesh> pSource = GetMemory().GetMesh( CACHE_ADDRESS(args.source,item) );
+                Ptr<const Mesh> pClip = GetMemory().GetMesh( CACHE_ADDRESS(args.clipMesh,item) );
+
+                // Only if both are valid.
+                MeshPtr pResult;
+                if (pSource && pClip)
+                {
+                    pResult = MeshClipWithMesh(pSource.get(), pClip.get());
+                }
+                else if (pSource)
+                {
+                    pResult = pSource->Clone();
+                }
+
+                GetMemory().SetMesh( item, pResult );
+                break;
+            }
+
+            default:
+                check(false);
+            }
+
+            break;
+        }
+		case OP_TYPE::ME_CLIPDEFORM:
+		{
+			OP::MeshClipDeformArgs args = pModel->GetPrivate()->m_program.GetOpArgs<OP::MeshClipDeformArgs>(item.at);
+			switch (item.stage)
+			{
+			case 0:
+				if (args.mesh)
+				{
+						AddOp(SCHEDULED_OP(item.at, item, 1),
+							SCHEDULED_OP(args.mesh, item),
+							SCHEDULED_OP(args.clipShape, item));
+					}
+					else
+					{
+					GetMemory().SetMesh(item, nullptr);
+				}
+				break;
+			
+			case 1:
+			{
+				Ptr<const Mesh> BaseMesh = GetMemory().GetMesh(CACHE_ADDRESS(args.mesh, item));
+				Ptr<const Mesh> ClipShape = GetMemory().GetMesh(CACHE_ADDRESS(args.clipShape, item));
+				Ptr<Mesh> pResult;
+
+					if ( BaseMesh && ClipShape )
+				{
+					pResult = MeshClipDeform(BaseMesh.get(), ClipShape.get(), args.clipWeightThreshold);
+				}
+				else if( BaseMesh )
+				{
+					pResult = BaseMesh->Clone();
+				}
+
+				GetMemory().SetMesh(item, pResult);
+				break;
+			}
+
+			default:
+				check(false);
+			}
+
+			break;
+		}
+
+        case OP_TYPE::ME_REMAPINDICES:
+        {
+            auto args = pModel->GetPrivate()->m_program.GetOpArgs<OP::MeshRemapIndicesArgs>(item.at);
+            switch (item.stage)
+            {
+            case 0:
+                if (args.source)
+                {
+                        AddOp( SCHEDULED_OP( item.at, item, 1),
+                               SCHEDULED_OP( args.source, item),
+                               SCHEDULED_OP( args.reference, item) );
+                    }
+                    else
+                    {
+                    GetMemory().SetMesh( item, nullptr );
+                }
+                break;
+
+            case 1:
+            {
+                Ptr<const Mesh> pSource = GetMemory().GetMesh( CACHE_ADDRESS(args.source,item) );
+                Ptr<const Mesh> pReference = GetMemory().GetMesh( CACHE_ADDRESS(args.reference,item) );
+
+                // Only if both are valid.
+                MeshPtr pResult;
+                if (pSource && pReference)
+                {
+                    pResult = MeshRemapIndices(pSource.get(), pReference.get());
+                }
+                else if (pSource)
+                {
+                    pResult = pSource->Clone();
+                }
+
+                GetMemory().SetMesh( item, pResult );
+                break;
+            }
+
+            default:
+                check(false);
+            }
+
+            break;
+        }
+
+
+        case OP_TYPE::ME_APPLYPOSE:
+        {
+            auto args = pModel->GetPrivate()->m_program.GetOpArgs<OP::MeshApplyPoseArgs>(item.at);
+            switch (item.stage)
+            {
+            case 0:
+                if (args.base)
+                {
+                        AddOp( SCHEDULED_OP( item.at, item, 1),
+                               SCHEDULED_OP( args.base, item),
+                               SCHEDULED_OP( args.pose, item) );
+                    }
+                    else
+                    {
+                    GetMemory().SetMesh( item, nullptr );
+                }
+                break;
+
+            case 1:
+            {
+                Ptr<const Mesh> pBase = GetMemory().GetMesh( CACHE_ADDRESS(args.base,item) );
+                Ptr<const Mesh> pPose = GetMemory().GetMesh( CACHE_ADDRESS(args.pose,item) );
+
+                // Only if both are valid.
+                MeshPtr pResult;
+                if (pBase && pPose)
+                {
+                    pResult = MeshApplyPose(pBase.get(), pPose.get());
+                }
+                else if (pBase)
+                {
+                    pResult = pBase->Clone();
+                }
+
+                GetMemory().SetMesh( item, pResult );
+                break;
+            }
+
+            default:
+                check(false);
+            }
+
+            break;
+        }
+
+
+		case OP_TYPE::ME_GEOMETRYOPERATION:
+		{
+			auto args = pModel->GetPrivate()->m_program.GetOpArgs<OP::MeshGeometryOperationArgs>(item.at);
+			switch (item.stage)
+			{
+			case 0:
+				if (args.meshA)
+				{
+						AddOp(SCHEDULED_OP(item.at, item, 1),
+							SCHEDULED_OP(args.meshA, item),
+							SCHEDULED_OP(args.meshB, item),
+							SCHEDULED_OP(args.scalarA, item),
+							SCHEDULED_OP(args.scalarB, item));
+					}
+					else
+					{
+					GetMemory().SetMesh(item, nullptr);
+				}
+				break;
+
+			case 1:
+			{
+				Ptr<const Mesh> MeshA = GetMemory().GetMesh(CACHE_ADDRESS(args.meshA, item));
+				Ptr<const Mesh> MeshB = GetMemory().GetMesh(CACHE_ADDRESS(args.meshB, item));
+				float ScalarA = GetMemory().GetScalar(CACHE_ADDRESS(args.scalarA, item));
+				float ScalarB = GetMemory().GetScalar(CACHE_ADDRESS(args.scalarB, item));
+
+				MeshPtr pResult = MeshGeometryOperation(MeshA.get(),MeshB.get(),ScalarA,ScalarB);
+
+				GetMemory().SetMesh(item, pResult);
+				break;
+			}
+
+			default:
+				check(false);
+			}
+
+			break;
+		}
+
+
+		case OP_TYPE::ME_BINDSHAPE:
+		{
+			OP::MeshBindShapeArgs args = pModel->GetPrivate()->m_program.GetOpArgs<OP::MeshBindShapeArgs>(item.at);
+			const uint8_t* data = pModel->GetPrivate()->m_program.GetOpArgsPointer(item.at);
+
+			switch (item.stage)
+			{
+			case 0:
+				if (args.mesh)
+				{
+					AddOp(SCHEDULED_OP(item.at, item, 1),
+							SCHEDULED_OP(args.mesh, item),
+							SCHEDULED_OP(args.shape, item));
+				}
+				else
+				{
+					GetMemory().SetMesh(item, nullptr);
+				}
+				break;
+
+			case 1:
+			{
+				Ptr<const Mesh> BaseMesh = GetMemory().GetMesh(CACHE_ADDRESS(args.mesh, item));
+				Ptr<const Mesh> Shape = GetMemory().GetMesh(CACHE_ADDRESS(args.shape, item));
+				
+				EShapeBindingMethod BindingMethod = static_cast<EShapeBindingMethod>(args.bindingMethod); 
+
+				if (BindingMethod == EShapeBindingMethod::ReshapeClosestProject)
+				{ 
+					// Bones are stored after the args
+					data += sizeof(args);
+
+					// Rebuilding array of bone names ----
+					int32 NumBones;
+					FMemory::Memcpy(&NumBones, data, sizeof(int32)); 
+					data += sizeof(int32);
+					
+					TArray<string> BonesToDeform;
+					BonesToDeform.SetNum(NumBones);
+					for (int32 b = 0; b < NumBones; ++b)
+					{
+						BonesToDeform[b] = pModel->GetPrivate()->m_program.m_constantStrings[*data].c_str();
+						data += sizeof(int32);
+					}
+
+					int32 NumPhysicsBodies;
+					FMemory::Memcpy(&NumPhysicsBodies, data, sizeof(int32)); 
+					data += sizeof(int32);
+
+					TArray<string> PhysicsToDeform;
+					PhysicsToDeform.SetNum(NumPhysicsBodies);
+					for (int32 b = 0; b < NumPhysicsBodies; ++b)
+					{
+						PhysicsToDeform[b] = pModel->GetPrivate()->m_program.m_constantStrings[*data].c_str();
+						data += sizeof(int32);
+					}
+
+					// -----------
+					Ptr<Mesh> pResult = MeshBindShapeReshape( BaseMesh.get(), Shape.get(), 	
+						BonesToDeform, PhysicsToDeform,
+						args.flags & uint32(OP::EMeshBindShapeFlags::DeformAllBones),
+						args.flags & uint32(OP::EMeshBindShapeFlags::DeformAllPhysics),	
+						args.flags & uint32(OP::EMeshBindShapeFlags::ReshapeVertices),
+						args.flags & uint32(OP::EMeshBindShapeFlags::ReshapeSkeleton),
+						args.flags & uint32(OP::EMeshBindShapeFlags::ReshapePhysicsVolumes),
+						args.flags & uint32(OP::EMeshBindShapeFlags::EnableRigidParts));
+					
+					GetMemory().SetMesh(item, pResult);
+				}	
+				else
+				{
+					Ptr<Mesh> ResultPtr = MeshBindShapeClipDeform( BaseMesh.get(), Shape.get(), BindingMethod );
+					GetMemory().SetMesh(item, ResultPtr);
+				}
+
+				break;
+			}
+
+			default:
+				check(false);
+			}
+
+			break;
+		}
+
+
+		case OP_TYPE::ME_APPLYSHAPE:
+		{
+			auto args = pModel->GetPrivate()->m_program.GetOpArgs<OP::MeshApplyShapeArgs>(item.at);
+			switch (item.stage)
+			{
+			case 0:
+				if (args.mesh)
+				{
+						AddOp(SCHEDULED_OP(item.at, item, 1),
+							SCHEDULED_OP(args.mesh, item),
+							SCHEDULED_OP(args.shape, item));
+				}
+				else
+				{
+					GetMemory().SetMesh(item, nullptr);
+				}
+				break;
+
+			case 1:
+			{
+				Ptr<const Mesh> BaseMesh = GetMemory().GetMesh(CACHE_ADDRESS(args.mesh, item));
+				Ptr<const Mesh> Shape = GetMemory().GetMesh(CACHE_ADDRESS(args.shape, item));
+
+				Ptr<Mesh> pResult = MeshApplyShape(BaseMesh.get(), Shape.get(),
+					args.flags & uint32(OP::EMeshBindShapeFlags::ReshapeVertices),
+					args.flags & uint32(OP::EMeshBindShapeFlags::ReshapeSkeleton),
+					args.flags & uint32(OP::EMeshBindShapeFlags::ReshapePhysicsVolumes));
+
+				GetMemory().SetMesh(item, pResult);
+				break;
+			}
+
+			default:
+				check(false);
+			}
+
+			break;
+		}
+
+		case OP_TYPE::ME_MORPHRESHAPE:
+		{
+			OP::MeshMorphReshapeArgs Args = pModel->GetPrivate()->m_program.GetOpArgs<OP::MeshMorphReshapeArgs>(item.at);
+			switch(item.stage)
+			{
+			case 0:
+			{
+				if (Args.Morph)
+				{
+					AddOp(SCHEDULED_OP(item.at, item, 1), 
+						  SCHEDULED_OP(Args.Morph, item), 
+						  SCHEDULED_OP(Args.Reshape, item));
+				}
+				else 
+				{
+					GetMemory().SetMesh(item, nullptr);
+				}
+				break;
+			}
+			case 1:
+			{
+				Ptr<const Mesh> MorphedMesh = GetMemory().GetMesh(CACHE_ADDRESS(Args.Morph, item));
+				Ptr<const Mesh> ReshapeMesh = GetMemory().GetMesh(CACHE_ADDRESS(Args.Reshape, item));
+
+				if (ReshapeMesh)
+				{
+					// Clone without Skeleton, Physics or Poses 
+					EMeshCloneFlags CloneFlags = ~(
+							EMeshCloneFlags::WithSkeleton    | 
+							EMeshCloneFlags::WithPhysicsBody | 
+							EMeshCloneFlags::WithPoses);
+
+					Ptr<Mesh> ResultPtr = MorphedMesh->Clone(CloneFlags);
+					ResultPtr->SetSkeleton(ReshapeMesh->GetSkeleton().get());
+					ResultPtr->SetPhysicsBody(ReshapeMesh->GetPhysicsBody().get());
+					ResultPtr->m_bonePoses = ReshapeMesh->m_bonePoses;
+
+					GetMemory().SetMesh(item, ResultPtr);
+				}
+				else
+				{
+					Ptr<Mesh> ResultPtr = MorphedMesh->Clone();
+					GetMemory().SetMesh(item, ResultPtr);
+				}
+
+				break;
+			}
+
+			default:
+				check(false);
+			}
+
+			break;
+		}
+
+        case OP_TYPE::ME_SETSKELETON:
+        {
+            auto args = pModel->GetPrivate()->m_program.GetOpArgs<OP::MeshSetSkeletonArgs>(item.at);
+            switch (item.stage)
+            {
+            case 0:
+                if (args.source)
+                {
+                        AddOp( SCHEDULED_OP( item.at, item, 1),
+                               SCHEDULED_OP( args.source, item),
+                               SCHEDULED_OP( args.skeleton, item) );
+                    }
+                    else
+                    {
+                    GetMemory().SetMesh( item, nullptr );
+                }
+                break;
+
+            case 1:
+            {
+                Ptr<const Mesh> pSource = GetMemory().GetMesh( CACHE_ADDRESS(args.source,item) );
+                Ptr<const Mesh> pSkeleton = GetMemory().GetMesh( CACHE_ADDRESS(args.skeleton,item) );
+
+                // Only if both are valid.
+                MeshPtr pResult;
+                if (pSource && pSkeleton)
+                {
+                    if ( pSource->GetSkeleton()
+                         &&
+                         !pSource->GetSkeleton()->m_bones.IsEmpty() )
+                    {
+                        // For some reason we already have bone data, so we can't just overwrite it
+                        // or the skinning may break. This may happen because of a problem in the
+                        // optimiser that needs investigation.
+                        // \TODO Be defensive, for now.
+                        UE_LOG(LogMutableCore,Warning, TEXT("Performing a MeshRemapSkeleton, instead of MeshSetSkeletonData because source mesh already has some skeleton."));
+                        pResult = MeshRemapSkeleton( pSource.get(), pSkeleton->GetSkeleton().get() );
+                        if (!pResult)
+                        {
+                            pResult = pSource->Clone();
+                        }
+                    }
+                    else
+                    {
+                        pResult = pSource->Clone();
+                        pResult->SetSkeleton(pSkeleton->GetSkeleton().get());
+                    }
+                    //pResult->GetPrivate()->CheckIntegrity();
+                }
+                else if (pSource)
+                {
+                    pResult = pSource->Clone();
+                }
+
+                GetMemory().SetMesh( item, pResult );
+                break;
+            }
+
+            default:
+                check(false);
+            }
+
+            break;
+        }
+
+        case OP_TYPE::ME_REMOVEMASK:
+        {
+            // Decode op
+            // TODO: Partial decode for each stage
+            const uint8_t* data = pModel->GetPrivate()->m_program.GetOpArgsPointer(item.at);
+
+            OP::ADDRESS source;
+            FMemory::Memcpy(&source,data,sizeof(OP::ADDRESS)); data += sizeof(OP::ADDRESS);
+
+            vector<SCHEDULED_OP> conditions;
+            vector<OP::ADDRESS> masks;
+
+            uint16_t removes;
+			FMemory::Memcpy(&removes,data,sizeof(uint16_t)); data += sizeof(uint16_t);
+
+            for( uint16_t r=0; r<removes; ++r)
+            {
+                OP::ADDRESS condition;
+				FMemory::Memcpy(&condition,data,sizeof(OP::ADDRESS)); data += sizeof(OP::ADDRESS);
+                conditions.emplace_back( condition, item );
+
+                OP::ADDRESS mask;
+				FMemory::Memcpy(&mask,data,sizeof(OP::ADDRESS)); data += sizeof(OP::ADDRESS);
+                masks.push_back( mask );
+            }
+
+
+            // Schedule next stages
+            switch (item.stage)
+            {
+            case 0:
+                if (source)
+                {
+                    // Request the conditions
+                    AddOp( SCHEDULED_OP( item.at, item, 1), conditions );
+                }
+                else
+                {
+                    GetMemory().SetMesh( item, nullptr );
+                }
+                break;
+
+            case 1:
+            {
+                // Request the source and the necessary masks
+                // \todo: store condition values in heap?
+                vector<SCHEDULED_OP> deps;
+                deps.push_back( SCHEDULED_OP( source, item ) );
+                for( size_t r=0; source && r<conditions.size(); ++r )
+                {
+                    // If there is no expression, we'll assume true.
+                    bool value = true;
+                    if (conditions[r].at)
+                    {
+                        value = GetMemory().GetBool( CACHE_ADDRESS(conditions[r].at, item) );
+                    }
+
+                    if (value)
+                    {
+                        deps.emplace_back( masks[r], item );
+                    }
+                }
+
+                if (source)
+                {
+                        AddOp( SCHEDULED_OP( item.at, item, 2), deps );
+                    }
+                break;
+            }
+
+            case 2:
+            {
+                // \todo: single remove operation with all masks?
+                Ptr<const Mesh> pSource = GetMemory().GetMesh( CACHE_ADDRESS(source,item) );
+
+                MeshPtrConst pResult = pSource;
+
+                for( size_t r=0; pResult && r<conditions.size(); ++r )
+                {
+                    // If there is no expression, we'll assume true.
+                    bool value = true;
+                    if (conditions[r].at)
+                    {
+                        value = GetMemory().GetBool( CACHE_ADDRESS(conditions[r].at, item) );
+                    }
+
+                    if (value)
+                    {
+                        Ptr<const Mesh> pMask = GetMemory().GetMesh( CACHE_ADDRESS(masks[r],item) );
+                        if (pMask)
+                        {
+                            pResult = MeshRemoveMask(pResult.get(), pMask.get());
+                        }
+                    }
+                }
+
+                GetMemory().SetMesh( item, pResult );
+                break;
+            }
+
+            default:
+                check(false);
+            }
+
+            break;
+        }
+
+        case OP_TYPE::ME_PROJECT:
+        {
+            auto args = pModel->GetPrivate()->m_program.GetOpArgs<OP::MeshProjectArgs>(item.at);
+            switch (item.stage)
+            {
+            case 0:
+                if (args.mesh)
+                {
+                        AddOp( SCHEDULED_OP( item.at, item, 1),
+                               SCHEDULED_OP( args.mesh, item),
+                               SCHEDULED_OP( args.projector, item));
+                    }
+                    else
+                    {
+                    GetMemory().SetMesh( item, nullptr );
+                }
+                break;
+
+            case 1:
+            {
+                Ptr<const Mesh> pMesh = GetMemory().GetMesh( CACHE_ADDRESS(args.mesh,item) );
+                Ptr<const Projector> pProjector = GetMemory().GetProjector( CACHE_ADDRESS(args.projector,item) );
+
+                // Only if both are valid.
+                MeshPtr pResult;
+                if (pMesh
+                    && pMesh.get()
+                    && pMesh.get()->GetVertexBuffers().GetBufferCount() > 0 )
+                {
+                    if (pProjector)
+                    {
+                        pResult = MeshProject(pMesh.get(), pProjector->m_value);
+//                        if (pResult)
+//                        {
+//                            pResult->GetPrivate()->CheckIntegrity();
+//                        }
+                    }
+                    else if (pMesh)
+                    {
+                        pResult = pMesh->Clone();
+                    }
+                }
+                GetMemory().SetMesh( item, pResult );
+                break;
+            }
+
+            default:
+                check(false);
+            }
+
+            break;
+        }
+
+        default:
+            if (type!=OP_TYPE::NONE)
+            {
+                // Operation not implemented
+                check( false );
+            }
+            break;
+        }
+    }
+
+    //---------------------------------------------------------------------------------------------
+    void CodeRunner::RunCode_Image( SCHEDULED_OP& item,
+                                    const Parameters* pParams,
+                                    const Model* pModel
+                                    )
+    {
+		MUTABLE_CPUPROFILER_SCOPE(RunCode_Image);
+
+		OP_TYPE type = pModel->GetPrivate()->m_program.GetOpType(item.at);
+		switch (type)
+        {
+
+        case OP_TYPE::IM_PARAMETER:
+        {
+				auto args = pModel->GetPrivate()->m_program.GetOpArgs<OP::ParameterArgs>(item.at);
+				EXTERNAL_IMAGE_ID id = pParams->GetImageValue(args.variable);
+				ImagePtr pResult = LoadExternalImage(id);
+				GetMemory().SetImage(item, pResult);
+            break;
+        }
+
+        case OP_TYPE::IM_LAYERCOLOUR:
+        {
+            auto args = pModel->GetPrivate()->m_program.GetOpArgs<OP::ImageLayerColourArgs>(item.at);
+            switch (item.stage)
+            {
+            case 0:
+                    AddOp( SCHEDULED_OP( item.at, item, 1),
+                           SCHEDULED_OP( args.base, item),
+                           SCHEDULED_OP::FromOpAndOptions( args.colour, item, 0),
+                           SCHEDULED_OP( args.mask, item) );
+                break;
+
+            case 1:
+				// This has been moved to a task. It should have been intercepted in IssueOp.
+				check(false);
+				break;
+
+            default:
+                check(false);
+            }
+
+            break;
+        }
+
+        case OP_TYPE::IM_LAYER:
+        {
+            auto args = pModel->GetPrivate()->m_program.GetOpArgs<OP::ImageLayerArgs>(item.at);
+            switch (item.stage)
+            {
+            case 0:
+                    AddOp( SCHEDULED_OP( item.at, item, 1),
+                           SCHEDULED_OP( args.base, item),
+                           SCHEDULED_OP( args.blended, item),
+                           SCHEDULED_OP( args.mask, item) );
+                break;
+
+            case 1:
+				// This has been moved to a task. It should have been intercepted in IssueOp.
+				check(false);
+				break;
+			
+            default:
+                check(false);
+            }
+
+            break;
+        }
+
+        case OP_TYPE::IM_MULTILAYER:
+        {
+            auto args = pModel->GetPrivate()->m_program.GetOpArgs<OP::ImageMultiLayerArgs>(item.at);
+            switch (item.stage)
+            {
+            case 0:
+                AddOp( SCHEDULED_OP( item.at, item, 1),
+                       SCHEDULED_OP( args.rangeSize, item ),
+                       SCHEDULED_OP( args.base, item) );
+                break;
+
+            case 1:
+            {
+                // We now know the number of iterations
+                int iterations = 0;
+                if (args.rangeSize)
+                {
+                    CACHE_ADDRESS address(args.rangeSize,item);
+
+                    // We support both integers and scalars here, which is not common.
+                    // \todo: review if this is necessary or we can enforce it at compile time.
+                    DATATYPE rangeSizeType = GetOpDataType( pModel->GetPrivate()->m_program.GetOpType(args.rangeSize) );
+                    if (rangeSizeType==DT_INT)
+                    {
+                        iterations = GetMemory().GetInt( address );
+                    }
+                    else if (rangeSizeType==DT_SCALAR)
+                    {
+                        iterations = int( GetMemory().GetScalar( address ) );
+                    }
+                }
+
+                // \todo Check that we are not overwriting the index (it shouldn't be set when we reach here)
+                SCHEDULED_OP itemCopy = item;
+                ExecutionIndex index = GetMemory().GetRageIndex( item.executionIndex );
+
+                    vector<SCHEDULED_OP> deps(iterations*2+1);
+                    int d=0;
+                    deps[d++] = SCHEDULED_OP( args.base, item );
+                    for (int i=0; i<iterations; ++i)
+                    {
+                        index.SetFromModelRangeIndex( args.rangeId, i );
+
+                        itemCopy.executionIndex = GetMemory().GetRageIndexIndex(index);
+
+                        deps[d++] = SCHEDULED_OP( args.blended, itemCopy );
+                        deps[d++] = SCHEDULED_OP( args.mask, itemCopy );
+                    }
+                    AddOp( SCHEDULED_OP( item.at, item, 2, iterations), deps );
+
+                break;
+            }
+
+            case 2:
+            {
+                MUTABLE_CPUPROFILER_SCOPE(ImageMultiLayer);
+
+                Ptr<const Image> pResult;
+                Ptr<Image> pNew;
+
+				int32 iterations = item.customState;
+
+                // TODO: Reuse base if possible
+                Ptr<const Image> pBase = GetMemory().GetImage( CACHE_ADDRESS(args.base,item) );
+
+                if (iterations==0)
+                {
+                    pResult = pBase;
+                }
+                else
+                {
+                    EImageFormat resultFormat = GetUncompressedFormat(pBase->GetFormat());
+                    pNew = new Image( pBase->GetSizeX(), pBase->GetSizeY(),
+                                      pBase->GetLODCount(), resultFormat );
+					pResult = pNew;
+                }
+
+                // This shouldn't happen in optimised models, but it could happen in editors, etc.
+                // \todo: raise a performance warning?
+                EImageFormat baseFormat = GetUncompressedFormat(pBase->GetFormat());
+                if ( pBase->GetFormat()!=baseFormat )
+                {
+                    pBase = ImagePixelFormat( m_pSettings->GetPrivate()->m_imageCompressionQuality,
+                                              pBase.get(), baseFormat );
+                }
+
+                SCHEDULED_OP itemCopy = item;
+                ExecutionIndex index = GetMemory().GetRageIndex( item.executionIndex );
+
+                for (int i=0; i<iterations; ++i)
+                {
+                    index.SetFromModelRangeIndex( args.rangeId, i );
+                    itemCopy.executionIndex = GetMemory().GetRageIndexIndex(index);
+					itemCopy.customState = 0;
+
+                    Ptr<const Image> pBlended = GetMemory().GetImage( CACHE_ADDRESS(args.blended,itemCopy) );
+
+                    // This shouldn't happen in optimised models, but it could happen in editors, etc.
+                    // \todo: raise a performance warning?
+                    if (pBlended && pBlended->GetFormat()!=baseFormat )
+                    {
+                        pBlended =
+                            ImagePixelFormat( m_pSettings->GetPrivate()->m_imageCompressionQuality,
+                                              pBlended.get(), baseFormat );
+                    }
+
+					// TODO: This shouldn't happen, but be defensive.
+					FImageSize ResultSize = pBase->GetSize();
+					if (pBlended && pBlended->GetSize() != ResultSize)
+					{
+						MUTABLE_CPUPROFILER_SCOPE(ImageResize_MaskFixForMultilayer);
+						pBlended = ImageResizeLinear(0, pBlended.get(), ResultSize);
+					}
+
+					if (pBlended->GetLODCount() < pBase->GetLODCount())
+					{
+						MUTABLE_CPUPROFILER_SCOPE(ImageLayer_MipBlendedEmergencyFix);
+
+						int levelCount = pBase->GetLODCount();
+						ImagePtr pDest = new Image(pBlended->GetSizeX(), pBlended->GetSizeY(),
+							levelCount,
+							pBlended->GetFormat());
+
+						SCRATCH_IMAGE_MIPMAP scratch;
+						FMipmapGenerationSettings settings{};
+
+						ImageMipmap_PrepareScratch(pDest.get(), pBlended.get(), levelCount, &scratch);
+						ImageMipmap(m_pSettings->GetPrivate()->m_imageCompressionQuality, pDest.get(), pBlended.get(), levelCount,
+							&scratch, settings);
+
+						pBlended = pDest;
+					}
+
+                    if (args.mask)
+                    {
+                        Ptr<const Image> pMask = GetMemory().GetImage( CACHE_ADDRESS(args.mask,itemCopy) );
+
+						// TODO: This shouldn't happen, but be defensive.
+						if (pMask && pMask->GetSize() != ResultSize)
+						{
+							MUTABLE_CPUPROFILER_SCOPE(ImageResize_MaskFixForMultilayer);
+							pMask = ImageResizeLinear(0, pMask.get(), ResultSize);
+						}
+
+                        // emergy fix c36adf47-e40d-490f-b709-41142bafad78
+                        if (pMask->GetLODCount()<pBase->GetLODCount())
+                        {
+                            MUTABLE_CPUPROFILER_SCOPE(ImageLayer_EmergencyFix);
+
+                            int levelCount = pBase->GetLODCount();
+                            ImagePtr pDest = new Image( pMask->GetSizeX(), pMask->GetSizeY(),
+                                                        levelCount,
+                                                        pMask->GetFormat() );
+							
+                            SCRATCH_IMAGE_MIPMAP scratch;
+							FMipmapGenerationSettings mipSettings{};
+                        
+                            ImageMipmap_PrepareScratch( pDest.get(), pMask.get(), levelCount, &scratch );
+                            ImageMipmap( m_pSettings->GetPrivate()->m_imageCompressionQuality,
+                                         pDest.get(), pMask.get(), levelCount, &scratch, mipSettings );
+                            pMask = pDest;
+                        }
+
+                        switch (EBlendType(args.blendType))
+                        {
+						case EBlendType::BT_NORMAL_COMBINE: ImageNormalCombine(pNew.get(), pBase.get(), pMask.get(), pBlended.get()); break; 
+                        case EBlendType::BT_SOFTLIGHT: ImageSoftLight( pNew.get(), pBase.get(), pMask.get(), pBlended.get(), false ); break;
+                        case EBlendType::BT_HARDLIGHT: ImageHardLight( pNew.get(), pBase.get(), pMask.get(), pBlended.get(), false ); break;
+                        case EBlendType::BT_BURN: ImageBurn( pNew.get(), pBase.get(), pMask.get(), pBlended.get(), false ); break;
+                        case EBlendType::BT_DODGE: ImageDodge( pNew.get(), pBase.get(), pMask.get(), pBlended.get(), false ); break;
+                        case EBlendType::BT_SCREEN: ImageScreen( pNew.get(), pBase.get(), pMask.get(), pBlended.get(), false ); break;
+                        case EBlendType::BT_OVERLAY: ImageOverlay( pNew.get(), pBase.get(), pMask.get(), pBlended.get(), false ); break;
+                        case EBlendType::BT_ALPHA_OVERLAY: ImageAlphaOverlay( pNew.get(), pBase.get(), pMask.get(), pBlended.get(), false ); break;
+                        case EBlendType::BT_MULTIPLY: ImageMultiply( pNew.get(), pBase.get(), pMask.get(), pBlended.get(), false ); break;
+                        case EBlendType::BT_BLEND: ImageBlend( pNew.get(), pBase.get(), pMask.get(), pBlended.get(), false ); break;
+                        default: check(false);
+                        }
+
+                    }
+                    else
+                    {
+                        switch (EBlendType(args.blendType))
+                        {
+						case EBlendType::BT_NORMAL_COMBINE: ImageNormalCombine(pNew.get(), pBase.get(), pBlended.get()); break;
+                        case EBlendType::BT_SOFTLIGHT: ImageSoftLight( pNew.get(), pBase.get(), pBlended.get(), false ); break;
+                        case EBlendType::BT_HARDLIGHT: ImageHardLight( pNew.get(), pBase.get(), pBlended.get(), false ); break;
+                        case EBlendType::BT_BURN: ImageBurn( pNew.get(), pBase.get(), pBlended.get(), false ); break;
+                        case EBlendType::BT_DODGE: ImageDodge( pNew.get(), pBase.get(), pBlended.get(), false ); break;
+                        case EBlendType::BT_SCREEN: ImageScreen( pNew.get(), pBase.get(), pBlended.get(), false ); break;
+                        case EBlendType::BT_OVERLAY: ImageOverlay( pNew.get(), pBase.get(), pBlended.get(), false ); break;
+                        case EBlendType::BT_ALPHA_OVERLAY: ImageAlphaOverlay( pNew.get(), pBase.get(), pBlended.get(), false ); break;
+                        case EBlendType::BT_MULTIPLY: ImageMultiply( pNew.get(), pBase.get(), pBlended.get(), false ); break;
+                        case EBlendType::BT_BLEND: pNew = pBase->Clone(); break;
+                        default: check(false);
+                        }
+                    }
+
+                    pResult = pNew;
+                    pBase = pNew;
+                    // \todo: optimise with ping pong buffers or blend-in-place
+                    pNew = new Image( pBase->GetSizeX(), pBase->GetSizeY(),
+                                      pBase->GetLODCount(), pBase->GetFormat() );
+				}
+
+                GetMemory().SetImage( item, pResult );
+                break;
+            }
+
+            default:
+                check(false);
+            }
+
+            break;
+        }
+
+        case OP_TYPE::IM_DIFFERENCE:
+        {
+            auto args = pModel->GetPrivate()->m_program.GetOpArgs<OP::ImageDifferenceArgs>(item.at);
+            switch (item.stage)
+            {
+            case 0:
+                    AddOp( SCHEDULED_OP( item.at, item, 1),
+                           SCHEDULED_OP( args.a, item),
+                           SCHEDULED_OP( args.b, item) );
+                break;
+
+            case 1:
+            {
+                // TODO: Reuse base if possible
+                Ptr<const Image> pA = GetMemory().GetImage( CACHE_ADDRESS(args.a,item) );
+                Ptr<const Image> pB = GetMemory().GetImage( CACHE_ADDRESS(args.b,item) );
+
+                ImagePtr pResult = ImageDifference( pA.get(), pB.get() );
+
+                GetMemory().SetImage( item, pResult );
+                break;
+            }
+
+            default:
+                check(false);
+            }
+
+			break;
+		}
+
+		case OP_TYPE::IM_NORMALCOMPOSITE:
+		{
+			OP::ImageNormalCompositeArgs args = pModel->GetPrivate()->m_program.GetOpArgs<OP::ImageNormalCompositeArgs>(item.at);
+			switch (item.stage)
+			{
+			case 0:
+				if (args.base && args.normal)
+				{
+						AddOp(SCHEDULED_OP(item.at, item, 1),
+							  SCHEDULED_OP(args.base, item),
+							  SCHEDULED_OP(args.normal, item));
+					}
+					else
+					{
+					GetMemory().SetImage(item, nullptr);
+				}
+				break;
+
+			case 1:
+			{
+				Ptr<const Image> Base = GetMemory().GetImage(CACHE_ADDRESS(args.base, item));
+				Ptr<const Image> Normal = GetMemory().GetImage(CACHE_ADDRESS(args.normal, item));
+
+				if (Normal->GetLODCount() < Base->GetLODCount())
+				{
+					MUTABLE_CPUPROFILER_SCOPE(ImageNormalComposite_EmergencyFix);
+
+					int levelCount = Base->GetLODCount();
+					ImagePtr pDest = new Image(Normal->GetSizeX(), Normal->GetSizeY(),
+						levelCount,
+						Normal->GetFormat());
+
+					SCRATCH_IMAGE_MIPMAP scratch;
+					FMipmapGenerationSettings mipSettings{};
+					ImageMipmap_PrepareScratch(pDest.get(), Normal.get(), levelCount, &scratch);
+					ImageMipmap(m_pSettings->GetPrivate()->m_imageCompressionQuality,
+						pDest.get(), Normal.get(), levelCount, &scratch, mipSettings);
+					Normal = pDest;
+				}
+
+
+                ImagePtr pResult = new Image( Base->GetSizeX(), Base->GetSizeY(), Base->GetLODCount(), Base->GetFormat() );
+				ImageNormalComposite(pResult.get(), Base.get(), Normal.get(), args.mode, args.power);
+
+				GetMemory().SetImage(item, pResult);
+				break;
+			}
+
+			default:
+				check(false);
+			}
+
+			break;
+		}
+
+        case OP_TYPE::IM_PIXELFORMAT:
+        {
+			OP::ImagePixelFormatArgs args = pModel->GetPrivate()->m_program.GetOpArgs<OP::ImagePixelFormatArgs>(item.at);
+            switch (item.stage)
+            {
+            case 0:
+                    AddOp( SCHEDULED_OP( item.at, item, 1),
+                           SCHEDULED_OP( args.source, item) );
+                break;
+
+            case 1:
+            {
+				// This has been moved to a task. It should have been intercepted in IssueOp.
+				check(false);
+				break;
+            }
+
+            default:
+                check(false);
+            }
+
+            break;
+        }
+
+        case OP_TYPE::IM_MIPMAP:
+        {
+			OP::ImageMipmapArgs args = pModel->GetPrivate()->m_program.GetOpArgs<OP::ImageMipmapArgs>(item.at);
+            switch (item.stage)
+            {
+            case 0:
+				AddOp( SCHEDULED_OP( item.at, item, 1),
+					SCHEDULED_OP( args.source, item) );
+                break;
+
+            case 1:
+				// This has been moved to a task. It should have been intercepted in IssueOp.
+				check(false);
+				break;
+
+            default:
+                check(false);
+            }
+
+            break;
+        }
+
+        case OP_TYPE::IM_RESIZE:
+        {
+			OP::ImageResizeArgs args = pModel->GetPrivate()->m_program.GetOpArgs<OP::ImageResizeArgs>(item.at);
+            switch (item.stage)
+            {
+            case 0:
+                    AddOp( SCHEDULED_OP( item.at, item, 1),
+                           SCHEDULED_OP( args.source, item) );
+                break;
+
+            case 1:
+            {
+                Ptr<const Image> pBase = GetMemory().GetImage( CACHE_ADDRESS(args.source,item) );
+
+                FImageSize destSize = FImageSize
+                    (
+                        args.size[0],
+                        args.size[1]
+                    );
+
+                ImagePtr pResult;
+                if ( destSize[0]!=pBase->GetSizeX()
+                     ||
+                     destSize[1]!=pBase->GetSizeY() )
+                {
+                    //pResult = ImageResize( pBase.get(), destSize );
+                    pResult =
+                        ImageResizeLinear( m_pSettings->GetPrivate()->m_imageCompressionQuality,
+                                           pBase.get(), destSize );
+
+                    // If the source image had mips, generate them as well for the resized image.
+                    // This shouldn't happen often since "ResizeLike" should be usually optimised out
+                    // during model compilation. The mipmap generation below is not very precise with
+                    // the number of mips that are needed and will probably generate too many
+                    bool sourceHasMips = pBase->GetLODCount()>1;
+                    if (sourceHasMips)
+                    {
+                        int levelCount = Image::GetMipmapCount( pResult->GetSizeX(), pResult->GetSizeY() );
+                        ImagePtr pMipmapped = new Image( pResult->GetSizeX(), pResult->GetSizeY(),
+                                                         levelCount,
+                                                         pResult->GetFormat() );
+
+                        SCRATCH_IMAGE_MIPMAP scratch;
+						FMipmapGenerationSettings mipSettings{};
+
+                        ImageMipmap_PrepareScratch( pMipmapped.get(), pResult.get(), levelCount, &scratch );
+                        ImageMipmap( m_pSettings->GetPrivate()->m_imageCompressionQuality,
+                                     pMipmapped.get(), pResult.get(), levelCount, &scratch, mipSettings );
+
+                        pResult = pMipmapped;
+                    }
+
+                }
+                else
+                {
+                    pResult = pBase->Clone();
+                }
+
+                GetMemory().SetImage( item, pResult );
+                break;
+            }
+
+            default:
+                check(false);
+            }
+
+            break;
+        }
+
+        case OP_TYPE::IM_RESIZELIKE:
+        {
+			OP::ImageResizeLikeArgs args = pModel->GetPrivate()->m_program.GetOpArgs<OP::ImageResizeLikeArgs>(item.at);
+            switch (item.stage)
+            {
+            case 0:
+                    AddOp( SCHEDULED_OP( item.at, item, 1),
+                           SCHEDULED_OP( args.source, item),
+                           SCHEDULED_OP( args.sizeSource, item) );
+                break;
+
+            case 1:
+            {
+                Ptr<const Image> pBase = GetMemory().GetImage( CACHE_ADDRESS(args.source,item) );
+                Ptr<const Image> pSizeBase = GetMemory().GetImage( CACHE_ADDRESS(args.sizeSource,item) );
+                ImagePtr pResult;
+
+                if ( pBase->GetSizeX()!=pSizeBase->GetSizeX()
+                     ||
+                     pBase->GetSizeY()!=pSizeBase->GetSizeY() )
+                {
+                    FImageSize destSize( pSizeBase->GetSizeX(), pSizeBase->GetSizeY() );
+
+                    pResult = ImageResize( pBase.get(), destSize );
+
+                    // If the source image had mips, generate them as well for the resized image.
+                    // This shouldn't happen often since "ResizeLike" should be usually optimised out
+                    // during model compilation. The mipmap generation below is not very precise with
+                    // the number of mips that are needed and will probably generate too many
+                    bool sourceHasMips = pBase->GetLODCount()>1;
+                    if (sourceHasMips)
+                    {
+                        int levelCount = Image::GetMipmapCount( pResult->GetSizeX(), pResult->GetSizeY() );
+                        ImagePtr pMipmapped = new Image( pResult->GetSizeX(), pResult->GetSizeY(),
+                                                         levelCount,
+                                                         pResult->GetFormat() );
+
+                        SCRATCH_IMAGE_MIPMAP scratch;
+						FMipmapGenerationSettings mipSettings{};
+
+                        ImageMipmap_PrepareScratch( pMipmapped.get(), pResult.get(), levelCount, &scratch );
+                        ImageMipmap( m_pSettings->GetPrivate()->m_imageCompressionQuality,
+                                     pMipmapped.get(), pResult.get(), levelCount, &scratch, mipSettings );
+
+                        pResult = pMipmapped;
+                    }
+                }
+                else
+                {
+                    pResult = pBase->Clone();
+                }
+
+                GetMemory().SetImage( item, pResult );
+                break;
+            }
+
+            default:
+                check(false);
+            }
+
+            break;
+        }
+
+        case OP_TYPE::IM_RESIZEREL:
+        {
+            auto args = pModel->GetPrivate()->m_program.GetOpArgs<OP::ImageResizeRelArgs>(item.at);
+            switch (item.stage)
+            {
+            case 0:
+                    AddOp( SCHEDULED_OP( item.at, item, 1),
+                           SCHEDULED_OP( args.source, item) );
+                break;
+
+            case 1:
+            {
+                Ptr<const Image> pBase = GetMemory().GetImage( CACHE_ADDRESS(args.source,item) );
+
+                FImageSize destSize(
+                            uint16_t( pBase->GetSizeX()*args.factor[0] + 0.5f ),
+                            uint16_t( pBase->GetSizeY()*args.factor[1] + 0.5f ) );
+
+                //pResult = ImageResize( pBase.get(), destSize );
+                ImagePtr pResult = ImageResizeLinear(
+                    m_pSettings->GetPrivate()->m_imageCompressionQuality, pBase.get(), destSize );
+
+                // If the source image had mips, generate them as well for the resized image.
+                // This shouldn't happen often since "ResizeLike" should be usually optimised out
+                // during model compilation. The mipmap generation below is not very precise with
+                // the number of mips that are needed and will probably generate too many
+                bool sourceHasMips = pBase->GetLODCount()>1;
+                if (sourceHasMips)
+                {
+                    int levelCount = Image::GetMipmapCount( pResult->GetSizeX(), pResult->GetSizeY() );
+                    ImagePtr pMipmapped = new Image( pResult->GetSizeX(), pResult->GetSizeY(),
+                                                     levelCount,
+                                                     pResult->GetFormat() );
+
+                    SCRATCH_IMAGE_MIPMAP scratch;
+					FMipmapGenerationSettings mipSettings{};
+
+                    ImageMipmap_PrepareScratch( pMipmapped.get(), pResult.get(), levelCount, &scratch );
+                    ImageMipmap( m_pSettings->GetPrivate()->m_imageCompressionQuality,
+                                 pMipmapped.get(), pResult.get(), levelCount, &scratch, mipSettings );
+
+                    pResult = pMipmapped;
+                }
+
+                GetMemory().SetImage( item, pResult );
+                break;
+            }
+
+            default:
+                check(false);
+            }
+
+
+            break;
+        }
+
+        case OP_TYPE::IM_BLANKLAYOUT:
+        {
+            auto args = pModel->GetPrivate()->m_program.GetOpArgs<OP::ImageBlankLayoutArgs>(item.at);
+            switch (item.stage)
+            {
+            case 0:
+                    AddOp( SCHEDULED_OP( item.at, item, 1),
+                           SCHEDULED_OP::FromOpAndOptions( args.layout, item, 0) );
+                break;
+
+            case 1:
+            {
+                Ptr<const Layout> pLayout = GetMemory().GetLayout(SCHEDULED_OP::FromOpAndOptions(args.layout, item, 0));
+
+                FIntPoint SizeInBlocks = pLayout->GetGridSize();
+
+				FIntPoint BlockSizeInPixels(args.blockSize[0], args.blockSize[1]);
+
+				// Image size if we don't skip any mipmap
+				FIntPoint FullImageSizeInPixels = SizeInBlocks * BlockSizeInPixels;
+				int32 FullImageMipCount = Image::GetMipmapCount(FullImageSizeInPixels.X, FullImageSizeInPixels.Y);
+
+				FIntPoint ImageSizeInPixels = FullImageSizeInPixels;
+				int32 MipsToSkip = item.executionOptions;
+				MipsToSkip = FMath::Min(MipsToSkip, FullImageMipCount);
+				if (MipsToSkip > 0)
+				{
+					//FIntPoint ReducedBlockSizeInPixels;
+
+					// This method tries to reduce only the block size, but it fails if the image is still too big
+					// If we want to generate only a subset of mipmaps, reduce the layout block size accordingly.
+					//ReducedBlockSizeInPixels.X = BlockSizeInPixels.X >> MipsToSkip;
+					//ReducedBlockSizeInPixels.Y = BlockSizeInPixels.Y >> MipsToSkip;
+					//const FImageFormatData& FormatData = GetImageFormatData((EImageFormat)args.format);
+					//int MinBlockSize = FMath::Max(FormatData.m_pixelsPerBlockX, FormatData.m_pixelsPerBlockY);
+					//ReducedBlockSizeInPixels.X = FMath::Max<int32>(ReducedBlockSizeInPixels.X, FormatData.m_pixelsPerBlockX);
+					//ReducedBlockSizeInPixels.Y = FMath::Max<int32>(ReducedBlockSizeInPixels.Y, FormatData.m_pixelsPerBlockY);
+					//FIntPoint ReducedImageSizeInPixels = SizeInBlocks * ReducedBlockSizeInPixels;
+
+					// This method simply reduces the size and assumes all the other operations will handle degeenrate cases.
+					ImageSizeInPixels = FullImageSizeInPixels / (1 << MipsToSkip);
+					
+					//if (ReducedImageSizeInPixels!= ImageSizeInPixels)
+					//{
+					//	check(false);
+					//}
+				}
+
+                int MipsToGenerate = 1;
+                if ( args.generateMipmaps )
+                {
+                    if ( args.mipmapCount==0 )
+                    {
+						MipsToGenerate = Image::GetMipmapCount(ImageSizeInPixels.X, ImageSizeInPixels.Y);
+                    }
+                    else
+                    {
+						MipsToGenerate = FMath::Max(args.mipmapCount-MipsToSkip,1);
+                    }
+                }
+
+                ImagePtr pNew = new Image(ImageSizeInPixels.X, ImageSizeInPixels.Y, MipsToGenerate, EImageFormat(args.format) );
+
+                FMemory::Memzero( pNew->GetData(), pNew->GetDataSize() );
+
+                GetMemory().SetImage( item, pNew );
+                break;
+            }
+
+            default:
+                check(false);
+            }
+
+
+            break;
+        }
+
+        case OP_TYPE::IM_COMPOSE:
+        {
+            auto args = pModel->GetPrivate()->m_program.GetOpArgs<OP::ImageComposeArgs>( item.at );
+            switch ( item.stage )
+            {
+            case 0:
+                AddOp( SCHEDULED_OP( item.at, item, 1 ), 
+					SCHEDULED_OP::FromOpAndOptions( args.layout, item, 0 ) );
+                break;
+
+            case 1:
+            {
+                Ptr<const Layout> pLayout =
+                    GetMemory().GetLayout( CACHE_ADDRESS( args.layout, SCHEDULED_OP::FromOpAndOptions(args.layout, item, 0)) );
+
+                SCHEDULED_OP_DATA data;
+                data.layout = pLayout;
+				int32 dataPos = m_heapData.Add( data );
+
+                int relBlockIndex = pLayout->FindBlock( args.blockIndex );
+                if ( relBlockIndex >= 0 )
+                {
+                    AddOp( SCHEDULED_OP( item.at, item, 2, dataPos ),
+                           SCHEDULED_OP( args.base, item ),
+                           SCHEDULED_OP( args.blockImage, item ),
+                           SCHEDULED_OP( args.mask, item ) );
+                }
+                else
+                {
+                    AddOp( SCHEDULED_OP( item.at, item, 2, dataPos ),
+                           SCHEDULED_OP( args.base, item ) );
+                }
+                break;
+            }
+
+            case 2:
+				// This has been moved to a task. It should have been intercepted in IssueOp.
+				check(false);
+                break;
+
+            default:
+                check( false );
+            }
+
+            break;
+        }
+
+        case OP_TYPE::IM_INTERPOLATE:
+        {
+            auto args = pModel->GetPrivate()->m_program.GetOpArgs<OP::ImageInterpolateArgs>(item.at);
+            switch (item.stage)
+            {
+            case 0:
+                AddOp( SCHEDULED_OP( item.at, item, 1),
+                       SCHEDULED_OP( args.factor, item) );
+                break;
+
+            case 1:
+            {
+                // Targets must be consecutive
+                int count = 0;
+                for ( int i=0
+                    ; i<MUTABLE_OP_MAX_INTERPOLATE_COUNT && args.targets[i]
+                    ; ++i )
+                {
+                    count++;
+                }
+
+                float factor = GetMemory().GetScalar( CACHE_ADDRESS(args.factor,item) );
+
+                float delta = 1.0f/(count-1);
+                int min = (int)floorf( factor/delta );
+                int max = (int)ceilf( factor/delta );
+
+                float bifactor = factor/delta - min;
+
+                SCHEDULED_OP_DATA data;
+                data.bifactor = bifactor;
+				data.min = FMath::Clamp(min, 0, count - 1);
+				data.max = FMath::Clamp(max, 0, count - 1);
+				uint32 dataPos = uint32(m_heapData.Add(data));
+
+                if ( bifactor < SMALL_NUMBER )
+                {
+                        AddOp( SCHEDULED_OP( item.at, item, 2, dataPos),
+                               SCHEDULED_OP( args.targets[min], item) );
+                    }
+                else if ( bifactor > 1.0f-SMALL_NUMBER )
+                {
+                        AddOp( SCHEDULED_OP( item.at, item, 2, dataPos),
+                               SCHEDULED_OP( args.targets[max], item) );
+                    }
+                    else
+                    {
+                        AddOp( SCHEDULED_OP( item.at, item, 2, dataPos),
+                               SCHEDULED_OP( args.targets[min], item),
+                               SCHEDULED_OP( args.targets[max], item) );
+                    }
+                break;
+            }
+
+            case 2:
+            {
+                // Targets must be consecutive
+                int count = 0;
+                for ( int i=0
+                    ; i<MUTABLE_OP_MAX_INTERPOLATE_COUNT && args.targets[i]
+                    ; ++i )
+                {
+                    count++;
+                }
+
+                // Factor from 0 to 1 between the two targets
+                const SCHEDULED_OP_DATA& data = m_heapData[(size_t)item.customState];
+                float bifactor = data.bifactor;
+                int min = data.min;
+                int max = data.max;
+
+                ImagePtr pResult;
+                if ( bifactor < SMALL_NUMBER )
+                {
+                    Ptr<const Image> pSource = GetMemory().GetImage( CACHE_ADDRESS(args.targets[min],item) );
+                    pResult = pSource->Clone();
+                }
+                else if ( bifactor > 1.0f-SMALL_NUMBER )
+                {
+                    Ptr<const Image> pSource = GetMemory().GetImage( CACHE_ADDRESS(args.targets[max],item) );
+                    pResult = pSource->Clone();
+                }
+                else
+                {
+                    Ptr<const Image> pMin = GetMemory().GetImage( CACHE_ADDRESS(args.targets[min],item) );
+                    Ptr<const Image> pMax = GetMemory().GetImage( CACHE_ADDRESS(args.targets[max],item) );
+
+                    if (pMin && pMax)
+                    {
+                        ImagePtr pNew = new Image( pMin->GetSizeX(),
+                                             pMin->GetSizeY(),
+                                             pMin->GetLODCount(),
+                                             pMin->GetFormat() );
+
+						// Be defensive: ensure image sizes match.
+						if (pMin->GetSize() != pMax->GetSize())
+						{
+							MUTABLE_CPUPROFILER_SCOPE(ImageResize_ForInterpolate);
+							pMax = ImageResizeLinear(0, pMax.get(), pMin->GetSize());
+						}
+
+						if (pMax->GetLODCount() != pMin->GetLODCount())
+						{
+							MUTABLE_CPUPROFILER_SCOPE(Mipmap_ForInterpolate);
+
+							int32 levelCount = pMin->GetLODCount();
+							ImagePtr pDest = new Image(pMax->GetSizeX(), pMax->GetSizeY(), levelCount, pMax->GetFormat());
+
+							SCRATCH_IMAGE_MIPMAP scratch;
+							FMipmapGenerationSettings settings{};
+
+							ImageMipmap_PrepareScratch(pDest.get(), pMax.get(), levelCount, &scratch);
+							ImageMipmap(m_pSettings->GetPrivate()->m_imageCompressionQuality, pDest.get(), pMax.get(), levelCount,
+								&scratch, settings);
+
+							pMax = pDest;
+						}
+
+
+                        ImageInterpolate( pNew.get(), pMin.get(), pMax.get(), bifactor );
+
+                        pResult = pNew;
+                    }
+                    else if (pMin)
+                    {
+                        pResult = pMin->Clone();
+                    }
+                    else if (pMax)
+                    {
+                        pResult = pMax->Clone();
+                    }
+                }
+
+                GetMemory().SetImage( item, pResult );
+                break;
+            }
+
+            default:
+                check(false);
+            }
+
+            break;
+        }
+
+        case OP_TYPE::IM_INTERPOLATE3:
+        {
+            auto args = pModel->GetPrivate()->m_program.GetOpArgs<OP::ImageInterpolate3Args>(item.at);
+            switch (item.stage)
+            {
+            case 0:
+                    AddOp( SCHEDULED_OP( item.at, item, 1),
+                           SCHEDULED_OP::FromOpAndOptions( args.factor1, item, 0),
+                           SCHEDULED_OP::FromOpAndOptions( args.factor2, item, 0),
+                           SCHEDULED_OP( args.target0, item),
+                           SCHEDULED_OP( args.target1, item),
+                           SCHEDULED_OP( args.target2, item) );
+                break;
+
+            case 1:
+            {
+                // \TODO Optimise limit cases
+
+                float factor1 = GetMemory().GetScalar(SCHEDULED_OP::FromOpAndOptions(args.factor1, item, 0));
+                float factor2 = GetMemory().GetScalar(SCHEDULED_OP::FromOpAndOptions(args.factor2, item, 0));
+
+                Ptr<const Image> pTarget0 = GetMemory().GetImage( CACHE_ADDRESS(args.target0,item) );
+                Ptr<const Image> pTarget1 = GetMemory().GetImage( CACHE_ADDRESS(args.target1,item) );
+                Ptr<const Image> pTarget2 = GetMemory().GetImage( CACHE_ADDRESS(args.target2,item) );
+
+                ImagePtr pResult = ImageInterpolate( pTarget0.get(), pTarget1.get(), pTarget2.get(),
+                                                     factor1, factor2 );
+
+                GetMemory().SetImage( item, pResult );
+                break;
+            }
+
+            default:
+                check(false);
+            }
+
+            break;
+        }
+
+        case OP_TYPE::IM_SATURATE:
+        {
+            auto args = pModel->GetPrivate()->m_program.GetOpArgs<OP::ImageSaturateArgs>(item.at);
+            switch (item.stage)
+            {
+            case 0:
+                    AddOp( SCHEDULED_OP( item.at, item, 1),
+                           SCHEDULED_OP( args.base, item ),
+                           SCHEDULED_OP::FromOpAndOptions( args.factor, item, 0 ));
+                break;
+
+            case 1:
+            {
+                Ptr<const Image> pBase = GetMemory().GetImage( CACHE_ADDRESS(args.base,item) );
+                float factor = GetMemory().GetScalar(SCHEDULED_OP::FromOpAndOptions(args.factor, item, 0));
+
+                ImagePtr pResult = ImageSaturate( pBase.get(), factor );
+
+                GetMemory().SetImage( item, pResult );
+                break;
+            }
+
+            default:
+                check(false);
+            }
+
+            break;
+        }
+
+        case OP_TYPE::IM_LUMINANCE:
+        {
+            auto args = pModel->GetPrivate()->m_program.GetOpArgs<OP::ImageLuminanceArgs>(item.at);
+            switch (item.stage)
+            {
+            case 0:
+                AddOp( SCHEDULED_OP( item.at, item, 1),
+                        SCHEDULED_OP( args.base, item ) );
+                break;
+
+            case 1:
+            {
+                Ptr<const Image> pBase = GetMemory().GetImage( CACHE_ADDRESS(args.base,item) );
+
+                ImagePtr pResult = ImageLuminance( pBase.get() );
+                GetMemory().SetImage( item, pResult );
+                break;
+            }
+
+            default:
+                check(false);
+            }
+
+            break;
+        }
+
+        case OP_TYPE::IM_SWIZZLE:
+        {
+            auto args = pModel->GetPrivate()->m_program.GetOpArgs<OP::ImageSwizzleArgs>(item.at);
+            switch (item.stage)
+            {
+            case 0:
+                AddOp( SCHEDULED_OP( item.at, item, 1),
+                        SCHEDULED_OP( args.sources[0], item ),
+                        SCHEDULED_OP( args.sources[1], item ),
+                        SCHEDULED_OP( args.sources[2], item ),
+                        SCHEDULED_OP( args.sources[3], item ) );
+                break;
+
+            case 1:
+				// This has been moved to a task. It should have been intercepted in IssueOp.
+				check(false);
+				break;
+
+            default:
+                check(false);
+            }
+
+            break;
+        }
+
+        case OP_TYPE::IM_SELECTCOLOUR:
+        {
+            auto args = pModel->GetPrivate()->m_program.GetOpArgs<OP::ImageSelectColourArgs>(item.at);
+            switch (item.stage)
+            {
+            case 0:
+                AddOp( SCHEDULED_OP( item.at, item, 1),
+                        SCHEDULED_OP( args.base, item ),
+                        SCHEDULED_OP::FromOpAndOptions( args.colour, item, 0 ) );
+                break;
+
+            case 1:
+            {
+                Ptr<const Image> pBase = GetMemory().GetImage( CACHE_ADDRESS(args.base,item) );
+                vec4<float> colour = GetMemory().GetColour(SCHEDULED_OP::FromOpAndOptions(args.colour, item, 0));
+
+               ImagePtr pResult = ImageSelectColour( pBase.get(), colour.xyz() );
+
+                GetMemory().SetImage( item, pResult );
+                break;
+            }
+
+            default:
+                check(false);
+            }
+
+            break;
+        }
+
+        case OP_TYPE::IM_COLOURMAP:
+        {
+            auto args = pModel->GetPrivate()->m_program.GetOpArgs<OP::ImageColourMapArgs>(item.at);
+            switch (item.stage)
+            {
+            case 0:
+                    AddOp( SCHEDULED_OP( item.at, item, 1),
+                           SCHEDULED_OP( args.base, item ),
+                           SCHEDULED_OP( args.mask, item ),
+                           SCHEDULED_OP( args.map, item ) );
+                break;
+
+            case 1:
+            {
+                Ptr<const Image> pSource = GetMemory().GetImage( CACHE_ADDRESS(args.base,item) );
+                Ptr<const Image> pMask = GetMemory().GetImage( CACHE_ADDRESS(args.mask,item) );
+                Ptr<const Image> pMap = GetMemory().GetImage( CACHE_ADDRESS(args.map,item) );
+
+				// Be defensive: ensure image sizes match.
+				if (pMask->GetSize() != pSource->GetSize())
+				{
+					MUTABLE_CPUPROFILER_SCOPE(ImageResize_ForColourmap);
+					pMask = ImageResizeLinear(0, pMask.get(), pSource->GetSize());
+				}
+
+
+                ImagePtr pResult = ImageColourMap( pSource.get(), pMask.get(), pMap.get() );
+
+                GetMemory().SetImage( item, pResult );
+                break;
+            }
+
+            default:
+                check(false);
+            }
+
+            break;
+        }
+
+        case OP_TYPE::IM_GRADIENT:
+        {
+            auto args = pModel->GetPrivate()->m_program.GetOpArgs<OP::ImageGradientArgs>(item.at);
+            switch (item.stage)
+            {
+            case 0:
+                    AddOp( SCHEDULED_OP( item.at, item, 1),
+                           SCHEDULED_OP::FromOpAndOptions( args.colour0, item, 0 ),
+                           SCHEDULED_OP::FromOpAndOptions( args.colour1, item, 0 ) );
+                break;
+
+            case 1:
+            {
+                vec4<float> colour0 = GetMemory().GetColour(SCHEDULED_OP::FromOpAndOptions(args.colour0, item, 0));
+                vec4<float> colour1 = GetMemory().GetColour(SCHEDULED_OP::FromOpAndOptions(args.colour1, item, 0));
+
+                ImagePtr pResult = ImageGradient( colour0.xyz(), colour1.xyz(),
+                                                  args.size[0],
+                                                  args.size[1] );
+
+                GetMemory().SetImage( item, pResult );
+                break;
+            }
+
+            default:
+                check(false);
+            }
+
+            break;
+        }
+
+        case OP_TYPE::IM_BINARISE:
+        {
+            auto args = pModel->GetPrivate()->m_program.GetOpArgs<OP::ImageBinariseArgs>(item.at);
+            switch (item.stage)
+            {
+            case 0:
+                    AddOp( SCHEDULED_OP( item.at, item, 1),
+                           SCHEDULED_OP( args.base, item ),
+                           SCHEDULED_OP::FromOpAndOptions( args.threshold, item, 0 ) );
+                break;
+
+            case 1:
+            {
+                Ptr<const Image> pA = GetMemory().GetImage( CACHE_ADDRESS(args.base,item) );
+
+                float c = GetMemory().GetScalar(SCHEDULED_OP::FromOpAndOptions(args.threshold, item, 0));
+
+                ImagePtr pResult = ImageBinarise( pA.get(), c );
+
+                GetMemory().SetImage( item, pResult );
+                break;
+            }
+
+            default:
+                check(false);
+            }
+
+            break;
+        }
+
+		case OP_TYPE::IM_INVERT:
+		{
+			OP::ImageInvertArgs args = pModel->GetPrivate()->m_program.GetOpArgs<OP::ImageInvertArgs>(item.at);
+			switch (item.stage)
+			{
+			case 0:
+					AddOp(SCHEDULED_OP(item.at, item, 1),
+						SCHEDULED_OP(args.base, item));
+				break;
+
+			case 1:
+			{
+				Ptr<const Image> pA = GetMemory().GetImage(CACHE_ADDRESS(args.base, item));
+
+				ImagePtr pResult = ImageInvert(pA.get());
+
+				GetMemory().SetImage(item, pResult);
+				break;
+			}
+
+			default:
+				check(false);
+			}
+
+			break;
+		}
+
+        case OP_TYPE::IM_PLAINCOLOUR:
+        {
+			OP::ImagePlainColourArgs args = pModel->GetPrivate()->m_program.GetOpArgs<OP::ImagePlainColourArgs>(item.at);
+            switch (item.stage)
+            {
+            case 0:
+				AddOp( SCHEDULED_OP( item.at, item, 1),
+					SCHEDULED_OP::FromOpAndOptions( args.colour, item, 0 ) );
+                break;
+
+            case 1:
+            {
+                mu::vec4f c = GetMemory().GetColour(SCHEDULED_OP::FromOpAndOptions(args.colour, item, 0));
+
+				uint16 SizeX = args.size[0];
+				uint16 SizeY = args.size[1];
+				for (int l=0; l<item.executionOptions; ++l)
+				{
+					SizeX = FMath::Max(uint16(1), FMath::DivideAndRoundUp(SizeX, uint16(2)));
+					SizeY = FMath::Max(uint16(1), FMath::DivideAndRoundUp(SizeY, uint16(2)));
+				}
+
+                ImagePtr pA = new Image( SizeX,
+                                         SizeY,
+                                         1,
+                                         (EImageFormat)args.format );
+
+                FillPlainColourImage(pA.get(), c);
+
+                GetMemory().SetImage( item, pA );
+                break;
+            }
+
+            default:
+                check(false);
+            }
+
+            break;
+        }
+
+        case OP_TYPE::IM_GPU:
+        {
+            check(false);
+            break;
+        }
+
+        case OP_TYPE::IM_CROP:
+        {
+			OP::ImageCropArgs args = pModel->GetPrivate()->m_program.GetOpArgs<OP::ImageCropArgs>(item.at);
+            switch (item.stage)
+            {
+            case 0:
+                AddOp( SCHEDULED_OP( item.at, item, 1),
+					SCHEDULED_OP( args.source, item ) );
+                break;
+
+            case 1:
+            {
+                Ptr<const Image> pA = GetMemory().GetImage( CACHE_ADDRESS(args.source,item) );
+
+                box< vec2<int> > rect;
+                rect.min[0] = args.minX;
+                rect.min[1] = args.minY;
+                rect.size[0] = args.sizeX;
+                rect.size[1] = args.sizeY;
+
+				// Apply ther mipmap reduction to the crop rectangle.
+				int32 MipsToSkip = item.executionOptions;
+				while ( MipsToSkip>0 && rect.size[0]>0 && rect.size[1]>0 )
+				{
+					rect.ShrinkToHalf();
+					MipsToSkip--;
+				}
+
+				ImagePtr pResult;
+				if (!rect.IsEmpty())
+				{
+					pResult = ImageCrop(m_pSettings->GetPrivate()->m_imageCompressionQuality, pA.get(), rect);
+				}
+
+                GetMemory().SetImage( item, pResult );
+                break;
+            }
+
+            default:
+                check(false);
+            }
+
+            break;
+        }
+
+        case OP_TYPE::IM_PATCH:
+        {
+			OP::ImagePatchArgs args = pModel->GetPrivate()->m_program.GetOpArgs<OP::ImagePatchArgs>(item.at);
+            switch (item.stage)
+            {
+            case 0:
+                AddOp( SCHEDULED_OP( item.at, item, 1),
+					SCHEDULED_OP( args.base, item ),
+					SCHEDULED_OP( args.patch, item ) );
+                break;
+
+            case 1:
+            {
+                MUTABLE_CPUPROFILER_SCOPE(ImagePatch);
+
+                Ptr<const Image> pA = GetMemory().GetImage( CACHE_ADDRESS(args.base,item) );
+                Ptr<const Image> pB = GetMemory().GetImage( CACHE_ADDRESS(args.patch,item) );
+
+                box<vec2<int>> rect;
+                rect.min[0] = args.minX;
+                rect.min[1] = args.minY;
+                rect.size[0] = pB->GetSizeX();
+                rect.size[1] = pB->GetSizeY();
+
+				// Apply ther mipmap reduction to the crop rectangle.
+				int32 MipsToSkip = item.executionOptions;
+				while (MipsToSkip > 0 && rect.size[0] > 0 && rect.size[1] > 0)
+				{
+					rect.min/=2;
+					MipsToSkip--;
+				}
+
+                ImagePtr pResult = pA->Clone();
+
+				bool bApplyPatch = !rect.IsEmpty();
+				if (bApplyPatch)
+				{
+					// Resize image if it doesn't fit in the new block size
+					if (pB->GetSizeX() != rect.size[0] ||
+						pB->GetSizeY() != rect.size[1])
+					{
+						MUTABLE_CPUPROFILER_SCOPE(ImagePatchResize_Emergency);
+
+						FImageSize blockSize((uint16_t)rect.size[0], (uint16_t)rect.size[1]);
+						pB = ImageResizeLinear(m_pSettings->GetPrivate()->m_imageCompressionQuality,pB.get(), blockSize);
+					}
+
+					// Change the block image format if it doesn't match the composed image
+					// This is usually enforced at object compilation time.
+					if (pResult->GetFormat() != pB->GetFormat())
+					{
+						MUTABLE_CPUPROFILER_SCOPE(ImagPatcheReformat);
+
+						EImageFormat format = GetMostGenericFormat(pResult->GetFormat(), pB->GetFormat());
+
+						const FImageFormatData& finfo = GetImageFormatData(format);
+						if (finfo.m_pixelsPerBlockX == 0)
+						{
+							format = GetUncompressedFormat(format);
+						}
+
+						if (pResult->GetFormat() != format)
+						{
+							pResult =
+								ImagePixelFormat(m_pSettings->GetPrivate()->m_imageCompressionQuality,
+									pResult.get(), format);
+						}
+						if (pB->GetFormat() != format)
+						{
+							pB = ImagePixelFormat(m_pSettings->GetPrivate()->m_imageCompressionQuality,
+								pB.get(), format);
+						}
+					}
+
+					// Don't patch if below the image compression block size.
+					const FImageFormatData& finfo = GetImageFormatData(pResult->GetFormat());
+					bApplyPatch =
+						(rect.min[0] % finfo.m_pixelsPerBlockX == 0) &&
+						(rect.min[1] % finfo.m_pixelsPerBlockY == 0) &&
+						(rect.size[0] % finfo.m_pixelsPerBlockX == 0) &&
+						(rect.size[1] % finfo.m_pixelsPerBlockY == 0) &&
+						(rect.min[0] + rect.size[0]) <= pResult->GetSizeX() &&
+						(rect.min[1] + rect.size[1]) <= pResult->GetSizeY()
+						;
+				}
+
+				if (bApplyPatch)
+				{
+					ImageCompose(pResult.get(), pB.get(), rect);
+					pResult->m_flags = 0;
+				}
+				else
+				{
+					// This happens very often when skipping mips, and floods the log.
+					//UE_LOG( LogMutableCore, Verbose, TEXT("Skipped patch operation for image not fitting the block compression size. Small image? Patch rect is (%d, %d), (%d, %d), base is (%d, %d)"),
+					//	rect.min[0], rect.min[1], rect.size[0], rect.size[1], pResult->GetSizeX(), pResult->GetSizeY());
+				}
+
+                GetMemory().SetImage( item, pResult );
+                break;
+            }
+
+            default:
+                check(false);
+            }
+
+            break;
+        }
+
+        case OP_TYPE::IM_RASTERMESH:
+        {
+            auto args = pModel->GetPrivate()->m_program.GetOpArgs<OP::ImageRasterMeshArgs>(item.at);
+            switch (item.stage)
+            {
+            case 0:
+                    if (!args.image)
+                    {
+                        AddOp( SCHEDULED_OP( item.at, item, 1),
+                            SCHEDULED_OP::FromOpAndOptions( args.mesh, item, 0 ) );
+                    }
+                    else
+                    {
+                        AddOp( SCHEDULED_OP( item.at, item, 1),
+							SCHEDULED_OP::FromOpAndOptions( args.mesh, item, 0),
+							SCHEDULED_OP( args.image, item ),
+							SCHEDULED_OP( args.mask, item ),
+							SCHEDULED_OP::FromOpAndOptions( args.angleFadeProperties, item, 0 ),
+							SCHEDULED_OP::FromOpAndOptions( args.projector, item, 0 ) );
+                    }
+                break;
+
+            case 1:
+            {
+                MUTABLE_CPUPROFILER_SCOPE(RunCode_RasterMesh);
+
+                Ptr<const Mesh> pMesh = GetMemory().GetMesh( SCHEDULED_OP::FromOpAndOptions(args.mesh, item, 0) );
+
+                ImagePtr pNew;
+
+				uint16 SizeX = args.sizeX;
+				uint16 SizeY = args.sizeY;
+
+				// Drop mips while possible
+				int32 MipsToDrop = item.executionOptions;
+				while (MipsToDrop && !(SizeX % 2) && !(SizeY % 2))
+				{
+					SizeX = FMath::Max(uint16(1),FMath::DivideAndRoundUp(SizeX, uint16(2)));
+					SizeY = FMath::Max(uint16(1),FMath::DivideAndRoundUp(SizeY, uint16(2)));
+					--MipsToDrop;
+				}
+
+                if (!pMesh.get())
+                {
+                    pNew = new Image(SizeX, SizeY, 1, EImageFormat::IF_L_UBYTE);
+                    check(false);
+                }
+                else if (args.image)
+                {
+                    // Raster with projection
+                    Ptr<const Image> pSource = GetMemory().GetImage( CACHE_ADDRESS(args.image,item) );
+
+                    Ptr<const Image> pMask = nullptr;
+                    if ( args.mask )
+                    {
+                        pMask = GetMemory().GetImage( CACHE_ADDRESS(args.mask,item) );
+
+						// TODO: This shouldn't happen, but be defensive.
+						FImageSize ResultSize(SizeX, SizeY);
+						if (pMask && pMask->GetSize()!= ResultSize)
+						{
+							MUTABLE_CPUPROFILER_SCOPE(ImageResize_MaskFixForProjection);
+							pMask = ImageResizeLinear(0, pMask.get(), ResultSize);
+						}
+                    }
+
+                    float fadeStart = 180.0f;
+                    float fadeEnd = 180.0f;
+                    if ( args.angleFadeProperties )
+                    {
+                        auto fadeProperties = GetMemory().GetColour(SCHEDULED_OP::FromOpAndOptions(args.angleFadeProperties, item, 0));
+                        fadeStart = fadeProperties[0];
+                        fadeEnd = fadeProperties[1];
+                    }
+                    fadeStart *= PI / 180.0f;
+                    fadeEnd *= PI / 180.0f;
+
+                    auto format = GetUncompressedFormat( pSource->GetFormat() );
+                    pNew = new Image( SizeX, SizeY, 1, format );
+
+                    if (pSource->GetFormat()!=format)
+                    {
+						MUTABLE_CPUPROFILER_SCOPE(RunCode_RasterMesh_ReformatSource);
+                        pSource = ImagePixelFormat( m_pSettings->GetPrivate()->m_imageCompressionQuality, pSource.get(), format );
+                    }
+
+                    // Allocate memory for the temporary buffers
+                    SCRATCH_IMAGE_PROJECT scratch;
+                    scratch.vertices.resize( pMesh->GetVertexCount() );
+                    scratch.culledVertex.resize( pMesh->GetVertexCount() );
+
+                    // Layout is always 0 because the previous mesh project operations take care of
+                    // moving the right layout channel to the 0.
+                    int layout = 0;
+
+                    float projectionAngle = 0;
+                    if ( args.projector )
+                    {
+                        auto pProjector = GetMemory().GetProjector(SCHEDULED_OP::FromOpAndOptions(args.projector, item, 0));
+                        if (pProjector)
+                        {
+                            projectionAngle = pProjector->m_value.projectionAngle;
+
+                            switch (pProjector->m_value.type)
+                            {
+                            case mu::PROJECTOR_TYPE::PLANAR:
+                                ImageRasterProjectedPlanar( pMesh.get(), pNew.get(),
+                                                            pSource.get(), pMask.get(),
+                                                            fadeStart, fadeEnd,
+                                                            layout, args.blockIndex,
+                                                            &scratch );
+                                break;
+
+                            case mu::PROJECTOR_TYPE::WRAPPING:
+                                ImageRasterProjectedWrapping( pMesh.get(), pNew.get(),
+                                                              pSource.get(), pMask.get(),
+                                                              fadeStart, fadeEnd,
+                                                              layout, args.blockIndex,
+                                                              &scratch );
+                                break;
+
+                            case mu::PROJECTOR_TYPE::CYLINDRICAL:
+                                ImageRasterProjectedCylindrical( pMesh.get(), pNew.get(),
+                                                                 pSource.get(), pMask.get(),
+                                                                 fadeStart, fadeEnd,
+                                                                 layout,
+                                                                 projectionAngle,
+                                                                 &scratch );
+                                break;
+
+                            default:
+                                check(false);
+                                break;
+                            }
+                        }
+                    }
+                }
+                else
+                {
+                    // Flat mesh UV raster
+                    pNew = new Image( SizeX, SizeY, 1, EImageFormat::IF_L_UBYTE );
+                    ImageRasterMesh( pMesh.get(), pNew.get(), args.blockIndex );
+                }
+
+                GetMemory().SetImage( item, pNew );
+                break;
+            }
+
+            default:
+                check(false);
+            }
+
+            break;
+        }
+
+        case OP_TYPE::IM_MAKEGROWMAP:
+        {
+            auto args = pModel->GetPrivate()->m_program.GetOpArgs<OP::ImageMakeGrowMapArgs>(item.at);
+            switch (item.stage)
+            {
+            case 0:
+                    AddOp( SCHEDULED_OP( item.at, item, 1),
+                           SCHEDULED_OP( args.mask, item) );
+                break;
+
+            case 1:
+            {
+                MUTABLE_CPUPROFILER_SCOPE(RunCode_ImageMakeGrowMap);
+
+                Ptr<const Image> pMask = GetMemory().GetImage( CACHE_ADDRESS(args.mask,item) );
+
+                ImagePtr pNew = new Image( pMask->GetSizeX(), pMask->GetSizeY(), 1, EImageFormat::IF_L_UBYTE );
+
+                ImageMakeGrowMap( pNew.get(), pMask.get(), args.border );
+				pNew->m_flags |= Image::IF_CANNOT_BE_SCALED;
+
+                GetMemory().SetImage( item, pNew );
+                break;
+            }
+
+            default:
+                check(false);
+            }
+
+            break;
+        }
+
+        case OP_TYPE::IM_DISPLACE:
+        {
+            auto args = pModel->GetPrivate()->m_program.GetOpArgs<OP::ImageDisplaceArgs>(item.at);
+            switch (item.stage)
+            {
+            case 0:
+                    AddOp( SCHEDULED_OP( item.at, item, 1),
+                           SCHEDULED_OP( args.source, item ),
+                           SCHEDULED_OP( args.displacementMap, item ) );
+                break;
+
+            case 1:
+            {
+                MUTABLE_CPUPROFILER_SCOPE(RunCode_ImageDisplace);
+
+                Ptr<const Image> pSource = GetMemory().GetImage( CACHE_ADDRESS(args.source,item) );
+                Ptr<const Image> pMap = GetMemory().GetImage( CACHE_ADDRESS(args.displacementMap,item) );
+
+				// TODO: This shouldn't happen: displacement maps cannot be scaled because their information
+				// is resolution sensitive (pixel offsets). If the size doesn't match, scale the source, apply 
+				// displacement and then unscale it.
+				FImageSize OriginalSourceScale = pSource->GetSize();
+				if (OriginalSourceScale != pMap->GetSize())
+				{
+					MUTABLE_CPUPROFILER_SCOPE(ImageResize_EmergencyHackForDisplacementStep1);
+					pSource = ImageResizeLinear(0, pSource.get(), pMap->GetSize());
+				}
+
+                ImagePtr pNew = new Image( pSource->GetSizeX(), pSource->GetSizeY(), 1, pSource->GetFormat() );
+
+                // TODO: Reuse source if possible
+                ImageDisplace( pNew.get(), pSource.get(), pMap.get() );
+
+				if (OriginalSourceScale != pNew->GetSize())
+				{
+					MUTABLE_CPUPROFILER_SCOPE(ImageResize_EmergencyHackForDisplacementStep2);
+					pNew = ImageResizeLinear(0, pNew.get(), OriginalSourceScale);
+				}
+
+                GetMemory().SetImage( item, pNew );
+                break;
+            }
+
+            default:
+                check(false);
+            }
+
+            break;
+        }
+
+        case OP_TYPE::IM_TRANSFORM:
+        {
+            const OP::ImageTransformArgs Args = pModel->GetPrivate()->m_program.GetOpArgs<OP::ImageTransformArgs>(item.at);
+
+            switch (item.stage)
+            {
+            case 0:
+			{
+				const TArray<SCHEDULED_OP, TInlineAllocator<6>> Deps = {
+						SCHEDULED_OP( Args.base, item),
+						SCHEDULED_OP( Args.offsetX, item),
+						SCHEDULED_OP( Args.offsetY, item),
+						SCHEDULED_OP( Args.scaleX, item),
+						SCHEDULED_OP( Args.scaleY, item),
+						SCHEDULED_OP( Args.rotation, item) };
+
+                AddOp( SCHEDULED_OP( item.at, item, 1), Deps );
+
+				break;
+			}
+            case 1:
+            {
+                Ptr<const Image> pBaseImage = GetMemory().GetImage( CACHE_ADDRESS(Args.base, item) );
+                
+                const FVector2f Offset = FVector2f(
+                        Args.offsetX ? GetMemory().GetScalar( CACHE_ADDRESS(Args.offsetX, item) ) : 0.0f,
+                        Args.offsetY ? GetMemory().GetScalar( CACHE_ADDRESS(Args.offsetY, item) ) : 0.0f );
+
+                const FVector2f Scale = FVector2f(
+                        Args.scaleX ? GetMemory().GetScalar( CACHE_ADDRESS(Args.scaleX, item) ) : 1.0f,
+                        Args.scaleY ? GetMemory().GetScalar( CACHE_ADDRESS(Args.scaleY, item) ) : 1.0f );
+
+				// Map Range 0-1 to a full rotation
+                const float Rotation = GetMemory().GetScalar( CACHE_ADDRESS(Args.rotation, item) ) * UE_TWO_PI;
+
+
+				Ptr<Image> pResult = new Image( pBaseImage->GetSizeX(), pBaseImage->GetSizeY(), 1, pBaseImage->GetFormat());
+			
+				FImageSize SampleImageSize =  FImageSize( 
+					static_cast<uint16>( FMath::Clamp( FMath::FloorToInt( float(pBaseImage->GetSizeX()) * FMath::Abs(Scale.X)), 2, pBaseImage->GetSizeX() ) ),
+					static_cast<uint16>( FMath::Clamp( FMath::FloorToInt( float(pBaseImage->GetSizeY()) * FMath::Abs(Scale.Y)), 2, pBaseImage->GetSizeY() ) ) );
+	
+				Ptr<Image> pSampleImage = ImageResizeLinear( 0, pBaseImage.get(), SampleImageSize );
+                ImageTransform( pResult.get(), pSampleImage.get(), Offset, Scale, Rotation );
+
+                GetMemory().SetImage( item, pResult );
+
+                break;
+            }
+
+            default:
+                check(false);
+            }
+
+			break;
+		}
+
+        default:
+            if (type!=OP_TYPE::NONE)
+            {
+                // Operation not implemented
+                check( false );
+            }
+            break;
+        }
+    }
+
+
+    //---------------------------------------------------------------------------------------------
+    Ptr<RangeIndex> CodeRunner::BuildCurrentOpRangeIndex( const SCHEDULED_OP& item,
+                                                        const Parameters* pParams,
+                                                        const Model* pModel,
+                                                        int parameterIndex
+                                                        )
+    {
+        if (!item.executionIndex)
+        {
+            return nullptr;
+        }
+
+        // \todo: optimise to avoid allocating the index here, we could access internal
+        // data directly.
+		Ptr<RangeIndex> index = pParams->NewRangeIndex( parameterIndex );
+        if (!index)
+        {
+            return nullptr;
+        }
+
+        const auto& program = pModel->GetPrivate()->m_program;
+        const PARAMETER_DESC& paramDesc = program.m_parameters[ parameterIndex ];
+        for( size_t rangeIndexInParam=0;
+             rangeIndexInParam<paramDesc.m_ranges.size();
+             ++rangeIndexInParam )
+        {
+            auto rangeIndexInModel = paramDesc.m_ranges[rangeIndexInParam];
+            const ExecutionIndex& currentIndex = GetMemory().GetRageIndex( item.executionIndex );
+            int position = currentIndex.GetFromModelRangeIndex(rangeIndexInModel);
+            index->GetPrivate()->m_values[rangeIndexInParam] = position;
+        }
+
+        return index;
+    }
+
+
+    //---------------------------------------------------------------------------------------------
+    void CodeRunner::RunCode_Bool( SCHEDULED_OP& item,
+                                   const Parameters* pParams,
+                                   const Model* pModel
+                                   )
+    {
+        //MUTABLE_CPUPROFILER_SCOPE(RunCode_Bool);
+
+        const auto& program = pModel->GetPrivate()->m_program;
+        auto type = program.GetOpType(item.at);
+        switch (type)
+        {
+
+        case OP_TYPE::BO_CONSTANT:
+        {
+            auto args = program.GetOpArgs<OP::BoolConstantArgs>(item.at);
+            bool result = args.value;
+            GetMemory().SetBool( item, result );
+            break;
+        }
+
+        case OP_TYPE::BO_PARAMETER:
+        {
+            auto args = program.GetOpArgs<OP::ParameterArgs>(item.at);
+            bool result = false;
+			Ptr<RangeIndex> index = BuildCurrentOpRangeIndex( item, pParams, pModel, args.variable );
+            result = pParams->GetBoolValue( args.variable, index );
+            GetMemory().SetBool( item, result );
+            break;
+        }
+
+        case OP_TYPE::BO_LESS:
+        {
+            auto args = program.GetOpArgs<OP::BoolLessArgs>(item.at);
+            switch (item.stage)
+            {
+            case 0:
+                    AddOp( SCHEDULED_OP( item.at, item, 1),
+                           SCHEDULED_OP( args.a, item),
+                           SCHEDULED_OP( args.b, item) );
+                break;
+
+            case 1:
+            {
+                float a = GetMemory().GetScalar( CACHE_ADDRESS(args.a,item) );
+                float b = GetMemory().GetScalar( CACHE_ADDRESS(args.b,item) );
+                bool result = a<b;
+                GetMemory().SetBool( item, result );
+                break;
+            }
+
+            default:
+                check(false);
+            }
+
+            break;
+        }
+
+        case OP_TYPE::BO_AND:
+        {
+            auto args = program.GetOpArgs<OP::BoolBinaryArgs>(item.at);
+            switch (item.stage)
+            {
+            case 0:
+                {
+                    // Try to avoid the op entirely if we have some children cached
+                    bool skip = false;
+                    if ( args.a && GetMemory().IsValid( CACHE_ADDRESS(args.a,item) ) )
+                    {
+                         bool a = GetMemory().GetBool( CACHE_ADDRESS(args.a,item) );
+                         if (!a)
+                         {
+                            GetMemory().SetBool( item, false );
+                            skip=true;
+                         }
+                    }
+
+                    if ( !skip && args.b && GetMemory().IsValid( CACHE_ADDRESS(args.b,item) ) )
+                    {
+                         bool b = GetMemory().GetBool( CACHE_ADDRESS(args.b,item) );
+                         if (!b)
+                         {
+                            GetMemory().SetBool( item, false );
+                            skip=true;
+                         }
+                    }
+
+                    if (!skip)
+                    {
+                        AddOp( SCHEDULED_OP( item.at, item, 1),
+                               SCHEDULED_OP( args.a, item));
+                    }
+				break;
+                }
+
+            case 1:
+            {
+                bool a = args.a ? GetMemory().GetBool( CACHE_ADDRESS(args.a,item) ) : true;
+                if (!a)
+                {
+                    GetMemory().SetBool( item, false );
+                }
+                else
+                {
+                    AddOp( SCHEDULED_OP( item.at, item, 2),
+                           SCHEDULED_OP( args.b, item));
+                }
+                break;
+            }
+
+            case 2:
+            {
+                // We arrived here because a is true
+                bool b = args.b ? GetMemory().GetBool( CACHE_ADDRESS(args.b,item) ) : true;
+                GetMemory().SetBool( item, b );
+                break;
+            }
+
+            default:
+                check(false);
+            }
+            break;
+        }
+
+        case OP_TYPE::BO_OR:
+        {
+            auto args = program.GetOpArgs<OP::BoolBinaryArgs>(item.at);
+            switch (item.stage)
+            {
+            case 0:
+                {
+                    // Try to avoid the op entirely if we have some children cached
+                    bool skip = false;
+                    if ( args.a && GetMemory().IsValid( CACHE_ADDRESS(args.a,item) ) )
+                    {
+                         bool a = GetMemory().GetBool( CACHE_ADDRESS(args.a,item) );
+                         if (a)
+                         {
+                            GetMemory().SetBool( item, true );
+                            skip=true;
+                         }
+                    }
+
+                    if ( !skip && args.b && GetMemory().IsValid( CACHE_ADDRESS(args.b,item) ) )
+                    {
+                         bool b = GetMemory().GetBool( CACHE_ADDRESS(args.b,item) );
+                         if (b)
+                         {
+                            GetMemory().SetBool( item, true );
+                            skip=true;
+                         }
+                    }
+
+                    if (!skip)
+                    {
+                        AddOp( SCHEDULED_OP( item.at, item, 1),
+                               SCHEDULED_OP( args.a, item));
+                    }
+				break;
+                }
+
+            case 1:
+            {
+                bool a = args.a ? GetMemory().GetBool( CACHE_ADDRESS(args.a,item) ) : false;
+                if (a)
+                {
+                    GetMemory().SetBool( item, true );
+                }
+                else
+                {
+                    AddOp( SCHEDULED_OP( item.at, item, 2),
+                           SCHEDULED_OP( args.b, item));
+                }
+                break;
+            }
+
+            case 2:
+            {
+                // We arrived here because a is false
+                bool b = args.b ? GetMemory().GetBool( CACHE_ADDRESS(args.b,item) ) : false;
+                GetMemory().SetBool( item, b );
+                break;
+            }
+
+            default:
+                check(false);
+            }
+            break;
+        }
+
+        case OP_TYPE::BO_NOT:
+        {
+            auto args = program.GetOpArgs<OP::BoolNotArgs>(item.at);
+            switch (item.stage)
+            {
+            case 0:
+                    AddOp( SCHEDULED_OP( item.at, item, 1),
+                           SCHEDULED_OP( args.source, item) );
+                break;
+
+            case 1:
+            {
+                bool result = !GetMemory().GetBool( CACHE_ADDRESS(args.source,item) );
+                GetMemory().SetBool( item, result );
+                break;
+            }
+
+            default:
+                check(false);
+            }
+            break;
+        }
+
+        case OP_TYPE::BO_EQUAL_INT_CONST:
+        {
+            auto args = program.GetOpArgs<OP::BoolEqualScalarConstArgs>(item.at);
+            switch (item.stage)
+            {
+            case 0:
+                    AddOp( SCHEDULED_OP( item.at, item, 1),
+                           SCHEDULED_OP( args.value, item) );
+                break;
+
+            case 1:
+            {
+                int a = GetMemory().GetInt( CACHE_ADDRESS(args.value,item) );
+                bool result = a == args.constant;
+                GetMemory().SetBool( item, result );
+                break;
+            }
+
+            default:
+                check(false);
+            }
+            break;
+        }
+
+        default:
+            check( false );
+            break;
+        }
+    }
+
+
+    //---------------------------------------------------------------------------------------------
+    void CodeRunner::RunCode_Int( SCHEDULED_OP& item,
+                                  const Parameters* pParams,
+                                  const Model* pModel
+                              )
+    {
+        MUTABLE_CPUPROFILER_SCOPE(RunCode_Int);
+
+		OP_TYPE type = pModel->GetPrivate()->m_program.GetOpType(item.at);
+        switch (type)
+        {
+
+        case OP_TYPE::NU_CONSTANT:
+        {
+            auto args = pModel->GetPrivate()->m_program.GetOpArgs<OP::IntConstantArgs>(item.at);
+            int result = args.value;
+            GetMemory().SetInt( item, result );
+            break;
+        }
+
+        case OP_TYPE::NU_PARAMETER:
+        {
+            auto args = pModel->GetPrivate()->m_program.GetOpArgs<OP::ParameterArgs>(item.at);
+			Ptr<RangeIndex> index = BuildCurrentOpRangeIndex( item, pParams, pModel, args.variable );
+            int result = pParams->GetIntValue( args.variable, index );
+
+            // Check that the value is actually valid. Otherwise set the default.
+            if ( pParams->GetIntPossibleValueCount( args.variable ) )
+            {
+                bool valid = false;
+                for ( int i=0;
+                      (!valid) && i<pParams->GetIntPossibleValueCount( args.variable );
+                      ++i )
+                {
+                    if ( result == pParams->GetIntPossibleValue( args.variable, i ) )
+                    {
+                        valid = true;
+                    }
+                }
+
+                if (!valid)
+                {
+                    result = pParams->GetIntPossibleValue( args.variable, 0 );
+                }
+            }
+
+            GetMemory().SetInt( item, result );
+            break;
+        }
+
+        default:
+            check( false );
+            break;
+        }
+    }
+
+
+    //---------------------------------------------------------------------------------------------
+    void CodeRunner::RunCode_Scalar( SCHEDULED_OP& item,
+                                     const Parameters* pParams,
+                                     const Model* pModel
+                                     )
+    {
+        MUTABLE_CPUPROFILER_SCOPE(RunCode_Scalar);
+
+		OP_TYPE type = pModel->GetPrivate()->m_program.GetOpType(item.at);
+        switch (type)
+        {
+
+        case OP_TYPE::SC_CONSTANT:
+        {
+            auto args = pModel->GetPrivate()->m_program.GetOpArgs<OP::ScalarConstantArgs>(item.at);
+            float result = args.value;
+            GetMemory().SetScalar( item, result );
+            break;
+        }
+
+        case OP_TYPE::SC_PARAMETER:
+        {
+            auto args = pModel->GetPrivate()->m_program.GetOpArgs<OP::ParameterArgs>(item.at);
+			Ptr<RangeIndex> index = BuildCurrentOpRangeIndex( item, pParams, pModel, args.variable );
+            float result = pParams->GetFloatValue( args.variable, index );
+            GetMemory().SetScalar( item, result );
+            break;
+        }
+
+        case OP_TYPE::SC_CURVE:
+        {
+            auto args = pModel->GetPrivate()->m_program.GetOpArgs<OP::ScalarCurveArgs>(item.at);
+            switch (item.stage)
+            {
+            case 0:
+                    AddOp( SCHEDULED_OP( item.at, item, 1),
+                           SCHEDULED_OP( args.time, item) );
+                break;
+
+            case 1:
+            {
+                float time = GetMemory().GetScalar( CACHE_ADDRESS(args.time,item) );
+
+                const Curve& curve = pModel->GetPrivate()->m_program.m_constantCurves[args.curve];
+                float result = EvalCurve(curve, time);
+
+                GetMemory().SetScalar( item, result );
+                break;
+            }
+
+            default:
+                check(false);
+            }
+            break;
+        }
+
+        case OP_TYPE::SC_MULTIPLYADD:
+            // \TODO
+            check( false );
+            break;
+
+        case OP_TYPE::SC_ARITHMETIC:
+        {
+            auto args = pModel->GetPrivate()->m_program.GetOpArgs<OP::ArithmeticArgs>(item.at);
+            switch (item.stage)
+            {
+            case 0:
+                    AddOp( SCHEDULED_OP( item.at, item, 1),
+                           SCHEDULED_OP( args.a, item),
+                           SCHEDULED_OP( args.b, item) );
+                break;
+
+            case 1:
+            {
+                float a = GetMemory().GetScalar( CACHE_ADDRESS(args.a,item) );
+                float b = GetMemory().GetScalar( CACHE_ADDRESS(args.b,item) );
+
+                float result = 1.0f;
+                switch (args.operation)
+                {
+                case OP::ArithmeticArgs::ADD:
+                    result = a + b;
+                    break;
+
+                case OP::ArithmeticArgs::MULTIPLY:
+                    result = a * b;
+                    break;
+
+                case OP::ArithmeticArgs::SUBTRACT:
+                    result = a - b;
+                    break;
+
+                case OP::ArithmeticArgs::DIVIDE:
+                    result = a / b;
+                    break;
+
+                default:
+                    checkf(false, TEXT("Arithmetic operation not implemented."));
+                    break;
+                }
+
+                GetMemory().SetScalar( item, result );
+                break;
+            }
+
+            default:
+                check(false);
+            }
+            break;
+        }
+
+        default:
+            check( false );
+            break;
+        }
+    }
+
+
+    //---------------------------------------------------------------------------------------------
+    void
+    CodeRunner::RunCode_String( SCHEDULED_OP& item, const Parameters* pParams, const Model* pModel )
+    {
+        MUTABLE_CPUPROFILER_SCOPE(RunCode_String );
+
+		OP_TYPE type = pModel->GetPrivate()->m_program.GetOpType( item.at );
+        switch ( type )
+        {
+
+        case OP_TYPE::ST_CONSTANT:
+        {
+            auto args =
+                pModel->GetPrivate()->m_program.GetOpArgs<OP::ResourceConstantArgs>( item.at );
+            check( args.value < pModel->GetPrivate()->m_program.m_constantStrings.size() );
+
+            auto result = pModel->GetPrivate()->m_program.m_constantStrings[args.value];
+            GetMemory().SetString( item, new String(result.c_str()) );
+
+            break;
+        }
+
+        case OP_TYPE::ST_PARAMETER:
+        {
+            auto args = pModel->GetPrivate()->m_program.GetOpArgs<OP::ParameterArgs>( item.at );
+			Ptr<RangeIndex> index = BuildCurrentOpRangeIndex( item, pParams, pModel, args.variable );
+            string result = pParams->GetStringValue( args.variable, index );
+            GetMemory().SetString( item, new String( result.c_str() ) );
+            break;
+        }
+
+        default:
+            check( false );
+            break;
+        }
+    }
+
+
+    //---------------------------------------------------------------------------------------------
+    void CodeRunner::RunCode_Colour( SCHEDULED_OP& item,
+                                     const Parameters* pParams,
+                                     const Model* pModel
+                                     )
+    {
+		MUTABLE_CPUPROFILER_SCOPE(RunCode_Colour);
+
+		OP_TYPE type = pModel->GetPrivate()->m_program.GetOpType(item.at);
+
+        const auto& program = pModel->GetPrivate()->m_program;
+
+        switch ( type )
+        {
+
+        case OP_TYPE::CO_CONSTANT:
+        {
+			OP::ColourConstantArgs args = program.GetOpArgs<OP::ColourConstantArgs>(item.at);
+            vec4f result;
+            result[0] = args.value[0];
+            result[1] = args.value[1];
+            result[2] = args.value[2];
+            result[3] = args.value[3];
+            GetMemory().SetColour( item, result );
+            break;
+        }
+
+        case OP_TYPE::CO_PARAMETER:
+        {
+			OP::ParameterArgs args = pModel->GetPrivate()->m_program.GetOpArgs<OP::ParameterArgs>(item.at);
+			Ptr<RangeIndex> index = BuildCurrentOpRangeIndex( item, pParams, pModel, args.variable );
+            float r=0.0f;
+            float g=0.0f;
+            float b=0.0f;            
+            pParams->GetColourValue( args.variable, &r, &g, &b, index );
+            GetMemory().SetColour( item, vec4f(r,g,b,1.0f) );
+            break;
+        }
+
+        case OP_TYPE::CO_SAMPLEIMAGE:
+        {
+			OP::ColourSampleImageArgs args = pModel->GetPrivate()->m_program.GetOpArgs<OP::ColourSampleImageArgs>(item.at);
+            switch (item.stage)
+            {
+            case 0:
+                    AddOp( SCHEDULED_OP( item.at, item, 1),
+                           SCHEDULED_OP( args.x, item),
+                           SCHEDULED_OP( args.y, item),
+						   // Don't skip mips for the texture to sample
+                           SCHEDULED_OP::FromOpAndOptions( args.image, item, 0) );
+                break;
+
+            case 1:
+            {
+                float x = args.x ? GetMemory().GetScalar( CACHE_ADDRESS(args.x,item) ) : 0.5f;
+                float y = args.y ? GetMemory().GetScalar( CACHE_ADDRESS(args.y,item) ) : 0.5f;
+
+                Ptr<const Image> pImage = GetMemory().GetImage(SCHEDULED_OP::FromOpAndOptions(args.image, item, 0));
+
+                vec4f result;
+                if (pImage)
+                {
+                    if (args.filter)
+                    {
+                        // TODO
+                        result = pImage->Sample(vec2<float>(x, y));
+                    }
+                    else
+                    {
+                        result = pImage->Sample(vec2<float>(x, y));
+                    }
+                }
+                else
+                {
+                    result = vec4f();
+                }
+
+                GetMemory().SetColour( item, result );
+                break;
+            }
+
+            default:
+                check(false);
+            }
+
+            break;
+        }
+
+        case OP_TYPE::CO_SWIZZLE:
+        {
+			OP::ColourSwizzleArgs args = pModel->GetPrivate()->m_program.GetOpArgs<OP::ColourSwizzleArgs>(item.at);
+            switch (item.stage)
+            {
+            case 0:
+                    AddOp( SCHEDULED_OP( item.at, item, 1),
+                           SCHEDULED_OP( args.sources[0], item),
+                           SCHEDULED_OP( args.sources[1], item),
+                           SCHEDULED_OP( args.sources[2], item),
+                           SCHEDULED_OP( args.sources[3], item) );
+                break;
+
+            case 1:
+            {
+                vec4f result;
+
+                for (int t=0;t<MUTABLE_OP_MAX_SWIZZLE_CHANNELS;++t)
+                {
+                    if ( args.sources[t] )
+                    {
+                        auto p = GetMemory().GetColour( CACHE_ADDRESS(args.sources[t],item) );
+                        result[t] = p[ args.sourceChannels[t] ];
+                    }
+                }
+
+                GetMemory().SetColour( item, result );
+                break;
+            }
+
+            default:
+                check(false);
+            }
+
+            break;
+        }
+
+        case OP_TYPE::CO_IMAGESIZE:
+        {
+			OP::ColourImageSizeArgs args = pModel->GetPrivate()->m_program.GetOpArgs<OP::ColourImageSizeArgs>(item.at);
+            switch (item.stage)
+            {
+            case 0:
+                    AddOp( SCHEDULED_OP( item.at, item, 1),
+                           SCHEDULED_OP( args.image, item) );
+                break;
+
+            case 1:
+            {
+                Ptr<const Image> pImage = GetMemory().GetImage( CACHE_ADDRESS(args.image,item) );
+
+                vec4f result = vec4f( (float)pImage->GetSizeX(), (float)pImage->GetSizeY(), 0.0f, 0.0f );
+
+                GetMemory().SetColour( item, result );
+                break;
+            }
+
+            default:
+                check(false);
+            }
+
+            break;
+        }
+
+        case OP_TYPE::CO_LAYOUTBLOCKTRANSFORM:
+        {
+			OP::ColourLayoutBlockTransformArgs args = pModel->GetPrivate()->m_program.GetOpArgs<OP::ColourLayoutBlockTransformArgs>(item.at);
+            switch (item.stage)
+            {
+            case 0:
+                    AddOp( SCHEDULED_OP( item.at, item, 1),
+                           SCHEDULED_OP( args.layout, item) );
+                break;
+
+            case 1:
+            {
+                Ptr<const Layout> pLayout = GetMemory().GetLayout( CACHE_ADDRESS(args.layout,item) );
+
+                vec4f result = vec4f(0,0,0,0);
+                if ( pLayout )
+                {
+                    int relBlockIndex = pLayout->FindBlock( args.block );
+
+                    if( relBlockIndex >=0 )
+                    {
+                        box< vec2<int> > rectInblocks;
+                        pLayout->GetBlock
+                                (
+                                    relBlockIndex,
+                                    &rectInblocks.min[0], &rectInblocks.min[1],
+                                    &rectInblocks.size[0], &rectInblocks.size[1]
+                                    );
+
+                        // Convert the rect from blocks to pixels
+                        FIntPoint grid = pLayout->GetGridSize();
+
+                        result = vec4<float>( float(rectInblocks.min[0]) / float(grid[0]),
+                                              float(rectInblocks.min[1]) / float(grid[1]),
+                                              float(rectInblocks.size[0]) / float(grid[0]),
+                                              float(rectInblocks.size[1]) / float(grid[1]) );
+                    }
+                }
+
+                GetMemory().SetColour( item, result );
+                break;
+            }
+
+            default:
+                check(false);
+            }
+
+            break;
+        }
+
+        case OP_TYPE::CO_FROMSCALARS:
+        {
+			OP::ColourFromScalarsArgs args = pModel->GetPrivate()->m_program.GetOpArgs<OP::ColourFromScalarsArgs>(item.at);
+            switch (item.stage)
+            {
+            case 0:
+                    AddOp( SCHEDULED_OP( item.at, item, 1),
+                           SCHEDULED_OP( args.x, item),
+                           SCHEDULED_OP( args.y, item),
+                           SCHEDULED_OP( args.z, item),
+                           SCHEDULED_OP( args.w, item));
+                break;
+
+            case 1:
+            {
+                vec4f result = vec4f(1, 1, 1, 1);
+
+                if (args.x)
+                {
+                    result[0] = GetMemory().GetScalar( CACHE_ADDRESS(args.x,item) );
+                }
+
+                if (args.y)
+                {
+                    result[1] = GetMemory().GetScalar( CACHE_ADDRESS(args.y,item) );
+                }
+
+                if (args.z)
+                {
+                    result[2] = GetMemory().GetScalar( CACHE_ADDRESS(args.z,item) );
+                }
+
+                if (args.w)
+                {
+                    result[3] = GetMemory().GetScalar( CACHE_ADDRESS(args.w,item) );
+                }
+
+                GetMemory().SetColour( item, result );
+                break;
+            }
+
+            default:
+                check(false);
+            }
+
+            break;
+        }
+
+        case OP_TYPE::CO_ARITHMETIC:
+        {
+			OP::ArithmeticArgs args = program.GetOpArgs<OP::ArithmeticArgs>(item.at);
+            switch (item.stage)
+            {
+            case 0:
+                    AddOp( SCHEDULED_OP( item.at, item, 1),
+                           SCHEDULED_OP( args.a, item),
+                           SCHEDULED_OP( args.b, item));
+                break;
+
+            case 1:
+            {
+				OP_TYPE otype = program.GetOpType( args.a );
+                DATATYPE dtype = GetOpDataType( otype );
+                check( dtype == DT_COLOUR );
+                otype = program.GetOpType( args.b );
+                dtype = GetOpDataType( otype );
+                check( dtype == DT_COLOUR );
+                vec4f a = args.a ? GetMemory().GetColour( CACHE_ADDRESS( args.a, item ) )
+                                 : vec4f( 0, 0, 0, 0 );
+                vec4f b = args.b ? GetMemory().GetColour( CACHE_ADDRESS( args.b, item ) )
+                                 : vec4f( 0, 0, 0, 0 );
+
+                vec4f result = vec4f(0,0,0,0);
+                switch (args.operation)
+                {
+                case OP::ArithmeticArgs::ADD:
+                    result = a + b;
+                    break;
+
+                case OP::ArithmeticArgs::MULTIPLY:
+                    result = a * b;
+                    break;
+
+                case OP::ArithmeticArgs::SUBTRACT:
+                    result = a - b;
+                    break;
+
+                case OP::ArithmeticArgs::DIVIDE:
+                    result = a / b;
+                    break;
+
+                default:
+                    checkf(false, TEXT("Arithmetic operation not implemented."));
+                    break;
+                }
+
+                GetMemory().SetColour( item, result );
+                break;
+            }
+
+            default:
+                check(false);
+            }
+
+            break;
+        }
+
+        default:
+            check( false );
+            break;
+        }
+    }
+
+
+    //---------------------------------------------------------------------------------------------
+    void CodeRunner::RunCode_Projector( SCHEDULED_OP& item,
+                                const Parameters* pParams,
+                                const Model* pModel
+                              )
+    {
+        MUTABLE_CPUPROFILER_SCOPE(RunCode_Projector);
+
+        const auto& program = pModel->GetPrivate()->m_program;
+		OP_TYPE type = program.GetOpType(item.at);
+        switch (type)
+        {
+
+        case OP_TYPE::PR_CONSTANT:
+        {
+            auto args = pModel->GetPrivate()->m_program.GetOpArgs<OP::ResourceConstantArgs>(item.at);
+            ProjectorPtr pResult = new Projector();
+            pResult->m_value = program.m_constantProjectors[args.value];
+            GetMemory().SetProjector( item, pResult );
+            break;
+        }
+
+        case OP_TYPE::PR_PARAMETER:
+        {
+            auto args = pModel->GetPrivate()->m_program.GetOpArgs<OP::ParameterArgs>(item.at);
+			Ptr<RangeIndex> index = BuildCurrentOpRangeIndex( item, pParams, pModel, args.variable );
+            ProjectorPtr pResult = new Projector();
+            pResult->m_value = pParams->GetPrivate()->GetProjectorValue(args.variable,index);
+
+            // The type cannot be changed, take it from the default value
+            const auto& def = program.m_parameters[args.variable].m_defaultValue.m_projector;
+            pResult->m_value.type = def.type;
+
+            GetMemory().SetProjector( item, pResult );
+            break;
+        }
+
+        default:
+            check( false );
+            break;
+        }
+    }
+
+
+    //---------------------------------------------------------------------------------------------
+    void CodeRunner::RunCode_Layout( SCHEDULED_OP& item,
+                                     const Model* pModel
+                                     )
+    {
+        //MUTABLE_CPUPROFILER_SCOPE(RunCode_Layout);
+
+		OP_TYPE type = pModel->GetPrivate()->m_program.GetOpType(item.at);
+        switch (type)
+        {
+
+        case OP_TYPE::LA_CONSTANT:
+        {
+            auto args = pModel->GetPrivate()->m_program.GetOpArgs<OP::ResourceConstantArgs>(item.at);
+            check( args.value
+                            <
+                            pModel->GetPrivate()->m_program.m_constantLayouts.size() );
+
+            LayoutPtrConst pResult = pModel->GetPrivate()->m_program.m_constantLayouts
+                    [ args.value ];
+            GetMemory().SetLayout( item, pResult );
+            break;
+        }
+
+        case OP_TYPE::LA_MERGE:
+        {
+            auto args = pModel->GetPrivate()->m_program.GetOpArgs<OP::LayoutMergeArgs>(item.at);
+            switch (item.stage)
+            {
+            case 0:
+                    AddOp( SCHEDULED_OP( item.at, item, 1),
+                           SCHEDULED_OP( args.base, item),
+                           SCHEDULED_OP( args.added, item) );
+                break;
+
+            case 1:
+            {
+                Ptr<const Layout> pA = GetMemory().GetLayout( CACHE_ADDRESS(args.base,item) );
+                Ptr<const Layout> pB = GetMemory().GetLayout( CACHE_ADDRESS(args.added,item) );
+
+                LayoutPtrConst pResult;
+
+                if (pA && pB)
+                {
+					pResult = LayoutMerge(pA.get(),pB.get());
+                }
+                else if (pA)
+                {
+                    pResult = pA->Clone();
+                }
+                else if (pB)
+                {
+                    pResult = pB->Clone();
+                }
+
+                GetMemory().SetLayout( item, pResult );
+                break;
+            }
+
+            default:
+                check(false);
+            }
+
+            break;
+        }
+
+        case OP_TYPE::LA_PACK:
+        {
+            auto args = pModel->GetPrivate()->m_program.GetOpArgs<OP::LayoutPackArgs>(item.at);
+            switch (item.stage)
+            {
+            case 0:
+                    AddOp( SCHEDULED_OP( item.at, item, 1),
+                           SCHEDULED_OP( args.layout, item) );
+                break;
+
+            case 1:
+            {
+                Ptr<const Layout> pSource = GetMemory().GetLayout( CACHE_ADDRESS(args.layout,item) );
+
+				LayoutPtr pResult;
+
+				if (pSource)
+				{
+					pResult = pSource->Clone();
+
+					SCRATCH_LAYOUT_PACK scratch;
+					scratch.blocks.resize(pSource->GetBlockCount());
+					scratch.sorted.resize(pSource->GetBlockCount());
+					scratch.positions.resize(pSource->GetBlockCount());
+					scratch.priorities.resize(pSource->GetBlockCount());
+					scratch.reductions.resize(pSource->GetBlockCount());
+
+					LayoutPack3(pResult.get(), pSource.get(), &scratch);
+				}
+
+                GetMemory().SetLayout( item, pResult );
+                break;
+            }
+
+            default:
+                check(false);
+            }
+
+            break;
+        }
+
+        case OP_TYPE::LA_REMOVEBLOCKS:
+        {
+            auto args = pModel->GetPrivate()->m_program.GetOpArgs<OP::LayoutRemoveBlocksArgs>(item.at);
+            switch (item.stage)
+            {
+            case 0:
+                    AddOp( SCHEDULED_OP( item.at, item, 1),
+                           SCHEDULED_OP( args.source, item),
+                           SCHEDULED_OP( args.mesh, item) );
+                break;
+
+            case 1:
+            {
+                Ptr<const Layout> pSource = GetMemory().GetLayout( CACHE_ADDRESS(args.source,item) );
+                Ptr<const Mesh> pMesh = GetMemory().GetMesh( CACHE_ADDRESS(args.mesh,item) );
+
+				LayoutPtr pResult;
+
+				if (pSource && pMesh)
+				{
+					pResult = LayoutRemoveBlocks(pSource.get(), pMesh.get(), args.meshLayoutIndex);
+				}
+
+                GetMemory().SetLayout( item, pResult );
+                break;
+            }
+
+            default:
+                check(false);
+            }
+
+            break;
+        }
+
+        default:
+            // Operation not implemented
+            check( false );
+            break;
+        }
+    }
+
+
+    //---------------------------------------------------------------------------------------------
+    void CodeRunner::RunCode( SCHEDULED_OP& item,
+                              const Parameters* pParams,
+                              const Model* pModel,
+                              uint32_t lodMask)
+    {
+		//UE_LOG( LogMutableCore, Log, TEXT("Running :%5d , %d "), item.at, item.stage );
+		check( item.Type == SCHEDULED_OP::EType::Full );
+
+		OP_TYPE type = pModel->GetPrivate()->m_program.GetOpType(item.at);
+		//UE_LOG(LogMutableCore, Log, TEXT("Running :%5d , %d, of type %d "), item.at, item.stage, type);
+		switch ( type )
+        {
+        case OP_TYPE::NONE:
+            break;
+
+        case OP_TYPE::NU_CONDITIONAL:
+        case OP_TYPE::SC_CONDITIONAL:
+        case OP_TYPE::CO_CONDITIONAL:
+        case OP_TYPE::IM_CONDITIONAL:
+        case OP_TYPE::ME_CONDITIONAL:
+        case OP_TYPE::LA_CONDITIONAL:
+        case OP_TYPE::IN_CONDITIONAL:
+            RunCode_Conditional(item, pModel);
+            break;
+
+        case OP_TYPE::ME_CONSTANT:
+        case OP_TYPE::IM_CONSTANT:
+            RunCode_ConstantResource(item, pModel);
+            break;
+
+        case OP_TYPE::NU_SWITCH:
+        case OP_TYPE::SC_SWITCH:
+        case OP_TYPE::CO_SWITCH:
+        case OP_TYPE::IM_SWITCH:
+        case OP_TYPE::ME_SWITCH:
+        case OP_TYPE::LA_SWITCH:
+        case OP_TYPE::IN_SWITCH:
+            RunCode_Switch(item, pModel);
+            break;
+
+        case OP_TYPE::IN_ADDMESH:
+        case OP_TYPE::IN_ADDIMAGE:
+            RunCode_InstanceAddResource(item, pModel, pParams);
+            break;
+
+		default:
+		{
+			DATATYPE DataType = GetOpDataType(type);
+			switch (DataType)
+			{
+			case DT_INSTANCE:
+				RunCode_Instance(item, pModel, lodMask);
+				break;
+
+			case DT_MESH:
+				RunCode_Mesh(item, pModel);
+				break;
+
+			case DT_IMAGE:
+				RunCode_Image(item, pParams, pModel);
+				break;
+
+			case DT_LAYOUT:
+				RunCode_Layout(item, pModel);
+				break;
+
+			case DT_BOOL:
+				RunCode_Bool(item, pParams, pModel);
+				break;
+
+			case DT_SCALAR:
+				RunCode_Scalar(item, pParams, pModel);
+				break;
+
+			case DT_STRING:
+				RunCode_String(item, pParams, pModel);
+				break;
+
+			case DT_INT:
+				RunCode_Int(item, pParams, pModel);
+				break;
+
+			case DT_PROJECTOR:
+				RunCode_Projector(item, pParams, pModel);
+				break;
+
+			case DT_COLOUR:
+				RunCode_Colour(item, pParams, pModel);
+				break;
+
+			default:
+				check(false);
+				break;
+			}
+			break;
+		}
+
+        }
+    }
+
+
+	//---------------------------------------------------------------------------------------------
+	void CodeRunner::RunCodeImageDesc(SCHEDULED_OP& item,
+		const Parameters* pParams,
+		const Model* pModel, 
+		uint32_t lodMask
+	)
+	{
+		MUTABLE_CPUPROFILER_SCOPE(RunCodeImageDesc);
+
+		check(item.Type == SCHEDULED_OP::EType::ImageDesc);
+
+		// Ensure there is room for the result in the heap.
+		if (item.customState >= uint32(m_heapData.Num()))
+		{
+			m_heapImageDesc.SetNum(item.customState+1);
+		}
+
+
+		const mu::PROGRAM& program = pModel->GetPrivate()->m_program;
+
+		OP_TYPE type = program.GetOpType(item.at);
+		switch (type)
+		{
+
+		case OP_TYPE::IM_CONSTANT:
+		{
+			check(item.stage == 0);
+			OP::ResourceConstantArgs args = program.GetOpArgs<OP::ResourceConstantArgs>(item.at);
+			int32 ImageIndex = args.value;
+			m_heapImageDesc[item.customState].m_format = EImageFormat::IF_NONE;	// TODO: precalculate if necessary
+			m_heapImageDesc[item.customState].m_size[0] = program.m_constantImages[ImageIndex].ImageSizeX;
+			m_heapImageDesc[item.customState].m_size[1] = program.m_constantImages[ImageIndex].ImageSizeY;
+			m_heapImageDesc[item.customState].m_lods = program.m_constantImages[ImageIndex].LODCount;
+			GetMemory().SetValidDesc(item);
+			break;
+		}
+
+		case OP_TYPE::IM_PARAMETER:
+		{
+			check(item.stage == 0);
+			OP::ParameterArgs args = program.GetOpArgs<OP::ParameterArgs>(item.at);
+			EXTERNAL_IMAGE_ID id = pParams->GetImageValue(args.variable);
+			ImagePtr pResult = LoadExternalImage(id);
+			m_heapImageDesc[item.customState].m_format = pResult->GetFormat();
+			m_heapImageDesc[item.customState].m_size = pResult->GetSize();
+			m_heapImageDesc[item.customState].m_lods = pResult->GetLODCount();
+			GetMemory().SetValidDesc(item);
+			break;
+		}
+
+		case OP_TYPE::IM_CONDITIONAL:
+		{
+			OP::ConditionalArgs args = program.GetOpArgs<OP::ConditionalArgs>(item.at);
+			switch (item.stage)
+			{
+			case 0:
+			{
+				// We need to run the full condition result
+				SCHEDULED_OP FullConditionOp(args.condition, item);
+				FullConditionOp.Type = SCHEDULED_OP::EType::Full;
+				AddOp(SCHEDULED_OP(item.at, item, 1), FullConditionOp);
+				break;
+			}
+
+			case 1:
+			{
+				bool value = GetMemory().GetBool(CACHE_ADDRESS(args.condition, item.executionIndex, item.executionOptions));
+				OP::ADDRESS resultAt = value ? args.yes : args.no;
+				AddOp(SCHEDULED_OP(item.at, item, 2), SCHEDULED_OP(resultAt, item));
+				break;
+			}
+
+			case 2: GetMemory().SetValidDesc(item); break;
+			default: check(false);
+			}
+			break;
+		}
+
+		case OP_TYPE::IM_SWITCH:
+		{
+			const uint8_t* data = program.GetOpArgsPointer(item.at);
+		
+			OP::ADDRESS VarAddress;
+			FMemory::Memcpy( &VarAddress, data, sizeof(OP::ADDRESS));
+			data += sizeof(OP::ADDRESS);
+
+			OP::ADDRESS DefAddress;
+			FMemory::Memcpy( &DefAddress, data, sizeof(OP::ADDRESS));
+			data += sizeof(OP::ADDRESS);
+
+			uint32_t CaseCount;
+			FMemory::Memcpy( &CaseCount, data, sizeof(uint32_t));
+			data += sizeof(uint32_t);
+	
+			switch (item.stage)
+			{
+			case 0:
+			{
+				if (VarAddress)
+				{
+					// We need to run the full condition result
+					SCHEDULED_OP FullVariableOp(VarAddress, item);
+					FullVariableOp.Type = SCHEDULED_OP::EType::Full;
+					AddOp(SCHEDULED_OP(item.at, item, 1), FullVariableOp);
+				}
+				else
+				{
+					GetMemory().SetValidDesc(item);
+				}
+				break;
+			}
+
+			case 1:
+			{
+				// Get the variable result
+				int var = GetMemory().GetInt(CACHE_ADDRESS(VarAddress, item));
+
+				OP::ADDRESS valueAt = DefAddress;
+				for (uint32_t C = 0; C < CaseCount; ++C)
+				{
+					int32_t Condition;
+					FMemory::Memcpy( &Condition, data, sizeof(int32_t) );
+					data += sizeof(int32_t);
+					
+					OP::ADDRESS At;
+					FMemory::Memcpy( &At, data, sizeof(OP::ADDRESS) );
+					data += sizeof(OP::ADDRESS);
+
+					if (At && var == (int)Condition)
+					{
+						valueAt = At;
+						break;
+					}
+				}
+
+				AddOp(SCHEDULED_OP(item.at, item, 2, valueAt),
+					  SCHEDULED_OP(valueAt, item));
+
+				break;
+			}
+
+			case 2: GetMemory().SetValidDesc(item); break;
+			default: check(false); break;
+			}
+			break;
+		}
+
+		case OP_TYPE::IM_LAYERCOLOUR:
+		{
+			OP::ImageLayerColourArgs args = program.GetOpArgs<OP::ImageLayerColourArgs>(item.at);
+			switch (item.stage)
+			{
+			case 0: AddOp(SCHEDULED_OP(item.at, item, 1), SCHEDULED_OP(args.base, item)); break;
+			case 1: GetMemory().SetValidDesc(item); break;
+			default: check(false);
+			}
+			break;
+		}
+
+		case OP_TYPE::IM_LAYER:
+		{
+			OP::ImageLayerArgs args = program.GetOpArgs<OP::ImageLayerArgs>(item.at);
+			switch (item.stage)
+			{
+			case 0: AddOp(SCHEDULED_OP(item.at, item, 1), SCHEDULED_OP(args.base, item)); break;
+			case 1: GetMemory().SetValidDesc(item); break;
+			default: check(false);
+			}
+			break;
+		}
+
+		case OP_TYPE::IM_MULTILAYER:
+		{
+			OP::ImageMultiLayerArgs args = program.GetOpArgs<OP::ImageMultiLayerArgs>(item.at);
+			switch (item.stage)
+			{
+			case 0: AddOp(SCHEDULED_OP(item.at, item, 1), SCHEDULED_OP(args.base, item)); break;
+			case 1: GetMemory().SetValidDesc(item); break;
+			default: check(false);
+			}
+			break;
+		}
+
+		case OP_TYPE::IM_DIFFERENCE:
+		{
+			OP::ImageDifferenceArgs args = program.GetOpArgs<OP::ImageDifferenceArgs>(item.at);
+			switch (item.stage)
+			{
+			case 0: AddOp(SCHEDULED_OP(item.at, item, 1), SCHEDULED_OP(args.a, item)); break;
+			case 1: GetMemory().SetValidDesc(item); break;
+			default: check(false);
+			}
+			break;
+		}
+
+		case OP_TYPE::IM_NORMALCOMPOSITE:
+		{
+			OP::ImageNormalCompositeArgs args = program.GetOpArgs<OP::ImageNormalCompositeArgs>(item.at);
+			switch (item.stage)
+			{
+			case 0: AddOp(SCHEDULED_OP(item.at, item, 1), SCHEDULED_OP(args.base, item)); break;
+			case 1: GetMemory().SetValidDesc(item); break;
+			default: check(false);
+			}
+			break;
+		}
+
+		case OP_TYPE::IM_PIXELFORMAT:
+		{
+			OP::ImagePixelFormatArgs args = program.GetOpArgs<OP::ImagePixelFormatArgs>(item.at);
+			switch (item.stage)
+			{
+			case 0:
+				AddOp(SCHEDULED_OP(item.at, item, 1), SCHEDULED_OP(args.source, item));
+				break;
+
+			case 1:
+			{
+				// Update directly in the heap
+				EImageFormat OldFormat = m_heapImageDesc[item.customState].m_format;
+				EImageFormat NewFormat = args.format;
+				if (args.formatIfAlpha != EImageFormat::IF_NONE
+					&&
+					GetImageFormatData(OldFormat).m_channels > 3)
+				{
+					NewFormat = args.formatIfAlpha;
+				}
+				m_heapImageDesc[item.customState].m_format = NewFormat;				
+				GetMemory().SetValidDesc(item);
+				break;
+			}
+
+			default:
+				check(false);
+			}
+
+			break;
+		}
+
+		case OP_TYPE::IM_MIPMAP:
+		{
+			OP::ImageMipmapArgs args = program.GetOpArgs<OP::ImageMipmapArgs>(item.at);
+			switch (item.stage)
+			{
+			case 0:
+				AddOp(SCHEDULED_OP(item.at, item, 1), SCHEDULED_OP(args.source, item));
+				break;
+
+			case 1:
+			{
+				// Somewhat synched with Full op execution code.
+				FImageDesc BaseDesc = m_heapImageDesc[item.customState];
+				int levelCount = args.levels;
+				int maxLevelCount = Image::GetMipmapCount(BaseDesc.m_size[0], BaseDesc.m_size[1]);
+				if (levelCount == 0)
+				{
+					levelCount = maxLevelCount;
+				}
+				else if (levelCount > maxLevelCount)
+				{
+					// If code generation is smart enough, this should never happen.
+					// \todo But apparently it does, sometimes.
+					levelCount = maxLevelCount;
+				}
+
+				// At least keep the levels we already have.
+				int startLevel = BaseDesc.m_lods;
+				levelCount = std::max(startLevel, levelCount);
+
+				// Update result.
+				m_heapImageDesc[item.customState].m_lods = levelCount;
+				GetMemory().SetValidDesc(item);
+				break;
+			}
+
+			default:
+				check(false);
+			}
+			break;
+		}
+
+		case OP_TYPE::IM_RESIZE:
+		{
+			auto args = program.GetOpArgs<OP::ImageResizeArgs>(item.at);
+			switch (item.stage)
+			{
+			case 0:
+				AddOp(SCHEDULED_OP(item.at, item, 1), SCHEDULED_OP(args.source, item));
+				break;
+
+			case 1:
+				m_heapImageDesc[item.customState].m_size[0] = args.size[0];
+				m_heapImageDesc[item.customState].m_size[1] = args.size[1];
+				GetMemory().SetValidDesc(item);
+				break;
+
+			default:
+				check(false);
+			}
+			break;
+		}
+
+		case OP_TYPE::IM_RESIZELIKE:
+		{
+			auto args = program.GetOpArgs<OP::ImageResizeLikeArgs>(item.at);
+			switch (item.stage)
+			{
+			case 0:
+			{
+				int32 ResultAndBaseDesc = item.customState;
+				int32 SourceDescAddress = m_heapImageDesc.Add({});
+				int32 SecondStageData = m_heapData.Add({ 0.0f, ResultAndBaseDesc, SourceDescAddress, nullptr });
+				AddOp(SCHEDULED_OP(item.at, item, 1, SecondStageData),
+					SCHEDULED_OP(args.source, item, 0, ResultAndBaseDesc),
+					SCHEDULED_OP(args.sizeSource, item, 0, SourceDescAddress));
+				break;
+			}
+
+			case 1:
+			{
+				const SCHEDULED_OP_DATA& SecondStageData = m_heapData[ item.customState ];
+				FImageDesc& ResultAndBaseDesc = m_heapImageDesc[SecondStageData.min];
+				const FImageDesc& SourceDesc = m_heapImageDesc[SecondStageData.max];
+				ResultAndBaseDesc.m_size = SourceDesc.m_size;
+				GetMemory().SetValidDesc(item);
+				break;
+			}
+
+			default:
+				check(false);
+			}
+
+			break;
+		}
+
+		case OP_TYPE::IM_RESIZEREL:
+		{
+			auto args = program.GetOpArgs<OP::ImageResizeRelArgs>(item.at);
+			switch (item.stage)
+			{
+			case 0:
+				AddOp(SCHEDULED_OP(item.at, item, 1), SCHEDULED_OP(args.source, item));
+				break;
+
+			case 1:
+			{
+				FImageDesc& ResultAndBaseDesc = m_heapImageDesc[item.customState];
+				FImageSize destSize(
+					uint16_t(ResultAndBaseDesc.m_size[0] * args.factor[0] + 0.5f),
+					uint16_t(ResultAndBaseDesc.m_size[1] * args.factor[1] + 0.5f));
+				ResultAndBaseDesc.m_size = destSize;
+				GetMemory().SetValidDesc(item);
+				break;
+			}
+
+			default:
+				check(false);
+			}
+
+			break;
+		}
+
+		case OP_TYPE::IM_BLANKLAYOUT:
+		{
+			OP::ImageBlankLayoutArgs args = program.GetOpArgs<OP::ImageBlankLayoutArgs>(item.at);
+			switch (item.stage)
+			{
+			case 0:
+			{
+				// We need to run the full layout
+				SCHEDULED_OP FullLayoutOp(args.layout, item);
+				FullLayoutOp.Type = SCHEDULED_OP::EType::Full;
+				AddOp(SCHEDULED_OP(item.at, item, 1), FullLayoutOp);
+				break;
+			}
+
+			case 1:
+			{
+				Ptr<const Layout> pLayout = GetMemory().GetLayout(CACHE_ADDRESS(args.layout, item));
+
+				FIntPoint SizeInBlocks = pLayout->GetGridSize();
+				FIntPoint BlockSizeInPixels(args.blockSize[0], args.blockSize[1]);
+				FIntPoint ImageSizeInPixels = SizeInBlocks * BlockSizeInPixels;
+
+				FImageDesc& ResultAndBaseDesc = m_heapImageDesc[item.customState];
+				FImageSize destSize(uint16_t(ImageSizeInPixels.X), uint16_t(ImageSizeInPixels.Y));
+				ResultAndBaseDesc.m_size = destSize;
+				
+				if (args.generateMipmaps)
+				{
+					if (args.mipmapCount == 0)
+					{
+						ResultAndBaseDesc.m_lods = Image::GetMipmapCount(ImageSizeInPixels.X, ImageSizeInPixels.Y);
+					}
+					else
+					{
+						ResultAndBaseDesc.m_lods = args.mipmapCount;
+					}
+				}
+				GetMemory().SetValidDesc(item);
+				break;
+			}
+
+			default:
+				check(false);
+			}
+
+			break;
+		}
+
+		case OP_TYPE::IM_COMPOSE:
+		{
+			OP::ImageComposeArgs args = program.GetOpArgs<OP::ImageComposeArgs>(item.at);
+			switch (item.stage)
+			{
+			case 0: AddOp(SCHEDULED_OP(item.at, item, 1), SCHEDULED_OP(args.base, item)); break;
+			case 1: GetMemory().SetValidDesc(item); break;
+			default: check(false);
+			}
+			break;
+		}
+
+		case OP_TYPE::IM_INTERPOLATE:
+		{
+			OP::ImageInterpolateArgs args = program.GetOpArgs<OP::ImageInterpolateArgs>(item.at);
+			switch (item.stage)
+			{
+			case 0: AddOp(SCHEDULED_OP(item.at, item, 1), SCHEDULED_OP(args.targets[0], item)); break;
+			case 1: GetMemory().SetValidDesc(item); break;
+			default: check(false);
+			}
+			break;
+		}
+
+		case OP_TYPE::IM_INTERPOLATE3:
+		{
+			OP::ImageInterpolate3Args args = program.GetOpArgs<OP::ImageInterpolate3Args>(item.at);
+			switch (item.stage)
+			{
+			case 0: AddOp(SCHEDULED_OP(item.at, item, 1), SCHEDULED_OP(args.target0, item)); break;
+			case 1: GetMemory().SetValidDesc(item); break;
+			default: check(false);
+			}
+			break;
+		}
+
+		case OP_TYPE::IM_SATURATE:
+		{
+			OP::ImageSaturateArgs args = program.GetOpArgs<OP::ImageSaturateArgs>(item.at);
+			switch (item.stage)
+			{
+			case 0: AddOp(SCHEDULED_OP(item.at, item, 1), SCHEDULED_OP(args.base, item)); break;
+			case 1: GetMemory().SetValidDesc(item); break;
+			default: check(false);
+			}
+			break;
+		}
+
+		case OP_TYPE::IM_LUMINANCE:
+		{
+			OP::ImageLuminanceArgs args = program.GetOpArgs<OP::ImageLuminanceArgs>(item.at);
+			switch (item.stage)
+			{
+			case 0:
+				AddOp(SCHEDULED_OP(item.at, item, 1), SCHEDULED_OP(args.base, item));
+				break;
+
+			case 1:
+				m_heapImageDesc[item.customState].m_format = EImageFormat::IF_L_UBYTE;
+				GetMemory().SetValidDesc(item);
+				break;
+
+			default:
+				check(false);
+			}
+
+			break;
+		}
+
+		case OP_TYPE::IM_SWIZZLE:
+		{
+			OP::ImageSwizzleArgs args = program.GetOpArgs<OP::ImageSwizzleArgs>(item.at);
+			switch (item.stage)
+			{
+			case 0:
+				AddOp(SCHEDULED_OP(item.at, item, 1), SCHEDULED_OP(args.sources[0], item));
+				break;
+
+			case 1:
+				m_heapImageDesc[item.customState].m_format = args.format;
+				GetMemory().SetValidDesc(item);
+				break;
+
+			default:
+				check(false);
+			}
+
+			break;
+		}
+
+		case OP_TYPE::IM_SELECTCOLOUR:
+		{
+			OP::ImageSelectColourArgs args = program.GetOpArgs<OP::ImageSelectColourArgs>(item.at);
+			switch (item.stage)
+			{
+			case 0:
+				AddOp(SCHEDULED_OP(item.at, item, 1), SCHEDULED_OP(args.base, item));
+				break;
+
+			case 1:
+				m_heapImageDesc[item.customState].m_format = EImageFormat::IF_L_UBYTE;
+				GetMemory().SetValidDesc(item);
+				break;
+
+			default:
+				check(false);
+			}
+
+			break;
+		}
+
+		case OP_TYPE::IM_COLOURMAP:
+		{
+			OP::ImageColourMapArgs args = program.GetOpArgs<OP::ImageColourMapArgs>(item.at);
+			switch (item.stage)
+			{
+			case 0: AddOp(SCHEDULED_OP(item.at, item, 1), SCHEDULED_OP(args.base, item)); break;
+			case 1: GetMemory().SetValidDesc(item); break;
+			default: check(false);
+			}
+			break;
+		}
+
+		case OP_TYPE::IM_GRADIENT:
+		{
+			OP::ImageGradientArgs args = program.GetOpArgs<OP::ImageGradientArgs>(item.at);
+			m_heapImageDesc[item.customState].m_size[0] = args.size[0];
+			m_heapImageDesc[item.customState].m_size[1] = args.size[1];
+			m_heapImageDesc[item.customState].m_lods = 1;
+			m_heapImageDesc[item.customState].m_format = EImageFormat::IF_RGB_UBYTE;
+			GetMemory().SetValidDesc(item);
+			break;
+		}
+
+		case OP_TYPE::IM_BINARISE:
+		{
+			OP::ImageBinariseArgs args = program.GetOpArgs<OP::ImageBinariseArgs>(item.at);
+			switch (item.stage)
+			{
+			case 0:
+				AddOp(SCHEDULED_OP(item.at, item, 1), SCHEDULED_OP(args.base, item));
+				break;
+
+			case 1:
+				m_heapImageDesc[item.customState].m_format = EImageFormat::IF_L_UBYTE;
+				GetMemory().SetValidDesc(item);
+				break;
+
+			default:
+				check(false);
+			}
+			break;
+		}
+
+		case OP_TYPE::IM_INVERT:
+		{
+			OP::ImageInvertArgs args = program.GetOpArgs<OP::ImageInvertArgs>(item.at);
+			switch (item.stage)
+			{
+			case 0: AddOp(SCHEDULED_OP(item.at, item, 1), SCHEDULED_OP(args.base, item)); break;
+			case 1: GetMemory().SetValidDesc(item); break;
+			default: check(false);
+			}
+			break;
+		}
+
+		case OP_TYPE::IM_PLAINCOLOUR:
+		{
+			OP::ImagePlainColourArgs args = program.GetOpArgs<OP::ImagePlainColourArgs>(item.at);
+			m_heapImageDesc[item.customState].m_size[0] = args.size[0];
+			m_heapImageDesc[item.customState].m_size[1] = args.size[1];
+			m_heapImageDesc[item.customState].m_lods = 1;
+			m_heapImageDesc[item.customState].m_format = EImageFormat::IF_RGB_UBYTE;
+			GetMemory().SetValidDesc(item);
+			break;
+		}
+
+		case OP_TYPE::IM_CROP:
+		{
+			OP::ImageCropArgs args = program.GetOpArgs<OP::ImageCropArgs>(item.at);
+			switch (item.stage)
+			{
+			case 0:
+				AddOp(SCHEDULED_OP(item.at, item, 1), SCHEDULED_OP(args.source, item));
+				break;
+
+			case 1:
+				m_heapImageDesc[item.customState].m_size[0] = args.sizeX;
+				m_heapImageDesc[item.customState].m_size[1] = args.sizeY;
+				m_heapImageDesc[item.customState].m_lods = 1;
+				GetMemory().SetValidDesc(item);
+				break;
+
+			default:
+				check(false);
+			}
+			break;
+		}
+
+		case OP_TYPE::IM_PATCH:
+		{
+			OP::ImagePatchArgs args = program.GetOpArgs<OP::ImagePatchArgs>(item.at);
+			switch (item.stage)
+			{
+			case 0: AddOp(SCHEDULED_OP(item.at, item, 1), SCHEDULED_OP(args.base, item)); break;
+			case 1: GetMemory().SetValidDesc(item); break;
+			default: check(false);
+			}
+			break;
+		}
+
+		case OP_TYPE::IM_RASTERMESH:
+		{
+			OP::ImageRasterMeshArgs args = program.GetOpArgs<OP::ImageRasterMeshArgs>(item.at);
+			m_heapImageDesc[item.customState].m_size[0] = args.sizeX;
+			m_heapImageDesc[item.customState].m_size[1] = args.sizeY;
+			m_heapImageDesc[item.customState].m_lods = 1;
+			m_heapImageDesc[item.customState].m_format = EImageFormat::IF_L_UBYTE;
+			GetMemory().SetValidDesc(item);
+			break;
+		}
+
+		case OP_TYPE::IM_MAKEGROWMAP:
+		{
+			OP::ImageMakeGrowMapArgs args = program.GetOpArgs<OP::ImageMakeGrowMapArgs>(item.at);
+			switch (item.stage)
+			{
+			case 0:
+				AddOp(SCHEDULED_OP(item.at, item, 1), SCHEDULED_OP(args.mask, item));
+				break;
+
+			case 1:
+				m_heapImageDesc[item.customState].m_format = EImageFormat::IF_L_UBYTE;
+				m_heapImageDesc[item.customState].m_lods = 1;
+				GetMemory().SetValidDesc(item);
+				break;
+
+			default:
+				check(false);
+			}
+
+			break;
+		}
+
+		case OP_TYPE::IM_DISPLACE:
+		{
+			OP::ImageDisplaceArgs args = program.GetOpArgs<OP::ImageDisplaceArgs>(item.at);
+			switch (item.stage)
+			{
+			case 0: AddOp(SCHEDULED_OP(item.at, item, 1), SCHEDULED_OP(args.source, item)); break;
+			case 1: GetMemory().SetValidDesc(item); break;
+			default: check(false);
+			}
+			break;
+		}
+
+        case OP_TYPE::IM_TRANSFORM:
+        {
+
+			OP::ImageTransformArgs Args = program.GetOpArgs<OP::ImageTransformArgs>(item.at);
+
+            switch (item.stage)
+            {
+            case 0:
+			{
+				AddOp(SCHEDULED_OP(item.at, item, 1), SCHEDULED_OP(Args.base, item));
+                break;
+			}
+            case 1:
+            {
+				GetMemory().SetValidDesc(item);
+                break;
+            }
+
+            default:
+                check(false);
+            }
+
+			break;
+		}
+
+		default:
+			if (type != OP_TYPE::NONE)
+			{
+				// Operation not implemented
+				check(false);
+				m_heapImageDesc[item.customState] = FImageDesc();
+			}
+			break;
+		}
+	}
+
+
+}
+
