[mutable] Move Mutable plugin from the restricted folder to the experimental plugins folder.

diff --git a/Engine/Plugins/Experimental/Mutable/Source/MutableRuntime/Private/CodeRunner.h b/Engine/Plugins/Experimental/Mutable/Source/MutableRuntime/Private/CodeRunner.h
new file mode 100644
--- /dev/null
+++ b/Engine/Plugins/Experimental/Mutable/Source/MutableRuntime/Private/CodeRunner.h
@@ -0,0 +1,335 @@
+// Copyright Epic Games, Inc. All Rights Reserved.
+
+#pragma once
+
+#include "System.h"
+
+#include "Operations.h"
+#include "ImagePrivate.h"
+#include "Layout.h"
+#include "InstancePrivate.h"
+#include "ParametersPrivate.h"
+#include "SystemPrivate.h"
+#include "ModelPrivate.h"
+
+// This define could come from MutableRuntime/Public/System.h
+#ifdef MUTABLE_USE_NEW_TASKGRAPH
+	#include "Tasks/Task.h"
+#else
+	#include "Async/Async.h"
+#endif
+
+
+namespace  mu
+{
+
+
+    //---------------------------------------------------------------------------------------------
+    //! Code execution of the mutable virtual machine.
+    //---------------------------------------------------------------------------------------------
+    class CodeRunner : public Base
+    {
+    public:
+		CodeRunner(const SettingsPtrConst&, class System::Private*, 
+			const Model* pModel, const Parameters* pParams,
+			OP::ADDRESS at, uint32 lodMask, uint8 executionOptions, SCHEDULED_OP::EType );
+
+    protected:
+
+        //! Settings that may affect the execution of some operations, like image conversion
+        //! quality.
+        SettingsPtrConst m_pSettings;
+
+        //! Type of data sometimes stored in the code runner heap to pass info between
+        //! operation stages.
+        struct SCHEDULED_OP_DATA
+        {
+            float bifactor=0.0f;
+            int32 min=0, max=0;
+            Ptr<const Layout> layout;
+        };
+
+
+        Ptr<RangeIndex> BuildCurrentOpRangeIndex( const SCHEDULED_OP& item,
+                                                const Parameters* pParams,
+                                                const Model* pModel,
+                                                int parameterIndex );
+
+
+        void RunCode( SCHEDULED_OP& item, const Parameters* pParams, const Model* pModel, uint32_t lodMask);
+
+        void RunCode_Conditional( SCHEDULED_OP& item, const Model* pModel );
+        void RunCode_Switch( SCHEDULED_OP& item, const Model* pModel );
+        void RunCode_Instance( SCHEDULED_OP& item, const Model* pModel, uint32_t lodMask );
+        void RunCode_InstanceAddResource( SCHEDULED_OP& item, const Model* pModel, const Parameters* pParams );
+        void RunCode_ConstantResource( SCHEDULED_OP& item, const Model* pModel );
+        void RunCode_Mesh( SCHEDULED_OP& item, const Model* pModel );
+        void RunCode_Image( SCHEDULED_OP& item, const Parameters* pParams, const Model* pModel );
+        void RunCode_Layout( SCHEDULED_OP& item, const Model* pModel );
+        void RunCode_Bool( SCHEDULED_OP& item, const Parameters* pParams, const Model* pModel );
+        void RunCode_Int( SCHEDULED_OP& item, const Parameters* pParams, const Model* pModel );
+        void RunCode_Scalar( SCHEDULED_OP& item, const Parameters* pParams, const Model* pModel );        
+        void RunCode_String( SCHEDULED_OP& item, const Parameters* pParams, const Model* pModel );
+        void RunCode_Colour( SCHEDULED_OP& item, const Parameters* pParams, const Model* pModel );
+        void RunCode_Projector( SCHEDULED_OP& item, const Parameters* pParams, const Model* pModel );
+
+		void RunCodeImageDesc(SCHEDULED_OP& item, const Parameters* pParams, const Model* pModel, uint32_t lodMask);
+
+
+        ImagePtr LoadExternalImage( EXTERNAL_IMAGE_ID id );
+
+        //! Heap of intermediate data pushed by some instructions and referred by others.
+        //! It is not released until no operations are pending.
+		TArray< SCHEDULED_OP_DATA > m_heapData;
+		TArray< FImageDesc > m_heapImageDesc;
+
+	public:
+
+		void Run();
+
+		//!
+		void GetImageDescResult(FImageDesc& OutDesc);
+
+		//!
+		FProgramCache& GetMemory();
+
+
+		struct FTask
+		{
+			FTask() {}
+			FTask(const SCHEDULED_OP& InOp) : Op(InOp) {}
+			FTask(const SCHEDULED_OP& InOp, const SCHEDULED_OP& InDep0) : Op(InOp) { Deps.Add(InDep0); }
+			FTask(const SCHEDULED_OP& InOp, const SCHEDULED_OP& InDep0, const SCHEDULED_OP& InDep1) : Op(InOp) { Deps.Add(InDep0); Deps.Add(InDep1); }
+			FTask(const SCHEDULED_OP& InOp, const SCHEDULED_OP& InDep0, const SCHEDULED_OP& InDep1, const SCHEDULED_OP& InDep2) : Op(InOp) { Deps.Add(InDep0); Deps.Add(InDep1); Deps.Add(InDep2); }
+			FTask(const SCHEDULED_OP& InOp, const SCHEDULED_OP& InDep0, const SCHEDULED_OP& InDep1, const SCHEDULED_OP& InDep2, const SCHEDULED_OP& InDep3) : Op(InOp) { Deps.Add(InDep0); Deps.Add(InDep1); Deps.Add(InDep2); Deps.Add(InDep3); }
+			FTask(const SCHEDULED_OP& InOp, const SCHEDULED_OP& InDep0, const SCHEDULED_OP& InDep1, const SCHEDULED_OP& InDep2, const SCHEDULED_OP& InDep3, const SCHEDULED_OP& InDep4) : Op(InOp) { Deps.Add(InDep0); Deps.Add(InDep1); Deps.Add(InDep2); Deps.Add(InDep3); Deps.Add(InDep4); }
+
+			SCHEDULED_OP Op;
+			TArray<CACHE_ADDRESS, TInlineAllocator<3>> Deps;
+		};
+
+		class FIssuedTask
+		{
+		public:
+			SCHEDULED_OP Op;
+
+#ifdef MUTABLE_USE_NEW_TASKGRAPH
+			UE::Tasks::FTask Event = {};
+#else
+			FGraphEventRef Event = nullptr;
+#endif
+
+			virtual ~FIssuedTask() {}
+			virtual bool Prepare(CodeRunner*, const Model*, bool& bOutFailed) { bOutFailed = false; return true; }
+			virtual void DoWork() {}
+			virtual void Complete(CodeRunner*) = 0;
+			virtual bool IsComplete(CodeRunner*)
+			{ 
+				// Event can be null if we forced single-threaded execution.
+#ifdef MUTABLE_USE_NEW_TASKGRAPH
+				return !Event.IsValid() || Event.IsCompleted(); 
+#else
+				return !Event.IsValid() || Event->IsComplete();
+#endif
+			}
+		};
+
+		class FLoadMeshRomTask : public CodeRunner::FIssuedTask
+		{
+		public:
+			FLoadMeshRomTask(SCHEDULED_OP InOp, int32 InRomIndex)
+			{
+				Op = InOp;
+				RomIndex = InRomIndex;
+			}
+
+			// FIssuedTask interface
+			bool Prepare(CodeRunner*, const Model*, bool& bOutFailed) override;
+			void Complete(CodeRunner*) override;
+			bool IsComplete(CodeRunner*) override;
+
+		private:
+			int32 RomIndex = -1;
+		};
+
+		class FLoadImageRomsTask : public CodeRunner::FIssuedTask
+		{
+		public:
+			FLoadImageRomsTask(SCHEDULED_OP InOp, int32 InLODIndexIndex, int32 InLODIndexCount )
+			{
+				Op = InOp;
+				LODIndexIndex = InLODIndexIndex;
+				LODIndexCount = InLODIndexCount;
+			}
+
+			// FIssuedTask interface
+			bool Prepare(CodeRunner*, const Model*, bool& bOutFailed) override;
+			void Complete(CodeRunner*) override;
+			bool IsComplete(CodeRunner*) override;
+
+		private:
+ 			int32 LODIndexIndex = -1;
+			int32 LODIndexCount = -1;
+		};
+
+		void AddOp(const SCHEDULED_OP& op)
+		{
+			OpenTasks.Add(op);
+			ScheduledStagePerOp[op] = op.stage + 1;
+		}
+
+		void AddOp(const SCHEDULED_OP& op,
+			const SCHEDULED_OP& dep0)
+		{
+			ClosedTasks.Add(FTask(op, dep0));
+			ScheduledStagePerOp[op] = op.stage + 1;
+			AddChildren(dep0);
+		}
+
+		void AddOp(const SCHEDULED_OP& op,
+			const SCHEDULED_OP& dep0,
+			const SCHEDULED_OP& dep1)
+		{
+			ClosedTasks.Add(FTask(op, dep0, dep1));
+			ScheduledStagePerOp[op] = op.stage + 1;
+			AddChildren(dep0);
+			AddChildren(dep1);
+		}
+
+		void AddOp(const SCHEDULED_OP& op,
+			const SCHEDULED_OP& dep0,
+			const SCHEDULED_OP& dep1,
+			const SCHEDULED_OP& dep2)
+		{
+			ClosedTasks.Add(FTask(op, dep0, dep1, dep2));
+			ScheduledStagePerOp[op] = op.stage + 1;
+			AddChildren(dep0);
+			AddChildren(dep1);
+			AddChildren(dep2);
+		}
+
+		void AddOp(const SCHEDULED_OP& op,
+			const SCHEDULED_OP& dep0,
+			const SCHEDULED_OP& dep1,
+			const SCHEDULED_OP& dep2,
+			const SCHEDULED_OP& dep3)
+		{
+			ClosedTasks.Add(FTask(op, dep0, dep1, dep2, dep3));
+			ScheduledStagePerOp[op] = op.stage + 1;
+			AddChildren(dep0);
+			AddChildren(dep1);
+			AddChildren(dep2);
+			AddChildren(dep3);
+		}
+
+		void AddOp(const SCHEDULED_OP& op,
+			const SCHEDULED_OP& dep0,
+			const SCHEDULED_OP& dep1,
+			const SCHEDULED_OP& dep2,
+			const SCHEDULED_OP& dep3,
+			const SCHEDULED_OP& dep4)
+		{
+			ClosedTasks.Add(FTask(op, dep0, dep1, dep2, dep3, dep4));
+			ScheduledStagePerOp[op] = op.stage + 1;
+			AddChildren(dep0);
+			AddChildren(dep1);
+			AddChildren(dep2);
+			AddChildren(dep3);
+			AddChildren(dep4);
+		}
+
+		void AddOp(const SCHEDULED_OP& op, const vector<SCHEDULED_OP>& deps)
+		{
+			FTask Task(op);
+			for (const auto& d : deps)
+			{
+				Task.Deps.Add(d);
+			}
+			ClosedTasks.Add(Task);
+			for (const auto& d : deps)
+			{
+				AddChildren(d);
+			}
+		}
+
+		template<class RangeType>
+		void AddOp(const SCHEDULED_OP& Op, const RangeType& Deps)
+		{
+			FTask Task(Op);
+
+			for (const SCHEDULED_OP& D : Deps)
+			{
+				Task.Deps.Add(D);
+			}
+
+			ClosedTasks.Add(Task);
+
+			ScheduledStagePerOp[Op] = Op.stage + 1;
+
+			for (const SCHEDULED_OP& D : Deps)
+			{
+				AddChildren(D);
+			}
+		}
+
+
+	protected:
+
+		//! Stack of pending operations, and the execution stage they are in.
+		TArray< FTask > ClosedTasks;
+		TArray< SCHEDULED_OP > OpenTasks;
+		CodeContainer<uint8> ScheduledStagePerOp;
+		TArray< TSharedPtr<FIssuedTask> > IssuedTasks;
+
+	// TODO: protect this.
+	public:
+		
+		/** This flag is turned on when a streaming error or similar happens. Results are not usable.
+		* This should only happen in-editor.
+		*/
+		bool bUnrecoverableError = false;
+
+		System::Private* m_pSystem = nullptr;
+		const Model* m_pModel = nullptr;
+		const Parameters* m_pParams = nullptr;
+		uint32 m_lodMask = 0;
+
+		// Async rom loading control
+		struct ROM_LOAD_OP
+		{
+			int32 m_romIndex = 0;
+			DATATYPE ConstantType = DT_NONE;
+			ModelStreamer::OPERATION_ID m_streamID;
+			vector<uint8> m_streamBuffer;
+		};
+		vector<ROM_LOAD_OP> m_romLoadOps;
+
+		//! Count of pending operations for every rom index
+		vector<uint16> m_romPendingOps;
+
+	private:
+
+		inline void AddChildren(const SCHEDULED_OP& dep)
+		{
+			CACHE_ADDRESS at(dep);
+			if (dep.at && !GetMemory().IsValid(at))
+			{
+				if (ScheduledStagePerOp.get(at) <= dep.stage)
+				{
+					OpenTasks.Add(dep);
+					ScheduledStagePerOp[at] = dep.stage + 1;
+				}
+			}
+
+			m_pSystem->m_memory->IncreaseHitCount(at);
+		}
+
+		/** Try to create a concurrent task for the given op. Return null if not possible. */
+		TSharedPtr<FIssuedTask> IssueOp(SCHEDULED_OP item);
+
+		/** */
+		void CompleteRomLoadOp(ROM_LOAD_OP& o);
+
+    };
+
+}
diff --git a/Engine/Plugins/Experimental/Mutable/Source/MutableRuntime/Private/CodeRunnerTasked.cpp b/Engine/Plugins/Experimental/Mutable/Source/MutableRuntime/Private/CodeRunnerTasked.cpp
new file mode 100644
--- /dev/null
+++ b/Engine/Plugins/Experimental/Mutable/Source/MutableRuntime/Private/CodeRunnerTasked.cpp
@@ -0,0 +1,1453 @@
+// Copyright Epic Games, Inc. All Rights Reserved.
+
+#include "CodeRunner.h"
+
+#include "SystemPrivate.h"
+#include "SettingsPrivate.h"
+#include "ModelPrivate.h"
+#include "InstancePrivate.h"
+#include "ImagePrivate.h"
+#include "MemoryPrivate.h"
+#include "MeshPrivate.h"
+
+#include "Operations.h"
+#include "OpImageResize.h"
+#include "OpImagePixelFormat.h"
+#include "OpImageMipmap.h"
+#include "OpImageSoftLight.h"
+#include "OpImageHardLight.h"
+#include "OpImageBurn.h"
+#include "OpImageDodge.h"
+#include "OpImageScreen.h"
+#include "OpImageOverlay.h"
+#include "OpImageAlphaOverlay.h"
+#include "OpImageCompose.h"
+#include "OpImageBlend.h"
+#include "OpImageDifference.h"
+#include "OpImageInterpolate.h"
+#include "OpImageMultiply.h"
+#include "OpImageSaturate.h"
+#include "OpImageLuminance.h"
+#include "OpImageSwizzle.h"
+#include "OpImageSelectColour.h"
+#include "OpImageColourMap.h"
+#include "OpImageGradient.h"
+#include "OpImageBinarise.h"
+#include "OpImageCrop.h"
+#include "OpImageRasterMesh.h"
+#include "OpImageGrow.h"
+#include "OpImageProject.h"
+#include "OpImageDisplace.h"
+#include "OpImageInvert.h"
+#include "OpImageNormalCombine.h"
+
+#include "CodeVisitor.h"
+#include "Platform.h"
+
+
+DECLARE_CYCLE_STAT(TEXT("MutableCoreTask"), STAT_MutableCoreTask, STATGROUP_Game);
+
+namespace mu
+{
+
+	void CodeRunner::CompleteRomLoadOp( ROM_LOAD_OP& o )
+	{
+		if (DebugRom && (DebugRomAll || o.m_romIndex == DebugRomIndex))
+			UE_LOG(LogMutableCore, Log, TEXT("CodeRunner::CompleteRomLoadOp for rom %d."), o.m_romIndex);
+
+		m_pSystem->m_pStreamInterface->EndRead(o.m_streamID);
+
+		PROGRAM& program = m_pModel->GetPrivate()->m_program;
+		{
+			MUTABLE_CPUPROFILER_SCOPE(Unserialise);
+
+			InputMemoryStream stream(o.m_streamBuffer.data(), o.m_streamBuffer.size());
+			InputArchive arch(&stream);
+
+			int32 ResIndex = program.m_roms[o.m_romIndex].ResourceIndex;
+			switch (o.ConstantType)
+			{
+			case DATATYPE::DT_MESH:
+			{
+				check(!program.m_constantMeshes[ResIndex].Value);
+				program.m_constantMeshes[ResIndex].Value = Mesh::StaticUnserialise(arch);
+				check(program.m_constantMeshes[ResIndex].Value);
+				break;
+			}
+			case DATATYPE::DT_IMAGE:
+			{
+				check(!program.m_constantImageLODs[ResIndex].Value);
+				program.m_constantImageLODs[ResIndex].Value = Image::StaticUnserialise(arch);
+				check(program.m_constantImageLODs[ResIndex].Value);
+				break;
+			}
+			default:
+				check(false);
+				break;
+			}
+
+			o.m_romIndex = -1;
+			o.m_streamBuffer = vector<uint8_t>();
+		}
+	}
+
+
+    //---------------------------------------------------------------------------------------------
+    void CodeRunner::Run ()
+    {
+		bUnrecoverableError = false;
+
+		m_heapData.SetNum(0, false);
+		m_heapImageDesc.SetNum(0, false);
+
+		// Profiling data. TODO: make optional
+		bool bProfile = false;
+		uint32 NumRunOps = 0;
+		uint32 RunOpsPerType[size_t(OP_TYPE::COUNT)] = {};
+
+        while( !OpenTasks.IsEmpty() || !ClosedTasks.IsEmpty() || !IssuedTasks.IsEmpty())
+        {
+			// Debug: log the amount of tasks that we'd be able to run concurrently:
+			//{
+			//	int32 ClosedReady = ClosedTasks.Num();
+			//	for (int Index = ClosedTasks.Num() - 1; Index >= 0; --Index)
+			//	{
+			//		for (const CACHE_ADDRESS& Dep : ClosedTasks[Index].Deps)
+			//		{
+			//			if (Dep.at && !GetMemory().IsValid(Dep))
+			//			{
+			//				--ClosedReady;
+			//				continue;
+			//			}
+			//		}
+			//	}
+
+			//	UE_LOG(LogMutableCore, Log, TEXT("Tasks: %5d open, %5d issued, %5d closed, %d closed ready"), OpenTasks.Num(), IssuedTasks.Num(), ClosedTasks.Num(), ClosedReady);
+			//}
+
+			for (int Index = 0; Index < IssuedTasks.Num(); )
+			{
+				check(IssuedTasks[Index]);
+
+				bool WorkDone = IssuedTasks[Index]->IsComplete(this);
+				if (WorkDone)
+				{
+					const SCHEDULED_OP& item = IssuedTasks[Index]->Op;
+					IssuedTasks[Index]->Complete(this);
+
+					if (ScheduledStagePerOp[item] == item.stage + 1)
+					{
+						// We completed everything that was requested, clear it otherwise if needed
+						// again it is not going to be rebuilt.
+						// \TODO: track rebuilds.
+						ScheduledStagePerOp[item] = 0;
+					}
+
+					IssuedTasks.RemoveAt(Index); // with swap? changes order of execution.
+				}
+				else
+				{
+					++Index;
+				}
+			}
+
+
+			while (!OpenTasks.IsEmpty())
+			{
+				SCHEDULED_OP item = OpenTasks.Pop();
+
+				// Special processing in case it is an ImageDesc operation
+				if ( item.Type==SCHEDULED_OP::EType::ImageDesc )
+				{
+					RunCodeImageDesc(item, m_pParams, m_pModel, m_lodMask);
+					continue;
+				}
+
+				// Don't run it if we already have the result.
+				CACHE_ADDRESS cat(item);
+				if (GetMemory().IsValid(cat))
+				{
+					continue;
+				}
+
+				// See if we can schedule this item concurrently
+				TSharedPtr<FIssuedTask> IssuedTask = IssueOp(item);
+				if (IssuedTask)
+				{
+					bool bFailed = false;
+					bool hasWork = IssuedTask->Prepare(this, m_pModel, bFailed);
+					if (bFailed)
+					{
+						bUnrecoverableError = true;
+						return;
+					}
+
+					// Launch it
+					if (hasWork)
+					{
+						// Optional hack to ensure deterministic single threading.
+						if (bForceSingleThread)
+						{
+							IssuedTask->Event = {};
+							IssuedTask->DoWork();
+						}
+						else
+						{
+#ifdef MUTABLE_USE_NEW_TASKGRAPH
+							IssuedTask->Event = UE::Tasks::Launch(TEXT("MutableCore_Task"),
+								[IssuedTask]() { IssuedTask->DoWork(); },
+								UE::Tasks::ETaskPriority::Inherit );
+#else
+							IssuedTask->Event = FFunctionGraphTask::CreateAndDispatchWhenReady(
+								[IssuedTask]() { IssuedTask->DoWork(); },
+								TStatId{},
+								nullptr,
+								ENamedThreads::AnyThread);
+#endif
+						}
+					}
+
+					// Remember it for later processing.
+					IssuedTasks.Add(IssuedTask);
+				}
+				else
+				{
+					// Run immediately
+					RunCode(item, m_pParams, m_pModel, m_lodMask);
+
+					if (ScheduledStagePerOp[item] == item.stage + 1)
+					{
+						// We completed everything that was requested, clear it otherwise if needed
+						// again it is not going to be rebuilt.
+						// \TODO: track rebuilds.
+						ScheduledStagePerOp[item] = 0;
+					}
+				}
+
+				if (bProfile)
+				{
+					++NumRunOps;
+					RunOpsPerType[size_t(m_pModel->GetPrivate()->m_program.GetOpType(item.at))]++;
+				}
+			}
+
+			// Look for completed streaming ops and complete the rom loading
+			for (auto& o : m_romLoadOps)
+			{
+				if (o.m_romIndex>=0 && m_pSystem->m_pStreamInterface->IsReadCompleted(o.m_streamID))
+				{
+					CompleteRomLoadOp(o);
+				}
+			}
+
+			// Look for a closed task with dependencies satisfied and move them to open.
+			bool bSomeWasReady = false;
+			for (int Index = 0; Index<ClosedTasks.Num(); )
+			{
+				bool Ready = true;
+				for ( const CACHE_ADDRESS& Dep: ClosedTasks[Index].Deps )
+				{
+					bool bDependencySatisfied = false;					
+					bDependencySatisfied = Dep.at && !GetMemory().IsValid(Dep);
+					
+					if (bDependencySatisfied)
+					{
+						Ready = false;
+						continue;
+					}
+				}
+
+				if (Ready)
+				{
+					bSomeWasReady = true;
+					FTask Task = ClosedTasks[Index];
+					ClosedTasks.RemoveAt(Index); // with swap? would change order of execution.
+					OpenTasks.Push(Task.Op);
+				}
+				else
+				{
+					++Index;
+				}
+			}
+
+			// Debug: Did we dead-lock?
+			bool bDeadLock = !(OpenTasks.Num() || IssuedTasks.Num() || !ClosedTasks.Num() || bSomeWasReady);
+			if (bDeadLock)
+			{
+				// Log the task graph
+				for (int Index = 0; Index < ClosedTasks.Num(); ++Index)
+				{
+					FString TaskDesc = FString::Printf(TEXT("Closed task %d-%d-%d depends on : "), ClosedTasks[Index].Op.at, ClosedTasks[Index].Op.executionIndex, ClosedTasks[Index].Op.stage );
+					for (const CACHE_ADDRESS& Dep : ClosedTasks[Index].Deps)
+					{
+						if (Dep.at && !GetMemory().IsValid(Dep))
+						{
+							TaskDesc += FString::Printf(TEXT("%d-%d, "), Dep.at, Dep.executionIndex);
+						}
+					}
+
+					UE_LOG(LogMutableCore, Log, TEXT("%s"), *TaskDesc);
+				}
+				check(false);
+			}
+
+			// If at this point there is no open op and we haven't finished, we need to wait for an issued op to complete.
+			if (OpenTasks.IsEmpty() && !IssuedTasks.IsEmpty())
+			{
+				MUTABLE_CPUPROFILER_SCOPE(CodeRunner_WaitIssued);
+
+				for (int32 IssuedIndex = 0; IssuedIndex<IssuedTasks.Num(); ++IssuedIndex)
+				{
+					if (IssuedTasks[IssuedIndex]->Event.IsValid())
+					{
+#ifdef MUTABLE_USE_NEW_TASKGRAPH
+						IssuedTasks[IssuedIndex]->Event.Wait();
+#else
+						IssuedTasks[IssuedIndex]->Event->Wait();
+#endif
+						break;
+					}
+				}
+
+				// If we reached here it means we didn't find an op to wait for. Try to wait for a loading op.
+				// \todo: unify laoding ops with normal ones?
+				for (auto& o : m_romLoadOps)
+				{
+					if (o.m_romIndex >= 0)
+					{
+						CompleteRomLoadOp(o);
+						break;
+					}
+
+					// We should never reach this, since it would mean we deadlocked.
+					//check(false);
+				}
+
+			}
+		}
+
+		if (bProfile)
+		{
+			UE_LOG(LogMutableCore, Log, TEXT("Mutable Heap Bytes: %d"), m_heapData.Num()* m_heapData.GetTypeSize());
+			UE_LOG(LogMutableCore, Log, TEXT("Ran ops : %5d "), NumRunOps);
+
+			constexpr int HistogramSize = 8;
+			int32 MostCommonOps[HistogramSize] = {};
+			for (size_t i=0; i<size_t(OP_TYPE::COUNT); ++i)
+			{
+				for (int32 h=0; h<HistogramSize;++h)
+				{
+					if (RunOpsPerType[i] > RunOpsPerType[MostCommonOps[h]])
+					{
+						// Displace others
+						int32 ElementsToMove = HistogramSize - h - 1;
+						if (ElementsToMove > 0)
+						{
+							FMemory::Memcpy(&MostCommonOps[h + 1], &MostCommonOps[h], sizeof(int32)* ElementsToMove);
+						}
+						// Set new value
+						MostCommonOps[h] = i;
+						break;
+					}
+				}
+			}
+
+			for (int h = 0; h < HistogramSize; ++h)
+			{
+				UE_LOG(LogMutableCore, Log, TEXT("    op %4d, %4d times."), MostCommonOps[h], RunOpsPerType[MostCommonOps[h]]);
+			}
+		}
+    }
+
+
+	//---------------------------------------------------------------------------------------------
+	void CodeRunner::GetImageDescResult(FImageDesc& OutDesc)
+	{
+		check( m_heapImageDesc.Num()>0 );
+		OutDesc = m_heapImageDesc[0];
+	}
+
+
+	//---------------------------------------------------------------------------------------------
+	//---------------------------------------------------------------------------------------------
+	//---------------------------------------------------------------------------------------------
+	class FImageLayerTask : public CodeRunner::FIssuedTask
+	{
+	public:
+		FImageLayerTask(SCHEDULED_OP, FProgramCache&, const OP::ImageLayerArgs&, int imageCompressionQuality );
+
+		// FIssuedTask interface
+		void DoWork() override;
+		void Complete(CodeRunner* staged) override;
+
+	private:
+		int m_imageCompressionQuality = 0;
+		OP::ImageLayerArgs Args;
+		Ptr<const Image> m_base;
+		Ptr<const Image> m_blended;
+		Ptr<const Image> m_mask;
+		Ptr<const Image> Result;
+	};
+
+
+	//---------------------------------------------------------------------------------------------
+	FImageLayerTask::FImageLayerTask(SCHEDULED_OP InOp, FProgramCache& Memory, const OP::ImageLayerArgs& InArgs, int imageCompressionQuality)
+	{
+		Op = InOp;
+		m_imageCompressionQuality = imageCompressionQuality;
+		Args = InArgs;
+		m_base = Memory.GetImage({ Args.base, InOp.executionIndex, InOp.executionOptions });
+		m_blended = Memory.GetImage({ Args.blended, InOp.executionIndex, InOp.executionOptions });
+		if (Args.mask)
+		{
+			m_mask = Memory.GetImage({ Args.mask, InOp.executionIndex, InOp.executionOptions });
+		}
+	}
+
+
+	//---------------------------------------------------------------------------------------------
+	void FImageLayerTask::DoWork()
+	{
+		MUTABLE_CPUPROFILER_SCOPE(FImageLayerTask);
+
+		// Warning: This runs in a worker thread.
+
+		// TODO: a bug may cause this (issue T97)
+		EImageFormat finalFormat = EImageFormat::IF_NONE;
+		if (IsCompressedFormat(m_base->GetFormat()))
+		{
+			UE_LOG(LogMutableCore, Error, TEXT("Image layer on compressed image is not supported. Ignored."));
+			finalFormat = m_base->GetFormat();
+			auto uncompressedFormat = GetUncompressedFormat(finalFormat);
+			m_base = ImagePixelFormat(m_imageCompressionQuality, m_base.get(), uncompressedFormat);
+		}
+
+		if (m_blended && m_base->GetFormat() != m_blended->GetFormat())
+		{
+			m_blended =
+				ImagePixelFormat(m_imageCompressionQuality, m_blended.get(), m_base->GetFormat());
+		}
+
+		bool applyToAlpha = (Args.flags & OP::ImageLayerArgs::F_APPLY_TO_ALPHA) != 0;
+
+		ImagePtr pNew = new Image(m_base->GetSizeX(), m_base->GetSizeY(),
+			m_base->GetLODCount(), m_base->GetFormat());
+
+		check(pNew->GetDataSize() == m_base->GetDataSize());
+
+		bool bValid = pNew->GetSizeX() > 0 && pNew->GetSizeY() > 0;
+		if (bValid && m_blended)
+		{
+			// TODO: This shouldn't happen, but check it to be safe
+			if (m_base->GetSize() != m_blended->GetSize())
+			{
+				MUTABLE_CPUPROFILER_SCOPE(ImageResize_EmergencyFix);
+				m_blended = ImageResize(m_blended.get(), m_base->GetSize());
+			}
+
+			check(pNew->GetDataSize() == m_base->GetDataSize());
+
+			if (m_blended->GetLODCount() < m_base->GetLODCount())
+			{
+				MUTABLE_CPUPROFILER_SCOPE(ImageLayer_MipBlendedEmergencyFix);
+
+				int32 levelCount = m_base->GetLODCount();
+				ImagePtr pDest = new Image(m_blended->GetSizeX(), m_blended->GetSizeY(),
+					levelCount,
+					m_blended->GetFormat());
+
+				SCRATCH_IMAGE_MIPMAP scratch;
+				FMipmapGenerationSettings settings{};
+
+				ImageMipmap_PrepareScratch(pDest.get(), m_blended.get(), levelCount, &scratch);
+				ImageMipmap(m_imageCompressionQuality, pDest.get(), m_blended.get(), levelCount,
+					&scratch, settings);
+
+				m_blended = pDest;
+			}
+
+			if (Args.mask)
+			{
+				if (m_base->GetSize() != m_mask->GetSize())
+				{
+					MUTABLE_CPUPROFILER_SCOPE(ImageResize_EmergencyFix);
+					m_mask = ImageResizeLinear(m_imageCompressionQuality, m_mask.get(), m_base->GetSize());
+				}
+
+				// emergy fix c36adf47-e40d-490f-b709-41142bafad78
+				if (m_mask->GetLODCount() < m_base->GetLODCount())
+				{
+					MUTABLE_CPUPROFILER_SCOPE(ImageLayer_EmergencyFix);
+
+					int32 levelCount = m_base->GetLODCount();
+					ImagePtr pDest = new Image(m_mask->GetSizeX(), m_mask->GetSizeY(),
+						levelCount,
+						m_mask->GetFormat());
+
+					SCRATCH_IMAGE_MIPMAP scratch;
+					FMipmapGenerationSettings settings{};
+
+					ImageMipmap_PrepareScratch(pDest.get(), m_mask.get(), levelCount, &scratch);
+					ImageMipmap(m_imageCompressionQuality, pDest.get(), m_mask.get(), levelCount,
+						&scratch, settings);
+
+					m_mask = pDest;
+				}
+
+				switch (EBlendType(Args.blendType))
+				{
+				case EBlendType::BT_NORMAL_COMBINE: ImageNormalCombine(pNew.get(), m_base.get(), m_mask.get(), m_blended.get()); break;
+				case EBlendType::BT_SOFTLIGHT: ImageSoftLight(pNew.get(), m_base.get(), m_mask.get(), m_blended.get(), applyToAlpha); break;
+				case EBlendType::BT_HARDLIGHT: ImageHardLight(pNew.get(), m_base.get(), m_mask.get(), m_blended.get(), applyToAlpha); break;
+				case EBlendType::BT_BURN: ImageBurn(pNew.get(), m_base.get(), m_mask.get(), m_blended.get(), applyToAlpha); break;
+				case EBlendType::BT_DODGE: ImageDodge(pNew.get(), m_base.get(), m_mask.get(), m_blended.get(), applyToAlpha); break;
+				case EBlendType::BT_SCREEN: ImageScreen(pNew.get(), m_base.get(), m_mask.get(), m_blended.get(), applyToAlpha); break;
+				case EBlendType::BT_OVERLAY: ImageOverlay(pNew.get(), m_base.get(), m_mask.get(), m_blended.get(), applyToAlpha); break;
+				case EBlendType::BT_ALPHA_OVERLAY: ImageAlphaOverlay(pNew.get(), m_base.get(), m_mask.get(), m_blended.get(), applyToAlpha); break;
+				case EBlendType::BT_MULTIPLY: ImageMultiply(pNew.get(), m_base.get(), m_mask.get(), m_blended.get(), applyToAlpha); break;
+				case EBlendType::BT_BLEND: ImageBlend(pNew.get(), m_base.get(), m_mask.get(), m_blended.get(), applyToAlpha); break;
+				default: check(false);
+				}
+
+			}
+			else
+			{
+				switch (EBlendType(Args.blendType))
+				{
+				case EBlendType::BT_NORMAL_COMBINE: ImageNormalCombine(pNew.get(), m_base.get(), m_blended.get()); break;
+				case EBlendType::BT_SOFTLIGHT: ImageSoftLight(pNew.get(), m_base.get(), m_blended.get(), applyToAlpha); break;
+				case EBlendType::BT_HARDLIGHT: ImageHardLight(pNew.get(), m_base.get(), m_blended.get(), applyToAlpha); break;
+				case EBlendType::BT_BURN: ImageBurn(pNew.get(), m_base.get(), m_blended.get(), applyToAlpha); break;
+				case EBlendType::BT_DODGE: ImageDodge(pNew.get(), m_base.get(), m_blended.get(), applyToAlpha); break;
+				case EBlendType::BT_SCREEN: ImageScreen(pNew.get(), m_base.get(), m_blended.get(), applyToAlpha); break;
+				case EBlendType::BT_OVERLAY: ImageOverlay(pNew.get(), m_base.get(), m_blended.get(), applyToAlpha); break;
+				case EBlendType::BT_ALPHA_OVERLAY: ImageAlphaOverlay(pNew.get(), m_base.get(), m_blended.get(), applyToAlpha); break;
+				case EBlendType::BT_MULTIPLY: ImageMultiply(pNew.get(), m_base.get(), m_blended.get(), applyToAlpha); break;
+				case EBlendType::BT_BLEND: ImageBlend(pNew.get(), m_base.get(), m_blended.get(), applyToAlpha); break;
+				default: check(false);
+				}
+			}
+		}
+
+		if (finalFormat != EImageFormat::IF_NONE)
+		{
+			pNew = ImagePixelFormat(m_imageCompressionQuality, pNew.get(), finalFormat);
+		}
+
+		// Free resources
+		m_base = nullptr;
+		m_blended = nullptr;
+		m_mask = nullptr;
+		Result = pNew;
+	}
+
+
+	//---------------------------------------------------------------------------------------------
+	void FImageLayerTask::Complete( CodeRunner* runner )
+	{
+		// This runs in the runner thread
+		runner->GetMemory().SetImage(Op, Result);
+	}
+
+
+	//---------------------------------------------------------------------------------------------
+	//---------------------------------------------------------------------------------------------
+	//---------------------------------------------------------------------------------------------
+	class FImageLayerColourTask : public CodeRunner::FIssuedTask
+	{
+	public:
+		FImageLayerColourTask(SCHEDULED_OP, FProgramCache&, const OP::ImageLayerColourArgs&, int imageCompressionQuality);
+
+		// FIssuedTask interface
+		void DoWork() override;
+		void Complete(CodeRunner* staged) override;
+
+	private:
+		int m_imageCompressionQuality = 0;
+		OP::ImageLayerColourArgs Args;
+		Ptr<const Image> m_base;
+		vec4<float> m_col;
+		Ptr<const Image> m_mask;
+		Ptr<const Image> Result;
+	};
+
+
+	//---------------------------------------------------------------------------------------------
+	FImageLayerColourTask::FImageLayerColourTask(SCHEDULED_OP InOp, FProgramCache& Memory, const OP::ImageLayerColourArgs& InArgs, int imageCompressionQuality)
+	{
+		Op = InOp;
+		m_imageCompressionQuality = imageCompressionQuality;
+		Args = InArgs;
+		m_base = Memory.GetImage({ Args.base, InOp.executionIndex, InOp.executionOptions });
+		m_col = Memory.GetColour({ Args.colour, InOp.executionIndex, 0 });
+		if (Args.mask)
+		{
+			m_mask = Memory.GetImage({ Args.mask, InOp.executionIndex, InOp.executionOptions });
+		}
+	}
+
+
+	//---------------------------------------------------------------------------------------------
+	void FImageLayerColourTask::DoWork()
+	{
+		MUTABLE_CPUPROFILER_SCOPE(FImageLayerColourTask);
+
+		// Warning: This runs in a worker thread.
+
+		ImagePtr pNew = new Image(m_base->GetSizeX(), m_base->GetSizeY(),
+			m_base->GetLODCount(), m_base->GetFormat());
+
+		if (Args.mask)
+		{
+			if (m_base->GetSize() != m_mask->GetSize())
+			{
+				MUTABLE_CPUPROFILER_SCOPE(ImageResize_EmergencyFix);
+				m_mask = ImageResizeLinear(m_imageCompressionQuality, m_mask.get(), m_base->GetSize());
+			}
+
+			// emergy fix c36adf47-e40d-490f-b709-41142bafad78
+			if (m_mask->GetLODCount() < m_base->GetLODCount())
+			{
+				//MUTABLE_CPUPROFILER_SCOPE(ImageLayerColour_EmergencyFix);
+
+				int32 levelCount = m_base->GetLODCount();
+				ImagePtr pDest = new Image(m_mask->GetSizeX(), m_mask->GetSizeY(),
+					levelCount,
+					m_mask->GetFormat());
+
+				SCRATCH_IMAGE_MIPMAP scratch;
+				FMipmapGenerationSettings settings{};
+
+				ImageMipmap_PrepareScratch(pDest.get(), m_mask.get(), levelCount, &scratch);
+				ImageMipmap(m_imageCompressionQuality, pDest.get(), m_mask.get(), levelCount,
+					&scratch, settings);
+				m_mask = pDest;
+			}
+
+			// \todo: precalculated tables for softlight
+			switch (EBlendType(Args.blendType))
+			{
+			case EBlendType::BT_NORMAL_COMBINE: ImageNormalCombine(pNew.get(), m_base.get(), m_mask.get(), m_col.xyz()); break;
+			case EBlendType::BT_SOFTLIGHT: ImageSoftLight(pNew.get(), m_base.get(), m_mask.get(), m_col.xyz(), nullptr); break;
+			case EBlendType::BT_HARDLIGHT: ImageHardLight(pNew.get(), m_base.get(), m_mask.get(), m_col.xyz()); break;
+			case EBlendType::BT_BURN: ImageBurn(pNew.get(), m_base.get(), m_mask.get(), m_col.xyz()); break;
+			case EBlendType::BT_DODGE: ImageDodge(pNew.get(), m_base.get(), m_mask.get(), m_col.xyz()); break;
+			case EBlendType::BT_SCREEN: ImageScreen(pNew.get(), m_base.get(), m_mask.get(), m_col.xyz()); break;
+			case EBlendType::BT_OVERLAY: ImageOverlay(pNew.get(), m_base.get(), m_mask.get(), m_col.xyz()); break;
+			case EBlendType::BT_ALPHA_OVERLAY: ImageAlphaOverlay(pNew.get(), m_base.get(), m_mask.get(), m_col.xyz()); break;
+			case EBlendType::BT_MULTIPLY: ImageMultiply(pNew.get(), m_base.get(), m_mask.get(), m_col.xyz()); break;
+			case EBlendType::BT_BLEND: ImageBlend(pNew.get(), m_base.get(), m_mask.get(), m_col.xyz()); break;
+			default: check(false);
+			}
+
+		}
+		else
+		{
+			switch (EBlendType(Args.blendType))
+			{
+			case EBlendType::BT_NORMAL_COMBINE: ImageNormalCombine(pNew.get(), m_base.get(), m_col.xyz()); break;
+			case EBlendType::BT_SOFTLIGHT: ImageSoftLight(pNew.get(), m_base.get(), m_col.xyz()); break;
+			case EBlendType::BT_HARDLIGHT: ImageHardLight(pNew.get(), m_base.get(), m_col.xyz()); break;
+			case EBlendType::BT_BURN: ImageBurn(pNew.get(), m_base.get(), m_col.xyz()); break;
+			case EBlendType::BT_DODGE: ImageDodge(pNew.get(), m_base.get(), m_col.xyz()); break;
+			case EBlendType::BT_SCREEN: ImageScreen(pNew.get(), m_base.get(), m_col.xyz()); break;
+			case EBlendType::BT_OVERLAY: ImageOverlay(pNew.get(), m_base.get(), m_col.xyz()); break;
+			case EBlendType::BT_ALPHA_OVERLAY: ImageAlphaOverlay(pNew.get(), m_base.get(), m_col.xyz()); break;
+			case EBlendType::BT_MULTIPLY: ImageMultiply(pNew.get(), m_base.get(), m_col.xyz()); break;
+			case EBlendType::BT_BLEND: FillPlainColourImage(pNew.get(), m_col); break;
+			default: check(false);
+			}
+		}
+
+		m_base = nullptr;
+		m_mask = nullptr;
+		Result = pNew;
+	}
+
+
+	//---------------------------------------------------------------------------------------------
+	void FImageLayerColourTask::Complete(CodeRunner* runner)
+	{
+		// This runs in the runner thread
+		runner->GetMemory().SetImage(Op, Result);
+	}
+
+
+	//---------------------------------------------------------------------------------------------
+	//---------------------------------------------------------------------------------------------
+	//---------------------------------------------------------------------------------------------
+	class FImagePixelFormatTask : public CodeRunner::FIssuedTask
+	{
+	public:
+		FImagePixelFormatTask(SCHEDULED_OP, FProgramCache&, const OP::ImagePixelFormatArgs&, int imageCompressionQuality);
+
+		// FIssuedTask interface
+		void DoWork() override;
+		void Complete(CodeRunner* staged) override;
+
+	private:
+		int m_imageCompressionQuality = 0;
+		OP::ImagePixelFormatArgs Args;
+		Ptr<const Image> m_base;
+		Ptr<const Image> Result;
+	};
+
+
+	//---------------------------------------------------------------------------------------------
+	FImagePixelFormatTask::FImagePixelFormatTask(SCHEDULED_OP InOp, FProgramCache& Memory, const OP::ImagePixelFormatArgs& InArgs, int imageCompressionQuality)
+	{
+		Op = InOp;
+		m_imageCompressionQuality = imageCompressionQuality;
+		Args = InArgs;
+		m_base = Memory.GetImage({ Args.source, InOp.executionIndex, InOp.executionOptions });
+	}
+
+
+	//---------------------------------------------------------------------------------------------
+	void FImagePixelFormatTask::DoWork()
+	{
+		MUTABLE_CPUPROFILER_SCOPE(FImagePixelFormatTask);
+
+		// Warning: This runs in a worker thread.
+
+		if (!m_base)
+			return;
+
+		EImageFormat targetFormat = Args.format;
+		if (Args.formatIfAlpha != EImageFormat::IF_NONE
+			&&
+			GetImageFormatData(m_base->GetFormat()).m_channels > 3)
+		{
+			// If the alpha channel is all white, we don't need to compress in a format with
+			// alpha support, as the default value is always white in shaders, and mutable
+			// code. This assumption exist in other commercial engines as well.
+			bool hasNonWhiteAlpha = !m_base->IsFullAlpha();
+			if (hasNonWhiteAlpha)
+			{
+				targetFormat = Args.formatIfAlpha;
+			}
+		}
+
+		if (targetFormat != EImageFormat::IF_NONE)
+		{
+			vec2<int> resultSize;
+			int resultLODCount = 0;
+
+			resultSize[0] = m_base->GetSizeX();
+			resultSize[1] = m_base->GetSizeY();
+			resultLODCount = m_base->GetLODCount();
+
+			Ptr<Image> TempResult = ImagePixelFormat(m_imageCompressionQuality, m_base.get(), targetFormat , -1);
+			TempResult->m_flags = m_base->m_flags;
+
+			Result = TempResult;
+		}
+		else
+		{
+			Result = m_base;
+		}
+
+		m_base = nullptr;
+	}
+
+
+	//---------------------------------------------------------------------------------------------
+	void FImagePixelFormatTask::Complete(CodeRunner* runner)
+	{
+		// This runs in the runner thread
+		runner->GetMemory().SetImage(Op, Result);
+	}
+
+
+	//---------------------------------------------------------------------------------------------
+	//---------------------------------------------------------------------------------------------
+	//---------------------------------------------------------------------------------------------
+	class FImageMipmapTask : public CodeRunner::FIssuedTask
+	{
+	public:
+		FImageMipmapTask(SCHEDULED_OP, FProgramCache&, const OP::ImageMipmapArgs&, int imageCompressionQuality);
+
+		// FIssuedTask interface
+		void DoWork() override;
+		void Complete(CodeRunner* staged) override;
+
+	private:
+		int m_imageCompressionQuality = 0;
+		OP::ImageMipmapArgs Args;
+		Ptr<const Image> m_base;
+		Ptr<const Image> Result;
+	};
+
+
+	//---------------------------------------------------------------------------------------------
+	FImageMipmapTask::FImageMipmapTask(SCHEDULED_OP InOp, FProgramCache& Memory, const OP::ImageMipmapArgs& InArgs, int imageCompressionQuality)
+	{
+		Op = InOp;
+		m_imageCompressionQuality = imageCompressionQuality;
+		Args = InArgs;
+		m_base = Memory.GetImage({ Args.source, InOp.executionIndex, InOp.executionOptions });
+	}
+
+
+	//---------------------------------------------------------------------------------------------
+	void FImageMipmapTask::DoWork()
+	{
+		MUTABLE_CPUPROFILER_SCOPE(FImageMipmapTask);
+
+		// Warning: This runs in a worker thread.
+
+		if (!m_base)
+		{
+			return;
+		}
+
+		int32 levelCount = Args.levels;
+		int32 maxLevelCount = Image::GetMipmapCount(m_base->GetSizeX(), m_base->GetSizeY());
+		if (levelCount == 0)
+		{
+			levelCount = maxLevelCount;
+		}
+		else if (levelCount > maxLevelCount)
+		{
+			// If code generation is smart enough, this should never happen.
+			// \todo But apparently it does, sometimes.
+			levelCount = maxLevelCount;
+		}
+
+		// At least keep the levels we already have.
+		int startLevel = m_base->GetLODCount();
+		levelCount = FMath::Max(startLevel, levelCount);
+
+		ImagePtr pDest = new Image(m_base->GetSizeX(), m_base->GetSizeY(), levelCount, m_base->GetFormat());
+		pDest->m_flags = m_base->m_flags;
+
+		SCRATCH_IMAGE_MIPMAP scratch;
+		FMipmapGenerationSettings settings{};
+
+		ImageMipmap_PrepareScratch(pDest.get(), m_base.get(), levelCount, &scratch);
+		ImageMipmap(m_imageCompressionQuality, pDest.get(), m_base.get(), levelCount, &scratch, settings);
+
+		m_base = nullptr;
+		Result = pDest;
+	}
+
+
+	//---------------------------------------------------------------------------------------------
+	void FImageMipmapTask::Complete(CodeRunner* runner)
+	{
+		// This runs in the runner thread
+		runner->GetMemory().SetImage(Op, Result);
+	}
+
+
+	//---------------------------------------------------------------------------------------------
+	//---------------------------------------------------------------------------------------------
+	//---------------------------------------------------------------------------------------------
+	class FImageSwizzleTask : public CodeRunner::FIssuedTask
+	{
+	public:
+		FImageSwizzleTask(SCHEDULED_OP, FProgramCache&, const OP::ImageSwizzleArgs&);
+
+		// FIssuedTask interface
+		void DoWork() override;
+		void Complete(CodeRunner* staged) override;
+
+	private:
+		OP::ImageSwizzleArgs Args;
+		Ptr<const Image> Sources[MUTABLE_OP_MAX_SWIZZLE_CHANNELS];
+		Ptr<const Image> Result;
+	};
+
+
+	//---------------------------------------------------------------------------------------------
+	FImageSwizzleTask::FImageSwizzleTask(SCHEDULED_OP InOp, FProgramCache& Memory, const OP::ImageSwizzleArgs& InArgs)
+	{
+		Op = InOp;
+		Args = InArgs;
+		for (int i = 0; i < MUTABLE_OP_MAX_SWIZZLE_CHANNELS; ++i)
+		{
+			if (Args.sources[i])
+			{
+				Sources[i] = Memory.GetImage({ Args.sources[i], InOp.executionIndex, InOp.executionOptions });
+			}
+		}
+	}
+
+
+	//---------------------------------------------------------------------------------------------
+	void FImageSwizzleTask::DoWork()
+	{
+		MUTABLE_CPUPROFILER_SCOPE(FImageSwizzleTask);
+
+		// Warning: This runs in a worker thread.
+
+		// Be defensive: ensure image sizes match.
+		for (int i = 1; i < MUTABLE_OP_MAX_SWIZZLE_CHANNELS; ++i)
+		{
+			if (Sources[i] && Sources[i]->GetSize() != Sources[0]->GetSize())
+			{
+				MUTABLE_CPUPROFILER_SCOPE(ImageResize_ForSwizzle);
+				Sources[i] = ImageResizeLinear(0, Sources[i].get(), Sources[0]->GetSize());
+			}
+		}
+
+
+		EImageFormat format = (EImageFormat)Args.format;
+		Result = ImageSwizzle(format, Sources, Args.sourceChannels);
+		for (int i = 0; i < MUTABLE_OP_MAX_SWIZZLE_CHANNELS; ++i)
+		{
+			Sources[i] = nullptr;
+		}
+	}
+
+
+	//---------------------------------------------------------------------------------------------
+	void FImageSwizzleTask::Complete(CodeRunner* runner)
+	{
+		// This runs in the runner thread
+		runner->GetMemory().SetImage(Op, Result);
+	}
+
+
+	//---------------------------------------------------------------------------------------------
+	//---------------------------------------------------------------------------------------------
+	//---------------------------------------------------------------------------------------------
+	class FImageComposeTask : public CodeRunner::FIssuedTask
+	{
+	public:
+		FImageComposeTask(SCHEDULED_OP, FProgramCache&, const OP::ImageComposeArgs&, int imageCompressionQuality, const Ptr<const Layout>& );
+
+		// FIssuedTask interface
+		void DoWork() override;
+		void Complete(CodeRunner* staged) override;
+
+	private:
+		int m_imageCompressionQuality = 0;
+		OP::ImageComposeArgs Args;
+		Ptr<const Image> m_base;
+		Ptr<const Layout> m_layout;
+		Ptr<const Image> m_block;
+		Ptr<const Image> m_mask;
+		Ptr<const Image> Result;
+	};
+
+
+	//---------------------------------------------------------------------------------------------
+	FImageComposeTask::FImageComposeTask(SCHEDULED_OP InOp, FProgramCache& Memory, const OP::ImageComposeArgs& InArgs, int imageCompressionQuality, const Ptr<const Layout>& InLayout)
+	{
+		Op = InOp;
+		m_imageCompressionQuality = imageCompressionQuality;
+		Args = InArgs;
+		m_layout = InLayout;
+
+		m_base = Memory.GetImage({ Args.base, Op.executionIndex, InOp.executionOptions });
+		
+		int relBlockIndex = m_layout->FindBlock(Args.blockIndex);
+
+		if (relBlockIndex >= 0)
+		{
+			m_block = Memory.GetImage({ Args.blockImage, Op.executionIndex, InOp.executionOptions });
+			if (Args.mask)
+			{
+				m_mask = Memory.GetImage({ Args.mask, Op.executionIndex, InOp.executionOptions });
+			}
+		}
+	}
+
+
+	//---------------------------------------------------------------------------------------------
+	void FImageComposeTask::DoWork()
+	{
+		MUTABLE_CPUPROFILER_SCOPE(FImageComposeTask);
+
+		// Warning: This runs in a worker thread.
+		int relBlockIndex = m_layout->FindBlock(Args.blockIndex);
+
+		if (relBlockIndex < 0)
+		{
+			Result = m_base;
+		}
+		else
+		{
+			// Try to avoid the clone here
+			ImagePtr pComposed = CloneOrTakeOver<>(m_base.get());
+			pComposed->m_flags = 0;
+
+			box< vec2<int> > rectInblocks;
+			m_layout->GetBlock
+			(
+				relBlockIndex,
+				&rectInblocks.min[0], &rectInblocks.min[1],
+				&rectInblocks.size[0], &rectInblocks.size[1]
+			);
+
+			// Convert the rect from blocks to pixels
+			FIntPoint grid = m_layout->GetGridSize();
+
+			int blockSizeX = m_base->GetSizeX() / grid[0];
+			int blockSizeY = m_base->GetSizeY() / grid[1];
+			box< vec2<int> > rect = rectInblocks;
+			rect.min[0] *= blockSizeX;
+			rect.min[1] *= blockSizeY;
+			rect.size[0] *= blockSizeX;
+			rect.size[1] *= blockSizeY;
+
+			//
+			if (m_block && rect.size[0] && rect.size[1] && m_block->GetSizeX() && m_block->GetSizeY())
+			{
+				bool useMask = Args.mask != 0;
+
+				if (!useMask)
+				{
+					MUTABLE_CPUPROFILER_SCOPE(ImageComposeWithoutMask);
+
+					// Resize image if it doesn't fit in the new block size
+					if (m_block->GetSizeX() != rect.size[0] ||
+						m_block->GetSizeY() != rect.size[1])
+					{
+						// This now happens more often since the generation of specific mips on request. For this reason
+						// this warning is usually acceptable.
+						//UE_LOG(LogMutableCore, Log, TEXT("Required image resize for ImageCompose: performance warning."));
+						FImageSize blockSize((uint16_t)rect.size[0], (uint16_t)rect.size[1]);
+						m_block = ImageResizeLinear(m_imageCompressionQuality, m_block.get(), blockSize);
+					}
+
+					// Change the block image format if it doesn't match the composed image
+					// This is usually enforced at object compilation time.
+					if (pComposed->GetFormat() != m_block->GetFormat())
+					{
+						MUTABLE_CPUPROFILER_SCOPE(ImageComposeReformat);
+
+						EImageFormat format = GetMostGenericFormat(pComposed->GetFormat(), m_block->GetFormat());
+						if (pComposed->GetFormat() != format)
+						{
+							pComposed = ImagePixelFormat(m_imageCompressionQuality, pComposed.get(), format);
+						}
+						if (m_block->GetFormat() != format)
+						{
+							m_block = ImagePixelFormat(m_imageCompressionQuality, m_block.get(), format);
+						}
+					}
+
+					// Compose without a mask
+					ImageCompose(pComposed.get(), m_block.get(), rect);
+				}
+				else
+				{
+					MUTABLE_CPUPROFILER_SCOPE(ImageComposeWithMask);
+
+					// Compose with a mask
+					ImageBlendOnBaseNoAlpha(pComposed.get(), m_mask.get(), m_block.get(), rect);
+				}
+
+			}
+
+			Result = pComposed;
+		}
+
+		m_base = nullptr;
+		m_block = nullptr;
+		m_layout = nullptr;
+		m_mask = nullptr;
+	}
+
+
+	//---------------------------------------------------------------------------------------------
+	void FImageComposeTask::Complete(CodeRunner* runner)
+	{
+		// This runs in the runner thread
+		runner->GetMemory().SetImage(Op, Result);
+	}
+
+
+	//---------------------------------------------------------------------------------------------
+	//---------------------------------------------------------------------------------------------
+	//---------------------------------------------------------------------------------------------
+	bool CodeRunner::FLoadMeshRomTask::Prepare(CodeRunner* runner, const Model* InModel, bool& bOutFailed )
+	{
+		MUTABLE_CPUPROFILER_SCOPE(FLoadMeshRomTask_Prepare);
+
+		// This runs in th CodeRunner thread
+		bOutFailed = false;
+
+		PROGRAM& program = InModel->GetPrivate()->m_program;
+
+		check(RomIndex < program.m_roms.Num());
+		bool bRomIsLoaded = program.IsRomLoaded(RomIndex);
+		bool bRomHasPendingOps = runner->m_romPendingOps[RomIndex]!=0;
+		if (!bRomIsLoaded && !bRomHasPendingOps)
+		{
+			check(runner->m_pSystem->m_pStreamInterface);
+
+			// Free roms if necessary
+			{
+				MUTABLE_CPUPROFILER_SCOPE(FreeingRoms);
+				runner->m_pSystem->m_modelCache.UpdateForLoad(RomIndex, InModel,
+					[&](const Model* pModel, int romInd)
+					{
+						return (pModel == InModel)
+							&&
+							runner->m_romPendingOps[romInd] > 0;
+					});
+			}
+
+			uint32_t RomSize = program.m_roms[RomIndex].Size;
+			check(RomSize>0);
+
+			CodeRunner::ROM_LOAD_OP* op = nullptr;
+			for (auto& o : runner->m_romLoadOps)
+			{
+				if (o.m_romIndex < 0)
+				{
+					op = &o;
+				}
+			}
+
+			if (!op)
+			{
+				runner->m_romLoadOps.push_back(CodeRunner::ROM_LOAD_OP());
+				op = &runner->m_romLoadOps.back();
+			}
+
+			op->m_romIndex = RomIndex;
+			op->ConstantType = DT_MESH;
+
+			if (op->m_streamBuffer.size() < RomSize)
+			{
+				// This could happen in 32-bit platforms
+				check(RomSize < std::numeric_limits<size_t>::max());
+				op->m_streamBuffer.resize((size_t)RomSize);
+			}
+
+			uint32 RomId = program.m_roms[RomIndex].Id;
+			op->m_streamID = runner->m_pSystem->m_pStreamInterface->BeginReadBlock(InModel, RomId, op->m_streamBuffer.data(), RomSize);
+			if (op->m_streamID < 0)
+			{
+				bOutFailed = true;
+				return false;
+			}
+		}
+		++runner->m_romPendingOps[RomIndex];
+
+		//UE_LOG(LogMutableCore, Log, TEXT("FLoadMeshRomTask::Prepare romindex %d pending %d."), RomAt.m_rom, runner->m_romPendingOps[RomAt.m_rom]);
+
+		// No worker thread work
+		return false;
+	}
+
+
+	//---------------------------------------------------------------------------------------------
+	void CodeRunner::FLoadMeshRomTask::Complete(CodeRunner* runner)
+	{
+		MUTABLE_CPUPROFILER_SCOPE(FLoadMeshRomTask_Complete);
+
+		// This runs in the runner thread
+
+		// Process the constant op normally, now that the rom is loaded.
+		runner->RunCode(Op, runner->m_pParams, runner->m_pModel, runner->m_lodMask);
+
+		runner->m_pSystem->m_modelCache.MarkRomUsed(RomIndex, runner->m_pModel);
+		--runner->m_romPendingOps[RomIndex];
+	}
+
+
+	//---------------------------------------------------------------------------------------------
+	bool CodeRunner::FLoadMeshRomTask::IsComplete(CodeRunner* runner)
+	{ 
+		PROGRAM& program = runner->m_pModel->GetPrivate()->m_program;
+		return program.IsRomLoaded(RomIndex);
+	}
+
+
+	//---------------------------------------------------------------------------------------------
+	//---------------------------------------------------------------------------------------------
+	//---------------------------------------------------------------------------------------------
+	bool CodeRunner::FLoadImageRomsTask::Prepare(CodeRunner* runner, const Model* InModel, bool& bOutFailed )
+	{
+		MUTABLE_CPUPROFILER_SCOPE(FLoadImageRomsTask_Prepare);
+
+		// This runs in th CodeRunner thread
+		bOutFailed = false;
+
+		PROGRAM& program = InModel->GetPrivate()->m_program;
+
+		for (int32 i = 0; i<LODIndexCount; ++i )
+		{
+			int32 CurrentIndexIndex = LODIndexIndex + i;
+			int32 CurrentIndex = program.m_constantImageLODIndices[CurrentIndexIndex];
+
+			if (program.m_constantImageLODs[CurrentIndex].Key<0)
+			{
+				// This data is always resident.
+				continue;
+			}
+
+			int32 RomIndex = program.m_constantImageLODs[CurrentIndex].Key;
+			check(RomIndex < program.m_roms.Num());
+
+			bool bRomIsLoaded = program.IsRomLoaded(RomIndex);
+			bool bRomHasAlreadyBeenRequested = runner->m_romPendingOps[RomIndex] != 0;
+			++runner->m_romPendingOps[RomIndex];
+
+			if (DebugRom && (DebugRomAll || RomIndex == DebugRomIndex))
+				UE_LOG(LogMutableCore, Log, TEXT("Preparing rom %d, now peding ops is %d."), RomIndex, runner->m_romPendingOps[RomIndex]);
+
+			if (bRomIsLoaded || bRomHasAlreadyBeenRequested)
+			{
+				continue;
+			}
+
+			check(runner->m_pSystem->m_pStreamInterface);
+
+			// Free roms if necessary
+			{
+				MUTABLE_CPUPROFILER_SCOPE(FreeingRoms);
+				runner->m_pSystem->m_modelCache.UpdateForLoad(RomIndex, InModel,
+					[&](const Model* pModel, int romInd)
+					{
+						return (pModel == InModel)
+							&&
+							runner->m_romPendingOps[romInd] > 0;
+					});
+			}
+
+			uint32_t RomSize = program.m_roms[RomIndex].Size;
+			check(RomSize > 0);
+
+			CodeRunner::ROM_LOAD_OP* op = nullptr;
+			for (auto& o : runner->m_romLoadOps)
+			{
+				if (o.m_romIndex < 0)
+				{
+					op = &o;
+				}
+			}
+
+			if (!op)
+			{
+				runner->m_romLoadOps.push_back(CodeRunner::ROM_LOAD_OP());
+				op = &runner->m_romLoadOps.back();
+			}
+
+			op->m_romIndex = RomIndex;
+			op->ConstantType = DT_IMAGE;
+
+			if (op->m_streamBuffer.size() < RomSize)
+			{
+				// This could happen in 32-bit platforms
+				check(RomSize < std::numeric_limits<size_t>::max());
+				op->m_streamBuffer.resize((size_t)RomSize);
+			}
+
+			uint32 RomId = program.m_roms[RomIndex].Id;
+			op->m_streamID = runner->m_pSystem->m_pStreamInterface->BeginReadBlock(InModel, RomId, op->m_streamBuffer.data(), RomSize);
+			if (op->m_streamID < 0)
+			{
+				bOutFailed = true;
+				return false;
+			}
+		}
+
+		//UE_LOG(LogMutableCore, Log, TEXT("FLoadImageRomsTask::Prepare romindex %d pending %d."), RomAt.m_rom, runner->m_romPendingOps[RomAt.m_rom]);
+
+		// No worker thread work
+		return false;
+	}
+
+
+	//---------------------------------------------------------------------------------------------
+	void CodeRunner::FLoadImageRomsTask::Complete(CodeRunner* runner)
+	{
+		MUTABLE_CPUPROFILER_SCOPE(FLoadImageRomsTask_Complete);
+
+		// This runs in the runner thread
+
+		// Process the constant op normally, now that the rom is loaded.
+		runner->RunCode(Op, runner->m_pParams, runner->m_pModel, runner->m_lodMask);
+
+		PROGRAM& program = runner->m_pModel->GetPrivate()->m_program;
+		for (int32 i = 0; i < LODIndexCount; ++i)
+		{
+			int32 CurrentIndexIndex = LODIndexIndex + i;
+			int32 CurrentIndex = program.m_constantImageLODIndices[CurrentIndexIndex];
+
+			if (program.m_constantImageLODs[CurrentIndex].Key < 0)
+			{
+				// This data is always resident.
+				continue;
+			}
+
+			int32 RomIndex = program.m_constantImageLODs[CurrentIndex].Key;
+			check(RomIndex < program.m_roms.Num());
+
+			runner->m_pSystem->m_modelCache.MarkRomUsed(RomIndex, runner->m_pModel);
+			--runner->m_romPendingOps[RomIndex];
+
+			if (DebugRom && (DebugRomAll || RomIndex==DebugRomIndex))
+				UE_LOG(LogMutableCore, Log, TEXT("FLoadImageRomsTask::Complete rom %d, now peding ops is %d."), RomIndex, runner->m_romPendingOps[RomIndex]);
+		}
+	}
+
+
+	//---------------------------------------------------------------------------------------------
+	bool CodeRunner::FLoadImageRomsTask::IsComplete(CodeRunner* runner)
+	{
+		PROGRAM& program = runner->m_pModel->GetPrivate()->m_program;
+		for (int32 i = 0; i < LODIndexCount; ++i)
+		{
+			int32 CurrentIndexIndex = LODIndexIndex + i;
+			int32 CurrentIndex = program.m_constantImageLODIndices[CurrentIndexIndex];
+
+			if (program.m_constantImageLODs[CurrentIndex].Key < 0)
+			{
+				// This data is always resident.
+				continue;
+			}
+
+			int32 RomIndex = program.m_constantImageLODs[CurrentIndex].Key;
+			check(RomIndex < program.m_roms.Num());
+			if (!program.IsRomLoaded(RomIndex))
+			{
+				return false;
+			}
+		}
+
+		return true;
+	}
+
+
+	//---------------------------------------------------------------------------------------------
+	//---------------------------------------------------------------------------------------------
+	//---------------------------------------------------------------------------------------------
+	TSharedPtr<CodeRunner::FIssuedTask> CodeRunner::IssueOp(SCHEDULED_OP item)
+	{
+		TSharedPtr<FIssuedTask> Issued;
+
+		PROGRAM& program = m_pModel->GetPrivate()->m_program;
+
+		auto type = program.GetOpType(item.at);
+
+		switch (type)
+		{
+		case OP_TYPE::ME_CONSTANT:
+		{
+			OP::MeshConstantArgs args = program.GetOpArgs<OP::MeshConstantArgs>(item.at);
+			int32 RomIndex = program.m_constantMeshes[args.value].Key;
+			if (RomIndex >= 0 && !program.m_constantMeshes[args.value].Value )
+			{
+				Issued = MakeShared<FLoadMeshRomTask>(item, RomIndex);
+			}
+			else
+			{
+				// If already available, the rest of the constant code will run right away.
+			}
+			break;
+		}
+
+		case OP_TYPE::IM_CONSTANT:
+		{
+			OP::ResourceConstantArgs args = program.GetOpArgs<OP::ResourceConstantArgs>(item.at);
+			int32 MipsToSkip = item.executionOptions;
+			int32 ImageIndex = args.value;
+			int32 ReallySkip = FMath::Min(MipsToSkip, program.m_constantImages[ImageIndex].LODCount - 1);
+			int32 LODIndexIndex = program.m_constantImages[ImageIndex].FirstIndex + ReallySkip;
+			int32 LODIndexCount = program.m_constantImages[ImageIndex].LODCount - ReallySkip;
+			check(LODIndexCount > 0);
+
+			bool bAnyMissing = false;
+			for (int32 i=0; i<LODIndexCount; ++i)
+			{
+				uint32 LODIndex = program.m_constantImageLODIndices[LODIndexIndex+i];
+				if ( !program.m_constantImageLODs[LODIndex].Value )
+				{
+					bAnyMissing = true;
+					break;
+				}
+			}
+
+			if (bAnyMissing)
+			{
+				Issued = MakeShared<FLoadImageRomsTask>(item, LODIndexIndex, LODIndexCount);
+
+				if (DebugRom && (DebugRomAll || ImageIndex == DebugImageIndex))
+					UE_LOG(LogMutableCore, Log, TEXT("Issuing image %d skipping %d ."), ImageIndex, ReallySkip);
+			}
+			else
+			{
+				// If already available, the rest of the constant code will run right away.
+				if (DebugRom && (DebugRomAll || ImageIndex == DebugImageIndex))
+					UE_LOG(LogMutableCore, Log, TEXT("Image %d skipping %d is already loaded."), ImageIndex, ReallySkip);
+			}
+			break;
+		}		
+
+		case OP_TYPE::IM_PIXELFORMAT:
+		{
+			if (item.stage == 1)
+			{
+				OP::ImagePixelFormatArgs args = program.GetOpArgs<OP::ImagePixelFormatArgs>(item.at);
+				Issued = MakeShared<FImagePixelFormatTask>(item, GetMemory(), args, m_pSettings->GetPrivate()->m_imageCompressionQuality);
+			}
+			break;
+		}
+
+		case OP_TYPE::IM_LAYERCOLOUR:
+		{
+			if (item.stage == 1)
+			{
+				OP::ImageLayerColourArgs args = program.GetOpArgs<OP::ImageLayerColourArgs>(item.at);
+				Issued = MakeShared<FImageLayerColourTask>(item, GetMemory(), args, m_pSettings->GetPrivate()->m_imageCompressionQuality);
+			}
+			break;
+		}
+
+		case OP_TYPE::IM_LAYER:
+		{
+			if (item.stage == 1)
+			{
+				OP::ImageLayerArgs args = program.GetOpArgs<OP::ImageLayerArgs>(item.at);
+				Issued = MakeShared<FImageLayerTask>(item, GetMemory(), args, m_pSettings->GetPrivate()->m_imageCompressionQuality);
+			}
+			break;
+		}
+
+		case OP_TYPE::IM_MIPMAP:
+		{
+			if (item.stage == 1)
+			{
+				OP::ImageMipmapArgs args = program.GetOpArgs<OP::ImageMipmapArgs>(item.at);
+				Issued = MakeShared<FImageMipmapTask>(item, GetMemory(), args, m_pSettings->GetPrivate()->m_imageCompressionQuality);
+			}
+			break;
+		}
+
+		case OP_TYPE::IM_SWIZZLE:
+		{
+			if (item.stage == 1)
+			{
+				OP::ImageSwizzleArgs args = program.GetOpArgs<OP::ImageSwizzleArgs>(item.at);
+				Issued = MakeShared<FImageSwizzleTask>(item, GetMemory(), args);
+			}
+			break;
+		}
+
+		case OP_TYPE::IM_COMPOSE:
+		{
+			if (item.stage == 2)
+			{
+				OP::ImageComposeArgs args = program.GetOpArgs<OP::ImageComposeArgs>(item.at);
+				Ptr<const Layout> Layout = m_heapData[(size_t)item.customState].layout;
+				Issued = MakeShared<FImageComposeTask>(item, GetMemory(), args, m_pSettings->GetPrivate()->m_imageCompressionQuality, Layout);
+			}
+			break;
+		}
+
+		default:
+			break;
+		}
+
+		return Issued;
+	}
+
+
+} // namespace mu
diff --git a/Engine/Plugins/Experimental/Mutable/Source/MutableRuntime/Private/CodeVisitor.cpp b/Engine/Plugins/Experimental/Mutable/Source/MutableRuntime/Private/CodeVisitor.cpp
new file mode 100644
--- /dev/null
+++ b/Engine/Plugins/Experimental/Mutable/Source/MutableRuntime/Private/CodeVisitor.cpp
@@ -0,0 +1,90 @@
+// Copyright Epic Games, Inc. All Rights Reserved.
+
+
+#include "CodeVisitor.h"
+
+#include "MutableTrace.h"
+
+using namespace mu;
+
+
+//-------------------------------------------------------------------------------------------------
+//-------------------------------------------------------------------------------------------------
+//-------------------------------------------------------------------------------------------------
+void SubtreeParametersVisitor::Run( OP::ADDRESS root, PROGRAM& program )
+{
+    // Cached?
+    auto it = m_resultCache.find(root);
+    if (it!=m_resultCache.end())
+    {
+        m_params = it->second;
+        return;
+    }
+
+    // Not cached
+    {
+        MUTABLE_CPUPROFILER_SCOPE(SubtreeParametersVisitor);
+
+        m_visited.resize( program.m_opAddress.size() );
+        if (program.m_opAddress.size())
+        {
+			FMemory::Memzero( &m_visited[0], m_visited.size() );
+        }
+
+        m_currentParams.resize( program.m_parameters.size() );
+        if (m_currentParams.size())
+        {
+			FMemory::Memzero( &m_currentParams[0], sizeof(int)*m_currentParams.size() );
+        }
+
+        m_pending.reserve( program.m_opAddress.size()/4 );
+        m_pending.push_back(root);
+
+        while ( m_pending.size() )
+        {
+            OP::ADDRESS at = m_pending.back();
+            m_pending.pop_back();
+
+            if (!m_visited[at])
+            {
+                m_visited[at] = true;
+
+                switch ( program.GetOpType(at) )
+                {
+                case OP_TYPE::NU_PARAMETER:
+                case OP_TYPE::SC_PARAMETER:
+                case OP_TYPE::BO_PARAMETER:
+                case OP_TYPE::CO_PARAMETER:
+                case OP_TYPE::PR_PARAMETER:
+                case OP_TYPE::IM_PARAMETER:
+                    m_currentParams[ program.GetOpArgs<OP::ParameterArgs>(at).variable ]++;
+                    break;
+
+                default:
+                    break;
+                }
+
+                ForEachReference( program, at, [&](OP::ADDRESS ref)
+                {
+                    if (ref)
+                    {
+                        m_pending.push_back( ref );
+                    }
+                });
+            }
+        }
+
+
+        // Build result
+        m_params.clear();
+        for (size_t i=0; i<m_currentParams.size(); ++i )
+        {
+            if ( m_currentParams[i] )
+            {
+                m_params.push_back( int(i) );
+            }
+        }
+
+        m_resultCache[root] = m_params;
+    }
+}
diff --git a/Engine/Plugins/Experimental/Mutable/Source/MutableRuntime/Private/CodeVisitor.h b/Engine/Plugins/Experimental/Mutable/Source/MutableRuntime/Private/CodeVisitor.h
new file mode 100644
--- /dev/null
+++ b/Engine/Plugins/Experimental/Mutable/Source/MutableRuntime/Private/CodeVisitor.h
@@ -0,0 +1,725 @@
+// Copyright Epic Games, Inc. All Rights Reserved.
+
+#pragma once
+
+#include "Operations.h"
+#include "ModelPrivate.h"
+#include "SystemPrivate.h"
+
+
+#include <algorithm>
+
+
+namespace mu
+{
+
+    //---------------------------------------------------------------------------------------------
+    //! Decide what operations are an "add resource" since they are handled differently sometimes.
+    //---------------------------------------------------------------------------------------------
+    inline bool VisitorIsAddResource(const OP_TYPE& type)
+    {
+        return  type == OP_TYPE::IN_ADDIMAGE ||
+                type == OP_TYPE::IN_ADDMESH;
+    }
+
+
+
+    //---------------------------------------------------------------------------------------------
+    //! Code visitor that:
+    //! - is top-down
+    //! - cannot change the visited instructions.
+    //! - will not visit twice the same instruction with the same state.
+    //! - Its iterative
+    //---------------------------------------------------------------------------------------------
+    template<typename STATE=int>
+    class UniqueConstCodeVisitorIterative : public Base
+    {
+    public:
+
+        UniqueConstCodeVisitorIterative( bool skipResources=false )
+        {
+            // Default state
+            m_states.push_back(STATE());
+            m_currentState = 0;
+            m_skipResources = skipResources;
+        }
+
+        //! Ensure virtual destruction
+        virtual ~UniqueConstCodeVisitorIterative() {}
+
+    protected:
+
+        //!
+        void SetDefaultState(const STATE& s)
+        {
+            m_states[0] = s;
+        }
+
+        //!
+        const STATE& GetDefaultState() const
+        {
+            return m_states[0];
+        }
+
+        //! Use this from visit to access the state at the time of processing the current
+        //! instruction.
+        STATE GetCurrentState() const
+        {
+            return m_states[m_currentState];
+        }
+
+        //! For manual recursion that changes the state for a specific path.
+        void RecurseWithState(OP::ADDRESS at, const STATE& newState)
+        {
+            auto it = std::find(m_states.begin(),m_states.end(),newState);
+            if (it==m_states.end())
+            {
+                m_states.push_back(newState);
+                it = m_states.end() - 1;
+            }
+            int stateIndex = (int)(it - m_states.begin());
+
+            m_pending.push_back( PENDING(at,stateIndex) );
+        }
+
+        //! For manual recursion that doesn't change the state for a specific path.
+        void RecurseWithCurrentState(OP::ADDRESS at)
+        {
+            m_pending.push_back( PENDING(at,m_currentState) );
+        }
+
+        //! Can be called from visit to set the state to visit all children ops
+        void SetCurrentState(const STATE& newState)
+        {
+            auto it = std::find(m_states.begin(),m_states.end(),newState);
+            if (it==m_states.end())
+            {
+                m_states.push_back(newState);
+                it = m_states.end() - 1;
+            }
+            m_currentState = (int)(it - m_states.begin());
+        }
+
+
+        void Traverse( OP::ADDRESS root, PROGRAM& program, bool visitDecorators = true )
+        {
+            m_pending.reserve( program.m_opAddress.size() );
+
+            // Visit the given root
+            m_pending.push_back( PENDING(root,0) );
+            Recurse( program );
+
+            if (visitDecorators)
+            {
+                // Fix the code used in the parameter descriptions
+                for ( std::size_t p=0; p<program.m_parameters.size(); ++p )
+                {
+                    for ( std::size_t d=0; d<program.m_parameters[p].m_descImages.size(); ++d )
+                    {
+                        m_pending.push_back(PENDING(program.m_parameters[p].m_descImages[d],0) );
+                        Recurse( program );
+                    }
+                }
+            }
+        }
+
+        void FullTraverse( PROGRAM& program, bool visitDecorators = true )
+        {
+            // Visit all the state roots
+            for ( std::size_t p=0; p<program.m_states.size(); ++p )
+            {
+                m_pending.push_back(PENDING(program.m_states[p].m_root,0) );
+                Recurse( program );
+            }
+
+            if (visitDecorators)
+            {
+                // Fix the code used in the parameter descriptions
+                for ( std::size_t p=0; p<program.m_parameters.size(); ++p )
+                {
+                    for ( std::size_t d=0; d<program.m_parameters[p].m_descImages.size(); ++d )
+                    {
+                        m_pending.push_back(PENDING(program.m_parameters[p].m_descImages[d],0) );
+                        Recurse( program );
+                    }
+                }
+            }
+        }
+
+
+    private:
+
+        //! Do the actual work by overriding this in the derived classes.
+        //! Return true if the traverse has to continue with the children of "at"
+        virtual bool Visit( OP::ADDRESS at, PROGRAM& program ) = 0;
+
+        //! Operations to be processed
+        struct PENDING
+        {
+            PENDING()
+            {
+                at = 0;
+                stateIndex = 0;
+            }
+            PENDING(OP::ADDRESS _at, int _stateIndex)
+            {
+                at = _at;
+                stateIndex = _stateIndex;
+            }
+
+            OP::ADDRESS at;
+            int stateIndex;
+        };
+        vector<PENDING> m_pending;
+
+        //! States found so far
+        vector<STATE> m_states;
+
+        //! Index of the current state, from the m_states vector.
+        int m_currentState;
+
+        //! If true, operations adding resources (meshes or images) will only
+        //! have the base operation recursed, but not the resources.
+        bool m_skipResources;
+
+        //! Array of states visited for each operation.
+        //! Empty vector means operation not visited at all.
+        vector<vector<int>> m_visited;
+
+        //! Process all the pending operations and visit all children if necessary
+        void Recurse( PROGRAM& program )
+        {
+            m_visited = vector<vector<int>>(program.m_opAddress.size());
+
+            while ( m_pending.size() )
+            {
+                OP::ADDRESS at = m_pending.back().at;
+                m_currentState = m_pending.back().stateIndex;
+                m_pending.pop_back();
+
+                bool recurse = false;
+
+                bool visitedInThisState = std::find(m_visited[at].begin(), m_visited[at].end(), m_currentState)
+                    !=
+                    m_visited[at].end();
+                if (!visitedInThisState)
+                {
+                    m_visited[at].push_back(m_currentState);
+
+                    // Visit may change current state
+                    recurse = Visit( at, program );
+                }
+
+                if (recurse)
+                {
+                    ForEachReference( program, at, [&](OP::ADDRESS ref)
+                    {
+                        if (ref)
+                        {
+                            m_pending.push_back( PENDING(ref,m_currentState) );
+                        }
+                    });
+                }
+            }
+        }
+
+    };
+
+
+    //---------------------------------------------------------------------------------------------
+    //! Code visitor that:
+    //! - is top-down
+    //! - cannot change the instructions.
+    //! - will repeat visits to the instructions that have multiple references.
+    //! - Its iterative
+    //---------------------------------------------------------------------------------------------
+    template<class STATE=int>
+    class RepeatConstCodeVisitorIterative : public Base
+    {
+    public:
+
+        RepeatConstCodeVisitorIterative( bool skipResources=false )
+        {
+            // Default state
+            m_states.push_back(STATE());
+            m_currentState = 0;
+            m_skipResources = skipResources;
+        }
+
+        //! Ensure virtual destruction
+        virtual ~RepeatConstCodeVisitorIterative() {}
+
+    protected:
+
+        //!
+        void SetDefaultState(const STATE& s)
+        {
+            m_states[0] = s;
+        }
+
+        //!
+        const STATE& GetDefaultState() const
+        {
+            return m_states[0];
+        }
+
+        //! Use this from visit to access the state at the time of processing the current
+        //! instruction.
+        const STATE& GetCurrentState() const
+        {
+            return m_states[m_currentState];
+        }
+
+        //! For manual recursion that changes the state for a specific path.
+        void RecurseWithState(OP::ADDRESS at, const STATE& newState)
+        {
+            auto it = std::find(m_states.begin(),m_states.end(),newState);
+            if (it==m_states.end())
+            {
+                m_states.push_back(newState);
+                it = m_states.end() - 1;
+            }
+            int stateIndex = (int)(it - m_states.begin());
+
+            //check(at<1000000);
+            m_pending.push_back( PENDING(at,stateIndex) );
+        }
+
+        //! For manual recursion that doesn't change the state for a specific path.
+        void RecurseWithCurrentState(OP::ADDRESS at)
+        {
+            //check(at<1000000);
+            m_pending.push_back( PENDING(at,m_currentState) );
+        }
+
+        //! Can be called from visit to set the state to visit all children ops
+        void SetCurrentState(const STATE& newState)
+        {
+            auto it = std::find(m_states.begin(),m_states.end(),newState);
+            if (it==m_states.end())
+            {
+                m_states.push_back(newState);
+                it = m_states.end() - 1;
+            }
+            m_currentState = (int)(it - m_states.begin());
+        }
+
+
+        void Traverse( OP::ADDRESS root, PROGRAM& program, bool visitDecorators = true )
+        {
+            m_pending.reserve( program.m_opAddress.size() );
+
+            // Visit the given root
+            m_pending.push_back( PENDING(root,0) );
+            Recurse( program );
+
+            if (visitDecorators)
+            {
+                // Fix the code used in the parameter descriptions
+                for ( std::size_t p=0; p<program.m_parameters.size(); ++p )
+                {
+                    for ( std::size_t d=0; d<program.m_parameters[p].m_descImages.size(); ++d )
+                    {
+                        m_pending.push_back(PENDING(program.m_parameters[p].m_descImages[d],0) );
+                        Recurse( program );
+                    }
+                }
+            }
+        }
+
+        void FullTraverse( PROGRAM& program, bool visitDecorators = true )
+        {
+            // Visit all the state roots
+            for ( std::size_t p=0; p<program.m_states.size(); ++p )
+            {
+                m_pending.push_back(PENDING(program.m_states[p].m_root,0) );
+                Recurse( program );
+            }
+
+            if (visitDecorators)
+            {
+                // Fix the code used in the parameter descriptions
+                for ( std::size_t p=0; p<program.m_parameters.size(); ++p )
+                {
+                    for ( std::size_t d=0; d<program.m_parameters[p].m_descImages.size(); ++d )
+                    {
+                        m_pending.push_back(PENDING(program.m_parameters[p].m_descImages[d],0) );
+                        Recurse( program );
+                    }
+                }
+            }
+        }
+
+
+    private:
+
+        //! Do the actual work by overriding this in the derived classes.
+        //! Return true if the traverse has to continue with the children of "at"
+        virtual bool Visit( OP::ADDRESS at, PROGRAM& program ) = 0;
+
+        //! Operations to be processed
+        struct PENDING
+        {
+            PENDING()
+            {
+                at = 0;
+                stateIndex = 0;
+            }
+            PENDING(OP::ADDRESS _at, int _stateIndex)
+            {
+                at = _at;
+                stateIndex = _stateIndex;
+            }
+
+            OP::ADDRESS at;
+            int stateIndex;
+        };
+        vector<PENDING> m_pending;
+
+        //! States found so far
+        vector<STATE> m_states;
+
+        //! Index of the current state, from the m_states vector.
+        int m_currentState;
+
+        //! If true, operations adding resources (meshes or images) will only
+        //! have the base operation recursed, but not the resources.
+        bool m_skipResources;
+
+
+        //! Process all the pending operations and visit all children if necessary
+        void Recurse( PROGRAM& program )
+        {
+            while ( m_pending.size() )
+            {
+                OP::ADDRESS at = m_pending.back().at;
+                m_currentState = m_pending.back().stateIndex;
+                m_pending.pop_back();
+
+                bool recurse = false;
+
+                // Visit may change current state
+                recurse = Visit( at, program );
+
+                if (recurse)
+                {
+                    if (m_skipResources && VisitorIsAddResource(program.GetOpType(at)))
+                    {
+                        auto args = program.GetOpArgs<OP::InstanceAddArgs>(at);
+
+                        // Recurse only the base
+                        OP::ADDRESS base = args.instance;
+                        if (base)
+                        {
+							check(base<program.m_opAddress.size());
+                            m_pending.push_back( PENDING(base,m_currentState) );
+                        }
+                    }
+                    else
+                    {
+                        ForEachReference( program, at, [&](OP::ADDRESS ref)
+                        {
+                            if (ref)
+                            {
+								check(ref<program.m_opAddress.size());
+                                m_pending.push_back( PENDING(ref,m_currentState) );
+                            }
+                        });
+                    }
+                }
+            }
+        }
+
+    };
+
+
+    //---------------------------------------------------------------------------------------------
+    //! Code visitor template for visitors that:
+    //! - only traverses the operations that are relevant for a given set of parameter values. It
+    //! only considers the discrete parameters like integers and booleans. In the case of forks
+    //! caused by continuous parameters like float weights for interpolation, all the branches are
+    //! traversed.
+    //! - cannot change the instructinos
+    //---------------------------------------------------------------------------------------------
+    struct COVERED_CODE_VISITOR_STATE
+    {
+        uint16_t m_underResourceCount = 0;
+
+        bool operator==(const COVERED_CODE_VISITOR_STATE& o) const
+        {
+            return m_underResourceCount==o.m_underResourceCount;
+        }
+    };
+
+    template<typename PARENT,typename STATE>
+    class DiscreteCoveredCodeVisitorBase : public PARENT
+    {
+    public:
+
+        DiscreteCoveredCodeVisitorBase
+            (
+                System::Private* pSystem,
+                const ModelPtrConst& pModel,
+                const ParametersPtrConst& pParams,
+                unsigned lodMask,
+                bool skipResources=false
+            )
+            : PARENT(skipResources)
+        {
+            m_pSystem = pSystem;
+            m_pModel = pModel;
+            m_pParams = pParams.get();
+            m_lodMask = lodMask;
+
+            // Visiting state
+            PARENT::SetDefaultState( STATE() );
+        }
+
+        void Run( OP::ADDRESS at  )
+        {
+            PARENT::SetDefaultState( STATE() );
+
+            PARENT::Traverse( at, m_pModel->GetPrivate()->m_program );
+        }
+
+    protected:
+
+        virtual bool Visit( OP::ADDRESS at, PROGRAM& program )
+        {
+            bool recurse = true;
+
+            OP_TYPE type = program.GetOpType(at);
+
+            switch ( type )
+            {
+            case OP_TYPE::NU_CONDITIONAL:
+            case OP_TYPE::SC_CONDITIONAL:
+            case OP_TYPE::CO_CONDITIONAL:
+            case OP_TYPE::IM_CONDITIONAL:
+            case OP_TYPE::ME_CONDITIONAL:
+            case OP_TYPE::LA_CONDITIONAL:
+            case OP_TYPE::IN_CONDITIONAL:
+            {
+                auto args = program.GetOpArgs<OP::ConditionalArgs>(at);
+
+                recurse = false;
+
+                PARENT::RecurseWithCurrentState( args.condition );
+
+                // If there is no expression, we'll assume true.
+                bool value = true;
+
+                if (args.condition)
+                {
+                    value = m_pSystem->BuildBool(m_pModel, m_pParams, args.condition);
+                }
+
+                if (value)
+                {
+                    PARENT::RecurseWithCurrentState( args.yes );
+                }
+                else
+                {
+                    PARENT::RecurseWithCurrentState( args.no );
+                }
+                break;
+            }
+
+            case OP_TYPE::NU_SWITCH:
+            case OP_TYPE::SC_SWITCH:
+            case OP_TYPE::CO_SWITCH:
+            case OP_TYPE::IM_SWITCH:
+            case OP_TYPE::ME_SWITCH:
+            case OP_TYPE::LA_SWITCH:
+            case OP_TYPE::IN_SWITCH:
+            {
+                recurse = false;
+
+				const uint8_t* data = program.GetOpArgsPointer(at);
+				
+				OP::ADDRESS VarAddress;
+				FMemory::Memcpy( &VarAddress, data, sizeof(OP::ADDRESS));
+				data += sizeof(OP::ADDRESS);
+
+                if (VarAddress)
+                {
+					OP::ADDRESS DefAddress;
+					FMemory::Memcpy( &DefAddress, data, sizeof(OP::ADDRESS));
+					data += sizeof(OP::ADDRESS);
+
+					uint32_t CaseCount;
+					FMemory::Memcpy( &CaseCount, data, sizeof(uint32_t));
+					data += sizeof(uint32_t);
+
+                    PARENT::RecurseWithCurrentState( VarAddress );
+
+                    int var = m_pSystem->BuildInt( m_pModel, m_pParams, VarAddress );
+
+					OP::ADDRESS valueAt = DefAddress;
+					for (uint32_t C = 0; C < CaseCount; ++C)
+					{
+						int32_t Condition;
+						FMemory::Memcpy( &Condition, data, sizeof(int32_t));		
+						data += sizeof(int32_t);
+
+						OP::ADDRESS At;
+						FMemory::Memcpy( &At, data, sizeof(OP::ADDRESS));
+						data += sizeof(OP::ADDRESS);
+
+						if (At && var == (int)Condition)
+						{
+							valueAt = At;
+							break;	
+						}
+					}
+
+					PARENT::RecurseWithCurrentState( valueAt );
+                }
+
+                break;
+            }
+
+
+            case OP_TYPE::IN_ADDLOD:
+            {
+                auto args = program.GetOpArgs<OP::InstanceAddLODArgs>(at);
+
+                recurse = false;
+
+                STATE newState = PARENT::GetCurrentState();
+                for (int t=0;t<MUTABLE_OP_MAX_ADD_COUNT;++t)
+                {
+                    OP::ADDRESS lodAt = args.lod[t];
+                    if (lodAt)
+                    {
+                        bool selected = ( (1<<t) & m_lodMask ) != 0;
+                        if ( selected )
+                        {
+                            PARENT::RecurseWithState( lodAt, newState );
+                        }
+                    }
+                }
+                break;
+            }
+
+
+            case OP_TYPE::IN_ADDMESH:
+            {
+                auto args = program.GetOpArgs<OP::InstanceAddArgs>(at);
+
+                recurse = false;
+
+                PARENT::RecurseWithCurrentState(args.instance);
+
+                STATE newState = PARENT::GetCurrentState();
+                newState.m_underResourceCount=1;
+
+                OP::ADDRESS meshAt = args.value;
+                if (meshAt)
+                {
+                    PARENT::RecurseWithState(meshAt, newState);
+                }
+                break;
+            }
+
+
+            case OP_TYPE::IN_ADDIMAGE:
+            {
+                auto args = program.GetOpArgs<OP::InstanceAddArgs>(at);
+
+                recurse = false;
+
+                PARENT::RecurseWithCurrentState(args.instance);
+
+                STATE newState = PARENT::GetCurrentState();
+                newState.m_underResourceCount=1;
+
+                OP::ADDRESS imageAt = args.value;
+                if (imageAt)
+                {
+                    PARENT::RecurseWithState(imageAt, newState);
+                }
+                break;
+            }
+
+            default:
+                break;
+            }
+
+            return recurse;
+        }
+
+
+    protected:
+        System::Private* m_pSystem = nullptr;
+        ModelPtrConst m_pModel;
+        const Parameters* m_pParams = nullptr;
+        unsigned m_lodMask = 0;
+    };
+
+
+    //---------------------------------------------------------------------------------------------
+    //! Code visitor that:
+    //! - only traverses the operations that are relevant for a given set of parameter values. It
+    //! only considers the discrete parameters like integers and booleans. In the case of forks
+    //! caused by continuous parameters like float weights for interpolation, all the branches are
+    //! traversed.
+    //! - cannot change the instructions
+    //! - will not repeat visits to instructions with the same state
+    //! - the state has to be a compatible with COVERED_CODE_VISITOR_STATE
+    //---------------------------------------------------------------------------------------------
+    template<typename COVERED_STATE=COVERED_CODE_VISITOR_STATE>
+    class UniqueDiscreteCoveredCodeVisitor :
+            public DiscreteCoveredCodeVisitorBase
+            <
+            UniqueConstCodeVisitorIterative<COVERED_STATE>,
+            COVERED_STATE
+            >
+    {
+        using PARENT=DiscreteCoveredCodeVisitorBase<
+        UniqueConstCodeVisitorIterative<COVERED_STATE>,
+        COVERED_STATE
+        >;
+    public:
+
+        UniqueDiscreteCoveredCodeVisitor
+            (
+                System::Private* pSystem,
+                const ModelPtrConst& pModel,
+                const ParametersPtrConst& pParams,
+                unsigned lodMask
+            )
+            : PARENT( pSystem, pModel, pParams, lodMask )
+        {
+        }
+
+    };
+
+
+
+    //---------------------------------------------------------------------------------------------
+    //! Calculate all the parameters found under a particular operation
+    //! It has an internal cache, so don't reuse objects of this class if the program changes.
+    //---------------------------------------------------------------------------------------------
+    class MUTABLERUNTIME_API SubtreeParametersVisitor : public Base
+    {
+    public:
+
+        void Run( OP::ADDRESS root, PROGRAM& program );
+
+        //! After Run, list of relevant parameters.
+        vector<int> m_params;
+
+    private:
+
+        vector<int> m_currentParams;
+        vector<uint8_t> m_visited;
+        vector<OP::ADDRESS> m_pending;
+
+        // Result cache
+        map< OP::ADDRESS, vector<int> > m_resultCache;
+    };
+
+
+}
+
diff --git a/Engine/Plugins/Experimental/Mutable/Source/MutableRuntime/Private/Config.h b/Engine/Plugins/Experimental/Mutable/Source/MutableRuntime/Private/Config.h
new file mode 100644
--- /dev/null
+++ b/Engine/Plugins/Experimental/Mutable/Source/MutableRuntime/Private/Config.h
@@ -0,0 +1,54 @@
+// Copyright Epic Games, Inc. All Rights Reserved.
+
+#pragma once
+
+//! A config file may have been provided with a define in the command-linewith something like
+//! -D'MUTABLE_CONFIG="ConfigAndroid.h"'
+
+#ifdef MUTABLE_CONFIG
+
+#include MUTABLE_CONFIG
+
+//! otherwise, use this one
+#else
+
+// This file contains all the possible configuration flags that affect the build of the runtime
+
+//! Add support for OpenGL 4 gpu building operations.
+//! If this is enabled, the .cpp files in the mutable/runtime/impl/gpu-gl4 need to be linked in the
+//! runtime.
+//#define MUTABLE_GPU_GL4
+
+//! Add support for OpenGL ES 2 gpu building operations.
+//! If this is enabled, the .cpp files in the mutable/runtime/impl/gpu-gles2 need to be linked in
+//! the runtime.
+//#define MUTABLE_GPU_GLES2
+
+//! This flag enables collection of statistics in the System objects. If using the waf build system,
+//! the statistics are only collected in the debug and profile versions of the runtime.
+//#define MUTABLE_PROFILE
+
+//! This flag disables the internal asserts of the library. Without this flag, the asserts are only
+//! disabled in release builds. This can be overriden by MUTABLE_EXCEPTIONS
+//#define MUTABLE_NO_ASSERT
+
+//! This flags will enable asserts regardless of MUTABLE_NO_ASSERT but they will thrown an
+//! exception instead of stop the program. This is useful for tools only.
+//#define MUTABLE_EXCEPTIONS
+
+//! Exclude code for multilayer image operations
+//#define MUTABLE_DISABLE_MULTILAYER
+
+//! Exclude support for BC(DXTC) image formats
+//! This may also be set by the build system for some target platforms.
+//#define MUTABLE_EXCLUDE_BC
+
+//! Exclude support for ASTC image formats
+//! This may also be set by the build system for some target platforms.
+//#define MUTABLE_EXCLUDE_ASTC
+
+//! Exclude support for higher quality compression image conversion (see Miro.h)
+//! This may also be set by the build system for some target platforms.
+//#define MUTABLE_EXCLUDE_IMAGE_COMPRESSION_QUALITY 4
+
+#endif
