[mutable] Move Mutable plugin from the restricted folder to the experimental plugins folder.

diff --git a/Engine/Plugins/Experimental/Mutable/Source/MutableRuntime/Private/ConvertData.h b/Engine/Plugins/Experimental/Mutable/Source/MutableRuntime/Private/ConvertData.h
new file mode 100644
--- /dev/null
+++ b/Engine/Plugins/Experimental/Mutable/Source/MutableRuntime/Private/ConvertData.h
@@ -0,0 +1,1325 @@
+// Copyright Epic Games, Inc. All Rights Reserved.
+
+#pragma once
+
+#include "Mesh.h"
+#include "Platform.h"
+#include "MutableMath.h"
+
+
+namespace mu
+{
+
+	//---------------------------------------------------------------------------------------------
+	//! Convert one channel element
+	//---------------------------------------------------------------------------------------------
+	inline void ConvertData
+		(
+			int channel,
+			void* pResult, MESH_BUFFER_FORMAT resultFormat,
+			const void* pSource, MESH_BUFFER_FORMAT sourceFormat
+		)
+	{
+		switch ( resultFormat )
+		{
+		case MBF_FLOAT64:
+		{
+			double* pTypedResult = reinterpret_cast<double*>(pResult);
+			uint8_t* pByteResult = reinterpret_cast<uint8_t*>(pResult);
+			const uint8_t* pByteSource = reinterpret_cast<const uint8_t*>(pSource);
+
+			switch (sourceFormat)
+			{
+			case MBF_FLOAT64:
+			{
+				// Just dereferencing is not safe in all architectures. some like ARM require the
+				// floats to be 4-byte aligned and it may not be the case. We need memcpy.
+				memcpy(pByteResult + 8 * channel, pByteSource + 8* channel, 8);
+				break;
+			}
+
+			case MBF_FLOAT32:
+			{
+				const float* pTypedSource = reinterpret_cast<const float*>(pSource);
+				pTypedResult[channel] = pTypedSource[channel];
+				break;
+			}
+
+			case MBF_FLOAT16:
+			{
+				const float16* pTypedSource = reinterpret_cast<const float16*>(pSource);
+				pTypedResult[channel] = halfToFloat(pTypedSource[channel]);
+				break;
+			}
+
+			case MBF_INT32:
+			{
+				const int32_t* pTypedSource = reinterpret_cast<const int32_t*>(pSource);
+				pTypedResult[channel] = (double)(pTypedSource[channel]);
+				break;
+			}
+
+			case MBF_UINT32:
+			{
+				const uint32_t* pTypedSource = reinterpret_cast<const uint32_t*>(pSource);
+				pTypedResult[channel] = (double)(pTypedSource[channel]);
+				break;
+			}
+
+			case MBF_INT16:
+			{
+				const int16_t* pTypedSource = reinterpret_cast<const int16_t*>(pSource);
+				pTypedResult[channel] = (double)(pTypedSource[channel]);
+				break;
+			}
+
+			case MBF_UINT16:
+			{
+				const uint16_t* pTypedSource = reinterpret_cast<const uint16_t*>(pSource);
+				pTypedResult[channel] = (double)(pTypedSource[channel]);
+				break;
+			}
+
+			case MBF_INT8:
+			{
+				const int8_t* pTypedSource = reinterpret_cast<const int8_t*>(pSource);
+				pTypedResult[channel] = (double)(pTypedSource[channel]);
+				break;
+			}
+
+			case MBF_UINT8:
+			{
+				const uint8_t* pTypedSource = reinterpret_cast<const uint8_t*>(pSource);
+				pTypedResult[channel] = (double)(pTypedSource[channel]);
+				break;
+			}
+
+			case MBF_NINT32:
+			{
+				const int32_t* pTypedSource = reinterpret_cast<const int32_t*>(pSource);
+				pTypedResult[channel] = (double)(pTypedSource[channel]);
+				pTypedResult[channel] /= 65536.0f * 65536.0f / 2.0f;
+				break;
+			}
+
+			case MBF_NUINT32:
+			{
+				const uint32_t* pTypedSource = reinterpret_cast<const uint32_t*>(pSource);
+				pTypedResult[channel] = (double)(pTypedSource[channel]);
+				pTypedResult[channel] /= 65536.0f * 65536.0f - 1.0f;
+				break;
+			}
+
+			case MBF_NINT16:
+			{
+				const int16_t* pTypedSource = reinterpret_cast<const int16_t*>(pSource);
+				pTypedResult[channel] = (double)(pTypedSource[channel]);
+				pTypedResult[channel] /= 32768.0f;
+				break;
+			}
+
+			case MBF_NUINT16:
+			{
+				const uint16_t* pTypedSource = reinterpret_cast<const uint16_t*>(pSource);
+				pTypedResult[channel] = (double)(pTypedSource[channel]);
+				pTypedResult[channel] /= 65535.0f;
+				break;
+			}
+
+			case MBF_NINT8:
+			{
+				const int8_t* pTypedSource = reinterpret_cast<const int8_t*>(pSource);
+				pTypedResult[channel] = (double)(pTypedSource[channel]);
+				pTypedResult[channel] /= 128.0f;
+				break;
+			}
+
+			case MBF_NUINT8:
+			{
+				const uint8_t* pTypedSource = reinterpret_cast<const uint8_t*>(pSource);
+				pTypedResult[channel] = (double)(pTypedSource[channel]);
+				pTypedResult[channel] /= 255.0f;
+				break;
+			}
+
+			case MBF_PACKEDDIR8:
+			case MBF_PACKEDDIR8_W_TANGENTSIGN:
+			{
+				const uint8_t* pTypedSource = reinterpret_cast<const uint8_t*>(pSource);
+				pTypedResult[channel] = (double)(pTypedSource[channel]);
+				pTypedResult[channel] /= 127.5f;
+				pTypedResult[channel] -= 1.0f;
+				break;
+			}
+
+			case MBF_PACKEDDIRS8:
+			case MBF_PACKEDDIRS8_W_TANGENTSIGN:
+			{
+				const int8_t* pTypedSource = reinterpret_cast<const int8_t*>(pSource);
+				pTypedResult[channel] = (double)(pTypedSource[channel]);
+				pTypedResult[channel] /= 127.5f;
+				break;
+			}
+
+			default:
+				checkf(false, TEXT("Conversion not implemented."));
+				break;
+			}
+			break;
+		}
+
+		case MBF_FLOAT32:
+		{
+			float* pTypedResult = reinterpret_cast<float*>( pResult );
+            uint8_t* pByteResult = reinterpret_cast<uint8_t*>( pResult );
+            const uint8_t* pByteSource = reinterpret_cast<const uint8_t*>( pSource );
+
+			switch ( sourceFormat )
+			{
+			case MBF_FLOAT64:
+			{
+				const double* pTypedSource = reinterpret_cast<const double*>(pSource);
+				pTypedResult[channel] = float(pTypedSource[channel]);
+				break;
+			}
+
+			case MBF_FLOAT32:
+			{
+				// Just dereferencing is not safe in all architectures. some like ARM require the
+				// floats to be 4-byte aligned and it may not be the case. We need memcpy.
+				memcpy(pByteResult + 4 * channel, pByteSource + 4 * channel, 4);
+				break;
+			}
+
+			case MBF_FLOAT16:
+			{
+				const float16* pTypedSource = reinterpret_cast<const float16*>( pSource );
+				pTypedResult[channel] = halfToFloat(pTypedSource[channel]);
+				break;
+			}
+
+			case MBF_INT32:
+			{
+                const int32_t* pTypedSource = reinterpret_cast<const int32_t*>( pSource );
+				pTypedResult[channel] = (float)(pTypedSource[channel]);
+				break;
+			}
+
+			case MBF_UINT32:
+			{
+                const uint32_t* pTypedSource = reinterpret_cast<const uint32_t*>( pSource );
+				pTypedResult[channel] = (float)(pTypedSource[channel]);
+				break;
+			}
+
+			case MBF_INT16:
+			{
+                const int16_t* pTypedSource = reinterpret_cast<const int16_t*>( pSource );
+				pTypedResult[channel] = (float)(pTypedSource[channel]);
+				break;
+			}
+
+			case MBF_UINT16:
+			{
+                const uint16_t* pTypedSource = reinterpret_cast<const uint16_t*>( pSource );
+				pTypedResult[channel] = (float)(pTypedSource[channel]);
+				break;
+			}
+
+			case MBF_INT8:
+			{
+                const int8_t* pTypedSource = reinterpret_cast<const int8_t*>( pSource );
+				pTypedResult[channel] = (float)(pTypedSource[channel]);
+				break;
+			}
+
+			case MBF_UINT8:
+			{
+                const uint8_t* pTypedSource = reinterpret_cast<const uint8_t*>( pSource );
+				pTypedResult[channel] = (float)(pTypedSource[channel]);
+				break;
+			}
+
+			case MBF_NINT32:
+			{
+                const int32_t* pTypedSource = reinterpret_cast<const int32_t*>( pSource );
+				pTypedResult[channel] = (float)(pTypedSource[channel]);
+				pTypedResult[channel] /= 65536.0f*65536.0f/2.0f;
+				break;
+			}
+
+			case MBF_NUINT32:
+			{
+                const uint32_t* pTypedSource = reinterpret_cast<const uint32_t*>( pSource );
+				pTypedResult[channel] = (float)(pTypedSource[channel]);
+				pTypedResult[channel] /= 65536.0f*65536.0f-1.0f;
+				break;
+			}
+
+			case MBF_NINT16:
+			{
+                const int16_t* pTypedSource = reinterpret_cast<const int16_t*>( pSource );
+				pTypedResult[channel] = (float)(pTypedSource[channel]);
+				pTypedResult[channel] /= 32768.0f;
+				break;
+			}
+
+			case MBF_NUINT16:
+			{
+                const uint16_t* pTypedSource = reinterpret_cast<const uint16_t*>( pSource );
+				pTypedResult[channel] = (float)(pTypedSource[channel]);
+				pTypedResult[channel] /= 65535.0f;
+				break;
+			}
+
+			case MBF_NINT8:
+			{
+                const int8_t* pTypedSource = reinterpret_cast<const int8_t*>( pSource );
+				pTypedResult[channel] = (float)(pTypedSource[channel]);
+				pTypedResult[channel] /= 128.0f;
+				break;
+			}
+
+			case MBF_NUINT8:
+			{
+                const uint8_t* pTypedSource = reinterpret_cast<const uint8_t*>( pSource );
+				pTypedResult[channel] = (float)(pTypedSource[channel]);
+				pTypedResult[channel] /= 255.0f;
+				break;
+			}
+
+            case MBF_PACKEDDIR8:
+            case MBF_PACKEDDIR8_W_TANGENTSIGN:
+            {
+                const uint8_t* pTypedSource = reinterpret_cast<const uint8_t*>( pSource );
+                pTypedResult[channel] = (float)(pTypedSource[channel]);
+                pTypedResult[channel] /= 127.5f;
+                pTypedResult[channel] -= 1.0f;
+                break;
+            }
+
+            case MBF_PACKEDDIRS8:
+            case MBF_PACKEDDIRS8_W_TANGENTSIGN:
+            {
+                const int8_t* pTypedSource = reinterpret_cast<const int8_t*>( pSource );
+                pTypedResult[channel] = (float)(pTypedSource[channel]);
+                pTypedResult[channel] /= 127.5f;
+                break;
+            }
+
+			default:
+				checkf( false, TEXT("Conversion not implemented.") );
+				break;
+			}
+			break;
+		}
+
+		//-----------------------------------------------------------------------------------------
+		case MBF_FLOAT16:
+		{
+			float16* pTypedResult = reinterpret_cast<float16*>( pResult );
+
+			switch ( sourceFormat )
+			{
+			case MBF_FLOAT32:
+			{
+				const float* pTypedSource = reinterpret_cast<const float*>( pSource );
+				pTypedResult[channel] = floatToHalf(pTypedSource[channel]);
+				break;
+			}
+
+			case MBF_FLOAT16:
+			{
+				const float16* pTypedSource = reinterpret_cast<const float16*>( pSource );
+				pTypedResult[channel] = pTypedSource[channel];
+				break;
+			}
+
+			case MBF_UINT32:
+			{
+                const uint32_t* pTypedSource = reinterpret_cast<const uint32_t*>( pSource );
+				pTypedResult[channel] = floatToHalf((float)(pTypedSource[channel]));
+				break;
+			}
+
+			case MBF_INT32:
+			{
+                const int32_t* pTypedSource = reinterpret_cast<const int32_t*>( pSource );
+				pTypedResult[channel] = floatToHalf((float)(pTypedSource[channel]));
+				break;
+			}
+
+			case MBF_UINT16:
+			{
+                const uint16_t* pTypedSource = reinterpret_cast<const uint16_t*>( pSource );
+				pTypedResult[channel] = floatToHalf((float)(pTypedSource[channel]));
+				break;
+			}
+
+			case MBF_INT16:
+			{
+                const int16_t* pTypedSource = reinterpret_cast<const int16_t*>( pSource );
+				pTypedResult[channel] = floatToHalf((float)(pTypedSource[channel]));
+				break;
+			}
+
+			case MBF_UINT8:
+			{
+                const uint8_t* pTypedSource = reinterpret_cast<const uint8_t*>( pSource );
+				pTypedResult[channel] = floatToHalf((float)(pTypedSource[channel]));
+				break;
+			}
+
+			case MBF_INT8:
+			{
+                const int8_t* pTypedSource = reinterpret_cast<const int8_t*>( pSource );
+				pTypedResult[channel] = floatToHalf((float)(pTypedSource[channel]));
+				break;
+			}
+
+			default:
+				checkf( false, TEXT("Conversion not implemented.") );
+				break;
+			}
+			break;
+		}
+
+		//-----------------------------------------------------------------------------------------
+		case MBF_UINT8:
+		{
+            uint8_t* pTypedResult = reinterpret_cast<uint8_t*>( pResult );
+
+			switch ( sourceFormat )
+			{
+			case MBF_FLOAT32:
+			{
+				const float* pTypedSource = reinterpret_cast<const float*>( pSource );
+                pTypedResult[channel] = (uint8_t)
+                        std::min<uint32_t>(
+							0xFF,
+                            std::max<uint32_t>(
+								0,
+                                (uint32_t)(pTypedSource[channel])
+								)
+							);
+				break;
+			}
+
+			case MBF_FLOAT16:
+			{
+				const float16* pTypedSource = reinterpret_cast<const float16*>( pSource );
+                pTypedResult[channel] = (uint8_t)
+                        std::min<uint32_t>(
+							0xFF,
+                            std::max<uint32_t>(
+								0,
+                                (uint32_t)halfToFloat(pTypedSource[channel])
+								)
+							);
+				break;
+			}
+
+			case MBF_INT8:
+			{
+                const int8_t* pTypedSource = reinterpret_cast<const int8_t*>( pSource );
+                pTypedResult[channel] =  (uint8_t)std::max<int8_t>( 0, pTypedSource[channel] );
+				break;
+			}
+
+			case MBF_UINT8:
+			{
+                const uint8_t* pTypedSource = reinterpret_cast<const uint8_t*>( pSource );
+				pTypedResult[channel] = pTypedSource[channel];
+				break;
+			}
+
+			case MBF_INT16:
+			{
+                const int16_t* pTypedSource = reinterpret_cast<const int16_t*>( pSource );
+                pTypedResult[channel] =  (uint8_t)
+                        std::min<int16_t>(
+							0xFF,
+                            std::max<int16_t>(
+								0,
+								pTypedSource[channel]
+								)
+							);
+				break;
+			}
+
+			case MBF_UINT16:
+			{
+				// Clamp
+                const uint16_t* pTypedSource = reinterpret_cast<const uint16_t*>( pSource );
+                pTypedResult[channel] = (uint8_t)
+                        std::min<uint16_t>(
+							0xFF,
+                            std::max<uint16_t>(
+								0,
+								pTypedSource[channel]
+								)
+							);
+				break;
+			}
+
+			case MBF_INT32:
+			{
+                const int32_t* pTypedSource = reinterpret_cast<const int32_t*>( pSource );
+                pTypedResult[channel] =  (uint8_t)
+                        std::min<int32_t>(
+							0xFF,
+                            std::max<int32_t>(
+								0,
+								pTypedSource[channel]
+								)
+							);
+				break;
+			}
+
+			case MBF_UINT32:
+			{
+				// Clamp
+                const uint32_t* pTypedSource = reinterpret_cast<const uint32_t*>( pSource );
+                pTypedResult[channel] = (uint8_t)
+                        std::min<uint32_t>(
+							0xFF,
+                            std::max<uint32_t>(
+								0,
+								pTypedSource[channel]
+								)
+							);
+				break;
+			}
+
+			case MBF_NUINT8:
+			case MBF_NUINT16:
+			case MBF_NUINT32:
+			case MBF_NINT8:
+			case MBF_NINT16:
+			case MBF_NINT32:
+				pTypedResult = 0;
+				break;
+
+			default:
+				checkf( false, TEXT("Conversion not implemented." ) );
+				break;
+			}
+			break;
+		}
+
+		case MBF_UINT16:
+		{
+            uint16_t* pTypedResult = reinterpret_cast<uint16_t*>( pResult );
+
+			switch ( sourceFormat )
+			{
+			case MBF_FLOAT32:
+			{
+				const float* pTypedSource = reinterpret_cast<const float*>( pSource );
+                pTypedResult[channel] = (uint16_t)
+                        std::min<uint32_t>(
+							0xFFFF,
+                            std::max<uint32_t>(
+								0,
+                                (uint32_t)(pTypedSource[channel])
+								)
+							);
+				break;
+			}
+
+			case MBF_FLOAT16:
+			{
+				const float16* pTypedSource = reinterpret_cast<const float16*>( pSource );
+                pTypedResult[channel] = (uint16_t)
+                        std::min<uint32_t>(
+							0xFF,
+                            std::max<uint32_t>(
+								0,
+                                (uint32_t)halfToFloat(pTypedSource[channel])
+								)
+							);
+				break;
+			}
+
+			case MBF_UINT8:
+			{
+                const uint8_t* pTypedSource = reinterpret_cast<const uint8_t*>( pSource );
+				pTypedResult[channel] =  pTypedSource[channel];
+				break;
+			}
+
+			case MBF_INT8:
+			{
+                const int8_t* pTypedSource = reinterpret_cast<const int8_t*>( pSource );
+                pTypedResult[channel] =  (uint16_t)std::max<int8_t>( 0, pTypedSource[channel] );
+				break;
+			}
+
+			case MBF_UINT16:
+			{
+                const uint16_t* pTypedSource = reinterpret_cast<const uint16_t*>( pSource );
+				pTypedResult[channel] =  pTypedSource[channel];
+				break;
+			}
+
+			case MBF_INT16:
+			{
+                const int16_t* pTypedSource = reinterpret_cast<const int16_t*>( pSource );
+                pTypedResult[channel] =  (uint16_t)std::max<int16_t>( 0, pTypedSource[channel] );
+				break;
+			}
+
+			case MBF_UINT32:
+			{
+				// Clamp
+                const uint32_t* pTypedSource = reinterpret_cast<const uint32_t*>( pSource );
+                pTypedResult[channel] = (uint16_t)
+                        std::min<uint32_t>(
+							0xFFFF,
+                            std::max<uint32_t>(
+								0,
+								pTypedSource[channel]
+								)
+							);
+				break;
+			}
+
+			case MBF_INT32:
+			{
+				// Clamp
+                const int32_t* pTypedSource = reinterpret_cast<const int32_t*>( pSource );
+                pTypedResult[channel] = (uint16_t)
+                        std::min<int32_t>(
+							0xFFFF,
+                            std::max<int32_t>(
+								0,
+								pTypedSource[channel]
+								)
+							);
+				break;
+			}
+
+			case MBF_NUINT8:
+			case MBF_NUINT16:
+			case MBF_NUINT32:
+			case MBF_NINT8:
+			case MBF_NINT16:
+			case MBF_NINT32:
+				pTypedResult = 0;
+				break;
+
+			default:
+				checkf( false, TEXT("Conversion not implemented." ) );
+				break;
+			}
+			break;
+		}
+
+		case MBF_UINT32:
+		{
+            uint32_t* pTypedResult = reinterpret_cast<uint32_t*>( pResult );
+
+			switch ( sourceFormat )
+			{
+			case MBF_FLOAT32:
+			{
+				const float* pTypedSource = reinterpret_cast<const float*>( pSource );
+				pTypedResult[channel] =
+                        std::min<uint32_t>(
+							0xFFFF,
+                            std::max<uint32_t>(
+								0,
+                                (uint32_t)(pTypedSource[channel])
+								)
+							);
+				break;
+			}
+
+			case MBF_FLOAT16:
+			{
+				const float16* pTypedSource = reinterpret_cast<const float16*>( pSource );
+				pTypedResult[channel] =
+                        std::min<uint32_t>(
+							0xFF,
+                            std::max<uint32_t>(
+								0,
+                                (uint32_t)halfToFloat(pTypedSource[channel])
+								)
+							);
+				break;
+			}
+
+			case MBF_UINT8:
+			{
+                const uint8_t* pTypedSource = reinterpret_cast<const uint8_t*>( pSource );
+				pTypedResult[channel] =  pTypedSource[channel];
+				break;
+			}
+
+			case MBF_INT8:
+			{
+                const int8_t* pTypedSource = reinterpret_cast<const int8_t*>( pSource );
+                pTypedResult[channel] =  (uint16_t)std::max<int8_t>( 0, pTypedSource[channel] );
+				break;
+			}
+
+			case MBF_UINT16:
+			{
+                const uint16_t* pTypedSource = reinterpret_cast<const uint16_t*>( pSource );
+				pTypedResult[channel] =  pTypedSource[channel];
+				break;
+			}
+
+			case MBF_INT16:
+			{
+                const int16_t* pTypedSource = reinterpret_cast<const int16_t*>( pSource );
+                pTypedResult[channel] =  (uint16_t)std::max<int16_t>( 0, pTypedSource[channel] );
+				break;
+			}
+
+			case MBF_UINT32:
+			{
+                const uint32_t* pTypedSource = reinterpret_cast<const uint32_t*>( pSource );
+				pTypedResult[channel] = pTypedSource[channel];
+				break;
+			}
+
+			case MBF_INT32:
+			{
+				// Clamp
+                const int32_t* pTypedSource = reinterpret_cast<const int32_t*>( pSource );
+                pTypedResult[channel] = (uint32_t)std::max<int32_t>( 0, pTypedSource[channel] );
+				break;
+			}
+
+			case MBF_NUINT8:
+			case MBF_NUINT16:
+			case MBF_NUINT32:
+			case MBF_NINT8:
+			case MBF_NINT16:
+			case MBF_NINT32:
+				pTypedResult = 0;
+				break;
+
+			default:
+				checkf( false, TEXT("Conversion not implemented." ) );
+				break;
+			}
+			break;
+		}
+
+		case MBF_INT8:
+		{
+            int8_t* pTypedResult = reinterpret_cast<int8_t*>( pResult );
+
+			switch ( sourceFormat )
+			{
+			case MBF_FLOAT32:
+			{
+				const float* pTypedSource = reinterpret_cast<const float*>( pSource );
+                pTypedResult[channel] = (int8_t)
+                        std::min<int32_t>(
+							127,
+                            std::max<int32_t>(
+								-128,
+                                (int32_t)(pTypedSource[channel])
+								)
+							);
+				break;
+			}
+
+			case MBF_FLOAT16:
+			{
+				const float16* pTypedSource = reinterpret_cast<const float16*>( pSource );
+                pTypedResult[channel] = (int8_t)
+                        std::min<int32_t>(
+							127,
+                            std::max<int32_t>(
+								-128,
+                                (int32_t)halfToFloat(pTypedSource[channel])
+								)
+							);
+				break;
+			}
+
+			case MBF_INT8:
+			{
+                const int8_t* pTypedSource = reinterpret_cast<const int8_t*>( pSource );
+				pTypedResult[channel] = pTypedSource[channel];
+				break;
+			}
+
+			case MBF_NUINT8:
+			case MBF_NUINT16:
+			case MBF_NUINT32:
+			case MBF_NINT8:
+			case MBF_NINT16:
+			case MBF_NINT32:
+				pTypedResult = 0;
+				break;
+
+			default:
+				checkf( false, TEXT("Conversion not implemented." ) );
+				break;
+			}
+			break;
+		}
+
+		case MBF_INT16:
+		{
+            int16_t* pTypedResult = reinterpret_cast<int16_t*>( pResult );
+
+			switch ( sourceFormat )
+			{
+			case MBF_FLOAT32:
+			{
+				const float* pTypedSource = reinterpret_cast<const float*>( pSource );
+                pTypedResult[channel] = (int16_t)
+                        std::min<int32_t>(
+							32767,
+                            std::max<int32_t>(
+								-32768,
+                                (int32_t)(pTypedSource[channel])
+								)
+							);
+				break;
+			}
+
+			case MBF_FLOAT16:
+			{
+				const float16* pTypedSource = reinterpret_cast<const float16*>( pSource );
+                pTypedResult[channel] = (int16_t)
+                        std::min<int32_t>(
+							32767,
+                            std::max<int32_t>(
+								-32768,
+                                (int32_t)halfToFloat(pTypedSource[channel])
+								)
+							);
+				break;
+			}
+
+			case MBF_INT8:
+			{
+                const int8_t* pTypedSource = reinterpret_cast<const int8_t*>( pSource );
+                pTypedResult[channel] = (int16_t)pTypedSource[channel];
+				break;
+			}
+
+			case MBF_UINT8:
+			{
+                const uint8_t* pTypedSource = reinterpret_cast<const uint8_t*>( pSource );
+                pTypedResult[channel] = (int16_t)pTypedSource[channel];
+				break;
+			}
+
+			case MBF_UINT16:
+			{
+                const uint16_t* pTypedSource = reinterpret_cast<const uint16_t*>( pSource );
+                pTypedResult[channel] = (int16_t)
+                        std::min<int32_t>(
+                            32767, (int32_t)pTypedSource[channel]
+							);
+				break;
+			}
+
+			case MBF_INT32:
+			{
+                const int32_t* pTypedSource = reinterpret_cast<const int32_t*>( pSource );
+                pTypedResult[channel] = (int16_t)
+                        std::min<int32_t>(
+							32767,
+                            std::max<int32_t>(
+								-32768,
+								pTypedSource[channel]
+								)
+							);
+				break;
+			}
+
+			case MBF_UINT32:
+			{
+                const uint32_t* pTypedSource = reinterpret_cast<const uint32_t*>( pSource );
+                pTypedResult[channel] = (int16_t)
+                        std::min<int32_t>(
+							32767, pTypedSource[channel]
+							);
+				break;
+			}
+
+			case MBF_NUINT8:
+			case MBF_NUINT16:
+			case MBF_NUINT32:
+			case MBF_NINT8:
+			case MBF_NINT16:
+			case MBF_NINT32:
+				pTypedResult = 0;
+				break;
+
+			default:
+				checkf( false, TEXT("Conversion not implemented." ) );
+				break;
+			}
+			break;
+		}
+
+		case MBF_INT32:
+		{
+            int32_t* pTypedResult = reinterpret_cast<int32_t*>( pResult );
+
+			switch ( sourceFormat )
+			{
+			case MBF_FLOAT32:
+			{
+				const float* pTypedSource = reinterpret_cast<const float*>( pSource );
+                pTypedResult[channel] = (int32_t)(pTypedSource[channel]);
+				break;
+			}
+
+			case MBF_FLOAT16:
+			{
+				const float16* pTypedSource = reinterpret_cast<const float16*>( pSource );
+                pTypedResult[channel] = (int32_t)halfToFloat(pTypedSource[channel]);
+				break;
+			}
+
+			case MBF_INT8:
+			{
+                const int8_t* pTypedSource = reinterpret_cast<const int8_t*>( pSource );
+                pTypedResult[channel] = (int32_t)pTypedSource[channel];
+				break;
+			}
+
+			case MBF_UINT8:
+			{
+                const uint8_t* pTypedSource = reinterpret_cast<const uint8_t*>( pSource );
+                pTypedResult[channel] = (int32_t)pTypedSource[channel];
+				break;
+			}
+
+			case MBF_INT16:
+			{
+                const int16_t* pTypedSource = reinterpret_cast<const int16_t*>( pSource );
+                pTypedResult[channel] = (int32_t)pTypedSource[channel];
+				break;
+			}
+
+			case MBF_UINT16:
+			{
+                const uint16_t* pTypedSource = reinterpret_cast<const uint16_t*>( pSource );
+                pTypedResult[channel] = (int32_t)pTypedSource[channel];
+				break;
+			}
+
+			case MBF_UINT32:
+			{
+                const uint32_t* pTypedSource = reinterpret_cast<const uint32_t*>( pSource );
+                pTypedResult[channel] = (int32_t)pTypedSource[channel];
+				break;
+			}
+
+			case MBF_INT32:
+			{
+                const int32_t* pTypedSource = reinterpret_cast<const int32_t*>( pSource );
+                pTypedResult[channel] = pTypedSource[channel];
+				break;
+			}
+
+			case MBF_NUINT8:
+			case MBF_NUINT16:
+			case MBF_NUINT32:
+			case MBF_NINT8:
+			case MBF_NINT16:
+			case MBF_NINT32:
+				pTypedResult = 0;
+				break;
+
+			default:
+				checkf( false, TEXT("Conversion not implemented." ) );
+				break;
+			}
+			break;
+		}
+
+		//-----------------------------------------------------------------------------------------
+		case MBF_NUINT8:
+		{
+            uint8_t* pTypedResult = reinterpret_cast<uint8_t*>( pResult );
+
+			switch ( sourceFormat )
+			{
+			case MBF_NUINT8:
+			{
+				auto pTypedSource = reinterpret_cast<const uint8_t*>(pSource);
+				pTypedResult[channel] = pTypedSource[channel];
+				break;
+			}
+
+			case MBF_FLOAT32:
+			{
+				const float* pTypedSource = reinterpret_cast<const float*>(pSource);
+				pTypedResult[channel] = (uint8_t)
+					std::min<uint32_t>(
+						0xFF,
+						std::max<uint32_t>(
+							0,
+							(uint32_t)(((float)0xFF)*pTypedSource[channel] + 0.5f)
+							)
+						);
+				break;
+			}
+
+			case MBF_FLOAT16:
+			{
+				const float16* pTypedSource = reinterpret_cast<const float16*>( pSource );
+                pTypedResult[channel] = (uint8_t)
+                        std::min<uint32_t>(
+							0xFF,
+                            std::max<uint32_t>(
+								0,
+                                (uint32_t)(((float)0xFF)*halfToFloat(pTypedSource[channel])+0.5f)
+								)
+							);
+				break;
+			}
+
+			case MBF_UINT8:
+			case MBF_UINT16:
+			case MBF_UINT32:
+			case MBF_INT8:
+			case MBF_INT16:
+			case MBF_INT32:
+				// TODO: 1 or 0
+				pTypedResult = 0;
+				break;
+
+			default:
+				checkf( false, TEXT("Conversion not implemented." ) );
+				break;
+			}
+			break;
+		}
+
+		case MBF_NUINT16:
+		{
+            uint16_t* pTypedResult = reinterpret_cast<uint16_t*>( pResult );
+
+			switch ( sourceFormat )
+			{
+			case MBF_NUINT16:
+			{
+				auto pTypedSource = reinterpret_cast<const uint16_t*>(pSource);
+				pTypedResult[channel] = pTypedSource[channel];
+				break;
+			}
+
+			case MBF_FLOAT32:
+			{
+				const float* pTypedSource = reinterpret_cast<const float*>( pSource );
+                pTypedResult[channel] = (uint16_t)
+                        std::min<uint32_t>(
+							0xFFFF,
+                            std::max<uint32_t>(
+								0,
+                                (uint32_t)(((float)0xFFFF)*pTypedSource[channel]+0.5f)
+								)
+							);
+				break;
+			}
+
+			case MBF_FLOAT16:
+			{
+				const float16* pTypedSource = reinterpret_cast<const float16*>( pSource );
+                pTypedResult[channel] = (uint16_t)
+                        std::min<uint32_t>(
+							0xFFFF,
+                            std::max<uint32_t>(
+								0,
+                                (uint32_t)(((float)0xFFFF)*halfToFloat(pTypedSource[channel])+0.5f)
+								)
+							);
+				break;
+			}
+
+			case MBF_UINT8:
+			case MBF_UINT16:
+			case MBF_UINT32:
+			case MBF_INT8:
+			case MBF_INT16:
+			case MBF_INT32:
+				// TODO: 1 or 0
+				pTypedResult = 0;
+				break;
+
+			default:
+				checkf( false, TEXT("Conversion not implemented." ) );
+				break;
+			}
+			break;
+		}
+
+		case MBF_NUINT32:
+		{
+            uint32_t* pTypedResult = reinterpret_cast<uint32_t*>( pResult );
+
+			switch ( sourceFormat )
+			{
+			case MBF_FLOAT32:
+			{
+				const float* pTypedSource = reinterpret_cast<const float*>( pSource );
+                pTypedResult[channel] = (uint32_t)
+                        std::min<uint32_t>(
+							0xFFFFFFFF,
+                            std::max<uint32_t>(
+								0,
+                                (uint32_t)(((float)0xFFFFFFFF)*pTypedSource[channel])
+								)
+							);
+				break;
+			}
+
+			case MBF_FLOAT16:
+			{
+				const float16* pTypedSource = reinterpret_cast<const float16*>( pSource );
+                pTypedResult[channel] = (uint32_t)
+                        std::min<uint32_t>(
+							0xFFFFFFFF,
+                            std::max<uint32_t>(
+								0,
+                                (uint32_t)(((float)0xFFFFFFFF)*halfToFloat(pTypedSource[channel])+0.5f)
+								)
+							);
+				break;
+			}
+
+			case MBF_UINT8:
+			case MBF_UINT16:
+			case MBF_UINT32:
+			case MBF_INT8:
+			case MBF_INT16:
+			case MBF_INT32:
+				// TODO: 1 or 0
+				pTypedResult = 0;
+				break;
+
+			default:
+				checkf( false, TEXT("Conversion not implemented." ) );
+				break;
+			}
+			break;
+		}
+
+		case MBF_NINT8:
+		{
+            int8_t* pTypedResult = reinterpret_cast<int8_t*>( pResult );
+
+			switch ( sourceFormat )
+			{
+			case MBF_FLOAT32:
+			{
+				const float* pTypedSource = reinterpret_cast<const float*>( pSource );
+                pTypedResult[channel] = (int8_t)
+                        std::min<int32_t>(
+							127,
+                            std::max<int32_t>(
+								-128,
+                                (int32_t)(128.0f*pTypedSource[channel]+0.5f)
+								)
+							);
+				break;
+			}
+
+			case MBF_FLOAT16:
+			{
+				const float16* pTypedSource = reinterpret_cast<const float16*>( pSource );
+                pTypedResult[channel] = (int8_t)
+                        std::min<int32_t>(
+							127,
+                            std::max<int32_t>(
+								-128,
+                                (int32_t)(128.0f*halfToFloat(pTypedSource[channel])+0.5f)
+								)
+							);
+				break;
+			}
+
+			case MBF_UINT8:
+			case MBF_UINT16:
+			case MBF_UINT32:
+			case MBF_INT8:
+			case MBF_INT16:
+			case MBF_INT32:
+				// TODO: -1, 1 or 0
+				pTypedResult = 0;
+				break;
+
+			default:
+				checkf( false, TEXT("Conversion not implemented." ) );
+				break;
+			}
+			break;
+		}
+
+		case MBF_NINT16:
+		{
+            int16_t* pTypedResult = reinterpret_cast<int16_t*>( pResult );
+
+			switch ( sourceFormat )
+			{
+			case MBF_FLOAT32:
+			{
+				const float* pTypedSource = reinterpret_cast<const float*>( pSource );
+                pTypedResult[channel] = (int16_t)
+                        std::min<int32_t>(
+							32767,
+                            std::max<int32_t>(
+								-32768,
+                                (int32_t)(32768.0f*pTypedSource[channel]+0.5f)
+								)
+							);
+				break;
+			}
+
+			case MBF_FLOAT16:
+			{
+				const float16* pTypedSource = reinterpret_cast<const float16*>( pSource );
+                pTypedResult[channel] = (int16_t)
+                        std::min<int32_t>(
+							32767,
+                            std::max<int32_t>(
+								-32768,
+                                (int32_t)(32768.0f*halfToFloat(pTypedSource[channel])+0.5f)
+								)
+							);
+				break;
+			}
+
+			case MBF_UINT8:
+			case MBF_UINT16:
+			case MBF_UINT32:
+			case MBF_INT8:
+			case MBF_INT16:
+			case MBF_INT32:
+				// TODO: -1, 1 or 0
+				pTypedResult = 0;
+				break;
+
+			default:
+				checkf( false, TEXT("Conversion not implemented." ) );
+				break;
+			}
+			break;
+		}
+
+
+		case MBF_NINT32:
+		{
+            int32_t* pTypedResult = reinterpret_cast<int32_t*>( pResult );
+
+			switch ( sourceFormat )
+			{
+			case MBF_FLOAT32:
+			{
+				const float* pTypedSource = reinterpret_cast<const float*>( pSource );
+                pTypedResult[channel] = (int32_t)(2147483648.0f*pTypedSource[channel]+0.5f);
+				break;
+			}
+
+			case MBF_FLOAT16:
+			{
+				const float16* pTypedSource = reinterpret_cast<const float16*>( pSource );
+                pTypedResult[channel] = (int32_t)(2147483648.0f*halfToFloat(pTypedSource[channel])+0.5f);
+				break;
+			}
+
+			case MBF_UINT8:
+			case MBF_UINT16:
+			case MBF_UINT32:
+			case MBF_INT8:
+			case MBF_INT16:
+			case MBF_INT32:
+				// TODO: -1, 1 or 0
+				pTypedResult = 0;
+				break;
+
+			default:
+				checkf( false, TEXT("Conversion not implemented." ) );
+				break;
+			}
+			break;
+		}
+
+        case MBF_PACKEDDIR8:
+        case MBF_PACKEDDIR8_W_TANGENTSIGN:
+        {
+            uint8_t* pTypedResult = reinterpret_cast<uint8_t*>( pResult );
+
+            switch ( sourceFormat )
+            {
+            case MBF_PACKEDDIR8:
+            case MBF_PACKEDDIR8_W_TANGENTSIGN:
+            {
+                const uint8_t* pTypedSource = reinterpret_cast<const uint8_t*>( pSource );
+                uint8_t source = pTypedSource[channel];
+                pTypedResult[channel] = source;
+                break;
+            }
+
+            case MBF_FLOAT32:
+            {
+                const float* pTypedSource = reinterpret_cast<const float*>( pSource );
+                float source = pTypedSource[channel];
+                source = (source*0.5f+0.5f)*255.0f;
+                pTypedResult[channel] =
+                        (uint8_t)std::min<float>( 255.0f, std::max<float>( 0.0f, source ) );
+                break;
+            }
+
+            default:
+                checkf( false, TEXT("Conversion not implemented." ) );
+                break;
+            }
+            break;
+        }
+
+        case MBF_PACKEDDIRS8:
+        case MBF_PACKEDDIRS8_W_TANGENTSIGN:
+        {
+            int8_t* pTypedResult = reinterpret_cast<int8_t*>( pResult );
+
+            switch ( sourceFormat )
+            {
+            case MBF_PACKEDDIRS8:
+            case MBF_PACKEDDIRS8_W_TANGENTSIGN:
+            {
+                const int8_t* pTypedSource = reinterpret_cast<const int8_t*>( pSource );
+                int8_t source = pTypedSource[channel];
+                pTypedResult[channel] = source;
+                break;
+            }
+
+            case MBF_FLOAT32:
+            {
+                const float* pTypedSource = reinterpret_cast<const float*>( pSource );
+                float source = pTypedSource[channel];
+                source = source*0.5f*255.0f;
+                pTypedResult[channel] =
+                        (int8_t)std::min<float>( 127.0f, std::max<float>( -128.0f, source ) );
+                break;
+            }
+
+            default:
+                checkf( false, TEXT("Conversion not implemented." ) );
+                break;
+            }
+            break;
+        }
+
+		default:
+			checkf( false, TEXT("Conversion not implemented." ) );
+			break;
+		}
+
+	}
+
+}
diff --git a/Engine/Plugins/Experimental/Mutable/Source/MutableRuntime/Private/Image.cpp b/Engine/Plugins/Experimental/Mutable/Source/MutableRuntime/Private/Image.cpp
new file mode 100644
--- /dev/null
+++ b/Engine/Plugins/Experimental/Mutable/Source/MutableRuntime/Private/Image.cpp
@@ -0,0 +1,935 @@
+// Copyright Epic Games, Inc. All Rights Reserved.
+
+#include "Image.h"
+#include "ImagePrivate.h"
+#include "MutableTrace.h"
+#include "MutableRuntime/Private/OpImageResize.h"
+
+#include "Platform.h"
+
+namespace mu
+{
+
+    //---------------------------------------------------------------------------------------------
+    static FImageFormatData s_imageFormatData[uint32(EImageFormat::IF_COUNT)] =
+    {
+		FImageFormatData( 0, 0, 0, 0 ),	// IF_NONE
+		FImageFormatData( 1, 1, 3, 3 ),	// IF_RGB_UBYTE
+		FImageFormatData( 1, 1, 4, 4 ),	// IF_RGBA_UBYTE
+		FImageFormatData( 1, 1, 1, 1 ),	// IF_U_UBYTE
+
+		FImageFormatData( 0, 0, 0, 0 ),	// IF_PVRTC2 (deprecated)
+        FImageFormatData( 0, 0, 0, 0 ),	// IF_PVRTC4 (deprecated)
+        FImageFormatData( 0, 0, 0, 0 ),	// IF_ETC1 (deprecated)
+        FImageFormatData( 0, 0, 0, 0 ),	// IF_ETC2 (deprecated)
+
+        FImageFormatData( 0, 0, 0, 1 ),	// IF_L_UBYTE_RLE
+        FImageFormatData( 0, 0, 0, 3 ),	// IF_RGB_UBYTE_RLE
+        FImageFormatData( 0, 0, 0, 4 ),	// IF_RGBA_UBYTE_RLE
+        FImageFormatData( 0, 0, 0, 1 ),	// IF_L_UBIT_RLE
+
+        FImageFormatData( 4, 4, 8,  4 ),	// IF_BC1
+        FImageFormatData( 4, 4, 16, 4 ),	// IF_BC2
+        FImageFormatData( 4, 4, 16, 4 ),	// IF_BC3
+        FImageFormatData( 4, 4, 8,  1 ),	// IF_BC4
+        FImageFormatData( 4, 4, 16, 2 ),	// IF_BC5
+        FImageFormatData( 4, 4, 16, 3 ),	// IF_BC6
+        FImageFormatData( 4, 4, 16, 4 ),	// IF_BC7
+
+        FImageFormatData( 1, 1, 4, 4 ),	// IF_BGRA_UBYTE
+
+        FImageFormatData( 4, 4, 16, 3 ),	// IF_ASTC_4x4_RGB_LDR
+        FImageFormatData( 4, 4, 16, 4 ),	// IF_ASTC_4x4_RGBA_LDR
+        FImageFormatData( 4, 4, 16, 2 ),	// IF_ASTC_4x4_RG_LDR
+    };
+
+
+    //---------------------------------------------------------------------------------------------
+    const FImageFormatData& GetImageFormatData( EImageFormat format )
+    {
+        check( uint8(format) < int8(EImageFormat::IF_COUNT) );
+        return s_imageFormatData[ uint8(format) ];
+    }
+
+
+    //---------------------------------------------------------------------------------------------
+    Image::Image()
+    {
+    }
+
+
+    //---------------------------------------------------------------------------------------------
+    Image::Image( uint32_t sizeX, uint32_t sizeY, uint32_t lods, EImageFormat format )
+    {
+		MUTABLE_CPUPROFILER_SCOPE(NewImage)
+
+        check( format!= EImageFormat::IF_NONE );
+		LLM_SCOPE_BYNAME(TEXT("MutableRuntime"));
+
+        m_format = format;
+        m_size = FImageSize( (uint16_t)sizeX, (uint16_t)sizeY );
+        m_lods = (uint8_t)lods;
+        m_internalId = 0;
+
+        const FImageFormatData& fdata = GetImageFormatData( format );
+        int pixelsPerBlock = fdata.m_pixelsPerBlockX*fdata.m_pixelsPerBlockY;
+        if ( pixelsPerBlock )
+        {
+			int32 DataSize = CalculateDataSize();
+            m_data.SetNum( DataSize );
+        }
+    }
+
+
+    //---------------------------------------------------------------------------------------------
+    void Image::Serialise( const Image* p, OutputArchive& arch )
+    {
+        arch << *p;
+    }
+
+
+    //---------------------------------------------------------------------------------------------
+    ImagePtr Image::StaticUnserialise( InputArchive& arch )
+    {
+		LLM_SCOPE_BYNAME(TEXT("MutableRuntime"));
+        ImagePtr pResult = new Image();
+        arch >> *pResult;
+        return pResult;
+    }
+
+
+	//---------------------------------------------------------------------------------------------
+	ImagePtr Image::ExtractMip(int32 Mip) const
+	{
+		if (Mip == 0 && m_lods == 1)
+		{
+			return Clone();
+		}
+
+		vec2<int> MipSize = CalculateMipSize(Mip);
+
+		int32 Quality = 4;
+
+		if (m_lods > Mip)
+		{
+			const uint8* SourceData = GetMipData(Mip);
+			int32 DataSize = CalculateDataSize(Mip);
+
+			if (DataSize)
+			{
+				ImagePtr pResult = new Image(MipSize[0], MipSize[1], 1, m_format );
+				uint8* DestData = pResult->GetData();
+				FMemory::Memcpy(DestData, SourceData, DataSize);
+				return pResult;
+			}
+			else
+			{
+				EImageFormat uncompressedFormat = GetUncompressedFormat(m_format);
+				ImagePtr pTemp = ImagePixelFormat(Quality, this, uncompressedFormat);
+				pTemp = pTemp->ExtractMip(Mip);
+				ImagePtr pResult = ImagePixelFormat(Quality, pTemp.get(), m_format);
+				return pResult;
+			}
+		}
+
+		// We need to generate the mip
+		// \TODO: optimize, quality
+		return ImageResizeLinear(Quality, this, FImageSize(MipSize[0],MipSize[1]))->ExtractMip(0);
+	}
+
+
+    //---------------------------------------------------------------------------------------------
+    uint16_t Image::GetSizeX() const
+    {
+        return m_size[0];
+    }
+
+
+    //---------------------------------------------------------------------------------------------
+    uint16_t Image::GetSizeY() const
+    {
+        return m_size[1];
+    }
+
+
+	//---------------------------------------------------------------------------------------------
+	const FImageSize& Image::GetSize() const
+	{
+		return m_size;
+	}
+
+
+    //---------------------------------------------------------------------------------------------
+	EImageFormat Image::GetFormat() const
+    {
+        return m_format;
+    }
+
+
+    //---------------------------------------------------------------------------------------------
+    int Image::GetLODCount() const
+    {
+        return FMath::Max( 1, (int)m_lods );
+    }
+
+
+    //---------------------------------------------------------------------------------------------
+    const uint8_t* Image::GetData() const
+    {
+        return m_data.GetData();
+    }
+
+
+    //---------------------------------------------------------------------------------------------
+    int32_t Image::GetDataSize() const
+    {
+        return m_data.Num();
+    }
+
+
+    //---------------------------------------------------------------------------------------------
+    int32_t Image::GetLODDataSize( int lod ) const
+    {
+        return CalculateDataSize( lod );
+    }
+
+
+    //---------------------------------------------------------------------------------------------
+    uint8_t* Image::GetData()
+    {
+        uint8_t* pResult = 0;
+        if ( m_data.Num() )
+        {
+            pResult = &m_data[0];
+        }
+        return pResult;
+    }
+
+
+    //---------------------------------------------------------------------------------------------
+    uint32_t Image::GetId() const
+    {
+        return m_internalId;
+    }
+
+
+    //---------------------------------------------------------------------------------------------
+    //---------------------------------------------------------------------------------------------
+    //---------------------------------------------------------------------------------------------
+	int32 Image::CalculateDataSize() const
+    {
+		int32 res = 0;
+
+        const FImageFormatData& fdata = GetImageFormatData( m_format );
+        if (fdata.m_bytesPerBlock)
+        {
+            FImageSize s = m_size;
+
+            for ( int l=0; l<FMath::Max(1,(int)m_lods); ++l )
+            {
+                int blocksX = FMath::DivideAndRoundUp( s[0], (uint16_t)fdata.m_pixelsPerBlockX );
+                int blocksY = FMath::DivideAndRoundUp( s[1], (uint16_t)fdata.m_pixelsPerBlockY );
+
+                res += (blocksX*blocksY) * fdata.m_bytesPerBlock;
+
+                s[0] = FMath::DivideAndRoundUp( s[0], (uint16_t)2 );
+                s[1] = FMath::DivideAndRoundUp( s[1], (uint16_t)2 );
+            }
+        }
+
+        return res;
+    }
+
+
+    //---------------------------------------------------------------------------------------------
+	int32 Image::CalculateDataSize( int mip ) const
+    {
+		int32 res = 0;
+
+        const FImageFormatData& fdata = GetImageFormatData( m_format );
+        if (fdata.m_bytesPerBlock)
+        {
+            FImageSize s = m_size;
+
+            for ( int l=0; l< FMath::Max(1,(int)m_lods); ++l )
+            {
+                int blocksX = FMath::DivideAndRoundUp( s[0], (uint16_t)fdata.m_pixelsPerBlockX );
+                int blocksY = FMath::DivideAndRoundUp( s[1], (uint16_t)fdata.m_pixelsPerBlockY );
+
+                if ( mip==l )
+                {
+                    res = (blocksX*blocksY) * fdata.m_bytesPerBlock;
+                    break;
+                }
+
+                s[0] = FMath::DivideAndRoundUp( s[0], (uint16_t)2 );
+                s[1] = FMath::DivideAndRoundUp( s[1], (uint16_t)2 );
+            }
+        }
+
+        return res;
+    }
+
+
+    //---------------------------------------------------------------------------------------------
+    int32 Image::CalculatePixelCount() const
+    {
+		int32 res = 0;
+
+        const FImageFormatData& fdata = GetImageFormatData( m_format );
+        if (fdata.m_bytesPerBlock)
+        {
+            FImageSize s = m_size;
+
+            for ( int l=0; l<FMath::Max(1,(int)m_lods); ++l )
+            {
+                int blocksX = FMath::DivideAndRoundUp( s[0], (uint16_t)fdata.m_pixelsPerBlockX );
+                int blocksY = FMath::DivideAndRoundUp( s[1], (uint16_t)fdata.m_pixelsPerBlockY );
+
+                res += (blocksX*blocksY) * fdata.m_pixelsPerBlockX * fdata.m_pixelsPerBlockY;
+
+                s[0] = FMath::DivideAndRoundUp( s[0], (uint16_t)2 );
+                s[1] = FMath::DivideAndRoundUp( s[1], (uint16_t)2 );
+            }
+        }
+        else
+        {
+            // An RLE image.
+            for ( int l=0; l<FMath::Max(1,(int)m_lods); ++l )
+            {
+                auto mipSize = CalculateMipSize(l);
+                res += mipSize[0] * mipSize[1];
+            }
+        }
+
+        return res;
+    }
+
+
+    //---------------------------------------------------------------------------------------------
+	int32 Image::CalculatePixelCount( int mip ) const
+    {
+		int32 res = 0;
+
+        const FImageFormatData& fdata = GetImageFormatData( m_format );
+        if (fdata.m_bytesPerBlock)
+        {
+            FImageSize s = m_size;
+
+            for ( int l=0; l<mip+1; ++l )
+            {
+                int blocksX = FMath::DivideAndRoundUp( s[0], (uint16_t)fdata.m_pixelsPerBlockX );
+                int blocksY = FMath::DivideAndRoundUp( s[1], (uint16_t)fdata.m_pixelsPerBlockY );
+
+                if ( l==mip )
+                {
+                    res = (blocksX*blocksY) * fdata.m_pixelsPerBlockX * fdata.m_pixelsPerBlockY;
+                }
+
+                s[0] = FMath::DivideAndRoundUp( s[0], (uint16_t)2 );
+                s[1] = FMath::DivideAndRoundUp( s[1], (uint16_t)2 );
+            }
+        }
+        else
+        {
+            // An RLE image.
+            auto mipSize = CalculateMipSize(mip);
+            res += mipSize[0] * mipSize[1];
+        }
+
+        return res;
+    }
+
+
+    //---------------------------------------------------------------------------------------------
+    vec2<int> Image::CalculateMipSize( int mip ) const
+    {
+        vec2<int> res;
+
+        FImageSize s = m_size;
+
+        for ( int l=0; l<mip+1; ++l )
+        {
+            if ( l==mip )
+            {
+                res[0] = s[0];
+                res[1] = s[1];
+                return res;
+            }
+
+            s[0] = FMath::DivideAndRoundUp( s[0], (uint16_t)2 );
+            s[1] = FMath::DivideAndRoundUp( s[1], (uint16_t)2 );
+        }
+
+        return res;
+    }
+
+	
+	//---------------------------------------------------------------------------------------------
+	int Image::GetMipmapCount(int SizeX, int SizeY)
+	{
+		if (SizeX <= 0 || SizeY <= 0)
+		{
+			return 0;
+		}
+
+		int MaxLevel = FMath::CeilLogTwo(FMath::Max(SizeX, SizeY)) + 1;
+		return MaxLevel;
+	}
+
+
+	//---------------------------------------------------------------------------------------------
+	const uint8_t* Image::GetMipData(int mip) const
+	{
+		check((mip >= 0 && mip < m_lods) || (m_lods == 0));
+
+		const uint8_t* pResult = m_data.GetData();
+
+		const FImageFormatData& fdata = GetImageFormatData(m_format);
+
+		if (fdata.m_bytesPerBlock)
+		{
+			// Fixed-sized formats
+			FImageSize s = m_size;
+
+			for (int l = 0; l < mip; ++l)
+			{
+				int blocksX = FMath::DivideAndRoundUp(s[0], (uint16_t)fdata.m_pixelsPerBlockX);
+				int blocksY = FMath::DivideAndRoundUp(s[1], (uint16_t)fdata.m_pixelsPerBlockY);
+
+				pResult += (blocksX * blocksY) * fdata.m_bytesPerBlock;
+
+				s[0] = FMath::DivideAndRoundUp(s[0], (uint16_t)2);
+				s[1] = FMath::DivideAndRoundUp(s[1], (uint16_t)2);
+			}
+		}
+
+		else if (m_format == EImageFormat::IF_L_UBYTE_RLE ||
+			m_format == EImageFormat::IF_L_UBIT_RLE)
+		{
+			// Every mip has variable size, but it is stored in the first 4 bytes.
+			for (int l = 0; l < mip; ++l)
+			{
+				uint32_t mipSize = *(const uint32_t*)pResult;
+				pResult += mipSize;
+			}
+		}
+		else
+		{
+			checkf(false, TEXT("Trying to get mipmap pointer in an unsupported pixel format."));
+		}
+
+		return pResult;
+	}
+
+
+	//---------------------------------------------------------------------------------------------
+	uint8_t* Image::GetMipData(int mip)
+	{
+		check((mip >= 0 && mip < m_lods) || (m_lods == 0));
+
+		uint8_t* pResult = m_data.GetData();
+
+		const FImageFormatData& fdata = GetImageFormatData(m_format);
+
+		if (fdata.m_bytesPerBlock)
+		{
+			// Fixed-sized formats
+			FImageSize s = m_size;
+
+			for (int l = 0; l < mip; ++l)
+			{
+				int blocksX = FMath::DivideAndRoundUp(s[0], (uint16_t)fdata.m_pixelsPerBlockX);
+				int blocksY = FMath::DivideAndRoundUp(s[1], (uint16_t)fdata.m_pixelsPerBlockY);
+
+				pResult += (blocksX * blocksY) * fdata.m_bytesPerBlock;
+
+				s[0] = FMath::DivideAndRoundUp(s[0], (uint16_t)2);
+				s[1] = FMath::DivideAndRoundUp(s[1], (uint16_t)2);
+			}
+		}
+
+		else if (m_format == EImageFormat::IF_L_UBYTE_RLE ||
+			m_format == EImageFormat::IF_L_UBIT_RLE)
+		{
+			// Every mip has variable size, but it is stored in the first 4 bytes.
+			for (int l = 0; l < mip; ++l)
+			{
+				uint32_t mipSize = *(const uint32_t*)pResult;
+				pResult += mipSize;
+			}
+		}
+		else
+		{
+			checkf(false, TEXT("Trying to get mipmap pointer in an unsupported pixel format."));
+		}
+
+		return pResult;
+	}
+
+
+    //---------------------------------------------------------------------------------------------
+    vec4<float> Image::Sample( vec2<float> coords ) const
+    {
+        vec4<float> result;
+
+        if (m_size[0]==0 || m_size[1]==0) { return result; }
+
+        const FImageFormatData& fdata = GetImageFormatData( m_format );
+
+        int pixelX = FMath::Max( 0, FMath::Min( m_size[0]-1, (int)(coords[0] * m_size[0] )));
+        int blockX = pixelX / fdata.m_pixelsPerBlockX;
+        int blockPixelX = pixelX % fdata.m_pixelsPerBlockX;
+
+        int pixelY = FMath::Max( 0, FMath::Min( m_size[1]-1, (int)(coords[1] * m_size[1] )));
+        int blockY = pixelY / fdata.m_pixelsPerBlockY;
+        int blockPixelY = pixelY % fdata.m_pixelsPerBlockY;
+
+        int blocksPerRow = m_size[0] / fdata.m_pixelsPerBlockX;
+        int blockOffset = blockX + blockY * blocksPerRow;
+
+        // Non-generic part
+        if ( m_format== EImageFormat::IF_RGB_UBYTE )
+        {
+            int byteOffset = blockOffset * fdata.m_bytesPerBlock
+                    + ( blockPixelY * fdata.m_pixelsPerBlockX + blockPixelX ) * 3;
+
+            result[0] = m_data[ byteOffset+0 ] / 255.0f;
+            result[1] = m_data[ byteOffset+1 ] / 255.0f;
+            result[2] = m_data[ byteOffset+2 ] / 255.0f;
+            result[3] = 1;
+        }
+        else if ( m_format== EImageFormat::IF_RGBA_UBYTE )
+        {
+            int byteOffset = blockOffset * fdata.m_bytesPerBlock
+                    + ( blockPixelY * fdata.m_pixelsPerBlockX + blockPixelX ) * 4;
+
+            result[0] = m_data[ byteOffset+0 ] / 255.0f;
+            result[1] = m_data[ byteOffset+1 ] / 255.0f;
+            result[2] = m_data[ byteOffset+2 ] / 255.0f;
+            result[3] = m_data[ byteOffset+3 ] / 255.0f;
+        }
+        else if ( m_format== EImageFormat::IF_BGRA_UBYTE )
+        {
+            int byteOffset = blockOffset * fdata.m_bytesPerBlock
+                    + ( blockPixelY * fdata.m_pixelsPerBlockX + blockPixelX ) * 4;
+
+            result[0] = m_data[ byteOffset+2 ] / 255.0f;
+            result[1] = m_data[ byteOffset+1 ] / 255.0f;
+            result[2] = m_data[ byteOffset+0 ] / 255.0f;
+            result[3] = m_data[ byteOffset+3 ] / 255.0f;
+        }
+        else if ( m_format== EImageFormat::IF_L_UBYTE )
+        {
+            int byteOffset = blockOffset * fdata.m_bytesPerBlock
+                    + ( blockPixelY * fdata.m_pixelsPerBlockX + blockPixelX ) * 1;
+
+            result[0] = m_data[ byteOffset ] / 255.0f;
+            result[1] = m_data[ byteOffset ] / 255.0f;
+            result[2] = m_data[ byteOffset ] / 255.0f;
+            result[3] = 1.0f;
+        }
+        else
+        {
+            check( false );
+        }
+
+        return result;
+    }
+
+    //---------------------------------------------------------------------------------------------
+    bool Image::IsPlainColour( vec4<float>& colour ) const
+    {
+        bool res = true;
+
+        int pixelCount = m_size[0] * m_size[1];
+
+        if ( pixelCount )
+        {
+            switch( m_format )
+            {
+            case EImageFormat::IF_L_UBYTE:
+            {
+                const uint8_t* pData = &m_data[0];
+                uint8_t v = pData[0];
+
+                for ( int p=0; res && p<pixelCount; ++p )
+                {
+                    uint8_t nv = *pData++;
+                    res &= ( nv==v );
+                }
+
+                if (res)
+                {
+                    colour[0] = colour[1] = colour[2] = float(v)/255.0f;
+                    colour[3] = 1.0f;
+                }
+                break;
+            }
+
+            default:
+                res = false;
+                break;
+            }
+        }
+
+        return res;
+    }
+
+
+    //---------------------------------------------------------------------------------------------
+    bool Image::IsFullAlpha() const
+    {
+        int pixelCount = m_size[0] * m_size[1];
+
+        if ( pixelCount )
+        {
+            switch( m_format )
+            {
+            case EImageFormat::IF_RGBA_UBYTE:
+            case EImageFormat::IF_BGRA_UBYTE:
+            {
+                const uint8_t* pData = &m_data[3];
+
+                for ( int p=0; p<pixelCount; ++p )
+                {
+                    uint8_t nv = *pData;
+                    if (nv!=255) return false;
+                    pData+=4;
+                }
+                return true;
+                break;
+            }
+
+            case EImageFormat::IF_RGB_UBYTE:
+            {
+                return true;
+                break;
+            }
+
+            default:
+                return false;
+            }
+        }
+
+        return true;
+    }
+
+
+    //---------------------------------------------------------------------------------------------
+    namespace
+    {
+        bool is_zero( const uint8_t* buff, size_t size )
+        {
+            return (*buff==0) && (FMemory::Memcmp( buff, buff + 1, size - 1 )==0);
+        }
+    }
+
+
+    //---------------------------------------------------------------------------------------------
+     void Image::GetNonBlackRect_Reference(FImageRect& rect) const
+     {            
+         rect.min[0] = rect.min[1] = 0;
+         rect.size = m_size;
+
+         if ( !rect.size[0] || !rect.size[1] )
+             return;
+
+         bool first = true;
+         uint16_t sx = m_size[0];
+         uint16_t sy = m_size[1];
+
+         // Slow reference implementation
+         uint16_t top = 0;
+         uint16_t left = 0;
+         uint16_t right = sx - 1;
+         uint16_t bottom = sy - 1;
+
+         for ( uint16_t y=0; y<sy; ++y )
+         {
+             for ( uint16_t x=0; x<sx; ++x )
+             {
+                 bool black = false;
+                 switch( m_format )
+                 {
+                 case EImageFormat::IF_L_UBYTE:
+                 {
+                     uint8_t nv = m_data[y*sx+x];
+                     if (nv==0)
+                     {
+                         black = true;
+                     }
+                     break;
+                 }
+
+                 case EImageFormat::IF_RGB_UBYTE:
+                 {
+                     const uint8_t* pData = &m_data[(y*sx+x)*3];
+                     if (!pData[0] && !pData[1] && !pData[2])
+                     {
+                         black = true;
+                     }
+                     break;
+                 }
+
+                 case EImageFormat::IF_RGBA_UBYTE:
+                 case EImageFormat::IF_BGRA_UBYTE:
+                 {
+                     const uint8_t* pData = &m_data[(y*sx+x)*4];
+                     if (!pData[0] && !pData[1] && !pData[2] && !pData[3])
+                     {
+                         black = true;
+                     }
+                     break;
+                 }
+
+                 default:
+                     check(false);
+                     break;
+                 }
+
+                 if (!black)
+                 {
+                     if (first)
+                     {
+                         first = false;
+                         left = right = x;
+                         top = bottom = y;
+                     }
+                     else
+                     {
+                         left = FMath::Min( left, x );
+                         right = FMath::Max( right, x );
+                         top = FMath::Min( top, y );
+                         bottom = FMath::Max( bottom, y );
+                     }
+                 }
+             }
+         }
+
+         rect.min[0] = left;
+         rect.min[1] = top;
+         rect.size[0] = right - left + 1;
+         rect.size[1] = bottom - top + 1;
+     }
+
+
+    //---------------------------------------------------------------------------------------------
+    void Image::GetNonBlackRect(FImageRect& rect) const
+    {            
+		// TODO: There is a bug here with cyborg-windows.
+		// Meanwhile do this.
+		GetNonBlackRect_Reference(rect);
+		return;
+
+  //      rect.min[0] = rect.min[1] = 0;
+  //      rect.size = m_size;
+
+		//check(rect.size[0] > 0);
+		//check(rect.size[1] > 0);
+
+  //      if ( !rect.size[0] || !rect.size[1] )
+  //          return;
+
+  //      uint16_t sx = m_size[0];
+  //      uint16_t sy = m_size[1];
+
+  //      // Somewhat faster implementation
+  //      uint16_t top = 0;
+  //      uint16_t left = 0;
+  //      uint16_t right = sx - 1;
+  //      uint16_t bottom = sy - 1;
+
+  //      switch ( m_format )
+  //      {
+  //      case IF_L_UBYTE:
+  //      {
+  //          size_t rowStride = sx;
+
+  //          // Find top
+  //          const uint8_t* pRow = m_data.GetData();
+  //          while ( top < sy )
+  //          {
+  //              if ( !is_zero( pRow, rowStride ) )
+  //                  break;
+  //              pRow += rowStride;
+  //              ++top;
+  //          }
+
+  //          // Find bottom
+  //          pRow = m_data.GetData() + rowStride * bottom;
+  //          while ( bottom > top )
+  //          {
+  //              if ( !is_zero( pRow, rowStride ) )
+  //                  break;
+  //              pRow -= rowStride;
+  //              --bottom;
+  //          }
+
+  //          // Find left and right
+  //          left = sx - 1;
+  //          right = 0;
+  //          int16_t currentRow = top;
+  //          pRow = m_data.GetData() + rowStride * currentRow;
+  //          while ( currentRow <= bottom && ( left > 0 || right < ( sx - 1 ) ) )
+  //          {
+  //              for ( uint16_t x = 0; x < left; ++x )
+  //              {
+  //                  if ( pRow[x] )
+  //                  {
+  //                      left = x;
+  //                      break;
+  //                  }
+  //              }
+  //              for ( uint16_t x = sx - 1; x > right; --x )
+  //              {
+  //                  if ( pRow[x] )
+  //                  {
+  //                      right = x;
+  //                      break;
+  //                  }
+  //              }
+  //              pRow += rowStride;
+  //              ++currentRow;
+  //          }
+
+  //          break;
+  //      }
+
+  //      case IF_RGB_UBYTE:
+  //      case IF_RGBA_UBYTE:
+  //      case IF_BGRA_UBYTE:
+  //      {
+  //          size_t bytesPerPixel = GetImageFormatData( m_format ).m_bytesPerBlock;
+  //          size_t rowStride = sx * bytesPerPixel;
+
+  //          // Find top
+  //          const uint8_t* pRow = m_data.GetData();
+  //          while(top<sy)
+  //          {
+  //              if ( !is_zero( pRow, rowStride ) )
+  //                  break;
+  //              pRow += rowStride;
+  //              ++top;
+  //          }
+
+  //          // Find bottom
+  //          pRow = m_data.GetData() + rowStride * bottom;
+  //          while ( bottom > top )
+  //          {
+  //              if ( !is_zero( pRow, rowStride ) )
+  //                  break;
+  //              pRow -= rowStride;
+  //              --bottom;
+  //          }
+
+  //          // Find left and right
+  //          left = sx - 1;
+  //          right = 0;
+  //          int16_t currentRow = top;
+  //          pRow = m_data.GetData() + rowStride * currentRow;
+  //          uint8_t zeroPixel[16] = { 0 };
+  //          check(bytesPerPixel<16);
+  //          while ( currentRow <= bottom && (left > 0 || right<(sx-1)) )
+  //          {
+  //              for ( uint16_t x = 0; x < left; ++x )
+  //              {
+  //                  if ( memcmp( pRow + x * bytesPerPixel, zeroPixel, bytesPerPixel ) )
+  //                  {
+  //                      left = x;
+  //                      break;
+  //                  }
+  //              }
+  //              for ( uint16_t x = sx - 1; x > right; --x )
+  //              {
+  //                  if ( memcmp( pRow + x * bytesPerPixel, zeroPixel, bytesPerPixel ) )
+  //                  {
+  //                      right = x;
+  //                      break;
+  //                  }
+  //              }
+  //              pRow += rowStride;
+  //              ++currentRow;
+  //          }
+
+  //          break;
+  //      }
+
+  //      default:
+  //          check(false);
+  //          break;
+  //      }
+
+  //      rect.min[0] = left;
+  //      rect.min[1] = top;
+  //      rect.size[0] = right - left + 1;
+  //      rect.size[1] = bottom - top + 1;
+
+  //      // debug
+  //       FImageRect debugRect;
+  //       GetNonBlackRect_Reference( debugRect );
+  //       if (!(rect==debugRect))
+  //       {
+  //           mu::Halt();
+  //       }
+    }
+
+
+	//---------------------------------------------------------------------------------------------
+	void Image::ReduceLODsTo(int32 NewLODCount)
+	{
+		bool bIsBlockBased = GetImageFormatData(m_format).m_bytesPerBlock > 0;
+		int32 MaxLODs = GetMipmapCount(m_size[0], m_size[1]);
+		int32 LODSToSkip = MaxLODs - NewLODCount;
+		if (LODSToSkip > 0 && bIsBlockBased)
+		{
+			check(LODSToSkip < m_lods);
+			
+			uint32 DataToSkip = 0;
+			for (int32 l = 0; l < LODSToSkip; ++l)
+			{
+				uint32 MipDataSize = CalculateDataSize(l);
+				check(MipDataSize > 0);
+				DataToSkip += MipDataSize;
+			}
+
+			vec2<int> FinalSize = CalculateMipSize(LODSToSkip);
+			m_size[0] = uint16_t(FinalSize[0]);
+			m_size[1] = uint16_t(FinalSize[1]);
+			m_lods = NewLODCount;
+
+			int32 FinalDataSize = m_data.Num() - int32(DataToSkip);
+			FMemory::Memmove(m_data.GetData(), m_data.GetData() + DataToSkip, FinalDataSize);
+		}
+	}
+
+
+	//---------------------------------------------------------------------------------------------
+	void Image::ReduceLODs(int32 LODSToSkip)
+	{
+		bool bIsBlockBased = GetImageFormatData(m_format).m_bytesPerBlock > 0;
+		if (LODSToSkip > 0 && bIsBlockBased)
+		{
+			check(LODSToSkip < m_lods);
+			
+			uint32 DataToSkip = 0;
+			for (int32 l = 0; l < LODSToSkip; ++l)
+			{
+				uint32 MipDataSize = CalculateDataSize(l);
+				check(MipDataSize>0);
+				DataToSkip += MipDataSize;
+			}
+
+			vec2<int> FinalSize = CalculateMipSize(LODSToSkip);
+			m_size[0] = uint16_t(FinalSize[0]);
+			m_size[1] = uint16_t(FinalSize[1]);
+			m_lods -= LODSToSkip;
+
+			int32 FinalDataSize = m_data.Num() - int32(DataToSkip);
+			FMemory::Memmove(m_data.GetData(), m_data.GetData() + DataToSkip, FinalDataSize);
+		}
+	}
+
+
+}
+
diff --git a/Engine/Plugins/Experimental/Mutable/Source/MutableRuntime/Private/ImagePrivate.h b/Engine/Plugins/Experimental/Mutable/Source/MutableRuntime/Private/ImagePrivate.h
new file mode 100644
--- /dev/null
+++ b/Engine/Plugins/Experimental/Mutable/Source/MutableRuntime/Private/ImagePrivate.h
@@ -0,0 +1,230 @@
+// Copyright Epic Games, Inc. All Rights Reserved.
+
+#pragma once
+
+#include "Image.h"
+
+#include "SerialisationPrivate.h"
+#include "MutableMath.h"
+#include "MemoryPrivate.h"
+
+
+namespace mu
+{
+
+	MUTABLE_DEFINE_ENUM_SERIALISABLE( EBlendType );
+	MUTABLE_DEFINE_ENUM_SERIALISABLE( EMipmapFilterType );
+	MUTABLE_DEFINE_ENUM_SERIALISABLE( ECompositeImageMode );
+
+
+    struct FImageFormatData
+	{
+		FImageFormatData
+			(
+				unsigned pixelsPerBlockX = 0,
+				unsigned pixelsPerBlockY = 0,
+				unsigned bytesPerBlock = 0,
+				unsigned channels = 0
+			)
+		{
+			m_pixelsPerBlockX = (uint8)pixelsPerBlockX;
+			m_pixelsPerBlockY = (uint8)pixelsPerBlockY;
+			m_bytesPerBlock = (uint16)bytesPerBlock;
+			m_channels = (uint16)channels;
+		}
+
+		//! For block based formats, size of the block size. For uncompressed formats it will
+		//! always be 1,1. For non-block-based compressed formats, it will be 0,0.
+        uint8 m_pixelsPerBlockX, m_pixelsPerBlockY;
+
+		//! Number of bytes used by every pixel block, if uncompressed or block-compressed format.
+		//! For non-block-compressed formats, it returns 0.
+        uint16 m_bytesPerBlock;
+
+		//! Channels in every pixel of the image.
+        uint16 m_channels;
+	};
+
+
+	struct FMipmapGenerationSettings
+	{
+		float m_sharpenFactor = 0.0f;
+		EMipmapFilterType m_filterType = EMipmapFilterType::MFT_SimpleAverage;
+		EAddressMode m_addressMode = EAddressMode::AM_NONE;
+		bool m_ditherMipmapAlpha = false;
+
+		void Serialise( OutputArchive& arch ) const
+		{
+			uint32 ver = 0;
+			arch << ver;
+
+			arch << m_sharpenFactor;
+			arch << m_filterType;
+			arch << m_ditherMipmapAlpha;
+		}
+
+		void Unserialise( InputArchive& arch )
+		{
+			uint32 ver = 0;
+			arch >> ver;
+			check(ver == 0);
+
+			arch >> m_sharpenFactor;
+			arch >> m_filterType;
+			arch >> m_ditherMipmapAlpha;
+		}
+	};
+
+
+	//---------------------------------------------------------------------------------------------
+	//!
+	//---------------------------------------------------------------------------------------------
+	MUTABLERUNTIME_API const FImageFormatData& GetImageFormatData(EImageFormat format );
+
+
+	//---------------------------------------------------------------------------------------------
+	//!
+	//---------------------------------------------------------------------------------------------
+	MUTABLERUNTIME_API inline EImageFormat GetUncompressedFormat(EImageFormat f )
+	{
+		EImageFormat r = f;
+
+		switch ( r )
+		{
+		case EImageFormat::IF_L_UBIT_RLE: r = EImageFormat::IF_L_UBYTE; break;
+		case EImageFormat::IF_L_UBYTE_RLE: r = EImageFormat::IF_L_UBYTE; break;
+		case EImageFormat::IF_RGB_UBYTE_RLE: r = EImageFormat::IF_RGB_UBYTE; break;
+        case EImageFormat::IF_RGBA_UBYTE_RLE: r = EImageFormat::IF_RGBA_UBYTE; break;
+        case EImageFormat::IF_BC1: r = EImageFormat::IF_RGBA_UBYTE; break;
+        case EImageFormat::IF_BC2: r = EImageFormat::IF_RGBA_UBYTE; break;
+        case EImageFormat::IF_BC3: r = EImageFormat::IF_RGBA_UBYTE; break;
+        case EImageFormat::IF_BC4: r = EImageFormat::IF_L_UBYTE; break;
+        case EImageFormat::IF_BC5: r = EImageFormat::IF_RGB_UBYTE; break;
+        case EImageFormat::IF_ASTC_4x4_RGB_LDR: r = EImageFormat::IF_RGB_UBYTE; break;
+        case EImageFormat::IF_ASTC_4x4_RGBA_LDR: r = EImageFormat::IF_RGBA_UBYTE; break;
+        case EImageFormat::IF_ASTC_4x4_RG_LDR: r = EImageFormat::IF_RGB_UBYTE; break;
+        default: break;
+		}
+
+		return r;
+	}
+
+    //---------------------------------------------------------------------------------------------
+    //!
+    //---------------------------------------------------------------------------------------------
+	inline EImageFormat GetMostGenericFormat(EImageFormat a, EImageFormat b)
+    {
+        if (a==b) return a;
+        if ( GetImageFormatData(a).m_channels>GetImageFormatData(b).m_channels ) return a;
+        if ( GetImageFormatData(b).m_channels>GetImageFormatData(a).m_channels ) return b;
+        if (a== EImageFormat::IF_BC2 || a== EImageFormat::IF_BC3 || a== EImageFormat::IF_ASTC_4x4_RGBA_LDR) return a;
+        if (b== EImageFormat::IF_BC2 || b== EImageFormat::IF_BC3 || b== EImageFormat::IF_ASTC_4x4_RGBA_LDR) return b;
+
+        return a;
+    }
+
+
+    //---------------------------------------------------------------------------------------------
+    //!
+    //---------------------------------------------------------------------------------------------
+    inline bool IsCompressedFormat(EImageFormat f )
+    {
+        return f!=GetUncompressedFormat(f);
+    }
+
+
+    //---------------------------------------------------------------------------------------------
+    //!
+    //---------------------------------------------------------------------------------------------
+    inline void FillPlainColourImage( Image* pImage, vec4<float> c )
+    {
+        int pixelCount = pImage->GetSizeX() * pImage->GetSizeY();
+
+        switch (pImage->GetFormat())
+        {
+		case EImageFormat::IF_RGB_UBYTE:
+        {
+            uint8* pData = pImage->GetData();
+            uint8 r = uint8(FMath::Clamp(255.0f * c[0], 0.0f, 255.0f));
+            uint8 g = uint8(FMath::Clamp(255.0f * c[1], 0.0f, 255.0f));
+            uint8 b = uint8(FMath::Clamp(255.0f * c[2], 0.0f, 255.0f));
+            for ( int p = 0; p < pixelCount; ++p )
+            {
+                pData[0] = r;
+                pData[1] = g;
+                pData[2] = b;
+                pData += 3;
+            }
+            break;
+        }
+
+        case EImageFormat::IF_RGBA_UBYTE:
+        {
+            uint8* pData = pImage->GetData();
+			uint8 r = uint8(FMath::Clamp(255.0f * c[0], 0.0f, 255.0f));
+			uint8 g = uint8(FMath::Clamp(255.0f * c[1], 0.0f, 255.0f));
+			uint8 b = uint8(FMath::Clamp(255.0f * c[2], 0.0f, 255.0f));
+			for ( int p = 0; p < pixelCount; ++p )
+            {
+                pData[0] = r;
+                pData[1] = g;
+                pData[2] = b;
+                pData[3] = 255;
+                pData += 4;
+            }
+            break;
+        }
+
+        case EImageFormat::IF_BGRA_UBYTE:
+        {
+            uint8* pData = pImage->GetData();
+			uint8 r = uint8(FMath::Clamp(255.0f * c[0], 0.0f, 255.0f));
+			uint8 g = uint8(FMath::Clamp(255.0f * c[1], 0.0f, 255.0f));
+			uint8 b = uint8(FMath::Clamp(255.0f * c[2], 0.0f, 255.0f));
+			for ( int p = 0; p < pixelCount; ++p )
+            {
+                pData[0] = b;
+                pData[1] = g;
+                pData[2] = r;
+                pData[3] = 255;
+                pData += 4;
+            }
+            break;
+        }
+
+        case EImageFormat::IF_L_UBYTE:
+        {
+            uint8* pData = pImage->GetData();
+            uint8 v = std::min<uint8>( 255, std::max<uint8>( 0, uint8( 255.0f * c[0] ) ) );
+            FMemory::Memset( pData, v, pixelCount );
+            break;
+        }
+
+        default:
+            check(false);
+        }
+
+    }
+
+
+    //---------------------------------------------------------------------------------------------
+    //! Use with care.
+    //---------------------------------------------------------------------------------------------
+    template<class T>
+    Ptr<T> CloneOrTakeOver( const T* source )
+    {
+        Ptr<T> result;
+        if (source->GetRefCount()==1)
+        {
+            result = const_cast<T*>(source);
+        }
+        else
+        {
+            result = source->Clone();
+        }
+
+        return result;
+    }
+
+}
+
diff --git a/Engine/Plugins/Experimental/Mutable/Source/MutableRuntime/Private/ImageRLE.cpp b/Engine/Plugins/Experimental/Mutable/Source/MutableRuntime/Private/ImageRLE.cpp
new file mode 100644
--- /dev/null
+++ b/Engine/Plugins/Experimental/Mutable/Source/MutableRuntime/Private/ImageRLE.cpp
@@ -0,0 +1,545 @@
+// Copyright Epic Games, Inc. All Rights Reserved.
+
+#include "ImageRLE.h"
+#include "Platform.h"
+
+namespace mu
+{
+
+    //---------------------------------------------------------------------------------------------
+    uint32_t CompressRLE_L( int32 width, int32 rows, const uint8_t* pBaseData, uint8_t* destData, uint32_t destDataSize )
+    {
+		const uint8_t* InitialBaseData = pBaseData;
+
+        uint8_t* rle = destData;
+
+        size_t maxSize = destDataSize;
+
+        // The first uint32_t will be the total mip size.
+        // Then there is an offset from the initial data pointer for each line.
+        uint32_t offset = sizeof( uint32_t ) * ( rows + 1 );
+
+        // Could happen in an image of 1x100, for example.
+        if ( offset >= maxSize )
+        {
+            return 0;
+        }
+
+        for ( int r = 0; r < rows; ++r )
+        {
+            uint32_t* pOffset = (uint32_t*)&rle[sizeof( uint32_t ) * ( r + 1 )];
+            FMemory::Memmove( pOffset, &offset, sizeof( uint32_t ) );
+
+            const uint8_t* pBaseRowEnd = pBaseData + width;
+            while ( pBaseData != pBaseRowEnd )
+            {
+                // Count equal pixels
+                uint8_t equalPixel = *pBaseData;
+
+                uint16_t equal = 0;
+                while ( pBaseData != pBaseRowEnd && equal < 65535 && pBaseData[0] == equalPixel )
+                {
+                    pBaseData++;
+                    equal++;
+                }
+
+                // Count different pixels
+                uint8_t different = 0;
+                const uint8_t* pDifferentPixels = pBaseData;
+                while ( pBaseData < pBaseRowEnd - 1 && different < 255 &&
+                        // Last in the row, or different from next
+                        ( pBaseData == pBaseRowEnd || pBaseData[0] != pBaseData[1] ) )
+                {
+                    pBaseData++;
+                    different++;
+                }
+
+                // Copy header
+                if ( maxSize < offset + 4 )
+                {
+                    return 0;
+                }
+                FMemory::Memmove( &rle[offset], &equal, sizeof( uint16_t ) );
+                offset += 2;
+                rle[offset] = different;
+                ++offset;
+
+                // Copy the equal pixel
+                rle[offset] = equalPixel;
+                ++offset;
+
+                // Copy the different pixels
+                if ( different )
+                {
+                    if ( maxSize < offset + different )
+                    {
+                        return 0;
+                    }
+                    FMemory::Memmove( &rle[offset], pDifferentPixels, different );
+                    offset += different;
+                }
+            }
+        }
+
+        uint32_t* pTotalSize = (uint32_t*)rle;
+        *pTotalSize = offset;
+
+#ifdef MUTABLE_DEBUG_RLE		
+		{
+			TArray<uint8> Temp;
+			Temp.SetNum(width*rows);
+			UncompressRLE_L(width,rows,destData,Temp.GetData());
+			int Difference = FMemory::Memcmp(Temp.GetData(), InitialBaseData, width * rows);
+			if (Difference)
+			{
+				// Different pos.
+				size_t Delta = 0;
+				for ( ; Delta<width*rows; ++Delta )
+				{
+					if (Temp[Delta] != InitialBaseData[Delta])
+					{
+						break;
+					}
+				}
+
+				UncompressRLE_L(width, rows, destData, Temp.GetData());
+				CompressRLE_L( width, rows, InitialBaseData, destData, destDataSize);
+			}
+		}
+#endif
+
+        // succeded
+        return offset;
+    }
+
+
+    //---------------------------------------------------------------------------------------------
+    uint32_t UncompressRLE_L( int32 width, int32 rows, const uint8* pStartBaseData, uint8* pStartDestData )
+    {
+		const uint8* pBaseData = pStartBaseData;
+		uint8* pDestData = pStartDestData;
+		pBaseData += sizeof(uint32); // Total mip size
+        pBaseData += rows*sizeof(uint32); // Size of each row.
+
+        for ( int r=0; r<rows; ++r )
+        {
+            const uint8* pDestRowEnd = pDestData + width;
+            while ( pDestData!=pDestRowEnd )
+            {
+                // Decode header
+                uint16 equal = 0;
+                FMemory::Memmove(&equal, pBaseData, sizeof(uint16));
+                pBaseData += 2;
+
+                uint8_t different = *pBaseData;
+                ++pBaseData;
+
+                uint8_t equalPixel = *pBaseData;
+                ++pBaseData;
+
+                if (equal)
+                {
+					check(pDestData + equal <= pStartDestData + width * rows);
+					FMemory::Memset(pDestData, equalPixel, equal);
+                    pDestData += equal;
+                }
+
+                if (different)
+                {
+					check(pDestData + different <= pStartDestData + width * rows);
+					FMemory::Memmove( pDestData, pBaseData, different );
+                    pDestData += different;
+                    pBaseData += different;
+                }
+            }
+        }
+
+        size_t totalSize = pBaseData-pStartBaseData;
+        check( totalSize==*(uint32*)pStartBaseData );
+
+        return (uint32)totalSize;
+    }
+
+    //---------------------------------------------------------------------------------------------
+    uint32_t CompressRLE_L1( int32 width, int32 rows,
+                             const uint8_t* pBaseData,
+                             uint8_t* destData,
+                             uint32_t destDataSize )
+    {
+        vector<int8_t> rle;
+        rle.reserve(  (width * rows) / 2 );
+
+        uint32_t offset = sizeof(uint32_t)*(rows+1);
+        rle.resize( offset );
+
+        for ( int r=0; r<rows; ++r )
+        {
+            uint32_t* pOffset = (uint32_t*) &rle[ sizeof(uint32_t)*(r+1) ];
+            *pOffset = offset;
+
+            const uint8_t* pBaseRowEnd = pBaseData + width;
+            while ( pBaseData!=pBaseRowEnd )
+            {
+                // Count 0 pixels
+                uint16_t zeroPixels = 0;
+                while ( pBaseData!=pBaseRowEnd
+                        && !*pBaseData )
+                {
+                    pBaseData++;
+                    zeroPixels++;
+                }
+
+                // Count 1 pixels
+                uint16_t onePixels = 0;
+                while ( pBaseData!=pBaseRowEnd
+                        && *pBaseData )
+                {
+                    pBaseData++;
+                    onePixels++;
+                }
+
+                // Copy block
+                rle.resize( rle.size()+4 );
+                FMemory::Memmove(&rle[ offset ], &zeroPixels, sizeof(uint16_t));
+                offset += 2;
+
+                FMemory::Memmove(&rle[ offset ], &onePixels, sizeof(uint16_t));
+                offset += 2;
+            }
+        }
+
+        if (destDataSize<rle.size())
+        {
+            // Failed
+            return 0;
+        }
+
+        uint32_t* pTotalSize = (uint32_t*) &rle[ 0 ];
+        *pTotalSize = offset;
+
+        FMemory::Memmove( destData, &rle[0], offset );
+
+        // succeded
+        return offset;
+    }
+
+
+    //---------------------------------------------------------------------------------------------
+    uint32_t UncompressRLE_L1( int width, int rows, const uint8_t* pStartBaseData, uint8_t* pDestData )
+    {
+        const uint8_t* pBaseData = pStartBaseData;
+        pBaseData += sizeof(uint32_t); // Total mip size
+        pBaseData += rows*sizeof(uint32_t);
+
+        for ( int r=0; r<rows; ++r )
+        {
+            const uint8_t* pDestRowEnd = pDestData + width;
+            while ( pDestData!=pDestRowEnd )
+            {
+                // Decode header
+                uint16_t zeroPixels = 0;
+                FMemory::Memmove(&zeroPixels, pBaseData, sizeof(uint16_t));
+                pBaseData += 2;
+
+                uint16_t onePixels = 0;
+                FMemory::Memmove(&onePixels, pBaseData, sizeof(uint16_t));
+                pBaseData += 2;
+
+                if (zeroPixels)
+                {
+                    FMemory::Memzero( pDestData, zeroPixels );
+                    pDestData += zeroPixels;
+                }
+
+                if (onePixels)
+                {
+                    FMemory::Memset( pDestData, 255, onePixels );
+                    pDestData += onePixels;
+                }
+            }
+        }
+
+        size_t totalSize = pBaseData-pStartBaseData;
+        check( totalSize==*(uint32_t*)pStartBaseData );
+
+        return (uint32_t)totalSize;
+    }
+
+
+    //---------------------------------------------------------------------------------------------
+    void CompressRLE_RGBA( int width, int rows,
+                           const uint8_t* pBaseDataByte,
+                           TArray<uint8>& destData )
+    {
+        // TODO: Support for compression from compressed data size, like L_RLE formats.
+        vector<int8_t> rle;
+        rle.reserve(  (width*rows) );
+
+        const uint32_t* pBaseData = (const uint32_t*)pBaseDataByte;
+        rle.resize( rows*4 );
+        uint32_t offset = sizeof(uint32_t)*rows;
+        for ( int r=0; r<rows; ++r )
+        {
+            uint32_t* pOffset = (uint32_t*) &rle[ sizeof(uint32_t)*r ];
+            FMemory::Memmove(pOffset, &offset, sizeof(uint32_t));
+
+            const uint32_t* pBaseRowEnd = pBaseData + width;
+            while ( pBaseData!=pBaseRowEnd )
+            {
+                // Count equal pixels
+                uint32_t equalPixel = *pBaseData;
+
+                uint16_t equal = 0;
+                while ( pBaseData<pBaseRowEnd-3 && equal<65535
+                        && pBaseData[0]==equalPixel && pBaseData[1]==equalPixel
+                        && pBaseData[2]==equalPixel && pBaseData[3]==equalPixel )
+                {
+                    pBaseData+=4;
+                    equal++;
+                }
+
+                // Count different pixels
+                uint16_t different = 0;
+                const uint32_t* pDifferentPixels = pBaseData;
+                while ( pBaseData!=pBaseRowEnd
+                        &&
+                        different<65535
+                        &&
+                        // Last in the row, or different from next
+                        ( pBaseData>pBaseRowEnd-4
+                          || pBaseData[0]!=pBaseData[1]
+                          || pBaseData[0]!=pBaseData[2]
+                          || pBaseData[0]!=pBaseData[3]
+                          )
+                        )
+                {
+					pBaseData += FMath::Min(int64(4), int64(pBaseRowEnd - pBaseData));
+                    different++;
+                }
+
+                // Copy header
+                rle.resize( rle.size()+8 );
+                FMemory::Memmove(&rle[ offset ], &equal, sizeof(uint16_t));
+                offset += 2;
+                FMemory::Memmove(&rle[ offset ], &different, sizeof(uint16_t));
+                offset += 2;
+                FMemory::Memmove(&rle[ offset ], &equalPixel, sizeof(uint32_t));
+                offset += 4;
+
+                // Copy the different pixels
+				if (different)
+				{
+					// If we are at the end of a row, maybe there isn't a block of 4 pixels
+					uint16_t BytesToCopy = FMath::Min(different * 4 * 4, uint16_t(pBaseRowEnd - pDifferentPixels) * 4);
+
+					rle.resize( rle.size()+ BytesToCopy);
+                    FMemory::Memmove( &rle[offset], pDifferentPixels, BytesToCopy);
+					offset += BytesToCopy;
+				}
+            }
+        }
+
+        destData.SetNum( offset );
+        if ( offset )
+        {
+            FMemory::Memmove( &destData[0], &rle[0], offset );
+        }
+    }
+
+
+    //---------------------------------------------------------------------------------------------
+    void UncompressRLE_RGBA( int width, int rows, const uint8_t* pBaseData, uint8_t* pDestDataB )
+    {
+        uint32_t* pDestData = reinterpret_cast<uint32_t*>( pDestDataB );
+
+        pBaseData += rows*sizeof(uint32_t);
+
+        int pendingPixels = width*rows;
+
+        for ( int r=0; r<rows; ++r )
+        {
+            const uint32_t* pDestRowEnd = pDestData + width;
+            while ( pDestData!=pDestRowEnd )
+            {
+                // Decode header
+                uint16_t equal = 0;
+                FMemory::Memmove(&equal, pBaseData, sizeof(uint16_t));
+                pBaseData += 2;
+
+                uint16_t different = 0;
+                FMemory::Memmove(&different, pBaseData, sizeof(uint16_t));
+                pBaseData += 2;
+
+                uint32_t equalPixel = 0;
+                FMemory::Memmove(&equalPixel, pBaseData, sizeof(uint32_t));
+                pBaseData += 4;
+
+                check((equal+different)*4<=pendingPixels);
+
+                for ( int e=0; e<equal*4; ++e )
+                {
+                    FMemory::Memmove( pDestData, &equalPixel, 4 );
+                    ++pDestData;
+                    pendingPixels--;
+                }
+
+				if (different)
+				{
+					// If we are at the end of a row, maybe there isn't a block of 4 pixels
+					uint16_t PixelsToCopy = FMath::Min(uint16_t(different * 4), uint16_t(pDestRowEnd - pDestData));
+
+					FMemory::Memmove( pDestData, pBaseData, PixelsToCopy *4 );
+					pDestData += PixelsToCopy;
+					pBaseData += PixelsToCopy *4;
+                    pendingPixels-= PixelsToCopy;
+                }
+            }
+        }
+
+        check(pendingPixels==0);
+    }
+
+
+    //---------------------------------------------------------------------------------------------
+    struct UINT24
+    {
+        uint8_t d[3];
+
+        bool operator==( const UINT24& o ) const
+        {
+            return d[0]==o.d[0] && d[1]==o.d[1] && d[2]==o.d[2];
+        }
+
+        bool operator!=( const UINT24& o ) const
+        {
+            return d[0]!=o.d[0] || d[1]!=o.d[1] || d[2]!=o.d[2];
+        }
+    };
+    static_assert( sizeof(UINT24)==3, "Uint24SizeCheck" );
+
+
+    //---------------------------------------------------------------------------------------------
+    void CompressRLE_RGB( int width, int rows,
+                          const uint8_t* pBaseDataByte,
+                          TArray<uint8>& destData )
+    {
+        vector<int8_t> rle;
+        rle.reserve(  (width*rows) );
+
+        const UINT24* pBaseData = (const UINT24*)pBaseDataByte;
+        rle.resize( rows*4 );
+        uint32_t offset = sizeof(uint32_t)*rows;
+        for ( int r=0; r<rows; ++r )
+        {
+            uint32_t* pOffset = (uint32_t*) &rle[ sizeof(uint32_t)*r ];
+            *pOffset = offset;
+
+            const UINT24* pBaseRowEnd = pBaseData + width;
+            while ( pBaseData!=pBaseRowEnd )
+            {
+                // Count equal pixels
+                UINT24 equalPixel = *pBaseData;
+
+                uint16_t equal = 0;
+                while ( pBaseData<pBaseRowEnd-3 && equal<65535
+                        && pBaseData[0]==equalPixel && pBaseData[1]==equalPixel
+                        && pBaseData[2]==equalPixel && pBaseData[3]==equalPixel )
+                {
+                    pBaseData+=4;
+                    equal++;
+                }
+
+                // Count different pixels
+                uint16_t different = 0;
+                const UINT24* pDifferentPixels = pBaseData;
+                while ( pBaseData!=pBaseRowEnd
+                        &&
+                        different<65535
+                        &&
+                        // Last pixels in the row, or different from next
+                        ( pBaseData>pBaseRowEnd-4
+                          || pBaseData[0]!=pBaseData[1]
+                          || pBaseData[0]!=pBaseData[2]
+                          || pBaseData[0]!=pBaseData[3]
+                          )
+                        )
+                {
+                    pBaseData+=FMath::Min(int64(4), int64(pBaseRowEnd-pBaseData));
+                    different++;
+                }
+
+                // Copy header
+                rle.resize( rle.size()+8 );
+                FMemory::Memmove( &rle[offset], &equal, sizeof(uint16_t) );
+                offset += 2;
+                FMemory::Memmove( &rle[offset], &different, sizeof(uint16_t) );
+                offset += 2;
+                FMemory::Memmove( &rle[offset], &equalPixel, sizeof(UINT24) );
+                offset += 4;
+
+                // Copy the different pixels
+				if (different)
+				{
+					// If we are at the end of a row, maybe there isn't a block of 4 pixels
+					uint16_t BytesToCopy = FMath::Min(different * 4 * 3, uint16_t(pBaseRowEnd- pDifferentPixels)*3 );
+
+					rle.resize( rle.size()+BytesToCopy );
+                    FMemory::Memmove( &rle[offset], pDifferentPixels, BytesToCopy );
+					offset += BytesToCopy;
+				}
+            }
+        }
+
+        destData.SetNum( offset );
+        if ( offset )
+        {
+            FMemory::Memmove( &destData[0], &rle[0], offset );
+        }
+    }
+
+
+    //---------------------------------------------------------------------------------------------
+    void UncompressRLE_RGB( int width, int rows, const uint8_t* pBaseData, uint8_t* pDestDataB )
+    {
+        UINT24* pDestData = reinterpret_cast<UINT24*>( pDestDataB );
+
+        pBaseData += rows*sizeof(uint32_t);
+
+        for ( int r=0; r<rows; ++r )
+        {
+            const UINT24* pDestRowEnd = pDestData + width;
+            while ( pDestData!=pDestRowEnd )
+            {
+                // Decode header
+                uint16_t equal = *(const uint16_t*)pBaseData;
+                pBaseData += 2;
+
+                uint16_t different = *(const uint16_t*)pBaseData;
+                pBaseData += 2;
+
+                UINT24 equalPixel = *(const UINT24*)pBaseData;
+                pBaseData += 4;
+
+                for ( int e=0; e<equal*4; ++e )
+                {
+                    FMemory::Memmove( pDestData, &equalPixel, 3 );
+                    ++pDestData;
+                }
+
+				if (different)
+				{
+					// If we are at the end of a row, maybe there isn't a block of 4 pixels
+					uint16_t PixelsToCopy = FMath::Min(uint16_t(different * 4), uint16_t(pDestRowEnd - pDestData));
+
+					FMemory::Memmove( pDestData, pBaseData, PixelsToCopy*3);
+					pDestData += PixelsToCopy;
+					pBaseData += PixelsToCopy*3;
+				}
+            }
+        }
+    }
+
+
+
+}
+
