[mutable] Move Mutable plugin from the restricted folder to the experimental plugins folder.

diff --git a/Engine/Plugins/Experimental/Mutable/Source/MutableRuntime/Private/ImageRLE.h b/Engine/Plugins/Experimental/Mutable/Source/MutableRuntime/Private/ImageRLE.h
new file mode 100644
--- /dev/null
+++ b/Engine/Plugins/Experimental/Mutable/Source/MutableRuntime/Private/ImageRLE.h
@@ -0,0 +1,260 @@
+// Copyright Epic Games, Inc. All Rights Reserved.
+
+#pragma once
+
+#include "ImagePrivate.h"
+
+/** This define enabled additional checks when using RLE compression. These checks have a lot of
+* of verhead so it should usually be disabled.
+*/
+//#define MUTABLE_DEBUG_RLE
+
+
+namespace mu
+{
+
+    //---------------------------------------------------------------------------------------------
+    //! Compress to RLE if the compressed image fits in destData, which should be preallocated.
+    //! If it doesn't fit, it returns 0, and the content of destData is undefined.
+    //! This method doesn't allocate any memory.
+    //!
+    //! This RLE format compresses the data per row. At the beginning of the data there is the
+    //! total compressed data size (uint32_t) followed by the row offsets for direct access.
+    //! Every offset is a uint32_t starting at the begining of the compressed data.
+    //! Every row consists of a series of header + data. The headers is:
+    //!     - a UINT16 with how many equal pixels are there
+    //!     - a UINT8 of how many different pixels follow the equal pixels
+    //!		- the value of the "equal" pixels
+    //! After the header there are as many pixel values as different pixels.
+    //---------------------------------------------------------------------------------------------
+	MUTABLERUNTIME_API extern uint32_t CompressRLE_L( int width, int rows,
+                                   const uint8_t* pBaseData,
+                                   uint8_t* destData,
+                                   uint32_t destDataSize );
+
+    //---------------------------------------------------------------------------------------------
+	MUTABLERUNTIME_API extern uint32_t UncompressRLE_L( int width, int rows,
+                                     const uint8_t* pBaseData,
+                                     uint8_t* pDestData );
+
+    //---------------------------------------------------------------------------------------------
+    //! This RLE format compresses the data per row. At the beginning of the data there are the row
+    //! offsets for direct access.
+    //! Every row consists of a series of blocks. The blocks are:
+    //!     - a UINT16 with how many 0 pixels are there
+    //!     - a UINT16 with how many 1 pixels follow the 0 pixels
+    //---------------------------------------------------------------------------------------------
+	MUTABLERUNTIME_API extern uint32_t CompressRLE_L1( int width, int rows,
+                                    const uint8_t* pBaseData,
+                                    uint8_t* destData,
+                                    uint32_t destDataSize );
+
+    //---------------------------------------------------------------------------------------------
+	MUTABLERUNTIME_API extern uint32_t UncompressRLE_L1( int width, int rows,
+                                      const uint8_t* pBaseData,
+                                      uint8_t* pDestData );
+
+
+    //---------------------------------------------------------------------------------------------
+    //! This RLE format compresses the data per row. At the beginning of the data there are the row
+    //! offsets for direct access. Every offset is a uint32_t.
+    //! Every row consists of a series of header + data. The headers is:
+    //!     - a UINT16 with how many equal pixels / 4
+    //!     - a UINT16 with how many different pixels / 4
+    //!     - 4 UINT8 with the value of the "equal" pixels
+    //! After the header there are as many pixel values as different pixels.
+    //---------------------------------------------------------------------------------------------
+	MUTABLERUNTIME_API extern void CompressRLE_RGBA( int width, int rows,
+									const uint8_t* pBaseDataByte,
+									TArray<uint8_t>& destData );
+
+    //---------------------------------------------------------------------------------------------
+	MUTABLERUNTIME_API extern void UncompressRLE_RGBA( int width, int rows,
+                                    const uint8_t* pBaseData,
+                                    uint8_t* pDestDataB );
+
+
+    //---------------------------------------------------------------------------------------------
+    //! This RLE format compresses the data per row. At the beginning of the data there are the row
+    //! offsets for direct access. Every offset is a uint32_t.
+    //! Every row consists of a series of header + data. The headers is:
+    //!     - a UINT16 with how many equal pixels / 4
+    //!     - a UINT16 with how many different pixels / 4
+    //!     - 4 UINT8 with the value of the "equal" pixels
+    //! After the header there are as many pixel values as different pixels.
+    //---------------------------------------------------------------------------------------------
+	MUTABLERUNTIME_API extern void CompressRLE_RGB( int width, int rows,
+									const uint8_t* pBaseDataByte,
+									TArray<uint8_t>& destData );
+
+    //---------------------------------------------------------------------------------------------
+	MUTABLERUNTIME_API extern void UncompressRLE_RGB( int width, int rows,
+                                   const uint8_t* pBaseData,
+                                   uint8_t* pDestDataB );
+
+
+	//---------------------------------------------------------------------------------------------
+	//!
+	//---------------------------------------------------------------------------------------------
+	inline void UncompressRLE_L(const Image* pBase, Image* pDest)
+	{
+		int sizeX = pBase->GetSizeX();
+		int sizeY = pBase->GetSizeY();
+		const uint8_t* pData = pBase->GetData();
+		uint8_t* pDestData = pDest->GetData();
+		for (int lod = 0; lod < pBase->GetLODCount(); ++lod)
+		{
+			uint32_t compressedSize = UncompressRLE_L(sizeX, sizeY, pData, pDestData);
+			pData += compressedSize;
+
+			pDestData += pDest->GetLODDataSize(lod);
+			sizeX = FMath::Max(1, FMath::DivideAndRoundUp(sizeX, 2));
+			sizeY = FMath::Max(1, FMath::DivideAndRoundUp(sizeY, 2));
+		}
+	}
+
+
+    //---------------------------------------------------------------------------------------------
+    //!
+    //---------------------------------------------------------------------------------------------
+    inline bool CompressRLE_L( const Image* pBase, Image* pDest )
+    {
+        int sizeX = pBase->GetSizeX();
+        int sizeY = pBase->GetSizeY();
+        const uint8_t* pData = pBase->GetData();
+        uint32_t destDataSize = pDest->GetDataSize();
+        uint8_t* pDestData = pDest->GetData();
+        for (int lod=0;lod<pBase->GetLODCount();++lod)
+        {
+            uint32_t compressedSize = CompressRLE_L( sizeX, sizeY, pData,
+                                                     pDestData,
+                                                     destDataSize );
+            if (!compressedSize)
+            {
+                return false;
+            }
+
+            destDataSize -= compressedSize;
+            pDestData += compressedSize;
+
+            pData += pBase->GetLODDataSize(lod);
+            sizeX = FMath::Max(1, FMath::DivideAndRoundUp(sizeX,2));
+            sizeY = FMath::Max(1, FMath::DivideAndRoundUp(sizeY,2));
+        }
+
+        pDest->m_data.SetNum( pDestData-pDest->GetData() );
+
+#ifdef MUTABLE_DEBUG_RLE		
+		{
+			ImagePtr Temp = new Image(pBase->GetSizeX(),pBase->GetSizeY(),pBase->GetLODCount(),pBase->GetFormat());
+			UncompressRLE_L(pDest,Temp.get());
+			check( !FMemory::Memcmp(Temp->GetData(),pBase->GetData(),pBase->GetDataSize() ) );
+		}
+#endif
+
+        return true;
+    }
+
+
+    //---------------------------------------------------------------------------------------------
+    //! This assumes the maximum data size has been preallocateds in pDest. It will return false
+    //! if the RLE-compressed data doesn't fit.
+    //---------------------------------------------------------------------------------------------
+    inline bool CompressRLE_L1( const Image* pBase, Image* pDest )
+    {
+        int sizeX = pBase->GetSizeX();
+        int sizeY = pBase->GetSizeY();
+        const uint8_t* pData = pBase->GetData();
+        uint32_t destDataSize = pDest->GetDataSize();
+        uint8_t* pDestData = pDest->GetData();
+        for (int lod=0;lod<pBase->GetLODCount();++lod)
+        {
+            uint32_t compressedSize = CompressRLE_L1( sizeX, sizeY, pData,
+                                                      pDestData,
+                                                      destDataSize );
+            if (!compressedSize)
+            {
+                return false;
+            }
+
+            destDataSize -= compressedSize;
+            pDestData += compressedSize;
+
+            pData += pBase->GetLODDataSize(lod);
+			sizeX = FMath::Max(1, FMath::DivideAndRoundUp(sizeX, 2));
+			sizeY = FMath::Max(1, FMath::DivideAndRoundUp(sizeY, 2));
+		}
+
+        pDest->m_data.SetNum( pDestData-pDest->GetData() );
+
+        return true;
+    }
+
+
+    //---------------------------------------------------------------------------------------------
+    inline void UncompressRLE_L1( const Image* pBase, Image* pDest )
+    {
+        int sizeX = pBase->GetSizeX();
+        int sizeY = pBase->GetSizeY();
+        const uint8_t* pData = pBase->GetData();
+        uint8_t* pDestData = pDest->GetData();
+        for (int lod=0;lod<pBase->GetLODCount();++lod)
+        {
+            uint32_t compressedSize = UncompressRLE_L1( sizeX, sizeY, pData, pDestData );
+            pData += compressedSize;
+
+            pDestData += pDest->GetLODDataSize(lod);
+			sizeX = FMath::Max(1, FMath::DivideAndRoundUp(sizeX, 2));
+			sizeY = FMath::Max(1, FMath::DivideAndRoundUp(sizeY, 2));
+		}
+    }
+
+
+    //---------------------------------------------------------------------------------------------
+    //!
+    //---------------------------------------------------------------------------------------------
+    inline void CompressRLE_RGB( const Image* pBase, Image* pDest )
+    {
+        CompressRLE_RGB( pBase->GetSizeX(),
+                         pBase->GetSizeY(),
+                         pBase->GetData(),
+                         pDest->m_data );
+    }
+
+
+    //---------------------------------------------------------------------------------------------
+    //!
+    //---------------------------------------------------------------------------------------------
+    inline void UncompressRLE_RGB( const Image* pBase, Image* pDest )
+    {
+        UncompressRLE_RGB( pBase->GetSizeX(),
+                           pBase->GetSizeY(),
+                           pBase->GetData(),
+                           pDest->GetData() );
+    }
+
+
+    //---------------------------------------------------------------------------------------------
+    //!
+    //---------------------------------------------------------------------------------------------
+    inline void CompressRLE_RGBA( const Image* pBase, Image* pDest )
+    {
+        CompressRLE_RGBA( pBase->GetSizeX(),
+                          pBase->GetSizeY(),
+                          pBase->GetData(),
+                          pDest->m_data );
+    }
+
+
+    //---------------------------------------------------------------------------------------------
+    //!
+    //---------------------------------------------------------------------------------------------
+    inline void UncompressRLE_RGBA( const Image* pBase, Image* pDest )
+    {
+        UncompressRLE_RGBA( pBase->GetSizeX(),
+                            pBase->GetSizeY(),
+                            pBase->GetData(),
+                            pDest->GetData() );
+    }
+
+}
diff --git a/Engine/Plugins/Experimental/Mutable/Source/MutableRuntime/Private/Instance.cpp b/Engine/Plugins/Experimental/Mutable/Source/MutableRuntime/Private/Instance.cpp
new file mode 100644
--- /dev/null
+++ b/Engine/Plugins/Experimental/Mutable/Source/MutableRuntime/Private/Instance.cpp
@@ -0,0 +1,699 @@
+// Copyright Epic Games, Inc. All Rights Reserved.
+
+
+#include "Instance.h"
+
+#include "Platform.h"
+
+#include "InstancePrivate.h"
+#include "ImagePrivate.h"
+#include "MeshPrivate.h"
+
+
+namespace mu
+{
+
+	//---------------------------------------------------------------------------------------------
+	//---------------------------------------------------------------------------------------------
+	//---------------------------------------------------------------------------------------------
+	Instance::Instance()
+	{
+		LLM_SCOPE_BYNAME(TEXT("MutableRuntime"));
+		m_pD = new Private();
+	}
+
+
+	//---------------------------------------------------------------------------------------------
+	Instance::~Instance()
+	{
+		LLM_SCOPE_BYNAME(TEXT("MutableRuntime"));
+        check( m_pD );
+		delete m_pD;
+		m_pD = 0;
+	}
+
+
+	//---------------------------------------------------------------------------------------------
+	Instance::Private* Instance::GetPrivate() const
+	{
+		return m_pD;
+	}
+
+
+    //---------------------------------------------------------------------------------------------
+    InstancePtr Instance::Clone() const
+    {
+		LLM_SCOPE_BYNAME(TEXT("MutableRuntime"));
+        InstancePtr pResult = new Instance();
+
+        *pResult->GetPrivate() = *m_pD;
+
+        return pResult;
+    }
+
+
+    //---------------------------------------------------------------------------------------------
+    Instance::ID Instance::GetId() const
+    {
+        return m_pD->m_id;
+    }
+
+
+    //---------------------------------------------------------------------------------------------
+    int Instance::GetLODCount() const
+    {
+        return m_pD->GetLODCount();
+    }
+
+
+	//---------------------------------------------------------------------------------------------
+	int Instance::Private::GetLODCount() const
+	{
+		return (int)m_lods.size();
+	}
+
+
+	//---------------------------------------------------------------------------------------------
+	int Instance::GetComponentCount( int lod ) const
+	{
+		return m_pD->GetComponentCount( lod );
+	}
+
+
+	//---------------------------------------------------------------------------------------------
+	int Instance::Private::GetComponentCount( int lod ) const
+	{
+		check( lod>=0 && lod<(int)m_lods.size() );
+        if ( lod>=0 && lod<(int)m_lods.size() )
+        {
+            return (int)m_lods[lod].m_components.size();
+        }
+
+        return 0;
+	}
+
+
+    //---------------------------------------------------------------------------------------------
+    const char* Instance::GetComponentName( int lod, int comp ) const
+    {
+        if ( lod>=0 && lod<(int)m_pD->m_lods.size() &&
+             comp>=0 && comp<(int)m_pD->m_lods[lod].m_components.size() )
+        {
+            return m_pD->m_lods[lod].m_components[comp].m_name.c_str();
+        }
+		else
+		{
+			check(false);
+		}
+
+        return "";
+    }
+
+	
+	//---------------------------------------------------------------------------------------------
+	uint16 Instance::GetComponentId( int lod, int comp ) const
+	{
+		if ( lod>=0 && lod<(int)m_pD->m_lods.size() &&
+			 comp>=0 && comp<(int)m_pD->m_lods[lod].m_components.size() )
+		{
+			return m_pD->m_lods[lod].m_components[comp].m_id;
+		}
+		else
+		{
+			check(false);
+		}
+
+		return 0;
+	}
+
+
+    //---------------------------------------------------------------------------------------------
+    int Instance::GetSurfaceCount( int lod, int comp ) const
+    {
+        return m_pD->GetSurfaceCount( lod, comp );
+    }
+
+
+    //---------------------------------------------------------------------------------------------
+    int Instance::Private::GetSurfaceCount( int lod, int comp ) const
+    {
+        if ( lod>=0 && lod<(int)m_lods.size() &&
+             comp>=0 && comp<(int)m_lods[lod].m_components.size() )
+        {
+            return (int)m_lods[lod].m_components[comp].m_surfaces.size();
+        }
+		else
+		{
+			check(false);
+		}
+
+        return 0;
+    }
+
+
+	//---------------------------------------------------------------------------------------------
+    const char* Instance::GetSurfaceName( int lod, int comp, int surf ) const
+    {
+        if ( lod>=0 && lod<(int)m_pD->m_lods.size() &&
+             comp>=0 && comp<(int)m_pD->m_lods[lod].m_components.size() &&
+             surf>=0 && surf<(int)m_pD->m_lods[lod].m_components[comp].m_surfaces.size() )
+        {
+            return m_pD->m_lods[lod].m_components[comp].m_surfaces[surf].m_name.c_str();
+        }
+		else
+		{
+			check(false);
+		}
+
+        return "";
+    }
+
+
+    //---------------------------------------------------------------------------------------------
+    uint32_t Instance::GetSurfaceId( int lod, int comp, int surf ) const
+    {
+        if ( lod>=0 && lod<(int)m_pD->m_lods.size() &&
+             comp>=0 && comp<(int)m_pD->m_lods[lod].m_components.size() &&
+             surf>=0 && surf<(int)m_pD->m_lods[lod].m_components[comp].m_surfaces.size() )
+        {
+            return m_pD->m_lods[lod].m_components[comp].m_surfaces[surf].m_internalID;
+        }
+		else
+		{
+			check(false);
+		}
+
+        return 0;
+    }
+
+
+    //---------------------------------------------------------------------------------------------
+    int Instance::FindSurfaceById( int lod, int comp, uint32_t id ) const
+    {
+		if (lod >= 0 && lod < (int)m_pD->m_lods.size() &&
+			comp >= 0 && comp < (int)m_pD->m_lods[lod].m_components.size())
+		{
+			for (int i = 0; i < (int)m_pD->m_lods[lod].m_components[comp].m_surfaces.size(); ++i)
+			{
+				if (m_pD->m_lods[lod].m_components[comp].m_surfaces[i].m_internalID == id)
+				{
+					return i;
+				}
+			}
+		}
+		else
+		{
+			check(false);
+		}
+
+        return -1;
+    }
+
+
+    //---------------------------------------------------------------------------------------------
+    uint32_t Instance::GetSurfaceCustomId( int lod, int comp, int surf ) const
+    {
+        if ( lod>=0 && lod<(int)m_pD->m_lods.size() &&
+             comp>=0 && comp<(int)m_pD->m_lods[lod].m_components.size() &&
+             surf>=0 && surf<(int)m_pD->m_lods[lod].m_components[comp].m_surfaces.size() )
+        {
+            return m_pD->m_lods[lod].m_components[comp].m_surfaces[surf].m_customID;
+        }
+		else
+		{
+			check(false);
+		}
+
+        return 0;
+    }
+
+
+	//---------------------------------------------------------------------------------------------
+    int Instance::GetMeshCount( int lod, int comp ) const
+	{
+        return m_pD->GetMeshCount( lod, comp );
+	}
+
+
+	//---------------------------------------------------------------------------------------------
+    int Instance::Private::GetMeshCount( int lod, int comp ) const
+	{
+		check( lod>=0 && lod<(int)m_lods.size() );
+		check( comp>=0 && comp<(int)m_lods[lod].m_components.size() );
+
+        return (int)m_lods[lod].m_components[comp].m_meshes.size();
+	}
+
+
+	//---------------------------------------------------------------------------------------------
+    int Instance::GetImageCount( int lod, int comp, int surf ) const
+	{
+        return m_pD->GetImageCount( lod, comp, surf );
+	}
+
+
+	//---------------------------------------------------------------------------------------------
+    int Instance::Private::GetImageCount( int lod, int comp, int surf ) const
+	{
+		check( lod>=0 && lod<(int)m_lods.size() );
+		check( comp>=0 && comp<(int)m_lods[lod].m_components.size() );
+        check( surf>=0 && surf<(int)m_lods[lod].m_components[comp].m_surfaces.size() );
+
+        return (int)m_lods[lod].m_components[comp].m_surfaces[surf].m_images.size();
+	}
+
+
+	//---------------------------------------------------------------------------------------------
+    int Instance::GetVectorCount( int lod, int comp, int surf ) const
+	{
+        return m_pD->GetVectorCount( lod, comp, surf );
+	}
+
+
+	//---------------------------------------------------------------------------------------------
+    int Instance::Private::GetVectorCount( int lod, int comp, int surf ) const
+	{
+		check( lod>=0 && lod<(int)m_lods.size() );
+		check( comp>=0 && comp<(int)m_lods[lod].m_components.size() );
+        check( surf>=0 && surf<(int)m_lods[lod].m_components[comp].m_surfaces.size() );
+
+        return (int)m_lods[lod].m_components[comp].m_surfaces[surf].m_vectors.size();
+	}
+
+
+	//---------------------------------------------------------------------------------------------
+    int Instance::GetScalarCount( int lod, int comp, int surf ) const
+	{
+        return m_pD->GetScalarCount( lod, comp, surf );
+	}
+
+
+	//---------------------------------------------------------------------------------------------
+    int Instance::Private::GetScalarCount( int lod, int comp, int surf ) const
+	{
+		check( lod>=0 && lod<(int)m_lods.size() );
+		check( comp>=0 && comp<(int)m_lods[lod].m_components.size() );
+        check( surf>=0 && surf<(int)m_lods[lod].m_components[comp].m_surfaces.size() );
+
+        return (int)m_lods[lod].m_components[comp].m_surfaces[surf].m_scalars.size();
+	}
+
+
+    //---------------------------------------------------------------------------------------------
+    int Instance::GetStringCount( int lod, int comp, int surf ) const
+    {
+        return m_pD->GetStringCount( lod, comp, surf );
+    }
+
+
+    //---------------------------------------------------------------------------------------------
+    int Instance::Private::GetStringCount( int lod, int comp, int surf ) const
+    {
+        check( lod >= 0 && lod < (int)m_lods.size() );
+        check( comp >= 0 && comp < (int)m_lods[lod].m_components.size() );
+        check( surf >= 0 && surf < (int)m_lods[lod].m_components[comp].m_surfaces.size() );
+
+        return (int)m_lods[lod].m_components[comp].m_surfaces[surf].m_strings.size();
+    }
+
+
+    //---------------------------------------------------------------------------------------------
+    RESOURCE_ID Instance::GetMeshId( int lod, int comp, int mesh ) const
+    {
+        check( lod>=0 && lod<(int)m_pD->m_lods.size() );
+        check( comp>=0 && comp<(int)m_pD->m_lods[lod].m_components.size() );
+        check( mesh>=0 && mesh<(int)m_pD->m_lods[lod].m_components[comp].m_meshes.size() );
+
+        RESOURCE_ID result = m_pD->m_lods[lod].m_components[comp].m_meshes[mesh].m_meshId;
+        return result;
+    }
+
+
+	//---------------------------------------------------------------------------------------------
+    const char* Instance::GetMeshName( int lod, int comp, int mesh ) const
+	{
+		check( lod>=0 && lod<(int)m_pD->m_lods.size() );
+		check( comp>=0 && comp<(int)m_pD->m_lods[lod].m_components.size() );
+        check( mesh>=0 && mesh<(int)m_pD->m_lods[lod].m_components[comp].m_meshes.size() );
+
+        const char* strResult = m_pD->m_lods[lod].m_components[comp].m_meshes[mesh].m_name.c_str();
+		return strResult;
+	}
+
+
+	//---------------------------------------------------------------------------------------------
+    RESOURCE_ID Instance::GetImageId( int lod, int comp, int surf, int img ) const
+	{
+		check( lod>=0 && lod<(int)m_pD->m_lods.size() );
+		check( comp>=0 && comp<(int)m_pD->m_lods[lod].m_components.size() );
+        check( surf>=0 && surf<(int)m_pD->m_lods[lod].m_components[comp].m_surfaces.size() );
+        check( img>=0 && img<(int)m_pD->m_lods[lod].m_components[comp].m_surfaces[surf].m_images.size() );
+
+        RESOURCE_ID result = m_pD->m_lods[lod].m_components[comp].m_surfaces[surf].m_images[img].m_imageId;
+        return result;
+	}
+
+
+	//---------------------------------------------------------------------------------------------
+    const char* Instance::GetImageName( int lod, int comp, int surf, int img ) const
+	{
+		check( lod>=0 && lod<(int)m_pD->m_lods.size() );
+		check( comp>=0 && comp<(int)m_pD->m_lods[lod].m_components.size() );
+        check( surf>=0 && surf<(int)m_pD->m_lods[lod].m_components[comp].m_surfaces.size() );
+        check( img>=0 && img<(int)m_pD->m_lods[lod].m_components[comp].m_surfaces[surf].m_images.size() );
+
+        const char* strResult = m_pD->m_lods[lod].m_components[comp].m_surfaces[surf].m_images[img].m_name.c_str();
+		return strResult;
+	}
+
+
+	//---------------------------------------------------------------------------------------------
+    void Instance::GetVector( int lod, int comp, int surf, int vec,
+							  float* pX, float* pY, float* pZ, float* pW ) const
+	{
+		check( lod>=0 && lod<(int)m_pD->m_lods.size() );
+		check( comp>=0 && comp<(int)m_pD->m_lods[lod].m_components.size() );
+        check( surf>=0 && surf<(int)m_pD->m_lods[lod].m_components[comp].m_surfaces.size() );
+        check( vec>=0 && vec<(int)m_pD->m_lods[lod].m_components[comp].m_surfaces[surf].m_vectors.size() );
+
+        vec4<float> r = m_pD->m_lods[lod].m_components[comp].m_surfaces[surf].m_vectors[vec].m_vec;
+
+		if ( pX ) { *pX = r[0]; }
+		if ( pY ) { *pY = r[1]; }
+		if ( pZ ) { *pZ = r[2]; }
+		if ( pW ) { *pW = r[3]; }
+	}
+
+
+	//---------------------------------------------------------------------------------------------
+    const char* Instance::GetVectorName( int lod, int comp, int surf, int vec ) const
+	{
+		check( lod>=0 && lod<(int)m_pD->m_lods.size() );
+		check( comp>=0 && comp<(int)m_pD->m_lods[lod].m_components.size() );
+        check( surf>=0 && surf<(int)m_pD->m_lods[lod].m_components[comp].m_surfaces.size() );
+        check( vec>=0 && vec<(int)m_pD->m_lods[lod].m_components[comp].m_surfaces[surf].m_vectors.size() );
+
+        const char* strResult = m_pD->m_lods[lod].m_components[comp].m_surfaces[surf].m_vectors[vec].m_name.c_str();
+		return strResult;
+	}
+
+
+	//---------------------------------------------------------------------------------------------
+    float Instance::GetScalar( int lod, int comp, int surf, int sca ) const
+	{
+		check( lod>=0 && lod<(int)m_pD->m_lods.size() );
+		check( comp>=0 && comp<(int)m_pD->m_lods[lod].m_components.size() );
+        check( surf>=0 && surf<(int)m_pD->m_lods[lod].m_components[comp].m_surfaces.size() );
+        check( sca>=0 && sca<(int)m_pD->m_lods[lod].m_components[comp].m_surfaces[surf].m_scalars.size() );
+
+        float result = m_pD->m_lods[lod].m_components[comp].m_surfaces[surf].m_scalars[sca].m_scalar;
+		return result;
+	}
+
+
+	//---------------------------------------------------------------------------------------------
+    const char* Instance::GetScalarName( int lod, int comp, int surf, int sca ) const
+	{
+		check( lod>=0 && lod<(int)m_pD->m_lods.size() );
+		check( comp>=0 && comp<(int)m_pD->m_lods[lod].m_components.size() );
+        check( surf>=0 && surf<(int)m_pD->m_lods[lod].m_components[comp].m_surfaces.size() );
+        check( sca>=0 && sca<(int)m_pD->m_lods[lod].m_components[comp].m_surfaces[surf].m_scalars.size() );
+
+        const char* strResult = m_pD->m_lods[lod].m_components[comp].m_surfaces[surf].m_scalars[sca].m_name.c_str();
+		return strResult;
+	}
+
+
+    //---------------------------------------------------------------------------------------------
+    const char* Instance::GetString( int lod, int comp, int surf, int str ) const
+    {
+        check( lod >= 0 && lod < (int)m_pD->m_lods.size() );
+        check( comp >= 0 && comp < (int)m_pD->m_lods[lod].m_components.size() );
+        check( surf >= 0 &&
+                        surf < (int)m_pD->m_lods[lod].m_components[comp].m_surfaces.size() );
+
+        bool valid =
+            str >= 0 &&
+            str < (int)m_pD->m_lods[lod].m_components[comp].m_surfaces[surf].m_strings.size();
+        check(valid);
+
+        if (valid)
+        {
+            const char* result = m_pD->m_lods[lod]
+                                     .m_components[comp]
+                                     .m_surfaces[surf]
+                                     .m_strings[str]
+                                     .m_string.c_str();
+            return result;
+        }
+
+        return "";
+    }
+
+
+    //---------------------------------------------------------------------------------------------
+    const char* Instance::GetStringName( int lod, int comp, int surf, int str ) const
+    {
+        check( lod >= 0 && lod < (int)m_pD->m_lods.size() );
+        check( comp >= 0 && comp < (int)m_pD->m_lods[lod].m_components.size() );
+        check( surf >= 0 &&
+                        surf < (int)m_pD->m_lods[lod].m_components[comp].m_surfaces.size() );
+        bool valid =
+            str >= 0 &&
+            str < (int)m_pD->m_lods[lod].m_components[comp].m_surfaces[surf].m_strings.size();
+        check( valid );
+
+        if (valid)
+        {
+            const char* strResult =
+                m_pD->m_lods[lod].m_components[comp].m_surfaces[surf].m_strings[str].m_name.c_str();
+            return strResult;
+        }
+
+        return "";
+    }
+
+
+    //---------------------------------------------------------------------------------------------
+	int Instance::Private::AddLOD()
+	{
+		int result = (int)m_lods.size();
+		m_lods.push_back( INSTANCE_LOD() );
+		return result;
+	}
+
+
+    //---------------------------------------------------------------------------------------------
+    int Instance::Private::AddComponent( int lod )
+    {
+        // Automatically create the necessary lods and components
+        while ( lod>=GetLODCount() )
+        {
+            AddLOD();
+        }
+
+        int result = (int)m_lods[lod].m_components.size();
+        m_lods[lod].m_components.push_back( INSTANCE_COMPONENT() );
+        return result;
+    }
+
+
+    //---------------------------------------------------------------------------------------------
+    int Instance::Private::AddSurface( int lod, int comp )
+    {
+        // Automatically create the necessary lods and components
+        while ( lod>=GetLODCount() )
+        {
+            AddLOD();
+        }
+        while ( comp>=GetComponentCount(lod) )
+        {
+            AddComponent(lod);
+        }
+
+        int result = (int)m_lods[lod].m_components[comp].m_surfaces.size();
+        m_lods[lod].m_components[comp].m_surfaces.push_back( INSTANCE_SURFACE() );
+        return result;
+    }
+
+
+    //---------------------------------------------------------------------------------------------
+    void Instance::Private::SetComponentName( int lod, int comp, const char* strName )
+    {
+        // Automatically create the necessary lods and components
+        while ( lod>=GetLODCount() )
+        {
+            AddLOD();
+        }
+        while ( comp>=GetComponentCount(lod) )
+        {
+            AddComponent( lod );
+        }
+
+        INSTANCE_COMPONENT& component = m_lods[lod].m_components[comp];
+        if (strName)
+        {
+            component.m_name = strName;
+        }
+        else
+        {
+            component.m_name = "";
+        }
+    }
+
+
+    //---------------------------------------------------------------------------------------------
+    void Instance::Private::SetSurfaceName( int lod, int comp, int surf, const char* strName )
+    {
+        // Automatically create the necessary lods and components
+        while ( lod>=GetLODCount() )
+        {
+            AddLOD();
+        }
+        while ( comp>=GetComponentCount(lod) )
+        {
+            AddComponent( lod );
+        }
+        while ( surf>=GetSurfaceCount(lod, comp) )
+        {
+            AddSurface( lod, comp );
+        }
+
+        INSTANCE_SURFACE& surface = m_lods[lod].m_components[comp].m_surfaces[surf];
+        if (strName)
+        {
+            surface.m_name = strName;
+        }
+        else
+        {
+            surface.m_name = "";
+        }
+    }
+
+
+	//---------------------------------------------------------------------------------------------
+    int Instance::Private::AddMesh( int lod, int comp,
+                                    RESOURCE_ID meshId,
+                                    const char* strName )
+	{
+		// Automatically create the necessary lods and components
+		while ( lod>=GetLODCount() )
+		{
+			AddLOD();
+		}
+		while ( comp>=GetComponentCount(lod) )
+		{
+			AddComponent( lod );
+		}
+
+		INSTANCE_COMPONENT& component = m_lods[lod].m_components[comp];
+        int result = (int)component.m_meshes.size();
+        component.m_meshes.emplace_back( meshId, strName );
+
+		return result;
+	}
+
+
+	//---------------------------------------------------------------------------------------------
+    int Instance::Private::AddImage( int lod, int comp, int surf,
+                                     RESOURCE_ID imageId,
+                                     const char* strName )
+	{
+		// Automatically create the necessary lods and components
+		while ( lod>=GetLODCount() )
+		{
+			AddLOD();
+		}
+		while ( comp>=GetComponentCount(lod) )
+		{
+			AddComponent( lod );
+		}
+        while ( surf>=GetSurfaceCount(lod, comp) )
+        {
+            AddSurface( lod, comp );
+        }
+
+        INSTANCE_SURFACE& surface = m_lods[lod].m_components[comp].m_surfaces[surf];
+        int result = (int)surface.m_images.size();
+        surface.m_images.emplace_back( imageId, strName );
+
+        return result;
+	}
+
+
+	//---------------------------------------------------------------------------------------------
+    int Instance::Private::AddVector( int lod, int comp, int surf, const vec4<float>& vec, const char* strName )
+	{
+		// Automatically create the necessary lods and components
+		while ( lod>=GetLODCount() )
+		{
+			AddLOD();
+		}
+		while ( comp>=GetComponentCount(lod) )
+		{
+			AddComponent( lod );
+		}
+        while ( surf>=GetSurfaceCount(lod, comp) )
+        {
+            AddSurface( lod, comp );
+        }
+
+        INSTANCE_SURFACE& surface = m_lods[lod].m_components[comp].m_surfaces[surf];
+        int result = (int)surface.m_vectors.size();
+        surface.m_vectors.emplace_back( vec, strName );
+
+        return result;
+	}
+
+
+    //---------------------------------------------------------------------------------------------
+    int Instance::Private::AddScalar( int lod, int comp, int surf, float sca, const char* strName )
+    {
+        // Automatically create the necessary lods and components
+        while ( lod >= GetLODCount() )
+        {
+            AddLOD();
+        }
+        while ( comp >= GetComponentCount( lod ) )
+        {
+            AddComponent( lod );
+        }
+        while ( surf >= GetSurfaceCount( lod, comp ) )
+        {
+            AddSurface( lod, comp );
+        }
+
+        INSTANCE_SURFACE& surface = m_lods[lod].m_components[comp].m_surfaces[surf];
+        int result = (int)surface.m_scalars.size();
+        surface.m_scalars.emplace_back( sca, strName );
+
+        return result;
+    }
+
+
+    //---------------------------------------------------------------------------------------------
+    int Instance::Private::AddString(
+        int lod, int comp, int surf, const char* str, const char* strName )
+    {
+        // Automatically create the necessary lods and components
+        while ( lod >= GetLODCount() )
+        {
+            AddLOD();
+        }
+        while ( comp >= GetComponentCount( lod ) )
+        {
+            AddComponent( lod );
+        }
+        while ( surf >= GetSurfaceCount( lod, comp ) )
+        {
+            AddSurface( lod, comp );
+        }
+
+        INSTANCE_SURFACE& surface = m_lods[lod].m_components[comp].m_surfaces[surf];
+        int result = (int)surface.m_strings.size();
+        surface.m_strings.emplace_back( str, strName );
+
+        return result;
+    }
+}
+
diff --git a/Engine/Plugins/Experimental/Mutable/Source/MutableRuntime/Private/InstancePrivate.h b/Engine/Plugins/Experimental/Mutable/Source/MutableRuntime/Private/InstancePrivate.h
new file mode 100644
--- /dev/null
+++ b/Engine/Plugins/Experimental/Mutable/Source/MutableRuntime/Private/InstancePrivate.h
@@ -0,0 +1,158 @@
+// Copyright Epic Games, Inc. All Rights Reserved.
+
+#pragma once
+
+#include "Instance.h"
+
+#include "MemoryPrivate.h"
+#include "ImagePrivate.h"
+#include "MeshPrivate.h"
+
+
+namespace mu
+{
+
+    struct INSTANCE_SURFACE
+	{
+		string m_name;
+        uint32_t m_internalID=0;
+        uint32_t m_customID=0;
+
+		struct IMAGE
+		{
+            IMAGE( RESOURCE_ID p, const char* strName )
+			{
+                m_imageId = p;
+				if (strName)
+				{
+					m_name = strName;
+				}
+			}
+
+            RESOURCE_ID m_imageId;
+			string m_name;
+		};
+
+		vector<IMAGE> m_images;
+
+		struct VECTOR
+		{
+			VECTOR( const vec4<float>& v, const char* strName )
+			{
+				m_vec = v;
+				if (strName)
+				{
+					m_name = strName;
+				}
+			}
+
+			vec4<float> m_vec;
+			string m_name;
+		};
+
+		vector<VECTOR> m_vectors;
+
+        struct SCALAR
+        {
+            SCALAR( float v, const char* strName )
+            {
+                m_scalar = v;
+                if ( strName )
+                {
+                    m_name = strName;
+                }
+            }
+
+            float m_scalar;
+            string m_name;
+        };
+
+        vector<SCALAR> m_scalars;
+
+        struct STRING
+        {
+            STRING( const char* strValue, const char* strName )
+            {
+                m_string = strValue ? strValue : "";
+                if ( strName )
+                {
+                    m_name = strName;
+                }
+            }
+
+            string m_string;
+            string m_name;
+        };
+
+        vector<STRING> m_strings;
+    };
+
+
+    struct INSTANCE_COMPONENT
+    {
+        string m_name;
+
+    	uint16 m_id;
+
+		struct MESH
+		{
+            MESH(RESOURCE_ID p, const char* strName)
+			{
+                m_meshId = p;
+				if (strName)
+				{
+					m_name = strName;
+				}
+			}
+
+            RESOURCE_ID m_meshId;
+			string m_name;
+		};
+		vector<MESH> m_meshes;
+
+		// The order must match the meshes surfaces
+		vector<INSTANCE_SURFACE> m_surfaces;
+	};
+
+
+    struct INSTANCE_LOD
+    {
+        vector<INSTANCE_COMPONENT> m_components;
+    };
+
+
+	class Instance::Private : public Base
+	{
+	public:
+
+        //!
+        Instance::ID m_id;
+
+		//!
+		vector<INSTANCE_LOD> m_lods;
+
+        Private()
+        {
+            m_id = 0;
+        }
+
+		int GetLODCount() const;
+        int GetComponentCount( int lod ) const;
+        int GetSurfaceCount( int lod, int comp ) const;
+        int GetMeshCount( int lod, int comp ) const;
+        int GetImageCount( int lod, int comp, int surf ) const;
+        int GetVectorCount( int lod, int comp, int surf ) const;
+        int GetScalarCount( int lod, int comp, int surf ) const;
+        int GetStringCount( int lod, int comp, int surf ) const;
+        int AddLOD();
+        int AddComponent( int lod );
+        void SetComponentName( int lod, int comp, const char* strName );
+        int AddMesh(int lod, int comp, RESOURCE_ID, const char* strName);
+		int AddSurface( int lod, int comp );
+        void SetSurfaceName( int lod, int comp, int surf, const char* strName );
+        int AddImage( int lod, int comp, int surf, RESOURCE_ID, const char* strName );
+        int AddVector( int lod, int comp, int surf, const vec4<float>&, const char* strName );
+        int AddScalar( int lod, int comp, int surf, float, const char* strName );
+        int AddString( int lod, int comp, int surf, const char* strValue, const char* strName );
+    };
+}
diff --git a/Engine/Plugins/Experimental/Mutable/Source/MutableRuntime/Private/Layout.cpp b/Engine/Plugins/Experimental/Mutable/Source/MutableRuntime/Private/Layout.cpp
new file mode 100644
--- /dev/null
+++ b/Engine/Plugins/Experimental/Mutable/Source/MutableRuntime/Private/Layout.cpp
@@ -0,0 +1,220 @@
+// Copyright Epic Games, Inc. All Rights Reserved.
+
+
+#include "MutableRuntime/Public/Layout.h"
+
+#include "MutableRuntime/Private/Platform.h"
+
+
+namespace mu {
+
+	//---------------------------------------------------------------------------------------------
+	Layout::Layout()
+	{
+	}
+
+
+	//---------------------------------------------------------------------------------------------
+	void Layout::Serialise( const Layout* p, OutputArchive& arch )
+	{
+		arch << *p;
+	}
+
+
+	//---------------------------------------------------------------------------------------------
+	LayoutPtr Layout::StaticUnserialise( InputArchive& arch )
+	{
+		LLM_SCOPE_BYNAME(TEXT("MutableRuntime"));
+		LayoutPtr pResult = new Layout();
+		arch >> *pResult;
+		return pResult;
+	}
+
+
+	//---------------------------------------------------------------------------------------------
+	LayoutPtr Layout::Clone() const
+	{
+		LLM_SCOPE_BYNAME(TEXT("MutableRuntime"));
+		LayoutPtr pResult = new Layout();		
+		pResult->m_size = m_size;
+		pResult->m_maxsize = m_maxsize;
+		pResult->m_blocks = m_blocks;
+		pResult->m_strategy = m_strategy;
+		return pResult;
+	}
+
+
+	//---------------------------------------------------------------------------------------------
+	bool Layout::operator==( const Layout& o ) const
+	{
+		return  (m_size == o.m_size) &&
+			(m_maxsize == o.m_maxsize) &&
+			(m_blocks == o.m_blocks) &&
+			(m_strategy == o.m_strategy);
+	}
+
+
+	//---------------------------------------------------------------------------------------------
+	FIntPoint Layout::GetGridSize() const
+	{
+		return FIntPoint(m_size[0], m_size[1]);
+	}
+
+
+	//---------------------------------------------------------------------------------------------
+	void Layout::SetGridSize( int sizeX, int sizeY )
+	{
+		check( sizeX>=0 && sizeY>=0 );
+		m_size[0] = (uint16_t)sizeX;
+		m_size[1] = (uint16_t)sizeY;
+	}
+
+
+	//---------------------------------------------------------------------------------------------
+	void Layout::GetMaxGridSize(int* pSizeX, int* pSizeY) const
+	{
+		check(pSizeX && pSizeY);
+
+		if (pSizeX && pSizeY)
+		{
+			*pSizeX = m_maxsize[0];
+			*pSizeY = m_maxsize[1];
+		}
+	}
+
+
+	//---------------------------------------------------------------------------------------------
+	void Layout::SetMaxGridSize(int sizeX, int sizeY)
+	{
+		check(sizeX >= 0 && sizeY >= 0);
+		m_maxsize[0] = (uint16_t)sizeX;
+		m_maxsize[1] = (uint16_t)sizeY;
+	}
+
+
+	//---------------------------------------------------------------------------------------------
+	int32 Layout::GetBlockCount() const
+	{
+		return m_blocks.Num();
+	}
+
+
+	//---------------------------------------------------------------------------------------------
+	void Layout::SetBlockCount( int32 n )
+	{
+		check( n>=0 );
+		LLM_SCOPE_BYNAME(TEXT("MutableRuntime"));
+		m_blocks.SetNum( n );
+	}
+
+
+	//---------------------------------------------------------------------------------------------
+	void Layout::GetBlock( int index, int* pMinX, int* pMinY, int* pSizeX, int* pSizeY ) const
+	{
+		check( index >=0 && index < m_blocks.Num() );
+		check( pMinX && pMinY && pSizeX && pSizeY );
+
+		if (pMinX && pMinY && pSizeX && pSizeY)
+		{
+			*pMinX = m_blocks[index].m_min[0];
+			*pMinY = m_blocks[index].m_min[1];
+			*pSizeX = m_blocks[index].m_size[0];
+			*pSizeY = m_blocks[index].m_size[1];
+		}
+	}
+
+
+	//---------------------------------------------------------------------------------------------
+	void Layout::GetBlockPriority(int index, int* pPriority) const
+	{
+		check(index >= 0 && index < m_blocks.Num());
+		check(pPriority);
+
+		if (pPriority)
+		{
+			*pPriority = m_blocks[index].m_priority;
+		}
+	}
+
+
+	//---------------------------------------------------------------------------------------------
+    void Layout::SetBlock( int index, int minx, int miny, int sizex, int sizey )
+	{
+		check( index >=0 && index < m_blocks.Num() );
+
+		// Keeps the id
+		m_blocks[index].m_min = vec2<uint16_t>((uint16_t)minx, (uint16_t)miny);
+		m_blocks[index].m_size = vec2<uint16_t>((uint16_t)sizex, (uint16_t)sizey);
+	}
+
+
+	//---------------------------------------------------------------------------------------------
+	void Layout::SetBlockPriority(int index, int priority)
+	{
+		check(index >= 0 && index < m_blocks.Num());
+
+		m_blocks[index].m_priority = priority;
+	}
+
+
+	//---------------------------------------------------------------------------------------------
+	void Layout::SetLayoutPackingStrategy(EPackStrategy _strategy)
+	{
+		m_strategy = _strategy;
+	}
+
+
+	//---------------------------------------------------------------------------------------------
+	EPackStrategy Layout::GetLayoutPackingStrategy() const
+	{
+		return m_strategy;
+	}
+
+
+	//---------------------------------------------------------------------------------------------
+	bool Layout::IsSimilar(const Layout& o) const
+	{
+		if (m_size != o.m_size || m_maxsize != o.m_maxsize ||
+			m_blocks.Num() != o.m_blocks.Num() || m_strategy != o.m_strategy)
+			return false;
+
+		for (int32 i = 0; i < m_blocks.Num(); ++i)
+		{
+			if (!m_blocks[i].IsSimilar(o.m_blocks[i])) return false;
+		}
+
+		return true;
+
+	}
+
+
+	//---------------------------------------------------------------------------------------------
+	int32 Layout::FindBlock(int32_t id) const
+	{
+		check(id >= 0);
+		int res = -1;
+		for (int32 i = 0; res < 0 && i < m_blocks.Num(); ++i)
+		{
+			if (m_blocks[i].m_id == id)
+			{
+				res = i;
+			}
+		}
+
+		return res;
+	}
+
+
+	//---------------------------------------------------------------------------------------------
+	bool Layout::IsSingleBlockAndFull() const
+	{
+		if (m_blocks.Num() == 1
+			&& m_blocks[0].m_min == vec2<uint16_t>(0, 0)
+			&& m_blocks[0].m_size == m_size)
+		{
+			return true;
+		}
+		return false;
+	}
+}
+
diff --git a/Engine/Plugins/Experimental/Mutable/Source/MutableRuntime/Private/Memory.cpp b/Engine/Plugins/Experimental/Mutable/Source/MutableRuntime/Private/Memory.cpp
new file mode 100644
--- /dev/null
+++ b/Engine/Plugins/Experimental/Mutable/Source/MutableRuntime/Private/Memory.cpp
@@ -0,0 +1,276 @@
+// Copyright Epic Games, Inc. All Rights Reserved.
+
+#include "MutableMemory.h"
+
+#include "Platform.h"
+
+#include "MemoryPrivate.h"
+
+#include "BlockCompression/BlockCompression.h"
+
+#include <atomic>
+#include <cstdlib>
+#include <cassert>
+
+using namespace mu;
+
+//! This would be nice, but apparently several compilers don't support it.
+#define MUTABLE_USE_STD_ALIGNED_ALLOC   1
+
+
+//! Static pointers to the user provided malloc and free methods.
+//! If they are null (default) the versions on the standard library will be used.
+static void* (*s_custom_malloc)(size_t, uint32_t) = nullptr;
+static void(*s_custom_free)(void*) = nullptr;
+
+
+static int s_initialized = 0;
+static int s_finalized = 0;
+
+
+//---------------------------------------------------------------------------------------------
+#if MUTABLE_USE_STD_ALIGNED_ALLOC
+
+inline void* mutable_aligned_alloc( size_t size, size_t align )
+{
+#if defined(_WIN32) && _WIN32
+    return _aligned_malloc(size,align);
+#elif defined(MUTABLE_PLATFORM_ANDROID)
+    return memalign( align, size );
+#else
+    return aligned_alloc( align, size );
+#endif
+}
+
+inline void mutable_aligned_free( void* p )
+{
+#if defined(_WIN32) && _WIN32
+    _aligned_free( p );
+#else
+    free( p );
+#endif
+}
+
+#else
+
+typedef uint16_t offset_t;
+#define PTR_OFFSET_SZ sizeof( offset_t ) 
+#ifndef align_up
+#define align_up( num, align ) ( ( ( num ) + ( (align)-1 ) ) & ~( (align)-1 ) )
+#endif
+
+inline void* mutable_aligned_alloc( size_t size, size_t align )
+{
+    void* ptr = nullptr;
+
+    // We want it to be a power of two since align_up operates on powers of two
+    assert( ( align & ( align - 1 ) ) == 0 );
+
+    if ( align && size )
+    {
+        // We know we have to fit an offset value
+        // We also allocate extra bytes to ensure we
+        // can meet the alignment
+        size_t hdr_size = PTR_OFFSET_SZ + ( align - 1 );
+        void* p = malloc( size + hdr_size );
+
+        if ( p )
+        {
+            // Add the offset size to malloc's pointer (we will always store that)
+            // Then align the resulting value to the
+            // target alignment
+            ptr = (void*)align_up( ( (uintptr_t)p + PTR_OFFSET_SZ ), align );
+
+            // Calculate the offset and store it
+            // behind our aligned pointer
+            *( (offset_t*)ptr - 1 ) = (offset_t)( (uintptr_t)ptr - (uintptr_t)p );
+
+        } // else NULL, could not malloc
+    }     // else NULL, invalid arguments
+
+    return ptr;
+}
+
+
+void mutable_aligned_free( void* ptr )
+{
+    assert( ptr );
+
+    // Walk backwards from the passed-in pointer to get the pointer offset. We convert to an 
+    // offset_t pointer and rely on pointer math to get the data
+    offset_t offset = *( (offset_t*)ptr - 1 );
+
+    // Once we have the offset, we can get our original pointer and call free
+    void* p = (void*)( (uint8_t*)ptr - offset );
+    free( p );
+}
+
+#endif
+
+
+//---------------------------------------------------------------------------------------------
+//! Call the system malloc or the user-provided replacement.
+//---------------------------------------------------------------------------------------------
+inline void* lowerlevel_malloc( size_t bytes )
+{
+    void* pMem = 0;
+    if ( s_custom_malloc )
+    {
+        pMem = s_custom_malloc( bytes, 1 );
+    }
+    else
+    {
+        pMem = malloc( bytes );
+    }
+    return pMem;
+}
+
+inline void* lowerlevel_malloc_aligned( size_t bytes, uint32_t alignment )
+{
+    void* pMem = 0;
+    if ( s_custom_malloc )
+    {
+        pMem = s_custom_malloc( bytes, alignment );
+    }
+    else
+    {
+        size_t padding = bytes % alignment;
+        bytes += padding ? alignment - padding : 0;
+        pMem = mutable_aligned_alloc( bytes, alignment );
+    }
+    return pMem;
+}
+
+//-------------------------------------------------------------------------------------------------
+//! Call the system free or the user-provided replacement.
+//-------------------------------------------------------------------------------------------------
+inline void lowerlevel_free( void* ptr )
+{
+    if ( s_custom_free )
+    {
+        s_custom_free( ptr );
+    }
+    else
+    {
+        free( ptr );
+    }
+}
+
+
+//-------------------------------------------------------------------------------------------------
+inline void lowerlevel_free_aligned( void* ptr )
+{
+    if ( s_custom_free )
+    {
+        s_custom_free( ptr );
+    }
+    else
+    {
+        mutable_aligned_free( ptr );
+    }
+}
+
+
+//-------------------------------------------------------------------------------------------------
+//! Memory management functions to be used inside the library. No other memory allocation is
+//! allowed.
+//-------------------------------------------------------------------------------------------------
+void* mu::mutable_malloc( size_t size )
+{
+	if (!size)
+	{
+		return nullptr;
+	}
+
+    void* pMem = nullptr;
+
+	pMem = lowerlevel_malloc(size);
+
+    return pMem;
+}
+
+
+void* mu::mutable_malloc_aligned( size_t size, uint32_t alignment )
+{
+	if (!size)
+	{
+		return nullptr;
+	}
+
+    void* pMem = nullptr;
+
+    pMem = lowerlevel_malloc_aligned(size, alignment); 
+
+    return pMem;
+}
+
+
+void mu::mutable_free( void* ptr )
+{
+	if (!ptr)
+	{
+		return;
+	}
+
+	lowerlevel_free(ptr);
+}
+
+
+void mu::mutable_free( void* ptr, size_t size )
+{
+	if (!ptr || !size)
+	{
+		return;
+	}
+
+	lowerlevel_free(ptr);
+}
+
+
+void mu::mutable_free_aligned( void* ptr, size_t size )
+{
+	if (!size || !ptr)
+	{
+		return;
+	}
+
+	lowerlevel_free_aligned(ptr);
+}
+
+
+//-------------------------------------------------------------------------------------------------
+void mu::Initialize
+    (
+        void* (*customMalloc)( size_t, uint32_t ),
+        void (*customFree)( void* )
+    )
+{
+	if (!s_initialized)
+	{
+		s_initialized = 1;
+		s_finalized = 0;
+
+		s_custom_malloc = customMalloc;
+		s_custom_free = customFree;
+
+		miro::initialize();
+	}
+}
+
+
+//-------------------------------------------------------------------------------------------------
+void mu::Finalize()
+{
+    if (s_initialized && !s_finalized)
+    {
+        miro::finalize();
+
+        s_finalized = 1;
+        s_initialized = 0;
+        s_custom_malloc = nullptr;
+        s_custom_free = nullptr;
+    }
+}
+
+
+
diff --git a/Engine/Plugins/Experimental/Mutable/Source/MutableRuntime/Private/MemoryPrivate.h b/Engine/Plugins/Experimental/Mutable/Source/MutableRuntime/Private/MemoryPrivate.h
new file mode 100644
--- /dev/null
+++ b/Engine/Plugins/Experimental/Mutable/Source/MutableRuntime/Private/MemoryPrivate.h
@@ -0,0 +1,115 @@
+// Copyright Epic Games, Inc. All Rights Reserved.
+
+#pragma once
+
+#include "MutableRuntime/Private/Platform.h"
+#include "MutableRuntime/Public/MutableMemory.h"
+
+#include <unordered_map>
+#include <map>
+#include <set>
+
+
+namespace mu
+{
+
+    //!
+	MUTABLERUNTIME_API void mutable_memlog();
+
+
+    //! Custom memory allocator to be used by STL continers inside the library.
+    //! It uses the functions from the memory hooks
+    template <class T>
+    class MemAlloc
+    {
+    public:
+
+        // type definitions
+        using value_type = T;
+        using pointer = T*;
+        using const_pointer = const T*;
+        using reference = T&;
+        using const_reference = const T&;
+        using size_type = std::size_t;
+        using difference_type = std::ptrdiff_t;
+
+        // constructors and destructor
+        // - nothing to do because the allocator has no state
+        MemAlloc() noexcept = default;
+
+        template <class U>
+        MemAlloc( const MemAlloc<U>& ) noexcept
+        {
+        }
+
+        // allocate but don't initialize num elements of type T
+        pointer allocate (size_type num, const void* = nullptr)
+        {
+            auto ret = (pointer)(mutable_malloc(num*sizeof(T)));
+            return ret;
+        }
+
+        // deallocate storage p of deleted elements
+        void deallocate (pointer p, size_type)
+        {
+			mutable_free((void*)p);
+        }
+    };
+
+    //! Return that all specializations of this allocator are interchangeable
+    template<class T1, class T2>
+    bool operator==( const MemAlloc<T1>&, const MemAlloc<T2>& ) noexcept
+    {
+        return true;
+    }
+
+    template<class T1, class T2>
+    bool operator!=( const MemAlloc<T1>& a, const MemAlloc<T2>& b ) noexcept
+    {
+        return !(a==b);
+    }
+
+
+    //! STL-like containers using this allocator
+    template< typename T >
+    using vector = std::vector<T, MemAlloc<T> >;
+
+    template< typename T >
+    using basic_string = std::basic_string<T, std::char_traits<T>, MemAlloc<T> >;
+
+    using string = std::basic_string<char, std::char_traits<char>, MemAlloc<char> >;
+
+    template< typename K, typename T >
+    using map = std::map< K, T, std::less<K>, MemAlloc< std::pair<const K,T> > >;
+
+    template< typename T >
+    using set = std::set< T, std::less<T>, MemAlloc<T> >;
+
+    template< typename T >
+    using multiset = std::multiset< T, std::less<T>, MemAlloc<T> >;
+
+    template< typename K, typename T >
+    using unordered_map = std::unordered_map< K, T, std::hash<K>, std::equal_to<K>, MemAlloc< std::pair<const K,T> > >;
+
+    template< typename K, typename T >
+    using pair = std::pair<K, T>;
+
+}
+
+
+#ifdef __GNUC__
+#include <string>
+namespace std
+{
+    // This is necessary for GCC, but it may be inefficient, causing additional string copies.
+    // In theory it can be avoided in C++0x17 with the use of string_view
+    template<>
+    struct hash<mu::string>
+    {
+        std::size_t operator()(mu::string const& s) const noexcept
+        {
+            return std::hash<std::basic_string<char, std::char_traits<char>>>{}(std::basic_string<char, std::char_traits<char>>(s.c_str()));
+        }
+    };
+}
+#endif //#ifdef __GNUC__
diff --git a/Engine/Plugins/Experimental/Mutable/Source/MutableRuntime/Private/Mesh.cpp b/Engine/Plugins/Experimental/Mutable/Source/MutableRuntime/Private/Mesh.cpp
new file mode 100644
--- /dev/null
+++ b/Engine/Plugins/Experimental/Mutable/Source/MutableRuntime/Private/Mesh.cpp
@@ -0,0 +1,1405 @@
+// Copyright Epic Games, Inc. All Rights Reserved.
+
+
+#include "Mesh.h"
+
+#include "MeshPrivate.h"
+#include "Platform.h"
+#include "MutableTrace.h"
+
+using namespace mu;
+
+
+//---------------------------------------------------------------------------------------------
+void Mesh::Serialise( const Mesh* p, OutputArchive& arch )
+{
+    //p->m_pD->CheckIntegrity();
+    arch << *p;
+}
+
+
+//---------------------------------------------------------------------------------------------
+MeshPtr Mesh::StaticUnserialise( InputArchive& arch )
+{
+    MUTABLE_CPUPROFILER_SCOPE(MeshUnserialise)
+	LLM_SCOPE_BYNAME(TEXT("MutableRuntime"));
+
+    MeshPtr pResult = new Mesh();
+    arch >> *pResult;
+
+    //pResult->m_pD->CheckIntegrity();
+
+    return pResult;
+}
+
+
+//---------------------------------------------------------------------------------------------
+MeshPtr Mesh::Clone() const
+{
+    //MUTABLE_CPUPROFILER_SCOPE(MeshClone);
+	LLM_SCOPE_BYNAME(TEXT("MutableRuntime"));
+
+    MeshPtr pResult = new Mesh();
+
+    pResult->m_internalId = m_internalId;
+	pResult->m_staticFormatFlags = m_staticFormatFlags;
+	pResult->m_surfaces = m_surfaces;
+	pResult->m_pSkeleton = m_pSkeleton;
+	pResult->m_pPhysicsBody = m_pPhysicsBody;
+	pResult->m_faceGroups = m_faceGroups;
+	pResult->m_tags = m_tags;
+
+    // Clone the main buffers
+    pResult->m_VertexBuffers = m_VertexBuffers;
+    pResult->m_IndexBuffers = m_IndexBuffers;
+    pResult->m_FaceBuffers = m_FaceBuffers;
+
+	// Clone additional buffers
+	pResult->m_AdditionalBuffers = m_AdditionalBuffers;
+
+    // Clone the layouts
+	pResult->m_layouts = m_layouts;
+
+    // The skeleton is not cloned because it is not owned by this mesh and it is always assumed
+    // to be shared.
+
+	// physics body doen't need to be deep cloned either as they are also assumed to be shared.
+	
+	// Clone bone poses
+	pResult->m_bonePoses = m_bonePoses;
+
+    return pResult;
+}
+
+MeshPtr Mesh::Clone( EMeshCloneFlags Flags ) const
+{
+    //MUTABLE_CPUPROFILER_SCOPE(MeshClone);
+	LLM_SCOPE_BYNAME(TEXT("MutableRuntime"));
+
+    MeshPtr pResult = new Mesh();
+
+    pResult->m_internalId = m_internalId;
+	pResult->m_staticFormatFlags = m_staticFormatFlags;
+
+	if ( EnumHasAnyFlags( Flags, EMeshCloneFlags::WithSurfaces ))
+	{
+		pResult->m_surfaces = m_surfaces;
+	}
+
+	if ( EnumHasAnyFlags( Flags, EMeshCloneFlags::WithSkeleton ))
+	{
+		pResult->m_pSkeleton = m_pSkeleton;
+	}
+
+	if ( EnumHasAnyFlags( Flags, EMeshCloneFlags::WithPhysicsBody ))
+	{
+		pResult->m_pPhysicsBody = m_pPhysicsBody;
+	}
+
+	if ( EnumHasAnyFlags( Flags, EMeshCloneFlags::WithFaceGroups ))
+	{
+		pResult->m_faceGroups = m_faceGroups;
+	}
+
+	if ( EnumHasAnyFlags( Flags, EMeshCloneFlags::WithTags ))
+	{
+		pResult->m_tags = m_tags;
+	}
+
+    // Clone the main buffers
+	if ( EnumHasAnyFlags( Flags, EMeshCloneFlags::WithVertexBuffers ))
+    {
+		pResult->m_VertexBuffers = m_VertexBuffers;
+	}
+
+	if ( EnumHasAnyFlags( Flags, EMeshCloneFlags::WithIndexBuffers ))
+    {
+		pResult->m_IndexBuffers = m_IndexBuffers;
+	}
+
+	if ( EnumHasAnyFlags( Flags, EMeshCloneFlags::WithFaceBuffers ))
+    {
+		pResult->m_FaceBuffers = m_FaceBuffers;
+	}
+
+	// Clone additional buffers
+	if ( EnumHasAnyFlags( Flags, EMeshCloneFlags::WithAdditionalBuffers ))
+	{
+		pResult->m_AdditionalBuffers = m_AdditionalBuffers;
+	}
+
+    // Clone the layout	
+	if ( EnumHasAnyFlags( Flags, EMeshCloneFlags::WithLayouts ))
+	{
+		pResult->m_layouts = m_layouts;
+	}
+    // The skeleton is not cloned because it is not owned by this mesh and it is always assumed
+    // to be shared.
+
+	// physics body doen't need to be deep cloned either as they are also assumed to be shared.
+	
+	// Clone bone poses
+	if (EnumHasAnyFlags(Flags, EMeshCloneFlags::WithPoses))
+	{
+		pResult->m_bonePoses = m_bonePoses;
+	}
+
+    return pResult;
+}
+
+
+//---------------------------------------------------------------------------------------------
+uint32_t Mesh::GetId() const
+{
+    return m_internalId;
+}
+
+
+//---------------------------------------------------------------------------------------------
+int Mesh::GetVertexCount() const
+{
+    return GetVertexBuffers().GetElementCount();
+}
+
+
+//---------------------------------------------------------------------------------------------
+FMeshBufferSet& Mesh::GetVertexBuffers()
+{
+    return m_VertexBuffers;
+}
+
+
+//---------------------------------------------------------------------------------------------
+const FMeshBufferSet& Mesh::GetVertexBuffers() const
+{
+    return m_VertexBuffers;
+}
+
+
+//---------------------------------------------------------------------------------------------
+mu::Ptr<const Skeleton> Mesh::GetSkeleton() const
+{
+    return m_pSkeleton;
+}
+
+
+//---------------------------------------------------------------------------------------------
+void Mesh::SetSkeleton( Ptr<const Skeleton> s )
+{
+    m_pSkeleton = s;
+}
+
+//---------------------------------------------------------------------------------------------
+mu::Ptr<const PhysicsBody> Mesh::GetPhysicsBody() const
+{
+    return m_pPhysicsBody;
+}
+
+
+//---------------------------------------------------------------------------------------------
+void Mesh::SetPhysicsBody( Ptr<const PhysicsBody> PhysicsBody )
+{
+    m_pPhysicsBody = PhysicsBody;
+}
+
+//---------------------------------------------------------------------------------------------
+int Mesh::GetFaceCount() const
+{
+    int result = GetFaceBuffers().GetElementCount();
+    if (result)
+    {
+        check(GetIndexBuffers().GetElementCount() == result * 3);
+    }
+    else
+    {
+        result = GetIndexBuffers().GetElementCount() / 3;
+    }
+    return result;
+}
+
+
+//---------------------------------------------------------------------------------------------
+FMeshBufferSet& Mesh::GetFaceBuffers()
+{
+    return m_FaceBuffers;
+}
+
+
+//---------------------------------------------------------------------------------------------
+const FMeshBufferSet& Mesh::GetFaceBuffers() const
+{
+    return m_FaceBuffers;
+}
+
+
+//---------------------------------------------------------------------------------------------
+int Mesh::GetIndexCount() const
+{
+    // It is possible to ignore face buffers altogether
+    //check( GetIndexBuffers().GetElementCount()==GetFaceBuffers().GetElementCount()*3 );
+    return GetIndexBuffers().GetElementCount();
+}
+
+
+//---------------------------------------------------------------------------------------------
+FMeshBufferSet& Mesh::GetIndexBuffers()
+{
+    return m_IndexBuffers;
+}
+
+
+//---------------------------------------------------------------------------------------------
+const FMeshBufferSet& Mesh::GetIndexBuffers() const
+{
+    return m_IndexBuffers;
+}
+
+
+//---------------------------------------------------------------------------------------------
+int Mesh::GetSurfaceCount() const
+{
+    return m_surfaces.Num();
+}
+
+
+//---------------------------------------------------------------------------------------------
+void Mesh::GetSurface( int surfaceIndex,
+                       int* firstVertex, int* vertexCount,
+                       int* firstIndex, int* indexCount ) const
+{
+    int count = GetSurfaceCount();
+
+    if (surfaceIndex>=0 && surfaceIndex<count)
+    {
+        if (surfaceIndex<m_surfaces.Num())
+        {
+            const MESH_SURFACE& surf = m_surfaces[surfaceIndex];
+            if (firstVertex) *firstVertex = surf.m_firstVertex;
+            if (vertexCount) *vertexCount = surf.m_vertexCount;
+            if (firstIndex) *firstIndex = surf.m_firstIndex;
+            if (indexCount) *indexCount = surf.m_indexCount;
+        }
+        else
+        {
+            // No surfaces defined, means only one surface using all the mesh
+            if (firstVertex) *firstVertex = 0;
+            if (vertexCount) *vertexCount = GetVertexCount();
+            if (firstIndex) *firstIndex = 0;
+            if (indexCount) *indexCount = GetIndexCount();
+        }
+    }
+    else
+    {
+        check( false );
+        if (firstVertex) *firstVertex = 0;
+        if (vertexCount) *vertexCount = 0;
+        if (firstIndex) *firstIndex = 0;
+        if (indexCount) *indexCount = 0;
+    }
+}
+
+
+//---------------------------------------------------------------------------------------------
+uint32_t Mesh::GetSurfaceId( int surfaceIndex ) const
+{
+    if (surfaceIndex>=0 && surfaceIndex<m_surfaces.Num())
+    {
+        const MESH_SURFACE& surf = m_surfaces[surfaceIndex];
+        return surf.m_id;
+    }
+
+    return 0;
+}
+
+
+//---------------------------------------------------------------------------------------------
+void Mesh::AddLayout(Ptr<const Layout> pLayout )
+{
+	LLM_SCOPE_BYNAME(TEXT("MutableRuntime"));
+    m_layouts.Add( pLayout );
+}
+
+
+//---------------------------------------------------------------------------------------------
+int Mesh::GetLayoutCount() const
+{
+    return m_layouts.Num();
+}
+
+
+//---------------------------------------------------------------------------------------------
+const Layout* Mesh::GetLayout( int i ) const
+{
+    check( i>=0 && i<m_layouts.Num() );
+
+    return m_layouts[i].get();
+}
+
+
+//---------------------------------------------------------------------------------------------
+void Mesh::SetLayout( int i, Ptr<const Layout> pLayout )
+{
+    check( i>=0 && i<m_layouts.Num() );
+
+    m_layouts[i] = pLayout;
+}
+
+
+//---------------------------------------------------------------------------------------------
+void Mesh::SetFaceGroupCount( int count )
+{
+    m_faceGroups.SetNum( count );
+}
+
+
+//---------------------------------------------------------------------------------------------
+int Mesh::GetFaceGroupCount() const
+{
+    return m_faceGroups.Num();
+}
+
+
+//---------------------------------------------------------------------------------------------
+const char* Mesh::GetFaceGroupName( int group ) const
+{
+    check( group>=0 && group<m_faceGroups.Num() );
+    return m_faceGroups[group].m_name.c_str();
+}
+
+
+//---------------------------------------------------------------------------------------------
+void Mesh::SetFaceGroupName( int group, const char* strName )
+{
+    check( group>=0 && group<m_faceGroups.Num() );
+    m_faceGroups[group].m_name = strName;
+}
+
+
+//---------------------------------------------------------------------------------------------
+int Mesh::GetFaceGroupFaceCount( int group ) const
+{
+    check( group>=0 && group<m_faceGroups.Num() );
+    return m_faceGroups[group].m_faces.Num();
+}
+
+
+//---------------------------------------------------------------------------------------------
+const int32* Mesh::GetFaceGroupFaces( int group ) const
+{
+    check( group>=0 && group<m_faceGroups.Num() );
+    const int32* pResult = m_faceGroups[group].m_faces.GetData();
+    return pResult;
+}
+
+
+//---------------------------------------------------------------------------------------------
+void Mesh::SetFaceGroupFaces( int group, int count, const int32* faces )
+{
+    check( group>=0 && group<m_faceGroups.Num() );
+    check( !count || faces );
+	LLM_SCOPE_BYNAME(TEXT("MutableRuntime"));
+
+    m_faceGroups[group].m_faces.SetNum( count );
+
+    if (count)
+    {
+		FMemory::Memcpy( m_faceGroups[group].m_faces.GetData(), faces, count*sizeof(int32) );
+    }
+
+}
+
+
+//---------------------------------------------------------------------------------------------
+int Mesh::GetTagCount() const
+{
+    return m_tags.Num();
+}
+
+
+//---------------------------------------------------------------------------------------------
+void Mesh::SetTagCount( int count )
+{
+	LLM_SCOPE_BYNAME(TEXT("MutableRuntime"));
+    m_tags.SetNum( count );
+}
+
+
+//---------------------------------------------------------------------------------------------
+const char* Mesh::GetTag( int tagIndex ) const
+{
+    check( tagIndex>=0 && tagIndex<GetTagCount() );
+
+    if (tagIndex >= 0 && tagIndex < GetTagCount())
+    {
+        return m_tags[tagIndex].c_str();
+    }
+    else
+    {
+        return nullptr;
+    }
+}
+
+
+//---------------------------------------------------------------------------------------------
+void Mesh::SetTag( int tagIndex, const char* strName )
+{
+    check( tagIndex>=0 && tagIndex<GetTagCount() );
+	LLM_SCOPE_BYNAME(TEXT("MutableRuntime"));
+
+    if (tagIndex >= 0 && tagIndex < GetTagCount())
+    {
+        m_tags[tagIndex] = strName ? strName : "";
+    }
+}
+
+
+//---------------------------------------------------------------------------------------------
+const char* mu::Mesh::GetBonePoseName(int32 boneIndex) const
+{
+	check(boneIndex >= 0 && boneIndex < GetBonePoseCount());
+	if (boneIndex >= 0 && boneIndex < GetBonePoseCount())
+	{
+		return m_bonePoses[boneIndex].m_boneName.c_str();
+	}
+	return "";
+}
+
+
+//---------------------------------------------------------------------------------------------
+int32 mu::Mesh::FindBonePose(const char* strName) const
+{
+	for (int32 i = 0; i < m_bonePoses.Num(); ++i)
+	{
+		if (m_bonePoses[i].m_boneName == strName)
+		{
+			return i;
+		}
+	}
+
+	return INDEX_NONE;
+}
+
+
+//---------------------------------------------------------------------------------------------
+void mu::Mesh::SetBonePoseCount(int32 count)
+{
+	LLM_SCOPE_BYNAME(TEXT("MutableRuntime"));
+	m_bonePoses.SetNum(count);
+}
+
+
+//---------------------------------------------------------------------------------------------
+int32 mu::Mesh::GetBonePoseCount() const
+{
+	return m_bonePoses.Num();
+}
+
+
+//---------------------------------------------------------------------------------------------
+void mu::Mesh::SetBonePose(int32 boneIndex, const char* strName, FTransform3f transform, bool bSkinned)
+{
+	if (boneIndex >= 0 && boneIndex < m_bonePoses.Num())
+	{
+		FBonePose& pose = m_bonePoses[boneIndex];
+		pose.m_boneName = strName;
+		pose.m_boneTransform = transform;
+		pose.m_boneSkinned = bSkinned;
+	}
+	else 
+	{
+		check(false);
+	}
+}
+
+
+//---------------------------------------------------------------------------------------------
+void mu::Mesh::GetBoneTransform(int32 boneIndex, FTransform3f& transform) const
+{
+	check(boneIndex >= 0 && boneIndex < m_bonePoses.Num());
+	transform = boneIndex >= INDEX_NONE ? m_bonePoses[boneIndex].m_boneTransform : FTransform3f::Identity;
+}
+
+
+//---------------------------------------------------------------------------------------------
+//---------------------------------------------------------------------------------------------
+//---------------------------------------------------------------------------------------------
+size_t Mesh::GetDataSize() const
+{
+    return m_IndexBuffers.GetDataSize()
+            + m_VertexBuffers.GetDataSize()
+            + m_FaceBuffers.GetDataSize();
+}
+
+
+//---------------------------------------------------------------------------------------------
+bool Mesh::HasCompatibleFormat( const Mesh* pOther ) const
+{
+    bool compatible = true;
+
+    compatible &= m_layouts.Num()==pOther->m_layouts.Num();
+    compatible &= m_VertexBuffers.GetBufferCount()
+            == pOther->m_VertexBuffers.GetBufferCount();
+
+
+    // Indices
+    //-----------------
+    if ( m_IndexBuffers.GetElementCount()>0 && pOther->GetIndexCount()>0 )
+    {
+        check( m_IndexBuffers.m_buffers.Num() == 1 );
+        check( pOther->GetIndexBuffers().m_buffers.Num() == 1 );
+        check( m_IndexBuffers.GetBufferChannelCount(0) == 1 );
+        check( pOther->GetIndexBuffers().GetBufferChannelCount(0) == 1 );
+
+        const MESH_BUFFER& dest = m_IndexBuffers.m_buffers[0];
+        const MESH_BUFFER& source = pOther->GetIndexBuffers().m_buffers[0];
+
+        compatible &= dest.m_channels[0].m_format == source.m_channels[0].m_format;
+    }
+
+
+    // Layouts
+    //-----------------
+    // TODO?
+
+
+    // Vertices
+    //-----------------
+    for ( int vb = 0; vb<m_VertexBuffers.GetBufferCount(); ++vb )
+    {
+        const MESH_BUFFER& dest = m_VertexBuffers.m_buffers[vb];
+        const MESH_BUFFER& source = pOther->GetVertexBuffers().m_buffers[vb];
+
+        // TODO: More checks about channels formats and semantics
+        //compatible &= GetVertexBufferElementSize(vb) == pOther->GetVertexBufferElementSize(vb);
+        compatible &= dest.m_channels.Num()==source.m_channels.Num();
+    }
+
+    return compatible;
+}
+
+
+//---------------------------------------------------------------------------------------------
+vec3<uint32_t>  Mesh::GetFaceVertexIndices( int f ) const
+{
+    vec3<uint32_t> res;
+
+    MeshBufferIteratorConst<MBF_UINT32,uint32_t,1> it( m_IndexBuffers, MBS_VERTEXINDEX );
+    it += f*3;
+
+    res[0] = (*it)[0];
+    ++it;
+
+    res[1] = (*it)[0];
+    ++it;
+
+    res[2] = (*it)[0];
+    ++it;
+
+    return res;
+}
+
+
+//---------------------------------------------------------------------------------------------
+bool Mesh::HasFace
+    (
+        const Mesh& other,
+        int otherFaceIndex,
+        const VERTEX_MATCH_MAP& vertexMap
+    ) const
+{
+    bool found = false;
+
+    vec3<uint32_t> ov = other.GetFaceVertexIndices( otherFaceIndex );
+
+    UntypedMeshBufferIteratorConst it( m_IndexBuffers, MBS_VERTEXINDEX );
+
+    for ( int f=0; !found && f<m_FaceBuffers.GetElementCount(); f++ )
+    {
+        vec3<uint32_t> v;
+        v[0] = it.GetAsUINT32(); ++it;
+        v[1] = it.GetAsUINT32(); ++it;
+        v[2] = it.GetAsUINT32(); ++it;
+
+        found = true;
+        for ( int vi=0; found && vi<3; ++vi )
+        {
+            found = vertexMap.Matches(v[vi],ov[0])
+                 || vertexMap.Matches(v[vi],ov[1])
+                 || vertexMap.Matches(v[vi],ov[2]);
+        }
+    }
+
+    return found;
+}
+
+
+//---------------------------------------------------------------------------------------------
+bool Mesh::IsSameVertex
+    (
+        uint32_t vertexIndex,
+        const Mesh& other,
+        uint32_t otherVertexIndex,
+        float tolerance
+    ) const
+{
+    bool same = true;
+
+    // For all the attributes in this mesh
+    for ( int b=0; same && b<m_VertexBuffers.GetBufferCount(); ++b )
+    {
+        for ( int c=0; same && c<m_VertexBuffers.GetBufferChannelCount(b); ++c )
+        {
+            MESH_BUFFER_SEMANTIC semantic;
+            int semanticIndex = 0;
+            MESH_BUFFER_FORMAT format;
+            int components;
+            int offset = 0;
+            m_VertexBuffers.GetChannel( b, c, &semantic, &semanticIndex, &format, &components, &offset );
+
+            // If it is not one of the relevant semantics
+            if ( semantic!=MBS_POSITION &&
+                    //semantic!=MBS_TEXCOORDS &&
+                    //semantic!=MBS_NORMAL &&
+                    //semantic!=MBS_TANGENT &&
+                    //semantic!=MBS_BINORMAL &&
+                    semantic!=MBS_BONEINDICES &&
+                    semantic!=MBS_BONEWEIGHTS )
+            {
+                break;
+            }
+
+            // Find the channel in the other mesh
+            int otherBuffer = -1;
+            int otherChannel = -1;
+            other.m_VertexBuffers.FindChannel( semantic, semanticIndex, &otherBuffer, &otherChannel );
+            check( otherBuffer>=0 && otherChannel>=0 );
+
+            MESH_BUFFER_SEMANTIC otherSemantic;
+            int otherSemanticIndex = 0;
+            MESH_BUFFER_FORMAT otherFormat;
+            int otherComponents;
+            int otherOffset;
+            other.m_VertexBuffers.GetChannel
+                    (
+                        b, c,
+                        &otherSemantic, &otherSemanticIndex,
+                        &otherFormat, &otherComponents,
+                        &otherOffset
+                    );
+            check( otherSemantic == semantic );
+            check( otherFormat == format );
+            check( otherComponents == components );
+
+            int elemSize = m_VertexBuffers.GetElementSize( b );
+            const uint8_t* pData = m_VertexBuffers.GetBufferData( b )
+                    + elemSize*vertexIndex + offset;
+
+            int otherElemSize = other.m_VertexBuffers.GetElementSize( b );
+            const uint8_t* pOtherData = other.m_VertexBuffers.GetBufferData( b )
+                    + otherElemSize*otherVertexIndex + otherOffset;
+
+            switch (format)
+            {
+            case MBF_FLOAT32:
+            {
+                const float* pfData = (const float*)pData;
+                const float* pfOtherData = (const float*)pOtherData;
+                for ( int d=0; same && d<components; ++d )
+                {
+                    float diff = fabs( (*pfData) - (*pfOtherData) );
+                    same = diff <= tolerance;
+                    ++pfData;
+                    ++pfOtherData;
+                }
+                break;
+            }
+
+            case MBF_UINT8:
+            {
+                for ( int d=0; same && d<components; ++d )
+                {
+                    same = (*pData) == (*pOtherData);
+                    ++pData;
+                    ++pOtherData;
+                }
+                break;
+            }
+
+            default:
+                check( false );
+            }
+        }
+    }
+
+    return same;
+}
+
+
+//---------------------------------------------------------------------------------------------
+bool Mesh::VERTEX_MATCH_MAP::Matches(int v, int ov) const
+{
+	if (v >= 0 && v < (int)m_firstMatch.size())
+	{
+		int start = m_firstMatch[v];
+		int end = v + 1 < (int)m_firstMatch.size() ? m_firstMatch[v + 1] : (int)m_matches.size();
+		bool res = false;
+
+		while (!res && start < end)
+		{
+			if (m_matches[start] == ov)
+			{
+				res = true;
+			}
+			++start;
+		}
+
+		return res;
+	}
+
+	return false;
+}
+
+
+//---------------------------------------------------------------------------------------------
+void Mesh::GetVertexMap
+    (
+        const Mesh& other,
+        VERTEX_MATCH_MAP& vertexMap,
+        float tolerance
+    ) const
+{
+    int vertexCount = m_VertexBuffers.GetElementCount();
+    vertexMap.m_firstMatch.resize( vertexCount );
+    vertexMap.m_matches.reserve( vertexCount+(vertexCount>>2) );
+
+    int otherVertexCount = other.m_VertexBuffers.GetElementCount();
+
+    if ( !vertexCount || !otherVertexCount )
+    {
+        return;
+    }
+
+
+    MeshBufferIteratorConst< MBF_FLOAT32, float, 3 > itp( m_VertexBuffers, MBS_POSITION);
+    MeshBufferIteratorConst< MBF_FLOAT32, float, 3 > itopBegin( other.m_VertexBuffers, MBS_POSITION);
+
+
+    // Bucket the other mesh
+#define MUTABLE_NUM_BUCKETS 256
+#define MUTABLE_BUCKET_CHANNEL 0
+
+    float rangeMin = std::numeric_limits<float>::max();
+    float rangeMax = -std::numeric_limits<float>::max();
+    MeshBufferIteratorConst< MBF_FLOAT32, float, 3 >  itop = itopBegin;
+    for ( int ov=0; ov<otherVertexCount; ++ov )
+    {
+        float v = (*itop)[MUTABLE_BUCKET_CHANNEL];
+        rangeMin = std::min( rangeMin, v );
+        rangeMax = std::max( rangeMax, v );
+        ++itop;
+    }
+    rangeMin -= tolerance;
+    rangeMax += tolerance;
+
+    vector<int> buckets[MUTABLE_NUM_BUCKETS];
+    for ( int b=0; b<MUTABLE_NUM_BUCKETS; ++b )
+    {
+        buckets[b].reserve( otherVertexCount/MUTABLE_NUM_BUCKETS*2 );
+    }
+
+    float bucketSize = (rangeMax-rangeMin)/float(MUTABLE_NUM_BUCKETS);
+    itop = itopBegin;
+    for ( int ov=0; ov<otherVertexCount; ++ov )
+    {
+        float v = (*itop)[MUTABLE_BUCKET_CHANNEL];
+
+        int bucket0 = int( floor( (v-tolerance-rangeMin)/bucketSize ) );
+        bucket0 = std::min( MUTABLE_NUM_BUCKETS-1, std::max( 0, bucket0 ) );
+        buckets[bucket0].push_back(ov);
+
+        int bucket1 = int( floor( (v+tolerance-rangeMin)/bucketSize ) );
+        bucket1 = std::min( MUTABLE_NUM_BUCKETS-1, std::max( 0, bucket1 ) );
+
+        if (bucket1!=bucket0)
+        {
+            buckets[bucket1].push_back(ov);
+        }
+
+        ++itop;
+    }
+
+
+    // TODO Compare only positions?
+    if (true)
+    {
+
+        /*
+        // Don't use buckets
+        for ( int v=0; v<vertexCount; ++v )
+        {
+            itop = itopBegin;
+
+            vertexMap.m_firstMatch[v] = (int)vertexMap.m_matches.size();
+            for ( int ov=0; ov<otherVertexCount; ++ov )
+            {
+                bool same = true;
+                for ( int d=0; same && d<3; ++d )
+                {
+                    float diff = fabs( (*itp)[d] - (*itop)[d] );
+                    same = diff <= tolerance;
+                }
+
+                if ( same )
+                {
+                    vertexMap.m_matches.push_back(ov);
+                }
+
+                ++itop;
+            }
+
+            ++itp;
+        }
+/*/
+        // Use buckets
+        for ( int v=0; v<vertexCount; ++v )
+        {
+            vertexMap.m_firstMatch[v] = (int)vertexMap.m_matches.size();
+
+            float vbucket = (*itp)[MUTABLE_BUCKET_CHANNEL];
+            int bucket = int( floor( (vbucket-rangeMin)/bucketSize ) );
+
+            if (bucket>=0 && bucket<MUTABLE_NUM_BUCKETS)
+            {
+                int bucketVertexCount = (int)buckets[bucket].size();
+                for ( int ov=0; ov<bucketVertexCount; ++ov )
+                {
+                    int otherVertexIndex = buckets[bucket][ov];
+                    vec3<float> p = itopBegin[ otherVertexIndex ];
+
+                    bool same = true;
+                    for ( int d=0; same && d<3; ++d )
+                    {
+                        float diff = fabs( (*itp)[d] - p[d] );
+                        same = diff <= tolerance;
+                    }
+
+                    if ( same )
+                    {
+                        vertexMap.m_matches.push_back( otherVertexIndex );
+                    }
+                }
+            }
+
+            ++itp;
+        }
+
+    }
+    else
+    {
+        // Slow generic way
+        for ( int v=0; v<vertexCount; ++v )
+        {
+            vertexMap.m_firstMatch[v] = (int)vertexMap.m_matches.size();
+            for ( int ov=0; ov<otherVertexCount; ++ov )
+            {
+                // TODO: Optimize this by not using IsSameVertex
+                if ( IsSameVertex( v, other, ov, tolerance ) )
+                {
+                    vertexMap.m_matches.push_back(ov);
+                }
+            }
+        }
+    }
+
+}
+
+
+//---------------------------------------------------------------------------------------------
+void Mesh::EnsureSurfaceData()
+{
+	if (!m_surfaces.Num() && m_VertexBuffers.GetElementCount())
+	{
+		MESH_SURFACE s;
+		s.m_vertexCount = m_VertexBuffers.GetElementCount();
+		s.m_indexCount = m_IndexBuffers.GetElementCount();
+		m_surfaces.Add(s);
+	}
+}
+
+
+//---------------------------------------------------------------------------------------------
+void Mesh::ResetBufferIndices()
+{
+	m_VertexBuffers.ResetBufferIndices();
+	m_IndexBuffers.ResetBufferIndices();
+	m_FaceBuffers.ResetBufferIndices();
+}
+
+
+//---------------------------------------------------------------------------------------------
+bool Mesh::IsSimilar(const Mesh& o) const
+{
+	// Some meshes are just vertex indices (masks) we don't consider them for similarity,
+	// because the kind of vertex channel data they store is the kind that is ignored.
+	if (m_IndexBuffers.GetElementCount() == 0)
+	{
+		return false;
+	}
+
+	bool equal = m_IndexBuffers == o.m_IndexBuffers;
+	if (equal) equal = m_FaceBuffers == o.m_FaceBuffers;
+	if (equal) equal = (m_layouts.Num() == o.m_layouts.Num());
+	if (equal && m_pSkeleton != o.m_pSkeleton)
+	{
+		if (m_pSkeleton && o.m_pSkeleton)
+		{
+			equal = (*m_pSkeleton == *o.m_pSkeleton);
+		}
+		else
+		{
+			equal = false;
+		}
+	}
+	
+	if (equal && m_pPhysicsBody != o.m_pPhysicsBody)
+	{
+		if (m_pPhysicsBody && o.m_pPhysicsBody)
+		{
+			equal = (*m_pPhysicsBody == *o.m_pPhysicsBody);
+		}
+		else
+		{
+			equal = false;
+		}
+	}
+
+	if (equal) equal = (m_surfaces == o.m_surfaces);
+	if (equal) equal = (m_faceGroups == o.m_faceGroups);
+	if (equal) equal = (m_tags == o.m_tags);
+
+	// Special comparison for layouts
+	for (int32 i = 0; equal && i < m_layouts.Num(); ++i)
+	{
+		equal &= m_layouts[i]->IsSimilar(*o.m_layouts[i]);
+	}
+
+	// Special comparison for vertex buffers
+	if (equal)
+	{
+		equal = m_VertexBuffers.IsSimilarRobust(o.m_VertexBuffers);
+	}
+
+	return equal;
+
+}
+
+
+//---------------------------------------------------------------------------------------------
+void Mesh::ResetStaticFormatFlags() const
+{
+    m_staticFormatFlags = 0;
+
+    for ( int f=0; f<SMF_COUNT; ++f )
+    {
+        if ( s_staticMeshFormatIdentify[f]
+             &&
+             s_staticMeshFormatIdentify[f]( this ) )
+        {
+            m_staticFormatFlags |= (1<<f);
+        }
+    }
+}
+
+
+//---------------------------------------------------------------------------------------------
+void Mesh::CheckIntegrity() const
+{
+#ifdef MUTABLE_DEBUG
+
+    // Check vertex indices
+    {
+        for ( int b=0; b<m_IndexBuffers.GetBufferCount(); ++b )
+        {
+            int elemSize = m_IndexBuffers.GetElementSize( b );
+
+            for ( int c=0; c<m_IndexBuffers.GetBufferChannelCount(b); ++c )
+            {
+                MESH_BUFFER_SEMANTIC semantic;
+                int semanticIndex = 0;
+                MESH_BUFFER_FORMAT format;
+                int components;
+                int offset = 0;
+                m_IndexBuffers.GetChannel
+                        ( b, c, &semantic, &semanticIndex, &format, &components, &offset );
+
+                if ( semantic==MBS_VERTEXINDEX )
+                {
+                    size_t icount = (size_t)m_IndexBuffers.GetElementCount();
+                    int elemCount = m_VertexBuffers.GetElementCount();
+                    for ( size_t indexIndex = 0; indexIndex < icount; ++indexIndex )
+                    {
+                        const uint8_t* pData = m_IndexBuffers.GetBufferData( b )
+                                + elemSize*indexIndex + offset;
+
+                        switch (format)
+                        {
+                        case MBF_UINT32:
+                        {
+                            uint32_t index = *(const uint32_t*)pData;
+                            check( index < uint32_t( elemCount ) );
+                            break;
+                        }
+                        case MBF_UINT16:
+                        {
+                            uint16_t index = *(const uint16_t*)pData;
+                            check( index < uint16_t( elemCount ) );
+                            break;
+                        }
+                        case MBF_UINT8:
+                        {
+                            uint8_t index = *(const uint8_t*)pData;
+                            check( index < uint8_t( elemCount ) );
+                            break;
+                        }
+                        default:
+                            check(false);
+                            break;
+                        }
+                    }
+                }
+            }
+        }
+    }
+
+
+    // Check bone indices, if there are bones. Bones could have been removed for later addition as an optimisation.
+    // For all the attributes in this mesh
+    int boneCount = m_pSkeleton ? m_pSkeleton->GetBoneCount() : 0;
+    if ( m_pSkeleton && boneCount )
+    {
+        for ( int b = 0; b < m_VertexBuffers.GetBufferCount(); ++b )
+        {
+            int channelCount = m_VertexBuffers.GetBufferChannelCount( b );
+            for ( int c = 0; c < channelCount; ++c )
+            {
+                MESH_BUFFER_SEMANTIC semantic;
+                int semanticIndex = 0;
+                MESH_BUFFER_FORMAT format;
+                int components;
+                int offset = 0;
+                m_VertexBuffers.GetChannel
+                        ( b, c, &semantic, &semanticIndex, &format, &components, &offset );
+
+                // If it is not one of the relevant semantics
+                if (
+                        //semantic!=MBS_POSITION &&
+                        //semantic!=MBS_TEXCOORDS &&
+                        //semantic!=MBS_NORMAL &&
+                        //semantic!=MBS_TANGENT &&
+                        //semantic!=MBS_BINORMAL &&
+                        semantic!=MBS_BONEINDICES
+                        // && semantic!=MBS_BONEWEIGHTS
+                        )
+                {
+                    continue;
+                }
+
+                size_t elemCount = (size_t)m_VertexBuffers.GetElementCount();
+                int elemSize = m_VertexBuffers.GetElementSize( b );
+
+                for ( size_t vertexIndex = 0;
+                      vertexIndex < elemCount; ++vertexIndex )
+                {
+                    const uint8_t* pData = m_VertexBuffers.GetBufferData( b )
+                            + elemSize*vertexIndex + offset;
+
+                    switch (format)
+                    {
+                    case MBF_UINT8:
+                    {
+                        for ( int d = 0; d < components; ++d )
+                        {
+                            uint8_t index = *pData;
+                            check( index < uint64_t(boneCount) );
+                            ++pData;
+                        }
+                        break;
+                    }
+
+                    case MBF_UINT16:
+                    {
+                        for ( int d = 0; d < components; ++d )
+                        {
+                            uint16_t index = *(uint16_t*)pData;
+                            check( index < uint64_t( boneCount ) );
+                            pData+=2;
+                        }
+                        break;
+                    }
+
+                    case MBF_UINT32:
+                    {
+                        for ( int d = 0; d < components; ++d )
+                        {
+                            uint32_t index = *(uint32_t*)pData;
+                            check( index < uint64_t( boneCount ) );
+                            pData+=4;
+                        }
+                        break;
+                    }
+
+                    default:
+                        check( false );
+                    }
+                }
+            }
+        }
+    }
+#endif
+}
+
+
+//---------------------------------------------------------------------------------------------
+//---------------------------------------------------------------------------------------------
+//---------------------------------------------------------------------------------------------
+static bool StaticMeshFormatIdentify_None( const Mesh* )
+{
+    return false;
+}
+
+//---------------------------------------------------------------------------------------------
+static bool StaticMeshFormatIdentify_Project( const Mesh* pM )
+{
+    // This format is used internally for the mesh project
+    bool res = true;
+
+    // The first vertex buffer must be texcoords(2f), position(3f), normal(3f)
+    // all tightly packed
+    res &= pM->m_VertexBuffers.GetBufferCount()>=1;
+
+    if ( res )
+    {
+        res &= pM->m_VertexBuffers.m_buffers[0].m_channels.Num()==3;
+    }
+
+    if ( res )
+    {
+        const MESH_BUFFER_CHANNEL& chan =
+                pM->m_VertexBuffers.m_buffers[0].m_channels[0];
+
+        res &= chan.m_semantic == MBS_TEXCOORDS;
+        res &= chan.m_format == MBF_FLOAT32;
+        res &= chan.m_componentCount == 2;
+        //we don't really care about the semantic index
+        //res &= chan.m_semanticIndex == 0;
+        res &= chan.m_offset == 0;
+    }
+
+    if ( res )
+    {
+        const MESH_BUFFER_CHANNEL& chan =
+                pM->m_VertexBuffers.m_buffers[0].m_channels[1];
+
+        res &= chan.m_semantic == MBS_POSITION;
+        res &= chan.m_format == MBF_FLOAT32;
+        res &= chan.m_componentCount == 3;
+        res &= chan.m_semanticIndex == 0;
+        res &= chan.m_offset == 8;
+    }
+
+    if ( res )
+    {
+        const MESH_BUFFER_CHANNEL& chan =
+                pM->m_VertexBuffers.m_buffers[0].m_channels[2];
+
+        res &= chan.m_semantic == MBS_NORMAL;
+        res &= chan.m_format == MBF_FLOAT32;
+        res &= chan.m_componentCount == 3;
+        res &= chan.m_semanticIndex == 0;
+        res &= chan.m_offset == 20;
+    }
+
+    // The first index buffer must be just index buffers u32
+    if ( res )
+    {
+        res &= pM->m_IndexBuffers.m_buffers[0].m_channels.Num()>=1;
+    }
+
+    if ( res )
+    {
+        const MESH_BUFFER_CHANNEL& chan =
+                pM->m_IndexBuffers.m_buffers[0].m_channels[0];
+
+        res &= chan.m_semantic == MBS_VERTEXINDEX;
+        res &= chan.m_format == MBF_UINT32;
+        res &= chan.m_componentCount == 1;
+        res &= chan.m_semanticIndex == 0;
+        res &= chan.m_offset == 0;
+    }
+
+    return res;
+}
+
+
+//---------------------------------------------------------------------------------------------
+static bool StaticMeshFormatIdentify_ProjectWrapping( const Mesh* pM )
+{
+    // This format is used internally for the mesh project
+    bool res = true;
+
+    // The first vertex buffer must be texcoords(2f), position(3f), normal(3f), layoutBlock(uint32_t)
+    // all tightly packed
+    res &= pM->m_VertexBuffers.GetBufferCount()>=1;
+
+    if ( res )
+    {
+        res &= pM->m_VertexBuffers.m_buffers[0].m_channels.Num()==4;
+    }
+
+    if ( res )
+    {
+        const MESH_BUFFER_CHANNEL& chan = pM->m_VertexBuffers.m_buffers[0].m_channels[0];
+
+        res &= chan.m_semantic == MBS_TEXCOORDS;
+        res &= chan.m_format == MBF_FLOAT32;
+        res &= chan.m_componentCount == 2;
+        // we don't really care about the semantic index as long as there is only one
+        // res &= chan.m_semanticIndex == 0;
+        res &= chan.m_offset == 0;
+    }
+
+    if ( res )
+    {
+        const MESH_BUFFER_CHANNEL& chan = pM->m_VertexBuffers.m_buffers[0].m_channels[1];
+
+        res &= chan.m_semantic == MBS_POSITION;
+        res &= chan.m_format == MBF_FLOAT32;
+        res &= chan.m_componentCount == 3;
+        res &= chan.m_semanticIndex == 0;
+        res &= chan.m_offset == 8;
+    }
+
+    if ( res )
+    {
+        const MESH_BUFFER_CHANNEL& chan = pM->m_VertexBuffers.m_buffers[0].m_channels[2];
+
+        res &= chan.m_semantic == MBS_NORMAL;
+        res &= chan.m_format == MBF_FLOAT32;
+        res &= chan.m_componentCount == 3;
+        res &= chan.m_semanticIndex == 0;
+        res &= chan.m_offset == 20;
+    }
+
+    if ( res )
+    {
+        const MESH_BUFFER_CHANNEL& chan = pM->m_VertexBuffers.m_buffers[0].m_channels[3];
+
+        res &= chan.m_semantic == MBS_LAYOUTBLOCK;
+        res &= chan.m_format == MBF_UINT32;
+        res &= chan.m_componentCount == 1;
+        // we don't really care about the semantic index as long as there is only one
+        // res &= chan.m_semanticIndex == 0;
+        res &= chan.m_offset == 32;
+    }
+
+    // The first index buffer must be just index buffers u32
+    if ( res )
+    {
+        res &= pM->m_IndexBuffers.m_buffers[0].m_channels.Num()>=1;
+    }
+
+    if ( res )
+    {
+        const MESH_BUFFER_CHANNEL& chan = pM->m_IndexBuffers.m_buffers[0].m_channels[0];
+
+        res &= chan.m_semantic == MBS_VERTEXINDEX;
+        res &= chan.m_format == MBF_UINT32;
+        res &= chan.m_componentCount == 1;
+        res &= chan.m_semanticIndex == 0;
+        res &= chan.m_offset == 0;
+    }
+
+    return res;
+}
+
+
+//---------------------------------------------------------------------------------------------
+STATIC_MESH_FORMAT_ID_FUNC mu::s_staticMeshFormatIdentify[] =
+{
+    StaticMeshFormatIdentify_None,
+    StaticMeshFormatIdentify_Project,
+    StaticMeshFormatIdentify_ProjectWrapping
+};
+
+
+//---------------------------------------------------------------------------------------------
+namespace
+{
+    void LogBuffer( FString& out, const FMeshBufferSet& bufset, int32 BufferElementLimit)
+    {
+        (void)out;
+        (void)bufset;
+
+        size_t elemCount = bufset.m_elementCount;
+        out += "  Set with "
+                + FString::Printf(TEXT("%d"), bufset.m_buffers.Num())
+                + " buffers and "
+                + FString::Printf(TEXT("%d"), elemCount)
+                + " elements.\n";
+
+        for( const auto& buf : bufset.m_buffers )
+        {
+            const uint8_t* pData = buf.m_data.GetData();
+
+            out += "    Buffer with "+ FString::Printf(TEXT("%d"), buf.m_channels.Num())
+                    + " channels and "+ FString::Printf(TEXT("%d"), buf.m_elementSize)+" elementsize\n";
+            for( const auto& chan : buf.m_channels )
+            {
+                out += "      Channel with format: "+ FString::Printf(TEXT("%d"), chan.m_format)
+                        + " semantic: "+ FString::Printf(TEXT("%d"), chan.m_semantic)
+                        + " " + FString::Printf(TEXT("%d"), chan.m_semanticIndex)
+                        + " components: " + FString::Printf(TEXT("%d"), chan.m_componentCount)
+                        + " offset: " + FString::Printf(TEXT("%d"), chan.m_offset)+"\n";
+                for( size_t e=0; e<elemCount && e<BufferElementLimit; ++e )
+                {
+                    const uint8_t* pElementData = pData+buf.m_elementSize*e;
+                    const uint8_t* pChanData = pElementData+chan.m_offset;
+                    out += "        ";
+                    for (int c=0; c<chan.m_componentCount; ++c)
+                    {
+                        out += "\t";
+                        switch (chan.m_format)
+                        {
+                        case MBF_UINT32:
+                        case MBF_NUINT32: out += FString::Printf(TEXT("%d"), *(const uint32_t*)pChanData); pChanData +=4; break;
+                        case MBF_UINT16:
+                        case MBF_NUINT16: out += FString::Printf(TEXT("%d"), *(const uint16_t*)pChanData); pChanData +=2; break;
+                        case MBF_UINT8:
+                        case MBF_NUINT8: out += FString::Printf(TEXT("%d"), *(const uint8_t*)pChanData); pChanData +=1; break;
+						case MBF_FLOAT32: out += FString::Printf(TEXT("%.3f"), *(const float*)pChanData); pChanData += 4; break;
+                        case MBF_FLOAT16: out += FString::Printf(TEXT("%d"), *(const uint16_t*)pChanData); pChanData +=2; break;
+                        default: break;
+                        }
+                        out += ",";
+                    }
+                    out += "\n";
+                }
+            }
+        }
+    }
+}
+
+
+void Mesh::Log( FString& out, int32 BufferElementLimit)
+{
+    out += "Mesh:\n";
+
+    out += "Indices:\n";
+    LogBuffer( out, m_IndexBuffers, BufferElementLimit);
+
+    out += "Vertices:\n";
+    LogBuffer( out, m_VertexBuffers, BufferElementLimit);
+
+    out += "Faces:\n";
+    LogBuffer( out, m_FaceBuffers, BufferElementLimit);
+}
