[mutable] Move Mutable plugin from the restricted folder to the experimental plugins folder.

diff --git a/Engine/Plugins/Experimental/Mutable/Source/MutableRuntime/Private/MeshBufferSet.cpp b/Engine/Plugins/Experimental/Mutable/Source/MutableRuntime/Private/MeshBufferSet.cpp
new file mode 100644
--- /dev/null
+++ b/Engine/Plugins/Experimental/Mutable/Source/MutableRuntime/Private/MeshBufferSet.cpp
@@ -0,0 +1,606 @@
+// Copyright Epic Games, Inc. All Rights Reserved.
+
+#include "MeshBufferSet.h"
+#include "MeshPrivate.h"
+#include "Platform.h"
+#include "MutableTrace.h"
+
+
+using namespace mu;
+
+
+//---------------------------------------------------------------------------------------------
+static FMeshBufferFormatData s_meshBufferFormatData[ MBF_COUNT ] =
+    {
+        { 0, 0 },
+        { 2, 0 },
+        { 4, 0 },
+
+		{ 1, 8 },
+        { 2, 16 },
+        { 4, 32 },
+        { 1, 7 },
+        { 2, 15 },
+        { 4, 31 },
+
+		{ 1, 0 },
+        { 2, 0 },
+        { 4, 0 },
+        { 1, 0 },
+        { 2, 0 },
+        { 4, 0 },
+
+		{ 1, 0 },
+        { 1, 0 },
+        { 1, 0 },
+        { 1, 0 },
+
+		{ 8, 0 }
+};
+
+
+//---------------------------------------------------------------------------------------------
+const FMeshBufferFormatData& mu::GetMeshFormatData( MESH_BUFFER_FORMAT format )
+{
+    check( format >= 0 );
+    check( format < MBF_COUNT );
+    return s_meshBufferFormatData[ format ];
+}
+
+
+//---------------------------------------------------------------------------------------------
+//---------------------------------------------------------------------------------------------
+//---------------------------------------------------------------------------------------------
+int32 FMeshBufferSet::GetElementCount() const
+{
+    return m_elementCount;
+}
+
+
+//---------------------------------------------------------------------------------------------
+void FMeshBufferSet::SetElementCount( int32 count )
+{
+    check( count>=0 );
+	LLM_SCOPE_BYNAME(TEXT("MutableRuntime"));
+
+    m_elementCount = count;
+
+    for ( auto& buf: m_buffers )
+    {
+        buf.m_data.SetNumUninitialized( buf.m_elementSize * count, false );
+    }
+}
+
+
+//---------------------------------------------------------------------------------------------
+int32 FMeshBufferSet::GetBufferCount() const
+{
+    return m_buffers.Num();
+}
+
+
+//---------------------------------------------------------------------------------------------
+void FMeshBufferSet::SetBufferCount( int32 count )
+{
+	LLM_SCOPE_BYNAME(TEXT("MutableRuntime"));
+    m_buffers.SetNum( count );
+}
+
+
+//---------------------------------------------------------------------------------------------
+int32 FMeshBufferSet::GetBufferChannelCount( int32 buffer ) const
+{
+    check(buffer >= 0 && buffer < m_buffers.Num());
+    if (buffer >= 0 && buffer < m_buffers.Num())
+    {
+        return m_buffers[buffer].m_channels.Num();
+    }
+    return 0;
+}
+
+
+//---------------------------------------------------------------------------------------------
+void FMeshBufferSet::GetChannel
+    (
+        int32 buffer,
+        int32 channel,
+        MESH_BUFFER_SEMANTIC* pSemantic,
+        int32* pSemanticIndex,
+        MESH_BUFFER_FORMAT* pFormat,
+        int32* pComponentCount,
+        int32* pOffset
+    ) const
+{
+    check( buffer>=0 && buffer<m_buffers.Num() );
+    check( channel>=0 && channel<m_buffers[buffer].m_channels.Num() );
+
+    const MESH_BUFFER_CHANNEL& chan = m_buffers[buffer].m_channels[channel];
+
+    if (pSemantic)
+    {
+        *pSemantic = chan.m_semantic;
+    }
+
+    if (pSemanticIndex)
+    {
+        *pSemanticIndex = chan.m_semanticIndex;
+    }
+
+    if (pFormat)
+    {
+        *pFormat = chan.m_format;
+    }
+
+    if (pComponentCount)
+    {
+        *pComponentCount = chan.m_componentCount;
+    }
+
+    if (pOffset)
+    {
+        *pOffset = chan.m_offset;
+    }
+}
+
+
+//---------------------------------------------------------------------------------------------
+void FMeshBufferSet::SetBuffer
+    (
+        int32 buffer,
+        int32 elementSize,
+        int32 channelCount,
+        const MESH_BUFFER_SEMANTIC* pSemantics,
+        const int32* pSemanticIndices,
+        const MESH_BUFFER_FORMAT* pFormats,
+        const int32* pComponentCount,
+        const int32* pOffsets
+    )
+{
+    check( buffer>=0 && buffer<m_buffers.Num() );
+	LLM_SCOPE_BYNAME(TEXT("MutableRuntime"));
+
+    MESH_BUFFER& buf = m_buffers[buffer];
+
+    unsigned minElemSize = 0;
+    buf.m_channels.SetNum( channelCount );
+    for ( int c=0; c<channelCount; ++c )
+    {
+        MESH_BUFFER_CHANNEL& chan = buf.m_channels[c];
+        chan.m_semantic = pSemantics ? pSemantics[c] : MBS_NONE;
+        chan.m_semanticIndex = pSemanticIndices ? ((uint8_t)pSemanticIndices[c]) : 0;
+        chan.m_format = pFormats ? pFormats[c] : MBF_NONE;
+        chan.m_componentCount = pComponentCount ? ((uint16_t)pComponentCount[c]) : 0;
+        chan.m_offset = pOffsets ? ((uint8_t)pOffsets[c]) : 0;
+
+        minElemSize = FMath::Max
+                (
+                    minElemSize,
+                    chan.m_offset+chan.m_componentCount*GetMeshFormatData(chan.m_format).m_size
+                );
+    }
+
+    // Set the user specified element size, or enlarge it if it was too small
+    buf.m_elementSize = FMath::Max( elementSize, int32(minElemSize) );
+
+    // Update the buffer data
+    buf.m_data.SetNumUninitialized( buf.m_elementSize * m_elementCount, false );
+}
+
+//---------------------------------------------------------------------------------------------
+void FMeshBufferSet::SetBufferChannel
+    (
+        int32 buffer,
+        int32 channelIndex,
+        MESH_BUFFER_SEMANTIC semantic,
+        int32 semanticIndex,
+        MESH_BUFFER_FORMAT format,
+        int32 componentCount,
+        int32 offset
+    )
+{
+    if (!( buffer>=0 && buffer<m_buffers.Num() ))
+    {
+		check(false);
+        return;
+    }
+
+	LLM_SCOPE_BYNAME(TEXT("MutableRuntime"));
+
+    MESH_BUFFER& buf = m_buffers[buffer];
+
+    if (!( channelIndex>=0 && channelIndex<buf.m_channels.Num() ))
+    {
+		check(false);
+        return;
+    }
+
+    MESH_BUFFER_CHANNEL& chan = buf.m_channels[channelIndex];
+    chan.m_semantic = semantic;
+    chan.m_semanticIndex = uint8_t(semanticIndex);
+    chan.m_format = format;
+    chan.m_componentCount = uint16_t(componentCount);
+    chan.m_offset = uint8(offset);
+}
+
+
+//---------------------------------------------------------------------------------------------
+uint8* FMeshBufferSet::GetBufferData( int32 buffer )
+{
+    check( buffer>=0 && buffer<m_buffers.Num() );
+	uint8_t* pResult = m_buffers[buffer].m_data.GetData();
+    return pResult;
+}
+
+
+//---------------------------------------------------------------------------------------------
+const uint8* FMeshBufferSet::GetBufferData( int32 buffer ) const
+{
+    check( buffer>=0 && buffer<m_buffers.Num() );
+    const uint8_t* pResult = m_buffers[buffer].m_data.GetData();
+    return pResult;
+}
+
+
+//---------------------------------------------------------------------------------------------
+void FMeshBufferSet::FindChannel
+    (
+        MESH_BUFFER_SEMANTIC semantic, int32 semanticIndex,
+        int32* pBuffer, int32* pChannel
+    ) const
+{
+    check( pBuffer && pChannel );
+
+    *pBuffer = -1;
+    *pChannel = -1;
+
+    int index = 0;
+
+    for( size_t b=0; b<m_buffers.Num(); ++b )
+    {
+        for( size_t c=0; c<m_buffers[b].m_channels.Num(); ++c )
+        {
+            if ( m_buffers[b].m_channels[c].m_semantic == semantic )
+            {
+                if (index==semanticIndex)
+                {
+                    *pBuffer = int(b);
+                    *pChannel = int(c);
+                    return;
+                }
+                else
+                {
+                    ++index;
+                }
+            }
+        }
+    }
+}
+
+
+//---------------------------------------------------------------------------------------------
+int32 FMeshBufferSet::GetElementSize( int32 buffer ) const
+{
+    check( buffer>=0 && buffer<m_buffers.Num() );
+
+    return m_buffers[buffer].m_elementSize;
+}
+
+
+//---------------------------------------------------------------------------------------------
+int32 FMeshBufferSet::GetChannelOffset
+    (
+        int32 buffer,
+        int32 channel
+    ) const
+{
+    check( buffer>=0 && buffer<m_buffers.Num() );
+    check( channel>=0 && channel<m_buffers[buffer].m_channels.Num() );
+
+    return m_buffers[buffer].m_channels[channel].m_offset;
+}
+
+
+//---------------------------------------------------------------------------------------------
+void FMeshBufferSet::AddBuffer( const FMeshBufferSet& Other, int32 buffer )
+{
+    check( GetElementCount()==Other.GetElementCount() );
+
+	LLM_SCOPE_BYNAME(TEXT("MutableRuntime"));
+
+    m_buffers.Add( Other.m_buffers[buffer] );
+}
+
+
+//---------------------------------------------------------------------------------------------
+void FMeshBufferSet::RemoveBuffer(int32 BufferIndex)
+{
+	if (BufferIndex >= 0 && BufferIndex < m_buffers.Num())
+	{
+		m_buffers.RemoveAt(BufferIndex);
+	}
+	else
+	{
+		check(false);
+	}
+}
+
+
+//---------------------------------------------------------------------------------------------
+bool FMeshBufferSet::HasSameFormat( const FMeshBufferSet& Other ) const
+{
+    int32 bc = GetBufferCount();
+    if (Other.GetBufferCount() != bc )
+    {
+        return false;
+    }
+
+    for ( int32 b=0; b<bc; ++b )
+    {
+        if ( !HasSameFormat( b, Other ) )
+        {
+            return false;
+        }
+    }
+
+    return true;
+}
+
+
+//---------------------------------------------------------------------------------------------
+bool FMeshBufferSet::HasSameFormat( int32 buffer, const FMeshBufferSet& Other ) const
+{
+    if ( m_buffers[buffer].m_channels ==  Other.m_buffers[buffer].m_channels
+         &&
+         m_buffers[buffer].m_elementSize == Other.m_buffers[buffer].m_elementSize
+         )
+    {
+        return true;
+    }
+
+    return false;
+}
+
+
+//---------------------------------------------------------------------------------------------
+size_t FMeshBufferSet::GetDataSize() const
+{
+    size_t res = 0;
+
+    for( int32 b=0; b<m_buffers.Num(); ++b )
+    {
+        res += m_buffers[b].m_elementSize * m_elementCount;
+    }
+
+    return res;
+}
+
+
+//-----------------------------------------------------------------------------------------
+void FMeshBufferSet::CopyElement(uint32_t fromIndex, uint32_t toIndex)
+{
+	check(fromIndex < m_elementCount);
+	check(toIndex < m_elementCount);
+
+	if (fromIndex != toIndex)
+	{
+		for (auto& b : m_buffers)
+		{
+			FMemory::Memcpy(&b.m_data[b.m_elementSize * toIndex],
+				&b.m_data[b.m_elementSize * fromIndex],
+				b.m_elementSize);
+		}
+	}
+}
+
+
+//-----------------------------------------------------------------------------------------
+bool FMeshBufferSet::IsSpecialBufferToIgnoreInSimilar(const MESH_BUFFER& b) const
+{
+	if (b.m_channels.Num() == 1
+		&&
+		b.m_channels[0].m_semantic == MBS_VERTEXINDEX)
+	{
+		return true;
+	}
+	if (b.m_channels.Num() == 1
+		&&
+		b.m_channels[0].m_semantic == MBS_LAYOUTBLOCK)
+	{
+		return true;
+	}
+	return false;
+}
+
+//-----------------------------------------------------------------------------------------
+bool FMeshBufferSet::IsSimilarRobust(const FMeshBufferSet& Other) const
+{
+
+    MUTABLE_CPUPROFILER_SCOPE(FMeshBufferSet::IsSimilarRobust);
+
+	if ( m_elementCount != Other.m_elementCount )
+	{ 
+		return false;
+	}
+
+	// IsSimilar is mush faster but can give false negatives if the buffer data description 
+	// omits parts of the data (e.g, memory layout paddings). It cannot have false positives. 
+	if ( IsSimilar(Other) )
+	{
+		return true;
+	}
+
+	const int32 ThisNumBuffers = m_buffers.Num();
+	const int32 OtherNumBuffers = Other.m_buffers.Num();
+	
+	int32 I = 0, J = 0;
+	while ( I < ThisNumBuffers && J < OtherNumBuffers )
+	{
+		if ( IsSpecialBufferToIgnoreInSimilar(m_buffers[I]) )
+		{
+			++I;
+			continue;
+		}
+
+		if ( IsSpecialBufferToIgnoreInSimilar(Other.m_buffers[J]) )
+		{
+			++J;
+			continue;
+		}
+		
+		const int32 ThisNumChannels = m_buffers[I].m_channels.Num();
+		const int32 OtherNumChannels = m_buffers[J].m_channels.Num();
+
+		if ( ThisNumChannels != OtherNumChannels )
+		{
+			return false;
+		}
+
+		const MESH_BUFFER& ThisBuffer = m_buffers[I];
+		const MESH_BUFFER& OtherBuffer = Other.m_buffers[J];
+
+		for ( uint32 Elem = 0; Elem < m_elementCount; ++Elem )
+		{
+			for ( int32 C = 0; C < ThisNumChannels; ++C )
+			{
+				const SIZE_T SizeA = GetMeshFormatData(ThisBuffer.m_channels[C].m_format).m_size * ThisBuffer.m_channels[C].m_componentCount;
+				const SIZE_T SizeB = GetMeshFormatData(OtherBuffer.m_channels[C].m_format).m_size * OtherBuffer.m_channels[C].m_componentCount;
+
+				if ( SizeA != SizeB )
+				{
+					return false;
+				}
+
+				const uint8* BuffA = ThisBuffer.m_data.GetData() + Elem * ThisBuffer.m_elementSize + ThisBuffer.m_channels[C].m_offset;
+				const uint8* BuffB = OtherBuffer.m_data.GetData() + Elem * OtherBuffer.m_elementSize + OtherBuffer.m_channels[C].m_offset;
+
+				if ( FMemory::Memcmp(BuffA, BuffB, SizeA) != 0 )
+				{
+					return false;
+				}
+			}
+		}		
+
+		++J;
+		++I;
+	}
+
+	if ( I + 1 < ThisNumBuffers || J + 1 < OtherNumBuffers )
+	{
+		return false;
+	}
+
+	if ( I < ThisNumBuffers && !IsSpecialBufferToIgnoreInSimilar(m_buffers[I]) )
+	{
+		return false;
+	}
+
+	if ( J < OtherNumBuffers && !IsSpecialBufferToIgnoreInSimilar(Other.m_buffers[J]) )
+	{
+		return false;
+	}
+	
+	return true;
+}
+
+//-----------------------------------------------------------------------------------------
+bool FMeshBufferSet::IsSimilar(const FMeshBufferSet& o) const
+{
+    MUTABLE_CPUPROFILER_SCOPE(FMeshBufferSet::IsSimilar);
+
+	if (m_elementCount != o.m_elementCount) return false;
+
+	// Compare all buffers except the vertex index channel, which should always be alone in
+	// the last buffer
+	int32 i = 0;
+	int32 oi = 0;
+
+	while (i < m_buffers.Num() && oi < o.m_buffers.Num())
+	{
+		// Is it a special buffer that we should ignore?
+		if (IsSpecialBufferToIgnoreInSimilar(m_buffers[i]))
+		{
+			++i;
+		}
+		else
+		{
+			if (IsSpecialBufferToIgnoreInSimilar(o.m_buffers[oi]))
+			{
+				++oi;
+			}
+			else
+			{
+				if (!(m_buffers[i] == o.m_buffers[oi]))
+				{
+					return false;
+				}
+				++i;
+				++oi;
+			}
+		}
+	}
+
+	// Remaining buffers?
+	if (i + 1 < m_buffers.Num()
+		||
+		oi + 1 < o.m_buffers.Num())
+	{
+		return false;
+	}
+
+	if (i < m_buffers.Num()
+		&&
+		!IsSpecialBufferToIgnoreInSimilar(m_buffers[i])
+		)
+	{
+		return false;
+	}
+
+	if (oi < o.m_buffers.Num()
+		&&
+		!IsSpecialBufferToIgnoreInSimilar(o.m_buffers[oi])
+		)
+	{
+		return false;
+	}
+
+	return true;
+}
+
+
+//-----------------------------------------------------------------------------------------
+void FMeshBufferSet::ResetBufferIndices()
+{
+	uint8_t currentIndices[MBS_COUNT];
+	memset(currentIndices, 0, sizeof(currentIndices));
+
+	for (auto& b : m_buffers)
+	{
+		for (auto& c : b.m_channels)
+		{
+			c.m_semanticIndex = currentIndices[c.m_semantic];
+			currentIndices[c.m_semantic]++;
+		}
+	}
+}
+
+
+//-----------------------------------------------------------------------------------------
+void FMeshBufferSet::UpdateOffsets(int32 b)
+{
+	uint8_t offset = 0;
+	for (auto& c : m_buffers[b].m_channels)
+	{
+		if (c.m_offset < offset)
+		{
+			c.m_offset = offset;
+		}
+		else
+		{
+			offset = c.m_offset;
+		}
+		offset += uint8_t(c.m_componentCount * GetMeshFormatData(c.m_format).m_size);
+	}
+
+	if (m_buffers[b].m_elementSize < offset)
+		m_buffers[b].m_elementSize = offset;
+}
diff --git a/Engine/Plugins/Experimental/Mutable/Source/MutableRuntime/Private/Model.cpp b/Engine/Plugins/Experimental/Mutable/Source/MutableRuntime/Private/Model.cpp
new file mode 100644
--- /dev/null
+++ b/Engine/Plugins/Experimental/Mutable/Source/MutableRuntime/Private/Model.cpp
@@ -0,0 +1,1155 @@
+// Copyright Epic Games, Inc. All Rights Reserved.
+
+
+#include "MutableRuntime/Public/Model.h"
+
+#include "MutableRuntime/Private/Platform.h"
+
+#include "ModelPrivate.h"
+#include "Operations.h"
+#include "CodeVisitor.h"
+
+#include "Hash/CityHash.h"
+
+#include <algorithm>
+
+namespace mu
+{
+
+    //---------------------------------------------------------------------------------------------
+    void PROGRAM::Check()
+    {
+//    #ifdef MUTABLE_DEBUG
+//        // Process all from root
+//        if (m_states.size())
+//        {
+//            // Super debug: enable in case of emergencies.
+////            InvariantCodeVisitor invariantVisitor;
+////            invariantVisitor.Apply( *this, m_states[0].m_root );
+//        }
+
+//        for ( size_t i=0; i<m_code.size(); ++i )
+//        {
+//            if ( m_code[i].type == OP_TYPE::IM_PIXELFORMAT )
+//            {
+//                check( m_code[i].args.ImagePixelFormat.format != IF_NONE );
+//                check( m_code[i].args.ImagePixelFormat.source != i );
+//            }
+//    //
+//    //            if ( m_code[i].type == OP_TYPE::IM_MIPMAP )
+//    //            {
+//    //                check( m_code[i].args.ImageMipmap.blockLevels != 0 );
+//    //            }
+
+//            if ( m_code[i].type == OP_TYPE::ME_APPLYLAYOUT )
+//            {
+//                check( m_code[i].args.MeshApplyLayout.mesh != 0 );
+//                check( m_code[i].args.MeshApplyLayout.mesh != i );
+//            }
+
+//            if ( m_code[i].type == OP_TYPE::ME_MORPH2 )
+//            {
+//                for (int j=0;j<MUTABLE_OP_MAX_MORPH2_TARGETS;++j)
+//                {
+//                    auto target = m_code[i].args.MeshMorph2.targets[j];
+//                    check( m_code[target].type!=OP_TYPE::ME_APPLYLAYOUT );
+//                }
+//            }
+
+//            if ( m_code[i].type == OP_TYPE::ME_CONDITIONAL ||
+//                 m_code[i].type == OP_TYPE::IN_CONDITIONAL ||
+//                 m_code[i].type == OP_TYPE::IM_CONDITIONAL ||
+//                 m_code[i].type == OP_TYPE::LA_CONDITIONAL )
+//            {
+//                for (int j=0;j<MUTABLE_OP_MAX_MORPH2_TARGETS;++j)
+//                {
+//                    check( m_code[i].args.Conditional.condition!=i );
+//                    check( m_code[i].args.Conditional.yes!=i );
+//                    check( m_code[i].args.Conditional.no!=i );
+//                }
+//            }
+
+//        }
+//    #endif
+    }
+
+
+    //---------------------------------------------------------------------------------------------
+    void PROGRAM::LogHistogram() const
+    {
+#if 0
+        uint64_t countPerType[(int)OP_TYPE::COUNT];
+        mutable_memset(countPerType,0,sizeof(countPerType));
+
+        for ( const auto& o: m_opAddress )
+        {
+            auto type = GetOpType(o);
+            countPerType[(int)type]++;
+        }
+
+        vector< pair<uint64_t,OP_TYPE> > sorted((int)OP_TYPE::COUNT);
+        for (int i=0; i<(int)OP_TYPE::COUNT; ++i)
+        {
+            sorted[i].second = (OP_TYPE)i;
+            sorted[i].first = countPerType[i];
+        }
+
+        std::sort(sorted.begin(),sorted.end(), []( const pair<uint64_t,OP_TYPE>& a, const pair<uint64_t,OP_TYPE>& b )
+        {
+            return a.first>b.first;
+        });
+
+        UE_LOG(LogMutableCore,Log, TEXT("Op histogram (%llu ops):"), m_opAddress.size());
+        for(int i=0; i<8; ++i)
+        {
+            float p = sorted[i].first/float(m_opAddress.size())*100.0f;
+            UE_LOG(LogMutableCore,Log, TEXT("  %3.2f%% : %d"), p, (int)sorted[i].second );
+        }
+#endif
+    }
+
+
+    //---------------------------------------------------------------------------------------------
+    //---------------------------------------------------------------------------------------------
+    //---------------------------------------------------------------------------------------------
+    Model::Model()
+    {
+        m_pD = new Private();
+    }
+
+
+    //---------------------------------------------------------------------------------------------
+    Model::~Model()
+    {
+        check( m_pD );
+        delete m_pD;
+        m_pD = 0;
+    }
+
+
+    //---------------------------------------------------------------------------------------------
+    void Model::Serialise( const Model* p, OutputArchive& arch )
+    {
+		LLM_SCOPE_BYNAME(TEXT("MutableRuntime"));
+
+        arch << *p->m_pD;
+    }
+
+    //---------------------------------------------------------------------------------------------
+    //---------------------------------------------------------------------------------------------
+    //---------------------------------------------------------------------------------------------
+    class OutputModelStream : public OutputStream
+    {
+    public:
+
+        //-----------------------------------------------------------------------------------------
+        // Life cycle
+        //-----------------------------------------------------------------------------------------
+        OutputModelStream( ModelStreamer* pStreamer )
+            : m_pStreamer( pStreamer )
+        {
+        }
+
+        //-----------------------------------------------------------------------------------------
+        // InputStream interface
+        //-----------------------------------------------------------------------------------------
+        void Write( const void* pData, uint64 size ) override
+        {
+            m_pStreamer->Write( pData, size );
+        }
+
+
+    private:
+
+        ModelStreamer* m_pStreamer;
+    };
+
+
+    //---------------------------------------------------------------------------------------------
+    void Model::Serialise( Model* p, ModelStreamer& streamer )
+    {
+		LLM_SCOPE_BYNAME(TEXT("MutableRuntime"));
+
+		mu::PROGRAM& program = p->m_pD->m_program;
+
+		TArray<TPair<int32, mu::ImagePtrConst>> InitialImages = program.m_constantImageLODs;
+		TArray<TPair<int32, mu::MeshPtrConst>> InitialMeshes = program.m_constantMeshes;
+
+		// Save images and unload from memory
+		for (int32 ResourceIndex = 0; ResourceIndex < program.m_constantImageLODs.Num(); ++ResourceIndex)
+		{
+			TPair<int32, mu::ImagePtrConst>& ResData = program.m_constantImageLODs[ResourceIndex];
+
+			// This shouldn't have been serialised with rom support before.
+			if (ResData.Key < 0)
+			{
+				continue;
+			}
+
+			int32 RomIndex = ResData.Key;
+			const FRomData& RomData = program.m_roms[RomIndex];
+			check(RomData.ResourceType == DT_IMAGE);
+			check(RomData.ResourceIndex == ResourceIndex);
+
+			// Serialize to memory, to find out final size of this rom
+			OutputMemoryStream MemStream(1024*1024);
+			OutputArchive MemoryArch(&MemStream);
+			Image::Serialise(ResData.Value.get(), MemoryArch);
+			check(RomData.Size == MemStream.GetBufferSize());
+
+			streamer.OpenWriteFile(p->GetLocation(), RomData.Id);
+			streamer.Write(MemStream.GetBuffer(), MemStream.GetBufferSize());
+			streamer.CloseWriteFile();
+
+			// We clear it to avoid including it also with the main model data
+			ResData.Value = nullptr;
+		}
+
+		// Save meshes and unload from memory
+		for (int32 ResourceIndex = 0; ResourceIndex < program.m_constantMeshes.Num(); ++ResourceIndex)
+		{
+			TPair<int32, mu::MeshPtrConst>& ResData = program.m_constantMeshes[ResourceIndex];
+
+			// This shouldn't have been serialised with rom support before.
+			if (ResData.Key < 0)
+			{
+				continue;
+			}
+
+			int32 RomIndex = ResData.Key;
+			const FRomData& RomData = program.m_roms[RomIndex];
+			check(RomData.ResourceType == DT_MESH);
+			check(RomData.ResourceIndex == ResourceIndex);
+
+			// Serialize to memory, to find out final size of this rom
+			OutputMemoryStream MemStream(1024 * 1024);
+			OutputArchive MemoryArch(&MemStream);
+			Mesh::Serialise(ResData.Value.get(), MemoryArch);
+			check(RomData.Size == MemStream.GetBufferSize());
+
+			streamer.OpenWriteFile(p->GetLocation(), RomData.Id);
+			streamer.Write(MemStream.GetBuffer(), MemStream.GetBufferSize());
+			streamer.CloseWriteFile();
+
+			// We clear it to avoid including it also with the main model data
+			ResData.Value = nullptr;
+		}
+
+		// Store the main data of the model
+		{
+			streamer.OpenWriteFile(p->GetLocation(), 0);
+			OutputModelStream stream(&streamer);
+			OutputArchive arch(&stream);
+
+			arch << *p->m_pD;
+
+			streamer.CloseWriteFile();
+		}
+
+		// Restore full data
+		program.m_constantImageLODs = InitialImages;
+		program.m_constantMeshes = InitialMeshes;
+	}
+
+
+    //---------------------------------------------------------------------------------------------
+    bool Model::HasExternalData() const
+    {
+		return m_pD->m_program.m_roms.Num() > 0;
+    }
+
+
+    //---------------------------------------------------------------------------------------------
+    void Model::UnloadExternalData()
+    {
+		LLM_SCOPE_BYNAME(TEXT("MutableRuntime"));
+
+		for (int32 ResIndex = 0; ResIndex < m_pD->m_program.m_constantImageLODs.Num(); ++ResIndex)
+		{
+			if (m_pD->m_program.m_constantImageLODs[ResIndex].Key >= 0)
+			{
+				m_pD->m_program.m_constantImageLODs[ResIndex].Value = nullptr;
+			}
+		}
+
+		for (int32 ResIndex = 0; ResIndex < m_pD->m_program.m_constantMeshes.Num(); ++ResIndex)
+		{
+			if (m_pD->m_program.m_constantMeshes[ResIndex].Key >= 0)
+			{
+				m_pD->m_program.m_constantMeshes[ResIndex].Value = nullptr;
+			}
+		}
+	}
+
+
+    //---------------------------------------------------------------------------------------------
+    ModelPtr Model::StaticUnserialise( InputArchive& arch )
+    {
+		LLM_SCOPE_BYNAME(TEXT("MutableRuntime"));
+
+        ModelPtr pResult = new Model();
+        arch >> *pResult->m_pD;
+        return pResult;
+    }
+
+
+    //---------------------------------------------------------------------------------------------
+    const char* Model::GetLocation( ) const
+    {
+        return m_pD->m_location.c_str();
+    }
+
+
+    //---------------------------------------------------------------------------------------------
+    void Model::SetLocation( const char* strLocation )
+    {
+        if (strLocation)
+        {
+            m_pD->m_location = strLocation;
+        }
+    }
+
+
+    //---------------------------------------------------------------------------------------------
+    Model::Private* Model::GetPrivate() const
+    {
+        return m_pD;
+    }
+
+
+    //---------------------------------------------------------------------------------------------
+    void Model::ClearCaches()
+    {
+		LLM_SCOPE_BYNAME(TEXT("MutableRuntime"));
+
+        if (m_pD)
+        {
+            m_pD->m_generatedResources.clear();
+        }
+    }
+
+
+    //---------------------------------------------------------------------------------------------
+    ParametersPtr Model::NewParameters( const Parameters* pOld ) const
+    {
+		LLM_SCOPE_BYNAME(TEXT("MutableRuntime"));
+
+        ParametersPtr pRes = new Parameters();
+
+        pRes->GetPrivate()->m_pModel = this;
+
+        pRes->GetPrivate()->m_values.resize( m_pD->m_program.m_parameters.size() );
+        for ( std::size_t p=0; p<m_pD->m_program.m_parameters.size(); ++p )
+        {
+            pRes->GetPrivate()->m_values[p] = m_pD->m_program.m_parameters[p].m_defaultValue;
+        }
+
+        // Copy old values
+        if ( pOld )
+        {
+            for ( int p=0; p<pOld->GetCount(); ++p )
+            {
+                int thisP = pRes->GetPrivate()->Find( pOld->GetName(p) );
+
+                if ( thisP>=0 )
+                {
+                    if ( pOld->GetType(p)==pRes->GetType(thisP) )
+                    {
+                        switch ( pRes->GetType(thisP) )
+                        {
+						case PARAMETER_TYPE::T_BOOL:
+                            pRes->SetBoolValue( thisP, pOld->GetBoolValue(p) );
+                            break;
+
+                        case PARAMETER_TYPE::T_INT:
+                            pRes->SetIntValue( thisP, pOld->GetIntValue(p) );
+                            break;
+
+                        case PARAMETER_TYPE::T_FLOAT:
+                            pRes->SetFloatValue( thisP, pOld->GetFloatValue(p) );
+                            break;
+
+                        case PARAMETER_TYPE::T_COLOUR:
+                        {
+                            float r,g,b;
+                            pOld->GetColourValue( p, &r, &g, &b );
+                            pRes->SetColourValue( thisP, r, g, b );
+                            break;
+                        }
+
+                        case PARAMETER_TYPE::T_PROJECTOR:
+                        {
+//							float m[16];
+//							pOld->GetProjectorValue( p, m );
+                            pRes->GetPrivate()->m_values[thisP].m_projector =
+                                    pOld->GetPrivate()->m_values[p].m_projector;
+                            break;
+                        }
+
+                        case PARAMETER_TYPE::T_IMAGE:
+                            pRes->SetImageValue( thisP, pOld->GetImageValue(p) );
+                            break;
+
+                        default:
+                            check(false);
+                            break;
+                        }
+                    }
+                }
+            }
+        }
+
+        return pRes;
+    }
+
+
+    //---------------------------------------------------------------------------------------------
+    int Model::GetStateCount() const
+    {
+        return (int)m_pD->m_program.m_states.size();
+    }
+
+
+    //---------------------------------------------------------------------------------------------
+    const char* Model::GetStateName( int index ) const
+    {
+        const char* strRes = 0;
+
+        if ( index>=0 && index<(int)m_pD->m_program.m_states.size() )
+        {
+            strRes = m_pD->m_program.m_states[index].m_name.c_str();
+        }
+
+        return strRes;
+    }
+
+
+    //---------------------------------------------------------------------------------------------
+    int Model::FindState( const char* strName ) const
+    {
+        int res = -1;
+
+        for ( int i=0; res<0 && i<(int)m_pD->m_program.m_states.size(); ++i )
+        {
+            if ( m_pD->m_program.m_states[i].m_name == strName )
+            {
+                res = i;
+            }
+        }
+
+        return res;
+    }
+
+
+    //---------------------------------------------------------------------------------------------
+    int Model::GetStateParameterCount( int stateIndex ) const
+    {
+        int res = -1;
+
+        if ( stateIndex>=0 && stateIndex<(int)m_pD->m_program.m_states.size() )
+        {
+            res = (int)m_pD->m_program.m_states[stateIndex].m_runtimeParameters.size();
+        }
+
+        return res;
+    }
+
+
+    //---------------------------------------------------------------------------------------------
+    int Model::GetStateParameterIndex( int stateIndex, int paramIndex ) const
+    {
+        int res = -1;
+
+        if ( stateIndex>=0 && stateIndex<(int)m_pD->m_program.m_states.size() )
+        {
+            const PROGRAM::STATE& state = m_pD->m_program.m_states[stateIndex];
+            if ( paramIndex>=0 && paramIndex<(int)state.m_runtimeParameters.size() )
+            {
+                res = (int)state.m_runtimeParameters[paramIndex];
+            }
+        }
+
+        return res;
+    }
+
+
+    //---------------------------------------------------------------------------------------------
+    static size_t AddMultiValueKeys( vector<uint8_t>& parameterValuesBlob, size_t pos,
+                                  const map< vector<int>, PARAMETER_VALUE >& multi )
+    {
+        parameterValuesBlob.resize( pos+4 );
+        uint32_t s = uint32_t( multi.size() );
+        memcpy( &parameterValuesBlob[pos], &s, sizeof(uint32_t) );
+        pos+=4;
+
+        for(const auto& v: multi)
+        {
+            uint32_t ds = uint32_t( v.first.size() );
+
+            parameterValuesBlob.resize( pos + 4 + ds*4 );
+
+            memcpy( &parameterValuesBlob[pos], &s, sizeof(uint32_t) );
+            pos+=4;
+
+            memcpy( &parameterValuesBlob[pos],v.first.data(), ds*sizeof(int32_t) );
+            pos += ds*sizeof(int32_t);
+        }
+
+        return pos;
+    }
+
+
+    //---------------------------------------------------------------------------------------------
+    uint32_t Model::Private::GetResourceKey( uint32_t paramListIndex,
+                                             OP::ADDRESS rootAt,
+                                             const Parameters* pParams )
+    {
+        // Find the list of relevant parameters
+        const vector<uint16_t>* params = nullptr;
+        if (paramListIndex<m_program.m_parameterLists.size())
+        {
+            params = &m_program.m_parameterLists[paramListIndex];
+        }
+        check(params);
+        if (!params)
+        {
+            return 0xffff;
+        }
+
+        // Generate the relevant parameters blob
+        vector<uint8_t> parameterValuesBlob;
+		parameterValuesBlob.reserve(1024);
+        for (int param: *params)
+        {
+            size_t pos = parameterValuesBlob.size();
+            size_t dataSize = 0;
+
+            switch(m_program.m_parameters[param].m_type)
+            {
+            case PARAMETER_TYPE::T_BOOL:
+                dataSize = 1;
+                parameterValuesBlob.push_back( pParams->GetPrivate()->m_values[param].m_bool ? 1 : 0 );
+                pos += dataSize;
+
+                // Multi-values
+                if ( param < int(pParams->GetPrivate()->m_multiValues.size()) )
+                {
+                    const auto& multi = pParams->GetPrivate()->m_multiValues[param];
+                    pos = AddMultiValueKeys( parameterValuesBlob, pos, multi );
+                    for(const auto& v: multi)
+                    {
+                        parameterValuesBlob.push_back( v.second.m_bool ? 1 : 0 );
+                        pos += dataSize;
+                    }
+                }
+                break;
+
+            case PARAMETER_TYPE::T_INT:
+                dataSize = sizeof(int32_t);
+                parameterValuesBlob.resize( pos+dataSize );
+                memcpy( &parameterValuesBlob[pos],
+                        &pParams->GetPrivate()->m_values[param].m_int,
+                        dataSize );
+                pos += dataSize;
+
+                // Multi-values
+                if ( param < int(pParams->GetPrivate()->m_multiValues.size()) )
+                {
+                    const auto& multi = pParams->GetPrivate()->m_multiValues[param];
+                    pos = AddMultiValueKeys( parameterValuesBlob, pos, multi );
+                    parameterValuesBlob.resize( pos+multi.size()*dataSize );
+                    for(const auto& v: multi)
+                    {
+                        memcpy( &parameterValuesBlob[pos], &v.second.m_int, dataSize );
+                        pos += dataSize;
+                    }
+                }
+                break;
+
+            //! Floating point value in the range of 0.0 to 1.0
+            case PARAMETER_TYPE::T_FLOAT:
+                dataSize = sizeof(float);
+                parameterValuesBlob.resize( pos+dataSize );
+                memcpy( &parameterValuesBlob[pos],
+                        &pParams->GetPrivate()->m_values[param].m_float,
+                        dataSize );
+                pos += dataSize;
+
+                // Multi-values
+                if ( param < int(pParams->GetPrivate()->m_multiValues.size()) )
+                {
+                    const auto& multi = pParams->GetPrivate()->m_multiValues[param];
+                    pos = AddMultiValueKeys( parameterValuesBlob, pos, multi );
+                    parameterValuesBlob.resize( pos+multi.size()*dataSize );
+                    for(const auto& v: multi)
+                    {
+                        memcpy( &parameterValuesBlob[pos], &v.second.m_float,dataSize );
+                        pos += dataSize;
+                    }
+                }
+                break;
+
+            //! Floating point RGBA colour, with each channel ranging from 0.0 to 1.0
+            case PARAMETER_TYPE::T_COLOUR:
+                dataSize = 4*sizeof(float);
+                parameterValuesBlob.resize( pos+dataSize );
+                memcpy( &parameterValuesBlob[pos],
+                        &pParams->GetPrivate()->m_values[param].m_colour,
+                        dataSize );
+                pos += dataSize;
+
+                // Multi-values
+                if ( param < int(pParams->GetPrivate()->m_multiValues.size()) )
+                {
+                    const auto& multi = pParams->GetPrivate()->m_multiValues[param];
+                    pos = AddMultiValueKeys( parameterValuesBlob, pos, multi );
+                    parameterValuesBlob.resize( pos+multi.size()*dataSize );
+                    for(const auto& v: multi)
+                    {
+                        memcpy( &parameterValuesBlob[pos], &v.second.m_colour,dataSize );
+                        pos += dataSize;
+                    }
+                }
+                break;
+
+            //! 3D Projector type, defining a position, scale and orientation. Basically used for
+            //! projected decals.
+            case PARAMETER_TYPE::T_PROJECTOR:
+                dataSize = sizeof(PROJECTOR);
+
+                // \todo: padding will be random?
+                parameterValuesBlob.resize( pos+dataSize );
+                memcpy( &parameterValuesBlob[pos],
+                        &pParams->GetPrivate()->m_values[param].m_projector,
+                        dataSize );
+				pos += dataSize;
+
+                // Multi-values
+                if ( param < int(pParams->GetPrivate()->m_multiValues.size()) )
+                {
+                    const auto& multi = pParams->GetPrivate()->m_multiValues[param];
+                    pos = AddMultiValueKeys( parameterValuesBlob, pos, multi );
+                    parameterValuesBlob.resize( pos+multi.size()*dataSize );
+                    for(const auto& v: multi)
+                    {
+                        memcpy( &parameterValuesBlob[pos], &v.second.m_projector,dataSize );
+                        pos += dataSize;
+                    }
+                }
+                break;
+
+            case PARAMETER_TYPE::T_IMAGE:
+                dataSize = sizeof(EXTERNAL_IMAGE_ID);
+                parameterValuesBlob.resize( pos+dataSize );
+                memcpy( &parameterValuesBlob[pos],
+                        &pParams->GetPrivate()->m_values[param].m_image,
+                        dataSize );
+				pos += dataSize;
+
+                // Multi-values
+                if ( param < int(pParams->GetPrivate()->m_multiValues.size()) )
+                {
+                    const auto& multi = pParams->GetPrivate()->m_multiValues[param];
+                    pos = AddMultiValueKeys( parameterValuesBlob, pos, multi );
+                    parameterValuesBlob.resize( pos+multi.size()*dataSize );
+                    for(const auto& v: multi)
+                    {
+                        memcpy( &parameterValuesBlob[pos], &v.second.m_image,dataSize );
+                        pos += dataSize;
+                    }
+                }
+                break;
+
+            default:
+                // unsupported parameter type
+                check(false);
+            }
+        }
+
+        // Increase the request id
+        ++m_lastResourceResquestId;
+
+        // See if we already have this id
+        size_t oldestCachePosition = 0;
+        for (size_t i=0; i<m_generatedResources.size(); ++i)
+        {
+            auto& key = m_generatedResources[i];
+            if (key.m_rootAddress==rootAt
+                    &&
+                    key.m_parameterValuesBlob==parameterValuesBlob)
+            {
+                key.m_lastRequestId = m_lastResourceResquestId;
+                return key.m_id;
+            }
+            else
+            {
+                if ( m_generatedResources[oldestCachePosition].m_lastRequestId
+                     >
+                     key.m_lastRequestId )
+                {
+                    oldestCachePosition = i;
+                }
+            }
+        }
+
+        // Generate a new id
+        uint32_t newId = ++m_lastResourceKeyId;
+        RESOURCE_KEY newKey;
+        newKey.m_id = newId;
+        newKey.m_lastRequestId = m_lastResourceResquestId;
+        newKey.m_rootAddress = rootAt;
+        newKey.m_parameterValuesBlob = std::move(parameterValuesBlob);
+
+        // TODO: Move the constant to settings?
+        const size_t maxGeneratedResourcesIDCacheSize = 1024;
+        if (m_generatedResources.size()>=maxGeneratedResourcesIDCacheSize)
+        {
+            m_generatedResources[oldestCachePosition] = std::move(newKey);
+        }
+        else
+        {
+            m_generatedResources.push_back(std::move(newKey));
+        }
+
+        return newId;
+    }
+
+
+    //---------------------------------------------------------------------------------------------
+    //---------------------------------------------------------------------------------------------
+    //---------------------------------------------------------------------------------------------
+    ModelParametersGenerator::ModelParametersGenerator
+        (
+            const Model* pModel,
+            System* pSystem,
+            bool considerRelevancy
+        )
+    {
+		LLM_SCOPE_BYNAME(TEXT("MutableRuntime"));
+
+        m_pD = new Private();
+        m_pD->m_pModel = pModel;
+        m_pD->m_pSystem = pSystem;
+
+        size_t paramCount = pModel->GetPrivate()->m_program.m_parameters.size();
+        m_pD->m_considerRelevancy = considerRelevancy;
+        if (m_pD->m_considerRelevancy)
+        {
+            vector<int> currentValues( paramCount, 0 );
+            m_pD->m_intervals.resize( paramCount );
+            m_pD->m_instanceCount = (int)m_pD->BuildIntervals( 0, 0, currentValues );
+        }
+        else
+        {
+            m_pD->m_instanceCount = 1;
+            for (size_t i=0;i<paramCount; ++i)
+            {
+                const PARAMETER_DESC& param =
+                        m_pD->m_pModel->GetPrivate()->m_program.m_parameters[ i ];
+                switch ( param.m_type )
+                {
+
+                case PARAMETER_TYPE::T_INT:
+                {
+                    m_pD->m_instanceCount *= param.m_possibleValues.size();
+                    break;
+                }
+
+                case PARAMETER_TYPE::T_BOOL:
+                {
+                    m_pD->m_instanceCount *= 2;
+                    break;
+                }
+
+                default:
+                    // Parameter not discrete, so ignored for combinations.
+                    break;
+                }
+            }
+        }
+    }
+
+
+    //---------------------------------------------------------------------------------------------
+    ModelParametersGenerator::~ModelParametersGenerator()
+    {
+		LLM_SCOPE_BYNAME(TEXT("MutableRuntime"));
+
+        check( m_pD );
+        delete m_pD;
+        m_pD = 0;
+    }
+
+
+    //---------------------------------------------------------------------------------------------
+    ModelParametersGenerator::Private* ModelParametersGenerator::GetPrivate() const
+    {
+        return m_pD;
+    }
+
+
+    //---------------------------------------------------------------------------------------------
+    int64_t ModelParametersGenerator::GetInstanceCount()
+    {
+        return m_pD->m_instanceCount;
+    }
+
+
+    //---------------------------------------------------------------------------------------------
+    ParametersPtr ModelParametersGenerator::GetInstance( int64_t index,
+                                                         float (*randomGenerator )() )
+    {
+		LLM_SCOPE_BYNAME(TEXT("MutableRuntime"));
+
+        ParametersPtr res = m_pD->m_pModel->NewParameters();
+
+        if (m_pD->m_considerRelevancy)
+        {
+            vector<int> values = m_pD->GetParameters( (int)index );
+            for ( size_t p=0; p<values.size(); ++p )
+            {
+                switch ( res->GetType( (int)p ) )
+                {
+                case PARAMETER_TYPE::T_BOOL:
+                    res->SetBoolValue( (int)p, values[p]!=0 );
+                    break;
+
+                case PARAMETER_TYPE::T_INT:
+                    res->SetIntValue( (int)p, res->GetIntPossibleValue( (int)p, values[p] ) );
+                    break;
+
+                case PARAMETER_TYPE::T_FLOAT:
+                    if ( randomGenerator )
+                    {
+                        res->SetFloatValue( (int)p, randomGenerator() );
+                    }
+                    break;
+
+                case PARAMETER_TYPE::T_COLOUR:
+                    if ( randomGenerator )
+                    {
+                        res->SetColourValue( (int)p,
+                                             randomGenerator(), randomGenerator(), randomGenerator() );
+                    }
+                    break;
+
+                default:
+                    break;
+                }
+            }
+        }
+        else
+        {
+            int paramCount = (int)m_pD->m_pModel->GetPrivate()->m_program.m_parameters.size();
+            int64_t currentInstance = index;
+            for (int i=0;i<paramCount;++i)
+            {
+                const PARAMETER_DESC& param = m_pD->m_pModel->GetPrivate()->m_program.m_parameters[ i ];
+                switch ( param.m_type )
+                {
+
+                case PARAMETER_TYPE::T_INT:
+                {
+                    int numOptions = res->GetIntPossibleValueCount( i );
+                    int value = res->GetIntPossibleValue( i, (int)(currentInstance%numOptions) );
+                    res->SetIntValue( i, value );
+                    currentInstance /= numOptions;
+                    break;
+                }
+
+                case PARAMETER_TYPE::T_BOOL:
+                {
+                    res->SetBoolValue( i, currentInstance%2!=0 );
+                    currentInstance /= 2;
+                    break;
+                }
+
+                case PARAMETER_TYPE::T_FLOAT:
+                    if ( randomGenerator )
+                    {
+                        res->SetFloatValue( i, randomGenerator() );
+                    }
+                    break;
+
+                case PARAMETER_TYPE::T_COLOUR:
+                    if ( randomGenerator )
+                    {
+                        res->SetColourValue( i, randomGenerator(), randomGenerator(), randomGenerator() );
+                    }
+                    break;
+
+                default:
+                    break;
+                }
+            }
+        }
+
+        return res;
+    }
+
+
+    //---------------------------------------------------------------------------------------------
+    ParametersPtr ModelParametersGenerator::GetRandomInstance( std::function<float()> randomGenerator)
+    {
+		LLM_SCOPE_BYNAME(TEXT("MutableRuntime"));
+
+        ParametersPtr res = m_pD->m_pModel->NewParameters();
+
+        int paramCount = (int)m_pD->m_pModel->GetPrivate()->m_program.m_parameters.size();
+
+        for (int i=0;i<paramCount;++i)
+        {
+            const PARAMETER_DESC& param = m_pD->m_pModel->GetPrivate()->m_program.m_parameters[ i ];
+            switch ( param.m_type )
+            {
+
+            case PARAMETER_TYPE::T_INT:
+            {
+                int numOptions = res->GetIntPossibleValueCount( i );
+                int valueIndex = (int)(std::min(numOptions-1,int(randomGenerator()*numOptions)));
+                int value = res->GetIntPossibleValue( i, valueIndex );
+                res->SetIntValue( i, value );
+                break;
+            }
+
+            case PARAMETER_TYPE::T_BOOL:
+            {
+                res->SetBoolValue( i, randomGenerator()>0.5f );
+                break;
+            }
+
+            case PARAMETER_TYPE::T_FLOAT:
+                if ( randomGenerator )
+                {
+                    res->SetFloatValue( i, randomGenerator() );
+                }
+                break;
+
+            case PARAMETER_TYPE::T_COLOUR:
+                if ( randomGenerator )
+                {
+                    res->SetColourValue( i, randomGenerator(), randomGenerator(), randomGenerator() );
+                }
+                break;
+
+            default:
+                break;
+            }
+        }
+
+        return res;
+    }
+
+
+    //---------------------------------------------------------------------------------------------
+    //---------------------------------------------------------------------------------------------
+    //---------------------------------------------------------------------------------------------
+//    class PartialRelevantParameterVisitor : public PartialDiscreteCoveredCodeVisitor
+//    {
+//    public:
+
+//        PartialRelevantParameterVisitor
+//            (
+//                System::Private* pSystem,
+//                Ptr<const Model>& pModel,
+//                int state,
+//                const Parameters* pParams,
+//                int relevantParameters,
+//                int thisParameter
+//            )
+//            : PartialDiscreteCoveredCodeVisitor( pSystem, pModel, pParams, relevantParameters )
+//        {
+//            m_visited.resize( pModel->GetPrivate()->m_program.m_opAddress.size(), false );
+//            m_thisParameter = thisParameter;
+//            m_relevant = false;
+
+//            OP::ADDRESS at = pModel->GetPrivate()->m_program.m_states[ state ].m_root;
+//            Run( at );
+
+//        }
+
+
+//        void Visit( OP::ADDRESS at, PROGRAM& program ) override
+//        {
+//            if ( !m_relevant && !m_visited[at])
+//            {
+//                m_visited[at] = true;
+
+//                switch ( program.m_code[at].type )
+//                {
+//                case OP_TYPE::BO_PARAMETER:
+//                case OP_TYPE::NU_PARAMETER:
+//                case OP_TYPE::SC_PARAMETER:
+//                case OP_TYPE::CO_PARAMETER:
+//                case OP_TYPE::PR_PARAMETER:
+//                case OP_TYPE::IM_PARAMETER:
+//                    if ( program.m_code[at].args.Parameter.variable == (size_t)m_thisParameter )
+//                    {
+//                        m_relevant = true;
+//                    }
+//                    break;
+
+//                default:
+//                    break;
+
+//                }
+
+//                PartialDiscreteCoveredCodeVisitor::Visit( at, program );
+//            }
+//        }
+
+//        //! Result
+//        bool m_relevant;
+
+//    private:
+
+//        int m_thisParameter;
+
+//        //! Flags for visited instructions.
+//        vector<bool> m_visited;
+//    };
+
+
+    //---------------------------------------------------------------------------------------------
+    size_t ModelParametersGenerator::Private::BuildIntervals( size_t currentInstanceIndex,
+                                                              size_t currentParameter,
+                                                              vector<int>& currentValues )
+    {
+        size_t paramCount = m_pModel->GetPrivate()->m_program.m_parameters.size();
+
+        if ( currentParameter>=paramCount )
+        {
+            currentInstanceIndex++;
+        }
+        else
+        {
+            // Find out about parameter relevancy
+            bool relevant = true;
+            {
+                ParametersPtr pParams = m_pModel->NewParameters();
+                for ( size_t p=0; p<currentParameter; ++p )
+                {
+                    switch ( pParams->GetType( (int)p ) )
+                    {
+                    case PARAMETER_TYPE::T_BOOL:
+                        pParams->SetBoolValue( (int)p, currentValues[p]!=0 );
+                        break;
+
+                    case PARAMETER_TYPE::T_INT:
+                        pParams->SetIntValue( (int)p,
+                                              pParams->GetIntPossibleValue( (int)p, currentValues[p] ) );
+                        break;
+
+                    default:
+                        break;
+                    }
+                }
+
+                // todo
+                relevant = true;
+//                PartialRelevantParameterVisitor relevancyVisitor( m_pSystem->GetPrivate(),
+//                                                                  m_pModel,
+//                                                                  0, // state
+//                                                                  pParams.get(),
+//                                                                  (int)currentParameter,
+//                                                                  (int)currentParameter
+//                                                                  );
+//                relevant = relevancyVisitor.m_relevant;
+            }
+
+            if (relevant)
+            {
+                const PARAMETER_DESC& param =
+                        m_pModel->GetPrivate()->m_program.m_parameters[ currentParameter ];
+                switch ( param.m_type )
+                {
+
+                case PARAMETER_TYPE::T_BOOL:
+                {
+                    // False
+                    PARAMETER_INTERVAL_VALUE falseValue;
+                    falseValue.m_minIndex = (int)currentInstanceIndex;
+                    falseValue.m_value = 0;
+                    m_intervals[ currentParameter ].m_intervalValue.push_back( falseValue );
+                    currentValues[ currentParameter ] = 0;
+                    currentInstanceIndex = BuildIntervals( currentInstanceIndex,
+                                                           currentParameter+1,
+                                                           currentValues);
+
+                    // True
+                    PARAMETER_INTERVAL_VALUE trueValue;
+                    trueValue.m_minIndex = (int)currentInstanceIndex;
+                    trueValue.m_value = 1;
+                    m_intervals[ currentParameter ].m_intervalValue.push_back( trueValue );
+                    currentValues[ currentParameter ] = 1;
+                    currentInstanceIndex = BuildIntervals( currentInstanceIndex,
+                                                           currentParameter+1,
+                                                           currentValues );
+                    break;
+                }
+
+                case PARAMETER_TYPE::T_INT:
+                {
+                    for ( size_t v=0; v<param.m_possibleValues.size(); ++v )
+                    {
+                        PARAMETER_INTERVAL_VALUE value;
+                        value.m_minIndex = (int)currentInstanceIndex;
+                        value.m_value = (int)v;
+                        m_intervals[ currentParameter ].m_intervalValue.push_back( value );
+                        currentValues[ currentParameter ] = (int)v;
+                        currentInstanceIndex = BuildIntervals( currentInstanceIndex,
+                                                               currentParameter+1,
+                                                               currentValues);
+                    }
+                    break;
+                }
+
+                default:
+                    // Continuous parameter.
+                    currentInstanceIndex = BuildIntervals( currentInstanceIndex,
+                                                           currentParameter+1,
+                                                           currentValues);
+                    break;
+                }
+            }
+            else
+            {
+                // Irrelevant parameter
+                currentInstanceIndex = BuildIntervals( currentInstanceIndex,
+                                                       currentParameter+1,
+                                                       currentValues);
+            }
+        }
+
+        return currentInstanceIndex;
+    }
+
+
+    //---------------------------------------------------------------------------------------------
+    vector<int> ModelParametersGenerator::Private::GetParameters( int instanceIndex )
+    {
+        vector<int> res( m_pModel->GetPrivate()->m_program.m_parameters.size(), 0 );
+
+        for ( size_t p=0; p<res.size(); ++p )
+        {
+            size_t e = 0;
+            while ( m_intervals[p].m_intervalValue.size()>e+1
+                    &&
+                    m_intervals[p].m_intervalValue[e+1].m_minIndex<=instanceIndex )
+            {
+                ++e;
+            }
+
+            // Avoid degenerated case (parameter always irrelevant?)
+            if ( m_intervals[p].m_intervalValue.size()>e )
+            {
+                res[p] = m_intervals[p].m_intervalValue[e].m_value;
+            }
+        }
+
+        return res;
+    }
+
+
+}
+
diff --git a/Engine/Plugins/Experimental/Mutable/Source/MutableRuntime/Private/ModelPrivate.h b/Engine/Plugins/Experimental/Mutable/Source/MutableRuntime/Private/ModelPrivate.h
new file mode 100644
--- /dev/null
+++ b/Engine/Plugins/Experimental/Mutable/Source/MutableRuntime/Private/ModelPrivate.h
@@ -0,0 +1,763 @@
+// Copyright Epic Games, Inc. All Rights Reserved.
+
+#pragma once
+
+#include "MutableRuntime/Public/Model.h"
+#include "MutableRuntime/Public/System.h"
+
+#include "MutableRuntime/Private/SerialisationPrivate.h"
+#include "MutableRuntime/Private/MemoryPrivate.h"
+#include "MutableRuntime/Private/Operations.h"
+#include "MutableRuntime/Private/ImagePrivate.h"
+#include "MutableRuntime/Private/MeshPrivate.h"
+#include "MutableRuntime/Private/ParametersPrivate.h"
+#include "MutableRuntime/Private/OpImageResize.h"
+#include "MutableRuntime/Private/OpImageMipmap.h"
+
+
+namespace mu
+{
+
+	/** Used to debug and log. */
+	constexpr bool bForceSingleThread = false;
+	constexpr bool DebugRom = false;
+	constexpr bool DebugRomAll = false;
+	constexpr int32 DebugRomIndex = 44;
+	constexpr int32 DebugImageIndex = 9;
+
+    //! Data stored for a rom even if it is not loaded
+    struct FRomData
+    {
+        //! This is used to identify a ROM file. It is usually a hash from its data.
+        uint32 Id;
+
+		//! Size of the rom
+		uint32 Size;
+
+		//! Index of the resource in its type-sepcific array
+		uint32 ResourceIndex;
+
+		//! Index of the resource in its type-sepcific array (one of DATATYPE values)
+		uint32 ResourceType;        
+    };
+
+    MUTABLE_DEFINE_POD_SERIALISABLE(FRomData);
+
+    //!
+    template<typename DATA>
+    inline void AppendCode( vector<uint8>& code, const DATA& data )
+    {
+        size_t pos = code.size();
+        code.resize( code.size()+sizeof(DATA) );
+        memcpy( &code[pos], &data, sizeof(DATA) );
+    }
+
+
+	//!
+	struct FImageLODRange
+	{
+		int32 FirstIndex = 0;
+		int32 LODCount = 0;
+		uint16 ImageSizeX = 0;
+		uint16 ImageSizeY = 0;
+	};
+	MUTABLE_DEFINE_POD_SERIALISABLE(FImageLODRange);
+
+    //!
+    struct PROGRAM
+    {
+        PROGRAM()
+        {
+            // Add the null instruction at address 0.
+            // TODO: Will do it in the linker
+            AppendCode( m_byteCode, OP_TYPE::NONE );
+            m_opAddress.push_back(0);
+        }
+
+        struct STATE
+        {
+            //! Name of the state
+            string m_name;
+
+            //! First instruction of the full build of an instance in this state
+            OP::ADDRESS m_root = 0;
+
+            //! List of parameters index (to PROGRAM::m_parameters) of the runtime parameters of
+            //! this state.
+            vector<int> m_runtimeParameters;
+
+            //! List of instructions that need to be cached to efficiently update this state
+            vector<OP::ADDRESS> m_updateCache;
+
+            //! List of root instructions for the dynamic resources that depend on the runtime
+            //! parameters of this state, with a mask of relevant runtime parameters.
+            //! The mask has a bit on for every runtime parameter in the m_runtimeParameters
+            //! vector.
+            vector< std::pair<OP::ADDRESS,uint64> > m_dynamicResources;
+
+            //!
+            inline void Serialise( OutputArchive& arch ) const
+            {
+                arch << m_name;
+                arch << m_root;
+                arch << m_runtimeParameters;
+                arch << m_updateCache;
+                arch << m_dynamicResources;
+            }
+
+            //!
+            inline void Unserialise( InputArchive& arch )
+            {
+                arch >> m_name;
+                arch >> m_root;
+                arch >> m_runtimeParameters;
+                arch >> m_updateCache;
+                arch >> m_dynamicResources;
+            }
+
+            //! Returns the parameters mask
+            uint64 IsDynamic( OP::ADDRESS at ) const
+            {
+                uint64 res = 0;
+
+                for ( size_t i=0; !res && i<m_dynamicResources.size(); ++i )
+                {
+                    if ( m_dynamicResources[i].first==at )
+                    {
+                        res = m_dynamicResources[i].second;
+                    }
+                }
+
+                return res;
+            }
+
+            //!
+            bool IsUpdateCache( OP::ADDRESS at ) const
+            {
+                bool res = false;
+
+                for ( size_t i=0; !res && i<m_updateCache.size(); ++i )
+                {
+                    if ( m_updateCache[i]==at )
+                    {
+                        res = true;
+                    }
+                }
+
+                return res;
+            }
+
+            //!
+            void AddUpdateCache( OP::ADDRESS at )
+            {
+                if ( !IsUpdateCache(at) )
+                {
+                    m_updateCache.push_back( at );
+                }
+            }
+        };
+
+        //! Location in the m_byteCode of the beginning of each operation
+        vector<uint32> m_opAddress;
+
+        //! Byte-coded representation of the program, using flexible-sized ops.
+        vector<uint8> m_byteCode;
+
+        //!
+        vector<STATE> m_states;
+
+		//! Data for every rom.
+		TArray<FRomData> m_roms;
+
+		//! Constant image mip data: the first is the index in m_roms for each image mip or -1 if it is always loaded.
+		TArray<TPair<int32, Ptr<const Image>>> m_constantImageLODs;
+
+		//! Constant image mip chain indices: ranges in this array are defined in FImageLODRange and the indices here refer to m_constantImageLODs.
+		TArray<uint32> m_constantImageLODIndices;
+
+		//! Constant image data.
+		TArray<FImageLODRange> m_constantImages;
+
+        //! Constant mesh data: the first is the index in m_roms for each mesh or -1 if it is always loaded.
+		TArray<TPair<int32, Ptr<const Mesh>>> m_constantMeshes;
+
+        //! Constant string data
+        vector<string> m_constantStrings;
+
+        //! Constant layout data
+        vector<Ptr<const Layout>> m_constantLayouts;
+
+        //! Constant projectors
+        vector<PROJECTOR> m_constantProjectors;
+
+        //! Constant matrices, usually used for transforms
+        vector<mat4f> m_constantMatrices;
+
+		//! Constant projectors
+		vector<SHAPE> m_constantShapes;
+
+        //! Constant curves
+        vector<Curve> m_constantCurves;
+
+        //! Constant curves
+        vector<Ptr<const Skeleton>> m_constantSkeletons;
+
+		//! Constant Physics Bodies
+		vector<Ptr<const PhysicsBody>> m_constantPhysicsBodies;
+
+        //! Parameters of the model.
+        //! The value stored here is the default value.
+        vector<PARAMETER_DESC> m_parameters;
+
+        //! Ranges for interation of the model operations.
+        vector<RANGE_DESC> m_ranges;
+
+        //! List of parameter lists. These are used in several places, like storing the
+        //! pregenerated list of parameters influencing a resource.
+        //! The parameter lists are sorted.
+        vector<vector<uint16>> m_parameterLists;
+
+
+        //!
+        void Serialise( OutputArchive& arch ) const
+        {
+            arch << m_opAddress;
+            arch << m_byteCode;
+            arch << m_states;
+			arch << m_roms;
+			arch << m_constantImageLODs;
+			arch << m_constantImageLODIndices;
+			arch << m_constantImages;
+			arch << m_constantMeshes;
+			arch << m_constantStrings;
+            arch << m_constantLayouts;
+            arch << m_constantProjectors;
+			arch << m_constantMatrices;
+			arch << m_constantShapes;
+            arch << m_constantCurves;
+            arch << m_constantSkeletons;
+			arch << m_constantPhysicsBodies;
+            arch << m_parameters;
+            arch << m_ranges;
+            arch << m_parameterLists;
+        }
+
+        //!
+        void Unserialise( InputArchive& arch )
+        {
+            arch >> m_opAddress;
+            arch >> m_byteCode;
+            arch >> m_states;
+			arch >> m_roms;
+			arch >> m_constantImageLODs;
+			arch >> m_constantImageLODIndices;
+			arch >> m_constantImages;
+			arch >> m_constantMeshes;
+			arch >> m_constantStrings;
+            arch >> m_constantLayouts;
+            arch >> m_constantProjectors;
+			arch >> m_constantMatrices;
+			arch >> m_constantShapes;
+            arch >> m_constantCurves;
+            arch >> m_constantSkeletons;
+			arch >> m_constantPhysicsBodies;
+            arch >> m_parameters;
+            arch >> m_ranges;
+            arch >> m_parameterLists;
+        }
+
+        //! Debug method that sanity-checks the program with a variety of tests.
+        void Check();
+
+        //! Debug method that logs the top used instruction types.
+        void LogHistogram() const;
+
+		//! Return true if the given ROM is loaded.
+		inline bool IsRomLoaded( int32 RomIndex ) const
+		{
+			switch (m_roms[RomIndex].ResourceType)
+			{
+			case DT_IMAGE:
+				return m_constantImageLODs[m_roms[RomIndex].ResourceIndex].Value.get() != nullptr;
+			case DT_MESH:
+				return m_constantMeshes[m_roms[RomIndex].ResourceIndex].Value.get() != nullptr;
+			default:
+				check(false);
+				break;
+			}
+			return false;
+		}
+
+		//! Unload a rom resource
+		inline void UnloadRom(int32 RomIndex)
+		{
+			if (DebugRom && (DebugRomAll||RomIndex == DebugRomIndex))
+				UE_LOG(LogMutableCore, Log, TEXT("Unloading rom %d."), RomIndex);
+
+			switch (m_roms[RomIndex].ResourceType)
+			{
+			case DT_IMAGE:
+				m_constantImageLODs[m_roms[RomIndex].ResourceIndex].Value = nullptr;
+				break;
+			case DT_MESH:
+				m_constantMeshes[m_roms[RomIndex].ResourceIndex].Value = nullptr;
+				break;
+			default:
+				check(false);
+				break;
+			}
+		}
+
+        //! Adds a constant data and returns its constant index.
+		int32 AddConstant(Ptr<const Image> pImage, int32 MinTextureResidentMipCount)
+		{
+			check(pImage->GetSizeX()*pImage->GetSizeY()>0);
+
+			// Mips to store
+			int32 MipsToStore = 1;
+			if (MinTextureResidentMipCount > 0 && MinTextureResidentMipCount<255)
+			{
+				int32 MaxMipmaps = Image::GetMipmapCount( pImage->GetSizeX(), pImage->GetSizeY() );
+				
+				// This is not true. We may want the full mipmaps for fragments of images, regardless of the resident mip size.
+				//MipsToStore = FMath::Max(1, MaxMipmaps-MinTextureResidentMipCount);
+				// \TODO: Calculate the mip ranges that makes sense to store.
+				MipsToStore = MaxMipmaps;
+			}
+
+			int32 FirstLODIndexIndex = m_constantImageLODIndices.Num();
+
+			// Some images cannot be resized or mipmaped
+			bool bCannotBeScaled = pImage->m_flags & Image::IF_CANNOT_BE_SCALED;
+			if (bCannotBeScaled)
+			{
+				MipsToStore = 1;
+			}
+
+			// TODO:
+			int32 CompressionQuality = 4;
+
+			// TODO: Not efficient if we don't make mips (no need to clone base)
+			// TODO: If the image already has mips, we will be duplicating them...
+			Ptr<Image> pMip = pImage->ExtractMip(0);			
+			for ( int Mip=0; Mip<MipsToStore; ++Mip )
+			{
+				check(pMip->GetFormat() == pImage->GetFormat());
+
+				// Shrink the buffer to the minimum necessary size.
+				uint32 CalculatedSize = pMip->CalculateDataSize();
+				if (CalculatedSize)
+				{
+					pMip->m_data.SetNum(CalculatedSize);
+				}
+
+				// Ensure unique at mip level
+				int32 MipIndex = -1;
+				for (int32 CandidateMipIndex = 0; CandidateMipIndex < m_constantImageLODs.Num(); ++CandidateMipIndex)
+				{
+					const Image* pCandidate = m_constantImageLODs[CandidateMipIndex].Value.get();
+					if (*pCandidate == *pMip)
+					{
+						// Reuse mip
+						MipIndex = CandidateMipIndex;
+					}
+				}
+
+				if (MipIndex < 0)
+				{
+					check(pMip->GetLODCount()==1);
+					MipIndex = m_constantImageLODs.Add(TPair<int32, Ptr<const Image>>(-1, pMip ));
+				}
+
+				m_constantImageLODIndices.Add(uint32(MipIndex));
+
+				// Generate next mip if necessary
+				if (Mip + 1 < MipsToStore)
+				{
+					pMip = pMip->ExtractMip(1);
+					check(pMip);
+				}
+			}
+
+			int32 ImageIndex = m_constantImages.Add({ FirstLODIndexIndex, MipsToStore, pImage->GetSizeX(), pImage->GetSizeY() });
+			return ImageIndex;
+		}
+
+
+		int32 AddConstant(Ptr<const Mesh> pMesh)
+		{
+			// Ensure unique
+			for (int32 i = 0; i < m_constantMeshes.Num(); ++i)
+			{
+				const Mesh* pCandidate = m_constantMeshes[i].Value.get();
+				if (*pCandidate == *pMesh)
+				{
+					return i;
+				}
+			}
+
+			return m_constantMeshes.Add(TPair<int32, Ptr<const Mesh>>( -1, pMesh.get() ));
+		}
+
+
+		OP::ADDRESS AddConstant( Ptr<const Layout> pLayout )
+        {
+            // Ensure unique
+            for ( size_t i=0; i<m_constantLayouts.size(); ++i)
+            {
+                if (m_constantLayouts[i]==pLayout)
+                {
+                    return (OP::ADDRESS)i;
+                }
+            }
+
+            OP::ADDRESS index = OP::ADDRESS( m_constantLayouts.size() );
+            m_constantLayouts.push_back( pLayout );
+            return index;
+        }
+
+        OP::ADDRESS AddConstant( Ptr<const Skeleton> pSkeleton )
+        {
+            // Ensure unique
+            for ( size_t i=0; i<m_constantSkeletons.size(); ++i)
+            {
+                if ( m_constantSkeletons[i]==pSkeleton
+                     ||
+                     *m_constantSkeletons[i]==*pSkeleton
+                     )
+                {
+                    return (OP::ADDRESS)i;
+                }
+            }
+
+            OP::ADDRESS index = OP::ADDRESS( m_constantSkeletons.size() );
+            m_constantSkeletons.push_back( pSkeleton );
+            return index;
+        }
+
+        OP::ADDRESS AddConstant( Ptr<const PhysicsBody> pPhysicsBody )
+        {
+            // Ensure unique
+            for ( size_t i=0; i<m_constantPhysicsBodies.size(); ++i)
+            {
+                if ( m_constantPhysicsBodies[i]==pPhysicsBody
+                     ||
+                     *m_constantPhysicsBodies[i]==*pPhysicsBody
+                     )
+                {
+                    return (OP::ADDRESS)i;
+                }
+            }
+
+            OP::ADDRESS index = OP::ADDRESS( m_constantPhysicsBodies.size() );
+            m_constantPhysicsBodies.push_back( pPhysicsBody );
+            return index;
+        }
+
+        OP::ADDRESS AddConstant( const string& str )
+        {            
+            // Ensure unique
+            for ( size_t i=0; i<m_constantStrings.size(); ++i)
+            {
+                if (m_constantStrings[i]==str)
+                {
+                    return (OP::ADDRESS)i;
+                }
+            }
+
+            OP::ADDRESS index = OP::ADDRESS( m_constantStrings.size() );
+            m_constantStrings.push_back( str );
+            return index;
+        }
+
+        OP::ADDRESS AddConstant( const mat4f& m )
+        {
+            // Ensure unique
+            for ( size_t i=0; i<m_constantMatrices.size(); ++i)
+            {
+                if (m_constantMatrices[i]==m)
+                {
+                    return (OP::ADDRESS)i;
+                }
+            }
+
+            OP::ADDRESS index = OP::ADDRESS( m_constantMatrices.size() );
+            m_constantMatrices.push_back( m );
+            return index;
+        }
+
+        OP::ADDRESS AddConstant( const SHAPE& m )
+        {
+            // Ensure unique
+            for ( size_t i=0; i<m_constantShapes.size(); ++i)
+            {
+                if (m_constantShapes[i]==m)
+                {
+                    return (OP::ADDRESS)i;
+                }
+            }
+
+            OP::ADDRESS index = OP::ADDRESS( m_constantShapes.size() );
+            m_constantShapes.push_back( m );
+            return index;
+        }
+
+        OP::ADDRESS AddConstant( const PROJECTOR& m )
+        {
+            // Ensure unique
+            for ( size_t i=0; i<m_constantProjectors.size(); ++i)
+            {
+                if (m_constantProjectors[i]==m)
+                {
+                    return (OP::ADDRESS)i;
+                }
+            }
+
+            OP::ADDRESS index = OP::ADDRESS( m_constantProjectors.size() );
+            m_constantProjectors.push_back( m );
+            return index;
+        }
+
+        OP::ADDRESS AddConstant( const Curve& m )
+        {
+            // Ensure unique
+//            for ( size_t i=0; i<m_constantCurves.size(); ++i)
+//            {
+//                if (m_constantCurves[i]==m)
+//                {
+//                    return (OP::ADDRESS)i;
+//                }
+//            }
+
+            OP::ADDRESS index = OP::ADDRESS( m_constantCurves.size() );
+            m_constantCurves.push_back( m );
+            return index;
+        }
+
+
+        //! Get a constant image, assuming it is fully loaded. The image constant will be composed with lodaded mips if necessary.
+        void GetConstant( int32 ConstantIndex, ImagePtrConst& res, int32 MipsToSkip) const
+        {
+			int32 ReallySkippedLODs = FMath::Min(m_constantImages[ConstantIndex].LODCount - 1, MipsToSkip);
+			int32 FirstLODIndexIndex = m_constantImages[ConstantIndex].FirstIndex;
+			int32 ResultLODIndexIndex = FirstLODIndexIndex + ReallySkippedLODs;
+			int32 FinalLODs = m_constantImages[ConstantIndex].LODCount - ReallySkippedLODs;
+			check(FinalLODs > 0);
+
+			// Get the first mip
+			int32 ResultLODIndex = m_constantImageLODIndices[ResultLODIndexIndex];
+			Ptr<const Image> CurrentMip = m_constantImageLODs[ResultLODIndex].Value;
+			check(CurrentMip);
+				
+			// Shortcut if we only want one mip
+			if (FinalLODs == 1)
+			{
+				res = CurrentMip;
+				return;
+			}
+
+			// Compose the result image
+			{
+				MUTABLE_CPUPROFILER_SCOPE(ComposeConstantImage);
+
+				Ptr<Image> Result = new Image( CurrentMip->GetSizeX(), CurrentMip->GetSizeY(), FinalLODs, CurrentMip->GetFormat() );
+				Result->m_flags = CurrentMip->m_flags;
+
+				// Some non-block pixel formats require separate memory size calculation
+				if (!Result->GetDataSize())
+				{
+					int32 TotalSize = 0;
+					for (int32 LOD = 0; LOD < FinalLODs; ++LOD)
+					{
+						int LODIndex = m_constantImageLODIndices[ResultLODIndexIndex + LOD];
+						int32 MipSizeBytes = m_constantImageLODs[LODIndex].Value->GetDataSize();
+						TotalSize += MipSizeBytes;
+					}
+					Result->m_data.SetNum(TotalSize);
+				}
+
+				int32 WrittenDataBytes = 0;
+				uint8* Data = Result->GetData();
+				for (int32 LOD=0; LOD<FinalLODs; ++LOD)
+				{
+					check(CurrentMip->GetLODCount() == 1);
+					check(CurrentMip->GetFormat() == Result->GetFormat());
+					int32 MipSizeBytes = CurrentMip->GetDataSize();
+					check( Result->GetDataSize()>=WrittenDataBytes+MipSizeBytes);
+
+					FMemory::Memcpy( Data, CurrentMip->GetData(), MipSizeBytes);
+					Data += MipSizeBytes;
+					WrittenDataBytes += MipSizeBytes;
+
+					if (LOD + 1 < FinalLODs)
+					{
+						ResultLODIndex = m_constantImageLODIndices[ResultLODIndexIndex + LOD + 1];
+						CurrentMip = m_constantImageLODs[ResultLODIndex].Value;
+						check(CurrentMip);
+					}
+				}
+
+				res = Result;
+			}
+		}
+
+
+        void GetConstant(int32 ConstantIndex, MeshPtrConst& res ) const
+        {
+			res = m_constantMeshes[ConstantIndex].Value;
+		}
+
+
+        inline OP_TYPE GetOpType( OP::ADDRESS at ) const
+        {
+            if (at>=OP::ADDRESS(m_opAddress.size())) return OP_TYPE::NONE;
+
+            OP_TYPE result;
+            uint64 byteCodeAddress = m_opAddress[at];
+            memcpy( &result, &m_byteCode[byteCodeAddress], sizeof(OP_TYPE) );
+            check( result<OP_TYPE::COUNT);
+            return result;
+        }
+
+        template<typename ARGS>
+        inline const ARGS GetOpArgs( OP::ADDRESS at ) const
+        {
+            ARGS result;
+            uint64 byteCodeAddress = m_opAddress[at];
+            byteCodeAddress += sizeof(OP_TYPE);
+            memcpy( &result, &m_byteCode[byteCodeAddress], sizeof(ARGS));
+            return result;
+        }
+
+        inline const uint8* GetOpArgsPointer( OP::ADDRESS at ) const
+        {
+            uint64 byteCodeAddress = m_opAddress[at];
+            byteCodeAddress += sizeof(OP_TYPE);
+            const uint8* pData = (const uint8*)&m_byteCode[byteCodeAddress];
+            return pData;
+        }
+
+        inline uint8* GetOpArgsPointer( OP::ADDRESS at )
+        {
+            uint64 byteCodeAddress = m_opAddress[at];
+            byteCodeAddress += sizeof(OP_TYPE);
+            uint8* pData = (uint8*)&m_byteCode[byteCodeAddress];
+            return pData;
+        }
+    };
+
+
+    //!
+    class MUTABLERUNTIME_API Model::Private : public Base
+    {
+    public:
+
+        //!
+        PROGRAM m_program;
+
+        //! Non-persistent, streamer-specific location information
+        string m_location;
+
+        //!
+        void Serialise( OutputArchive& arch ) const
+        {
+            uint32 ver = 0;
+            arch << ver;
+
+            arch << m_program;
+        }
+
+        //!
+        void Unserialise( InputArchive& arch )
+        {
+            uint32 ver;
+            arch >> ver;
+            check(ver==0);
+
+            arch >> m_program;
+        }
+
+        //! Management of generated resources
+        //! @{
+
+        //! This is used to uniquely identify a generated resource like meshes or images.
+        struct RESOURCE_KEY
+        {
+            //! The id assigned to the generated resource.
+            uint32 m_id;
+
+            //! The last request operation for this resource
+            uint32 m_lastRequestId;
+
+            //! The address that generated this resource
+            OP::ADDRESS m_rootAddress;
+
+            //! An opaque blob with the values of the relevant parameters
+            vector<uint8> m_parameterValuesBlob;
+        };
+
+        //! The last id generated for a resource
+        uint32 m_lastResourceKeyId = 0;
+
+        //! The last id generated for a resource request. This is used to check the
+        //! relevancy of the resources when flushing the cache
+        uint32 m_lastResourceResquestId = 0;
+
+        //! Cached ids for returned assets
+        //! This is non-persistent runtime data
+        vector<RESOURCE_KEY> m_generatedResources;
+
+        //! Get a resource key for a given resource with given parameter values.
+        uint32 GetResourceKey( uint32 paramListIndex, OP::ADDRESS rootAt, const Parameters* pParams );
+
+        //! @}
+
+    };
+
+
+    //!
+    class MUTABLERUNTIME_API ModelParametersGenerator::Private : public Base
+    {
+    public:
+
+        //!
+        ModelPtrConst m_pModel;
+        SystemPtr m_pSystem;
+
+        //! Number of possible instances
+        int64 m_instanceCount;
+
+        //!
+        struct PARAMETER_INTERVAL_VALUE
+        {
+            //! Minimum instance index to use these increments
+            int m_minIndex;
+
+            //! Parameter value in this interval
+            int m_value;
+        };
+
+        //!
+        struct PARAMETER_INTERVALS
+        {
+            //! As many entries as necessary
+            //! Sorted by PARAMETER_INCREMENTS_PER_VALUE::m_minIndex
+            vector<PARAMETER_INTERVAL_VALUE> m_intervalValue;
+        };
+
+        //! An entry for every parameter in the model
+        vector< PARAMETER_INTERVALS > m_intervals;
+
+        //! Whether to use brute force, or consider the parameter relevancy
+        bool m_considerRelevancy;
+
+        //! Build the interval information
+        size_t BuildIntervals( size_t currentInstanceIndex,
+                               size_t currentParameter,
+                               vector<int>& currentValues );
+
+        //! Get the parameter values for a particular instance index
+        vector<int> GetParameters( int instanceIndex );
+    };
+
+}
diff --git a/Engine/Plugins/Experimental/Mutable/Source/MutableRuntime/Private/MutableRuntimeModule.cpp b/Engine/Plugins/Experimental/Mutable/Source/MutableRuntime/Private/MutableRuntimeModule.cpp
new file mode 100644
--- /dev/null
+++ b/Engine/Plugins/Experimental/Mutable/Source/MutableRuntime/Private/MutableRuntimeModule.cpp
@@ -0,0 +1,97 @@
+// Copyright Epic Games, Inc. All Rights Reserved.
+
+#include "MutableRuntimeModule.h"
+#include "CoreMinimal.h"
+#include "Modules/ModuleManager.h"
+#include "Logging/LogMacros.h"
+#include "Platform.h"
+#include "MutableRuntime/Public/MutableMemory.h"
+#include "MutableRuntime/Public/System.h"
+
+IMPLEMENT_MODULE(FMutableRuntimeModule, MutableRuntime);
+
+DEFINE_LOG_CATEGORY(LogMutableCore);
+
+#if USE_STAT_MUTABLE_MEMORY
+DECLARE_MEMORY_STAT(TEXT("Mutable Memory Used"), STAT_MemoryMutableTotalAllocationSize, STATGROUP_Memory);
+#endif
+
+
+// Uncomment the following line to have allocation/deallocation messages printed to the output
+//#define MUTABLE_PRINT_ALLOCATIONS
+
+#ifdef MUTABLE_PRINT_ALLOCATIONS
+
+static int64 NumAllocations = 0;
+
+static void* CustomMalloc(std::size_t Size_t, uint32_t alignment)
+{
+	NumAllocations++;
+	UE_LOG(LogTemp, Warning, TEXT("NumAllocations=%d, Allocated %d bytes"), NumAllocations, Size_t);
+
+	return malloc(Size_t);
+}
+
+
+static void CustomFree(void* mem)
+{
+	NumAllocations--;
+	UE_LOG(LogTemp, Warning, TEXT("NumAllocations=%d, Deallocated"), NumAllocations);
+
+	free(mem);
+}
+
+#else //MUTABLE_PRINT_ALLOCATIONS
+
+
+
+static void* CustomMalloc(std::size_t Size_t, uint32_t alignment)
+{
+	void* Result = FMemory::Malloc(Size_t, alignment);
+
+#if USE_STAT_MUTABLE_MEMORY
+	SIZE_T ActualSize = FMemory::GetAllocSize(Result);
+	INC_DWORD_STAT_BY(STAT_MemoryMutableTotalAllocationSize, ActualSize);
+#endif
+	// This should always be fatal: mutable cannot handle out-of-memory errors yet.
+	if (Size_t > 0 && !Result)
+	{
+		LowLevelFatalError(TEXT("Not enough memory for mutable request %ul."), Size_t);
+		check(false);
+	}
+
+	return Result;
+}
+
+
+static void CustomFree(void* mem)
+{
+#if USE_STAT_MUTABLE_MEMORY
+	DEC_DWORD_STAT_BY(STAT_MemoryMutableTotalAllocationSize, FMemory::GetAllocSize(mem));
+#endif
+	return FMemory::Free(mem);
+}
+
+#endif //MUTABLE_PRINT_ALLOCATIONS
+
+
+void FMutableRuntimeModule::StartupModule()
+{
+	void* (*customMalloc)(size_t, uint32_t) = nullptr;
+	void(*customFree)(void*) = nullptr;
+
+#if USE_UNREAL_ALLOC_IN_MUTABLE
+	customMalloc = CustomMalloc;
+	customFree = CustomFree;
+#endif
+
+	mu::Initialize(customMalloc, customFree);
+}
+
+
+void FMutableRuntimeModule::ShutdownModule()
+{
+	// Finalize the mutable runtime in this module
+	mu::Finalize();
+}
+
diff --git a/Engine/Plugins/Experimental/Mutable/Source/MutableRuntime/Private/MutableString.cpp b/Engine/Plugins/Experimental/Mutable/Source/MutableRuntime/Private/MutableString.cpp
new file mode 100644
--- /dev/null
+++ b/Engine/Plugins/Experimental/Mutable/Source/MutableRuntime/Private/MutableString.cpp
@@ -0,0 +1,52 @@
+// Copyright Epic Games, Inc. All Rights Reserved.
+
+
+#include "MutableString.h"
+#include "Platform.h"
+
+
+namespace mu {
+
+	//---------------------------------------------------------------------------------------------
+	String::String( const char* v )
+	{
+		if (v)
+		{
+            m_value = v;
+        }
+	}
+
+
+	//---------------------------------------------------------------------------------------------
+	void String::Serialise( const String* p, OutputArchive& arch )
+	{
+		arch << *p;
+	}
+
+
+	//---------------------------------------------------------------------------------------------
+	Ptr<String> String::StaticUnserialise( InputArchive& arch )
+	{
+		Ptr<String> pResult = new String();
+		arch >> *pResult;
+		return pResult;
+	}
+
+
+	//---------------------------------------------------------------------------------------------
+	Ptr<String> String::Clone() const
+	{
+		Ptr<String> pResult = new String();
+		pResult->m_value = m_value;
+		return pResult;
+	}
+
+
+	//---------------------------------------------------------------------------------------------
+	const char* String::GetValue() const
+	{
+        return m_value.c_str();
+	}
+
+}
+
diff --git a/Engine/Plugins/Experimental/Mutable/Source/MutableRuntime/Private/MutableTrace.cpp b/Engine/Plugins/Experimental/Mutable/Source/MutableRuntime/Private/MutableTrace.cpp
new file mode 100644
--- /dev/null
+++ b/Engine/Plugins/Experimental/Mutable/Source/MutableRuntime/Private/MutableTrace.cpp
@@ -0,0 +1,9 @@
+// Copyright Epic Games, Inc. All Rights Reserved.
+
+#include "MutableTrace.h"
+
+#include "Trace/Detail/Channel.h"
+
+
+UE_TRACE_CHANNEL_DEFINE(MutableChannel)
+
diff --git a/Engine/Plugins/Experimental/Mutable/Source/MutableRuntime/Private/Octree.h b/Engine/Plugins/Experimental/Mutable/Source/MutableRuntime/Private/Octree.h
new file mode 100644
--- /dev/null
+++ b/Engine/Plugins/Experimental/Mutable/Source/MutableRuntime/Private/Octree.h
@@ -0,0 +1,196 @@
+// Copyright Epic Games, Inc. All Rights Reserved.
+
+#pragma once
+
+#include "Platform.h"
+
+#include "MutableMath.h"
+#include "Ptr.h"
+
+#include <vector>
+
+#define MAX_DIVISIONS 5
+#define MAX_VECTORS 20
+
+namespace mu
+{
+	struct VertexInfo
+	{
+		vec3f vector;
+		int index;
+	};
+
+	class Octree
+	{
+	public:
+
+		Octree(vec3f _minCorner, vec3f _maxCorner, float _padding = 0.0f,int div = 0)
+		{
+			minCorner = _minCorner;
+			maxCorner = _maxCorner;
+			padding = _padding;
+			division = div;
+
+			for (int i = 0; i < 8; ++i)
+			{
+				children[i] = nullptr;
+			}
+		}
+		
+		~Octree()
+		{
+			for (int i = 0; i < 8; ++i)
+			{
+				if (children[i])
+				{
+					delete children[i];
+					children[i] = nullptr;
+				}	
+			}
+
+            for (size_t i = 0; i < elements.size(); ++i)
+			{
+				delete elements[i];
+				elements[i] = nullptr;
+			}
+		}
+
+		void InsertElement(vec3f v, int idx)
+		{
+			if (!children[0])
+			{
+				VertexInfo* cur = new VertexInfo();
+				cur->vector = v;
+				cur->index = idx;
+
+				elements.push_back(cur);
+
+				if (elements.size() > MAX_VECTORS && division < MAX_DIVISIONS)
+				{
+					Divide();
+				}
+			}
+			else
+			{
+				for (int i = 0; i < 8; ++i)
+				{
+					if (children[i] != nullptr)
+					{
+						if (Collides(children[i], v))
+						{
+							children[i]->InsertElement(v, idx);
+						}
+					}
+				}
+			}
+		}
+
+		void Divide()
+		{
+			float sizeX = maxCorner[0] - minCorner[0];
+			float sizeY = maxCorner[1] - minCorner[1];
+			float sizeZ = maxCorner[2] - minCorner[2];
+
+			vec3f diagonal = { sizeX / 2.0f, sizeY / 2.0f , sizeZ / 2.0f };
+
+			for (int i = 0; i < 2; ++i)
+			{
+				float height = (sizeY*i) / 2.0f;
+
+				vec3f nMin = minCorner + vec3f(0.0f, height, 0.0f);
+				vec3f nMax = nMin + diagonal;
+				children[0 + i * 4] = new Octree(nMin, nMax, padding, division + 1);
+
+				nMin = minCorner + vec3f(0.0f, height, sizeZ / 2.0f);
+				nMax = nMin + diagonal;
+				children[1 + i * 4] = new Octree(nMin, nMax, padding, division + 1);
+
+				nMin = minCorner + vec3f(sizeX / 2.0f, height, 0.0f);
+				nMax = nMin + diagonal;
+				children[2 + i * 4] = new Octree(nMin, nMax, padding, division + 1);
+
+				nMin = minCorner + vec3f(sizeX / 2.0f, height, sizeZ / 2.0f);
+				nMax = nMin + diagonal;
+				children[3 + i * 4] = new Octree(nMin, nMax, padding, division + 1);
+			}
+
+            for (size_t i = 0; i < elements.size(); ++i)
+			{
+				InsertElement(elements[i]->vector, elements[i]->index);
+			}
+			
+            for (size_t i = 0; i < elements.size(); ++i)
+			{
+				delete elements[i];
+				elements[i] = nullptr;
+			}
+
+			elements.clear();
+		}
+
+		bool Collides(Octree* child, vec3f v)
+		{
+			bool ret = false;
+
+			if (child)
+			{
+				if (v[0] >= child->minCorner[0] - padding && v[1] >= child->minCorner[1] - padding && v[2] >= child->minCorner[2] - padding)
+				{
+					if (v[0] <= child->maxCorner[0] + padding && v[1] <= child->maxCorner[1] + padding && v[2] <= child->maxCorner[2] + padding)
+					{
+						ret = true;
+					}
+				}
+			}
+
+			return ret;
+		}
+
+		int GetNearests(vec3f v, float dist, int idx)
+		{
+			int ret = idx;
+
+			if (!children[0])
+			{
+                for (size_t i = 0; i < elements.size(); ++i)
+				{
+					vec3f r = elements[i]->vector - v;
+					
+					if (idx > elements[i]->index && dot(r,r) <= dist*dist)
+					{
+						ret = elements[i]->index;
+						break;
+					}
+				}
+				return ret;
+			}
+
+			for (int i = 0; i < 8; ++i)
+			{
+				if (Collides(children[i], v))
+				{
+					int result = children[i]->GetNearests(v, dist, idx);
+					if (result < idx && result < ret)
+					{
+						ret = result;
+					}
+				}
+			}
+
+			return ret;
+		}
+
+	private:
+		
+		int division = 0;
+
+		vec3f minCorner;
+		vec3f maxCorner;
+
+		float padding = 0.0f;
+
+		std::vector<VertexInfo*> elements;
+
+		Octree* children[8];
+	};
+}
diff --git a/Engine/Plugins/Experimental/Mutable/Source/MutableRuntime/Private/OpEvalCurve.h b/Engine/Plugins/Experimental/Mutable/Source/MutableRuntime/Private/OpEvalCurve.h
new file mode 100644
--- /dev/null
+++ b/Engine/Plugins/Experimental/Mutable/Source/MutableRuntime/Private/OpEvalCurve.h
@@ -0,0 +1,102 @@
+// Copyright Epic Games, Inc. All Rights Reserved.
+
+#pragma once
+
+#include "ConvertData.h"
+#include "Platform.h"
+
+
+namespace mu
+{
+    inline float lerp(float p0, float p1, float alpha)
+	{
+		return p0 * (1.f - alpha) + p1 * alpha;
+	}
+
+    inline float BezierInterpolation(float p0, float p1, float p2, float p3, float alpha)
+	{
+		const float p01 = lerp(p0, p1, alpha);
+		const float p12 = lerp(p1, p2, alpha);
+		const float p23 = lerp(p2, p3, alpha);
+		const float p012 = lerp(p01, p12, alpha);
+		const float p123 = lerp(p12, p23, alpha);
+		const float p0123 = lerp(p012, p123, alpha);
+
+		return p0123;
+	}
+
+	//---------------------------------------------------------------------------------------------
+	//! Reference version
+	//---------------------------------------------------------------------------------------------
+	inline float EvalCurve(const Curve& curve, float time)
+	{
+		float interpolated_value = curve.defaultValue;
+        const int32_t numKeyFrames = (int32_t)curve.keyFrames.size();
+
+		if (numKeyFrames < 1)
+		{
+			return curve.defaultValue;
+		}
+		else if (numKeyFrames < 2 || (time <= curve.keyFrames[0].time))
+		{
+			// Use the first value
+			interpolated_value = curve.keyFrames[0].value;
+		}
+		else if (time < curve.keyFrames[numKeyFrames - 1].time)
+		{
+			int32_t first = 1;
+			int32_t last = numKeyFrames - 1;
+			int32_t count = last - first;
+
+			while (count > 0)
+			{
+				int32_t step = count / 2;
+				int32_t middle = first + step;
+
+				if (time >= curve.keyFrames[middle].time)
+				{
+					first = middle + 1;
+					count -= step + 1;
+				}
+				else
+				{
+					count = step;
+				}
+			}
+
+			int32_t interpolated_node = first;
+			const float difference = curve.keyFrames[interpolated_node].time - curve.keyFrames[interpolated_node - 1].time;
+
+            if (difference > 0.f && curve.keyFrames[interpolated_node - 1].interp_mode != uint8_t(CurveKeyFrame::InterpMode::Constant))
+			{
+				const float p0 = curve.keyFrames[interpolated_node - 1].value;
+				const float p3 = curve.keyFrames[interpolated_node].value;
+				const float alpha = (time - curve.keyFrames[interpolated_node - 1].time) / difference;
+
+                if (curve.keyFrames[interpolated_node - 1].interp_mode == uint32_t(CurveKeyFrame::InterpMode::Linear))
+				{
+					interpolated_value = lerp(p0, p3, alpha);
+				}
+				else
+				{
+					const float aux = 1.0f / 3.0f;
+					const float p1 = p0 + (curve.keyFrames[interpolated_node - 1].out_tangent * difference * aux);
+					const float p2 = p3 - (curve.keyFrames[interpolated_node].in_tangent * difference * aux);
+
+					interpolated_value = BezierInterpolation(p0, p1, p2, p3, alpha);
+				}
+			}
+			else
+			{
+				interpolated_value = curve.keyFrames[interpolated_node - 1].value;
+			}
+		}
+		else
+		{
+			// Use the last value
+			interpolated_value = curve.keyFrames[numKeyFrames - 1].value;
+		}
+
+		return interpolated_value;
+	}
+}
diff --git a/Engine/Plugins/Experimental/Mutable/Source/MutableRuntime/Private/OpImageAlphaOverlay.h b/Engine/Plugins/Experimental/Mutable/Source/MutableRuntime/Private/OpImageAlphaOverlay.h
new file mode 100644
--- /dev/null
+++ b/Engine/Plugins/Experimental/Mutable/Source/MutableRuntime/Private/OpImageAlphaOverlay.h
@@ -0,0 +1,63 @@
+// Copyright Epic Games, Inc. All Rights Reserved.
+
+#pragma once
+
+#include "ImagePrivate.h"
+#include "OpImageHardLight.h"
+
+
+namespace mu
+{
+    inline unsigned AlphaOverlayChannelMasked( unsigned base, unsigned blended, unsigned mask )
+    {
+            // Photoshop like
+            //return HardLightChannel( blended, base, mask );
+
+            // gimp-like
+            //unsigned overlay = ( base * ( base + ( ( 2 * blended * ( 255 - base) ) >> 8  ) ) ) >> 8;
+            unsigned overlay = base + (blended * unsigned(255 - base) >> 8);
+            unsigned masked = ( ( ( 255 - mask ) * base ) + ( mask * overlay ) ) >> 8;
+            return masked;
+    }
+
+
+    inline unsigned AlphaOverlayChannel( unsigned base, unsigned blended )
+    {
+            // Photoshop like
+            //return HardLightChannel( blended, base, mask );
+
+            // gimp-like
+            //unsigned overlay = ( base * ( base + ( ( 2 * blended * ( 255 - base) ) >> 8  ) ) ) >> 8;
+            unsigned overlay = base + (blended * unsigned(255 - base) >> 8);
+            return overlay;
+    }
+
+
+    inline void ImageAlphaOverlay( Image* pResult, const Image* pBase, const Image* pMask, vec3<float> col )
+    {
+        ImageLayerColour<AlphaOverlayChannelMasked,AlphaOverlayChannel,true>( pResult, pBase, pMask, col );
+    }
+
+
+    inline void ImageAlphaOverlay( Image* pResult, const Image* pBase, const Image* pMask,
+                                   const Image* pBlended,
+                                   bool applyToAlpha )
+    {
+        ImageLayer<AlphaOverlayChannelMasked,AlphaOverlayChannel,true>
+                ( pResult, pBase, pMask, pBlended, applyToAlpha );
+    }
+
+
+    inline void ImageAlphaOverlay( Image* pResult, const Image* pBase, vec3<float> col )
+    {
+        ImageLayerColour<AlphaOverlayChannel,true>( pResult, pBase, col );
+    }
+
+
+    inline void ImageAlphaOverlay( Image* pResult, const Image* pBase, const Image* pBlended,
+                                   bool applyToAlpha )
+    {
+        ImageLayer<AlphaOverlayChannel,true>( pResult, pBase, pBlended, applyToAlpha );
+    }
+
+}
diff --git a/Engine/Plugins/Experimental/Mutable/Source/MutableRuntime/Private/OpImageApplyComposite.h b/Engine/Plugins/Experimental/Mutable/Source/MutableRuntime/Private/OpImageApplyComposite.h
new file mode 100644
--- /dev/null
+++ b/Engine/Plugins/Experimental/Mutable/Source/MutableRuntime/Private/OpImageApplyComposite.h
@@ -0,0 +1,82 @@
+// Copyright Epic Games, Inc. All Rights Reserved.
+
+#pragma once
+
+#include "Math/UnrealMathUtility.h"
+#include "Math/Vector.h"
+
+#include "ImagePrivate.h"
+#include "OpImageLayer.h"
+
+
+namespace mu
+{
+	// Combine layer oprations get packed color.
+	inline unsigned NormalComposite(
+			const uint32 Base, 
+			const uint32 Normal, 
+			const uint8 Channel, 
+			const float Power )
+	{
+		const int32 ChannelShift = Channel << 3;
+
+		float ReferencValue = static_cast<float>(Base >> ChannelShift & 0xFF) / 255.0f;
+
+		// Rougness is stored as Unorm8
+		float Roughness = static_cast<float>( (Base >> ChannelShift) & 0xFF ) / 255.0f;
+
+		FVector3f N = ( FVector3f( static_cast<float>((Normal >> 0 ) & 0xFF) / 255.0f,
+							       static_cast<float>((Normal >> 8 ) & 0xFF) / 255.0f,
+							       static_cast<float>((Normal >> 16) & 0xFF) / 255.0f ) * 2.0f - 1.0f );	
+	
+		// See TexturecCompressorModule.cpp::ApplyCompositeTexture
+		// Toksvig estimation of variance
+		float LengthN = FMath::Min( N.Size(), 1.0f );
+		float Variance = ( 1.0f - LengthN ) / LengthN;
+		Variance = FMath::Max( 0.0f, Variance - 0.00004f );
+
+		Variance *= Power;
+		
+		float a = Roughness * Roughness;
+		float a2 = a * a;
+		float B = 2.0f * Variance * (a2 - 1.0f);
+		a2 = ( B - a2 ) / ( B - 1.0f );
+		Roughness = FMath::Pow( a2, 0.25f );
+
+		const uint32 ResMask = ~(0xFF << ChannelShift);
+		const uint32 Value = FMath::Clamp<uint32>( static_cast<uint32>(Roughness * 255.0f), 0u, 255u );
+
+		return  (Value << ChannelShift) | (Base & ResMask); 
+	}
+
+	inline void ImageNormalComposite(Image* pResult, const Image* pBase, const Image* pNormal, ECompositeImageMode mode, float power)
+	{
+		if (mode == ECompositeImageMode::CIM_Disabled)
+		{
+			// Copy base to result. Using functor layer combine with identity to do the copy for simplicity.
+			// This could be optimized but will never happen since it is checked at compile time. 	
+			ImageLayerCombineFunctor( pResult, pBase, pNormal, 
+				[](uint32 base, uint32) -> uint32 { return base; });
+
+			return;
+		}	
+
+		const int8 channel = [mode]() 
+			{
+				switch (mode)
+				{
+					case ECompositeImageMode::CIM_NormalRoughnessToRed   : return 0;
+					case ECompositeImageMode::CIM_NormalRoughnessToGreen : return 1;
+					case ECompositeImageMode::CIM_NormalRoughnessToBlue  : return 2;
+					case ECompositeImageMode::CIM_NormalRoughnessToAlpha : return 3;
+				}
+			
+				check(false);
+				return 0;		
+			}();
+
+		ImageLayerCombineFunctor( pResult, pBase, pNormal,
+				[channel, power](uint32 base, uint32 normal) -> uint32 
+					{ return NormalComposite(base, normal, channel, power); });
+	}
+}
