[mutable] Move Mutable plugin from the restricted folder to the experimental plugins folder.

diff --git a/Engine/Plugins/Experimental/Mutable/Source/MutableRuntime/Private/OpImageBinarise.h b/Engine/Plugins/Experimental/Mutable/Source/MutableRuntime/Private/OpImageBinarise.h
new file mode 100644
--- /dev/null
+++ b/Engine/Plugins/Experimental/Mutable/Source/MutableRuntime/Private/OpImageBinarise.h
@@ -0,0 +1,74 @@
+// Copyright Epic Games, Inc. All Rights Reserved.
+
+#pragma once
+
+#include "ImagePrivate.h"
+#include "Platform.h"
+
+
+
+namespace mu
+{
+
+	//---------------------------------------------------------------------------------------------
+	//!
+	//---------------------------------------------------------------------------------------------
+	inline ImagePtr ImageBinarise( const Image* pA, float threshold )
+	{
+        ImagePtr pDest = new Image( pA->GetSizeX(), pA->GetSizeY(), pA->GetLODCount(), EImageFormat::IF_L_UBYTE );
+
+        uint8_t* pDestBuf = pDest->GetData();
+        const uint8_t* pABuf = pA->GetData();
+
+		// Generic implementation
+		int pixelCount = (int)pA->CalculatePixelCount();
+
+        uint32_t t_8 = uint32_t( threshold * 255.0f );
+
+		switch ( pA->GetFormat() )
+		{
+		case EImageFormat::IF_L_UBYTE:
+		{
+			for ( int i=0; i<pixelCount; ++i )
+			{
+                uint32_t a_8 = pABuf[i];
+				pDestBuf[i] = (a_8>=t_8) ? 255 : 0;
+			}
+			break;
+		}
+
+		case EImageFormat::IF_RGB_UBYTE:
+		{
+			for ( int i=0; i<pixelCount; ++i )
+			{
+                uint32_t a_8 = pABuf[3*i+0];
+				a_8 += pABuf[3*i+1];
+				a_8 += pABuf[3*i+2];
+				a_8 /= 3;
+				pDestBuf[i] = (a_8>=t_8) ? 255 : 0;
+			}
+			break;
+		}
+
+        case EImageFormat::IF_RGBA_UBYTE:
+        case EImageFormat::IF_BGRA_UBYTE:
+        {
+			for ( int i=0; i<pixelCount; ++i )
+			{
+                uint32_t a_8 = pABuf[4*i+0];
+				a_8 += pABuf[4*i+1];
+				a_8 += pABuf[4*i+2];
+				a_8 /= 3;
+				pDestBuf[i] = (a_8>=t_8) ? 255 : 0;
+			}
+			break;
+		}
+
+		default:
+			check(false);
+		}
+
+		return pDest;
+	}
+
+}
diff --git a/Engine/Plugins/Experimental/Mutable/Source/MutableRuntime/Private/OpImageBlend.h b/Engine/Plugins/Experimental/Mutable/Source/MutableRuntime/Private/OpImageBlend.h
new file mode 100644
--- /dev/null
+++ b/Engine/Plugins/Experimental/Mutable/Source/MutableRuntime/Private/OpImageBlend.h
@@ -0,0 +1,72 @@
+// Copyright Epic Games, Inc. All Rights Reserved.
+
+#pragma once
+
+#include "ImagePrivate.h"
+#include "OpImageLayer.h"
+
+
+namespace mu
+{
+
+	inline unsigned BlendChannel( unsigned // base
+		, unsigned blended )
+	{
+		return blended;
+	}
+
+	inline unsigned BlendChannelMasked( unsigned base, unsigned blended, unsigned mask )
+	{
+		unsigned blend = blended;
+		unsigned masked = ( ( ( 255 - mask ) * base ) + ( mask * blend ) ) / 255;
+		return masked;
+	}
+
+
+    inline void ImageBlend( Image* pResult, const Image* pBase, const Image* pMask, vec3<float> col )
+	{
+        ImageLayerColour<BlendChannelMasked, BlendChannel, false>( pResult, pBase, pMask, col );
+	}
+
+
+    inline void ImageBlend( Image* pResult, const Image* pBase, const Image* pMask,
+                            const Image* pBlended,
+                            bool applyToAlpha )
+	{
+        ImageLayer<BlendChannelMasked, BlendChannel, false>
+                ( pResult, pBase, pMask, pBlended, applyToAlpha );
+	}
+
+
+    inline void ImageBlend( Image* pResult, const Image* pBase, const Image* pBlended,
+                            bool applyToAlpha )
+    {
+        return ImageLayer<BlendChannel,false>( pResult, pBase, pBlended, applyToAlpha );
+    }
+
+
+	inline void ImageBlendOnBase( Image* pBase, const Image* pMask, vec3<float> col )
+	{
+		ImageLayerColourOnBase<BlendChannelMasked, BlendChannel, false>( pBase, pMask, col );
+	}
+
+
+    inline void ImageBlendOnBase( Image* pBase, const Image* pMask, const Image* pBlended,
+                                  bool applyToAlpha )
+	{
+        ImageLayerOnBase<BlendChannelMasked, BlendChannel, false>
+                 ( pBase, pMask, pBlended, applyToAlpha );
+	}
+
+	//! Blend a subimage on the base using a mask.
+	inline void ImageBlendOnBaseNoAlpha( Image* pBase,
+										 const Image* pMask,
+										 const Image* pBlended,
+										 const box< vec2<int> >& rect )
+	{
+		ImageLayerOnBaseNoAlpha< BlendChannel, false>
+				( pBase, pMask, pBlended, rect );
+	}
+
+
+}
diff --git a/Engine/Plugins/Experimental/Mutable/Source/MutableRuntime/Private/OpImageBurn.h b/Engine/Plugins/Experimental/Mutable/Source/MutableRuntime/Private/OpImageBurn.h
new file mode 100644
--- /dev/null
+++ b/Engine/Plugins/Experimental/Mutable/Source/MutableRuntime/Private/OpImageBurn.h
@@ -0,0 +1,60 @@
+// Copyright Epic Games, Inc. All Rights Reserved.
+
+#pragma once
+
+#include "ImagePrivate.h"
+
+
+namespace mu
+{
+
+	inline unsigned BurnChannel( unsigned base, unsigned blended )
+	{
+		// R = 1 - (1-Base) / Blend
+		unsigned burn =
+				std::min( 255,
+						  std::max( 0,
+									255 - ( ( ( 255 - (int)base  ) << 8 ) / ((int)blended+1) )
+									)
+						  );
+		return burn;
+	}
+
+	inline unsigned BurnChannelMasked( unsigned base, unsigned blended, unsigned mask )
+	{
+		unsigned burn = BurnChannel( base, blended );
+		unsigned masked = ( ( ( 255 - mask ) * base ) + ( mask * burn ) ) >> 8;
+		return masked;
+	}
+
+
+    inline void ImageBurn( Image* pResult, const Image* pBase, vec3<float> col )
+	{
+        ImageLayerColour<BurnChannel,true>( pResult, pBase, col );
+	}
+
+
+    inline void ImageBurn( Image* pResult, const Image* pBase, const Image* pBlended,
+                           bool applyToAlpha )
+	{
+        ImageLayer<BurnChannel,true>( pResult, pBase, pBlended, applyToAlpha );
+	}
+
+
+    inline void ImageBurn( Image* pResult, const Image* pBase, const Image* pMask, vec3<float> col )
+	{
+        ImageLayerColour<BurnChannelMasked,BurnChannel,true>( pResult, pBase, pMask, col );
+	}
+
+
+    inline void ImageBurn( Image* pResult,
+                           const Image* pBase,
+                           const Image* pMask,
+                           const Image* pBlended,
+                           bool applyToAlpha )
+	{
+        ImageLayer<BurnChannelMasked,BurnChannel,true>
+                ( pResult, pBase, pMask, pBlended, applyToAlpha );
+	}
+
+}
diff --git a/Engine/Plugins/Experimental/Mutable/Source/MutableRuntime/Private/OpImageColourMap.h b/Engine/Plugins/Experimental/Mutable/Source/MutableRuntime/Private/OpImageColourMap.h
new file mode 100644
--- /dev/null
+++ b/Engine/Plugins/Experimental/Mutable/Source/MutableRuntime/Private/OpImageColourMap.h
@@ -0,0 +1,133 @@
+// Copyright Epic Games, Inc. All Rights Reserved.
+
+#pragma once
+
+#include "ImagePrivate.h"
+#include "Platform.h"
+
+namespace mu
+{
+
+	inline ImagePtr ImageColourMap( const Image* pSource, const Image* pMask, const Image* pMap )
+	{
+		check( pSource->GetSizeX() == pMask->GetSizeX() );
+		check( pSource->GetSizeY() == pMask->GetSizeY() );
+        check( pMask->GetFormat() == EImageFormat::IF_L_UBYTE );
+
+		ImagePtr pDest = new Image( pSource->GetSizeX(),
+									pSource->GetSizeY(),
+                                    pSource->GetLODCount(),
+									pSource->GetFormat() );
+
+        uint8_t* pDestBuf = pDest->GetData();
+        const uint8_t* pSourceBuf = pSource->GetData();
+        const uint8_t* pMaskBuf = pMask->GetData();
+
+		// Generic implementation
+		int pixelCount = (int)pSource->CalculatePixelCount();
+
+		// Make a palette for faster conversion
+        uint8_t palette[256][4];
+		for ( int i=0; i<256; ++i )
+		{
+			vec4<float> c = pMap->Sample( vec2<float>( float(i)/256.0f, 0.0f ) );
+            palette[i][0] = (uint8_t)std::max( 0, std::min( 255, int(c[0]*255.0f) ) );
+            palette[i][1] = (uint8_t)std::max( 0, std::min( 255, int(c[1]*255.0f) ) );
+            palette[i][2] = (uint8_t)std::max( 0, std::min( 255, int(c[2]*255.0f) ) );
+            palette[i][3] = (uint8_t)std::max( 0, std::min( 255, int(c[3]*255.0f) ) );
+		}
+
+		switch ( pSource->GetFormat() )
+		{
+		case EImageFormat::IF_L_UBYTE:
+		{
+			for ( int i=0; i<pixelCount; ++i )
+			{
+                uint8_t m_8 = pMaskBuf[i];
+				if (m_8>127)
+				{
+					pDestBuf[i] = palette[ pSourceBuf[i] ][0];
+				}
+				else
+				{
+					pDestBuf[i] = pSourceBuf[i];
+				}
+			}
+			break;
+		}
+
+		case EImageFormat::IF_RGB_UBYTE:
+		{
+			for ( int i=0; i<pixelCount; ++i )
+			{
+                uint8_t m_8 = pMaskBuf[i];
+				if (m_8>127)
+				{
+					pDestBuf[3*i+0] = palette[ pSourceBuf[3*i+0] ][0];
+					pDestBuf[3*i+1] = palette[ pSourceBuf[3*i+1] ][1];
+					pDestBuf[3*i+2] = palette[ pSourceBuf[3*i+2] ][2];
+				}
+				else
+				{
+					pDestBuf[3*i+0] = pSourceBuf[3*i+0];
+					pDestBuf[3*i+1] = pSourceBuf[3*i+1];
+					pDestBuf[3*i+2] = pSourceBuf[3*i+2];
+				}
+			}
+			break;
+		}
+
+        case EImageFormat::IF_RGBA_UBYTE:
+        {
+            for ( int i=0; i<pixelCount; ++i )
+            {
+                uint8_t m_8 = pMaskBuf[i];
+                if (m_8>127)
+                {
+                    pDestBuf[4*i+0] = palette[ pSourceBuf[4*i+0] ][0];
+                    pDestBuf[4*i+1] = palette[ pSourceBuf[4*i+1] ][1];
+                    pDestBuf[4*i+2] = palette[ pSourceBuf[4*i+2] ][2];
+                    pDestBuf[4*i+3] = palette[ pSourceBuf[4*i+3] ][3];
+                }
+                else
+                {
+                    pDestBuf[4*i+0] = pSourceBuf[4*i+0];
+                    pDestBuf[4*i+1] = pSourceBuf[4*i+1];
+                    pDestBuf[4*i+2] = pSourceBuf[4*i+2];
+                    pDestBuf[4*i+3] = pSourceBuf[4*i+3];
+                }
+            }
+            break;
+        }
+
+        case EImageFormat::IF_BGRA_UBYTE:
+        {
+            for ( int i=0; i<pixelCount; ++i )
+            {
+                uint8_t m_8 = pMaskBuf[i];
+                if (m_8>127)
+                {
+                    pDestBuf[4*i+0] = palette[ pSourceBuf[4*i+0] ][2];
+                    pDestBuf[4*i+1] = palette[ pSourceBuf[4*i+1] ][1];
+                    pDestBuf[4*i+2] = palette[ pSourceBuf[4*i+2] ][0];
+                    pDestBuf[4*i+3] = palette[ pSourceBuf[4*i+3] ][3];
+                }
+                else
+                {
+                    pDestBuf[4*i+0] = pSourceBuf[4*i+0];
+                    pDestBuf[4*i+1] = pSourceBuf[4*i+1];
+                    pDestBuf[4*i+2] = pSourceBuf[4*i+2];
+                    pDestBuf[4*i+3] = pSourceBuf[4*i+3];
+                }
+            }
+            break;
+        }
+
+		default:
+			check(false);
+		}
+
+		return pDest;
+	}
+
+}
diff --git a/Engine/Plugins/Experimental/Mutable/Source/MutableRuntime/Private/OpImageCompose.h b/Engine/Plugins/Experimental/Mutable/Source/MutableRuntime/Private/OpImageCompose.h
new file mode 100644
--- /dev/null
+++ b/Engine/Plugins/Experimental/Mutable/Source/MutableRuntime/Private/OpImageCompose.h
@@ -0,0 +1,112 @@
+// Copyright Epic Games, Inc. All Rights Reserved.
+
+#pragma once
+
+#include "ImagePrivate.h"
+#include "MutableMath.h"
+#include "OpImagePixelFormat.h"
+
+
+namespace mu
+{
+
+	inline void ImageCompose( Image* pBase, const Image* pBlock, const box< vec2<int> >& rect )
+	{
+		check(pBase && pBlock);
+		check(pBase != pBlock);
+		check( rect.min[0]>=0 && rect.min[1]>=0 );
+		check( rect.size[0]>=0 && rect.size[1]>=0 );
+		check( rect.min[0]+rect.size[0]<=pBase->GetSizeX() );
+		check( rect.min[1]+rect.size[1]<=pBase->GetSizeY() );
+		check( rect.size[0]==pBlock->GetSizeX() );
+		check( rect.size[1]==pBlock->GetSizeY() );
+        check( pBase->GetFormat()==pBlock->GetFormat() );
+
+        // This assert may be acceptable if the missing lods are due to the image block format size
+        // limiting it, or disadjustment due to image scaling down. It is handled below.
+        //check( pBase->GetLODCount()<=pBlock->GetLODCount() );
+
+		EImageFormat OriginalFormat = pBase->GetFormat();
+		const FImageFormatData& finfo = GetImageFormatData( OriginalFormat );
+
+		// If this is true, we are composing coordinates or sizes that don't fall on pixel format
+		// block borders.
+		bool bSubPixelBlockCompose = 
+			rect.min[0] % finfo.m_pixelsPerBlockX != 0
+			|| rect.min[1] % finfo.m_pixelsPerBlockY != 0
+			|| rect.size[0] % finfo.m_pixelsPerBlockX != 0
+			|| rect.size[1] % finfo.m_pixelsPerBlockY != 0;
+
+		if (bSubPixelBlockCompose)
+		{
+			// Bad performance case, but likely a very small image anyway.
+			MUTABLE_CPUPROFILER_SCOPE(ImageCompose_SubBlock_Fix);
+			EImageFormat UncompressedFormat = GetUncompressedFormat(OriginalFormat);
+			constexpr int32 Quality = 0;
+			Ptr<Image> TempBase = ImagePixelFormat(Quality, pBase, UncompressedFormat);
+			Ptr<Image> TempBlock = ImagePixelFormat(Quality, pBlock, UncompressedFormat);
+			ImageCompose( TempBase.get(), TempBlock.get(), rect );
+			Ptr<Image> NewBase = ImagePixelFormat(Quality, TempBase.get(), OriginalFormat);
+			// ImageCompose above may have reduced the LODs because of the block lods.
+			pBase->m_lods = NewBase->m_lods;
+			pBase->m_data = MoveTemp(NewBase->m_data);
+			return;
+		}
+		else
+		{
+			// Sizes in blocks of the current mips
+			vec2<int> baseMipSize(pBase->GetSizeX() / finfo.m_pixelsPerBlockX,
+				pBase->GetSizeY() / finfo.m_pixelsPerBlockY);
+			vec2<int> blockMipPos(rect.min[0] / finfo.m_pixelsPerBlockX,
+				rect.min[1] / finfo.m_pixelsPerBlockY);
+			vec2<int> blockMipSize(rect.size[0] / finfo.m_pixelsPerBlockX,
+				rect.size[1] / finfo.m_pixelsPerBlockY);
+
+			int doneMips = 0;
+			for (doneMips = 0; doneMips < pBase->GetLODCount() && doneMips < pBlock->GetLODCount(); ++doneMips)
+			{
+				//UE_LOG(LogMutableCore, Warning, "Block Mip Pos : %d, %d", blockMipPos[0], blockMipPos[1]);
+
+				// Sizes of a row in bytes
+				int baseRowSize = finfo.m_bytesPerBlock * baseMipSize[0];
+				int blockRowSize = finfo.m_bytesPerBlock * blockMipSize[0];
+
+				uint8_t* pBaseBuf = pBase->GetMipData(doneMips);
+				const uint8_t* pBlockBuf = pBlock->GetMipData(doneMips);
+
+				int skipBytes = baseRowSize * blockMipPos[1] + blockMipPos[0] * finfo.m_bytesPerBlock;
+
+				pBaseBuf += skipBytes;
+				for (int y = 0; y < blockMipSize[1]; ++y)
+				{
+					check(pBaseBuf + blockRowSize <= pBase->GetData() + pBase->GetDataSize());
+					check(pBlockBuf + blockRowSize <= pBlock->GetData() + pBlock->GetDataSize());
+					FMemory::Memcpy(pBaseBuf, pBlockBuf, blockRowSize);
+					pBlockBuf += blockRowSize;
+					pBaseBuf += baseRowSize;
+				}
+
+				// We may not proceed to the next mip, if the layout block size is smaller than the
+				// pixel format block size.
+				if (blockMipPos[0] % 2 || blockMipPos[1] % 2 || blockMipSize[0] % 2 || blockMipSize[1] % 2)
+				{
+					break;
+				}
+
+				baseMipSize[0] = FMath::DivideAndRoundUp(baseMipSize[0], 2);
+				baseMipSize[1] = FMath::DivideAndRoundUp(baseMipSize[1], 2);
+				blockMipPos[0] = FMath::DivideAndRoundUp(blockMipPos[0], 2);
+				blockMipPos[1] = FMath::DivideAndRoundUp(blockMipPos[1], 2);
+				blockMipSize[0] = FMath::DivideAndRoundUp(blockMipSize[0], 2);
+				blockMipSize[1] = FMath::DivideAndRoundUp(blockMipSize[1], 2);
+			}
+
+			// Adjust the actually valid mips, maybe wasting some of the reserved memory.
+			// TODO: Review this.
+			pBase->m_lods = (uint8_t)doneMips;
+			pBase->m_data.SetNum(pBase->CalculateDataSize());
+		}
+    }
+
+}
+
diff --git a/Engine/Plugins/Experimental/Mutable/Source/MutableRuntime/Private/OpImageCrop.h b/Engine/Plugins/Experimental/Mutable/Source/MutableRuntime/Private/OpImageCrop.h
new file mode 100644
--- /dev/null
+++ b/Engine/Plugins/Experimental/Mutable/Source/MutableRuntime/Private/OpImageCrop.h
@@ -0,0 +1,86 @@
+// Copyright Epic Games, Inc. All Rights Reserved.
+
+#pragma once
+
+#include "ImagePrivate.h"
+#include "MutableMath.h"
+#include "Platform.h"
+#include "OpImagePixelFormat.h"
+
+
+namespace mu
+{
+
+	inline ImagePtr ImageCrop( int imageCompressionQuality, const Image* pBaseReal, const box< vec2<int> >& rect )
+	{
+		ImagePtrConst pBase = pBaseReal;
+
+		// RLE are not supported yet
+		if ( pBaseReal->GetFormat()==EImageFormat::IF_RGBA_UBYTE_RLE )
+		{
+            pBase = ImagePixelFormat( imageCompressionQuality, pBase.get(), EImageFormat::IF_RGBA_UBYTE );
+        }
+		else if ( pBaseReal->GetFormat()==EImageFormat::IF_RGB_UBYTE_RLE )
+		{
+            pBase = ImagePixelFormat( imageCompressionQuality, pBase.get(), EImageFormat::IF_RGB_UBYTE );
+        }
+		else if ( pBaseReal->GetFormat()==EImageFormat::IF_L_UBYTE_RLE )
+		{
+            pBase = ImagePixelFormat( imageCompressionQuality, pBase.get(), EImageFormat::IF_L_UBYTE );
+        }
+
+		check( pBase );
+
+		ImagePtr pCropped = new Image
+				(
+					rect.size[0], rect.size[1],
+					1,
+					pBase->GetFormat()
+					);
+
+//		check( rect.min[0]>=0 && rect.min[1]>=0 );
+//		check( rect.size[0]>=0 && rect.size[1]>=0 );
+//		check( rect.min[0]+rect.size[0]<=pBase->GetSizeX() );
+//		check( rect.min[1]+rect.size[1]<=pBase->GetSizeY() );
+		// TODO: better error control. This happens if some layouts are corrupt.
+		bool correct =
+				( rect.min[0]>=0 && rect.min[1]>=0 ) &&
+				( rect.size[0]>=0 && rect.size[1]>=0 ) &&
+				( rect.min[0]+rect.size[0]<=pBase->GetSizeX() ) &&
+				( rect.min[1]+rect.size[1]<=pBase->GetSizeY() );
+		if (!correct)
+		{
+			return pCropped;
+		}
+
+
+		const FImageFormatData& finfo = GetImageFormatData( pBase->GetFormat() );
+
+		// For now
+		check( finfo.m_pixelsPerBlockX == 1 );
+		check( finfo.m_pixelsPerBlockY == 1 );
+
+		checkf( rect.min[0] % finfo.m_pixelsPerBlockX == 0, TEXT("Rect must snap to blocks.") );
+		checkf( rect.min[1] % finfo.m_pixelsPerBlockY == 0, TEXT("Rect must snap to blocks.") );
+		checkf( rect.size[0] % finfo.m_pixelsPerBlockX == 0, TEXT("Rect must snap to blocks.") );
+		checkf( rect.size[1] % finfo.m_pixelsPerBlockY == 0, TEXT("Rect must snap to blocks.") );
+
+		int baseRowSize = finfo.m_bytesPerBlock * pBase->GetSizeX() / finfo.m_pixelsPerBlockX;
+		int cropRowSize = finfo.m_bytesPerBlock * rect.size[0] / finfo.m_pixelsPerBlockX;
+
+        const uint8_t* pBaseBuf = pBase->GetData();
+        uint8_t* pCropBuf = pCropped->GetData();
+
+		int skipPixels = pBase->GetSizeX() * rect.min[1] + rect.min[0];
+		pBaseBuf += finfo.m_bytesPerBlock * skipPixels / finfo.m_pixelsPerBlockX;
+		for ( int y=0; y<rect.size[1]; ++y )
+		{
+			memcpy( pCropBuf, pBaseBuf, cropRowSize );
+			pCropBuf += cropRowSize;
+			pBaseBuf += baseRowSize;
+		}
+
+		return pCropped;
+	}
+
+}
diff --git a/Engine/Plugins/Experimental/Mutable/Source/MutableRuntime/Private/OpImageDifference.h b/Engine/Plugins/Experimental/Mutable/Source/MutableRuntime/Private/OpImageDifference.h
new file mode 100644
--- /dev/null
+++ b/Engine/Plugins/Experimental/Mutable/Source/MutableRuntime/Private/OpImageDifference.h
@@ -0,0 +1,71 @@
+// Copyright Epic Games, Inc. All Rights Reserved.
+
+#pragma once
+
+#include "ImagePrivate.h"
+#include "Platform.h"
+
+
+namespace mu
+{
+
+	inline ImagePtr ImageDifference( const Image* pA, const Image* pB )
+	{
+		check( pA->GetSizeX() == pB->GetSizeX() );
+		check( pA->GetSizeY() == pB->GetSizeY() );
+		check( pA->GetFormat() == pB->GetFormat() );
+
+        ImagePtr pDest = new Image( pA->GetSizeX(), pA->GetSizeY(), pA->GetLODCount(), EImageFormat::IF_L_UBYTE );
+
+        uint8_t* pDestBuf = pDest->GetData();
+        const uint8_t* pABuf = pA->GetData();
+        const uint8_t* pBBuf = pB->GetData();
+
+		// Generic implementation
+		int pixelCount = (int)pA->CalculatePixelCount();
+
+		switch ( pA->GetFormat() )
+		{
+		case EImageFormat::IF_RGB_UBYTE:
+			for ( int i=0; i<pixelCount; ++i )
+			{
+                uint8_t diff = 0;
+				for ( int c=0; c<3; ++c )
+				{
+                    uint8_t a = pABuf[3*i+c];
+                    uint8_t b = pBBuf[3*i+c];
+					if ( abs(a-b)>3 )
+					{
+                        diff = std::max( diff, (uint8_t)255 );
+					}
+				}
+				pDestBuf[i] = diff;
+			}
+			break;
+
+        case EImageFormat::IF_BGRA_UBYTE:
+        case EImageFormat::IF_RGBA_UBYTE:
+            for ( int i=0; i<pixelCount; ++i )
+			{
+                uint8_t diff = 0;
+				for ( int c=0; c<3; ++c )
+				{
+                    uint8_t a = pABuf[4*i+c];
+                    uint8_t b = pBBuf[4*i+c];
+					if ( abs(a-b)>3 )
+					{
+                        diff = std::max( diff, (uint8_t)255 );
+					}
+				}
+				pDestBuf[i] = diff;
+			}
+			break;
+
+		default:
+			check(false);
+		}
+
+		return pDest;
+	}
+
+}
diff --git a/Engine/Plugins/Experimental/Mutable/Source/MutableRuntime/Private/OpImageDisplace.h b/Engine/Plugins/Experimental/Mutable/Source/MutableRuntime/Private/OpImageDisplace.h
new file mode 100644
--- /dev/null
+++ b/Engine/Plugins/Experimental/Mutable/Source/MutableRuntime/Private/OpImageDisplace.h
@@ -0,0 +1,330 @@
+// Copyright Epic Games, Inc. All Rights Reserved.
+
+#pragma once
+
+#include "ImagePrivate.h"
+#include "Platform.h"
+#include "ConvertData.h"
+#include "Async/ParallelFor.h"
+
+
+namespace mu
+{
+
+    inline uint8_t MutableEncodeOffset( int x, int y )
+	{
+        uint8_t c = uint8_t( ((7)<<4) | 7 );
+		if ( x<8 && x>-8 && y<8 && y>-8 )
+		{
+            c = uint8_t( ((x+7)<<4) | (y+7) );
+		}
+		return c;
+	}
+
+    inline void MutableDecodeOffset( uint8_t c, int& x, int& y )
+	{
+		x = int(c>>4) - 7;
+		y = int(c&0xf) - 7;
+	}
+
+	//---------------------------------------------------------------------------------------------
+	//!
+	//---------------------------------------------------------------------------------------------
+	inline void ImageMakeGrowMap( Image* pResult, const Image* pMask, int border )
+	{
+		MUTABLE_CPUPROFILER_SCOPE(ImageMakeGrowMap)
+
+		check( pMask->GetFormat() == EImageFormat::IF_L_UBYTE );
+		check( pResult->GetFormat() == EImageFormat::IF_L_UBYTE );
+		check( pResult->GetSizeX() == pMask->GetSizeX() );
+		check( pResult->GetSizeY() == pMask->GetSizeY() );
+
+		int sizeX = pResult->GetSizeX();
+		int sizeY = pResult->GetSizeY();
+		check( sizeX>0 && sizeY>0 );
+
+        if (sizeX<4 || sizeY<4)
+        {
+            return;
+        }
+
+		memset( pResult->GetData(), MutableEncodeOffset( 0, 0 ), sizeX*sizeY );
+
+		ImagePtrConst pThisMask = pMask;
+		for ( int b=0; b<border; ++b )
+		{
+			ImagePtr pNextMask = pThisMask->Clone();            
+
+			//for ( int y=0; y<sizeY; ++y )
+			const auto& ProcessRow = [
+				pThisMask, pNextMask, sizeX, sizeY, pResult
+			] (int32 y)
+			{
+				uint8_t* pNextData = pNextMask->GetData() + sizeX*y;
+
+				const uint8_t* pS0 = y > 0 ? pThisMask->GetData() + (y - 1) * sizeX : nullptr;
+				const uint8_t* pS1 = pThisMask->GetData() + y * sizeX;
+				bool bNotLastRow = y < sizeY - 1;
+				const uint8_t* pS2 = bNotLastRow ? pThisMask->GetData() + (y + 1) * sizeX : nullptr;
+
+				const uint8_t* pR0 = y > 0 ? pResult->GetData() + (y - 1) * sizeX : 0;
+				uint8_t* pR1 = pResult->GetData() + y * sizeX;
+
+				const uint8_t* pR2 = bNotLastRow ? pResult->GetData() + (y + 1) * sizeX : nullptr;
+
+				for (int x = 0; x < sizeX; ++x)
+				{
+					int dx, dy;
+					if (pS1[x])
+					{
+						*pNextData = 255;
+					}
+					else if (y > 0 && pS0[x])
+					{
+						MutableDecodeOffset(pR0[x], dx, dy);
+						pR1[x] = MutableEncodeOffset(dx, dy - 1);
+						*pNextData = 255;
+					}
+					else if (bNotLastRow && pS2 && pR2 && pS2[x])
+					{
+						MutableDecodeOffset(pR2[x], dx, dy);
+						pR1[x] = MutableEncodeOffset(dx, dy + 1);
+						*pNextData = 255;
+					}
+					else if (x < sizeX - 1 && pS1[x + 1])
+					{
+						MutableDecodeOffset(pR1[x + 1], dx, dy);
+						pR1[x] = MutableEncodeOffset(dx + 1, dy);
+						*pNextData = 255;
+					}
+					else if (x > 0 && pS1[x - 1])
+					{
+						MutableDecodeOffset(pR1[x - 1], dx, dy);
+						pR1[x] = MutableEncodeOffset(dx - 1, dy);
+						*pNextData = 255;
+					}
+
+					++pNextData;
+				}
+
+			};
+
+			ParallelFor(sizeY, ProcessRow);
+
+			pThisMask = pNextMask;
+		}
+
+	}
+
+
+	//---------------------------------------------------------------------------------------------
+	//!
+	//---------------------------------------------------------------------------------------------
+	inline void ImageDisplace( Image* pResult, const Image* pSource, const Image* pMap )
+	{
+		MUTABLE_CPUPROFILER_SCOPE(ImageDisplace)
+
+		check( pResult->GetFormat() == pSource->GetFormat()  );
+		check( pMap->GetFormat() == EImageFormat::IF_L_UBYTE
+						||
+						pMap->GetFormat() == EImageFormat::IF_L_UBYTE_RLE );
+		check( pResult->GetSizeX() == pSource->GetSizeX() );
+		check( pResult->GetSizeY() == pSource->GetSizeY() );
+		check( pResult->GetSizeX() == pMap->GetSizeX() );
+		check( pResult->GetSizeY() == pMap->GetSizeY() );
+
+		int sizeX = pResult->GetSizeX();
+		int sizeY = pResult->GetSizeY();
+
+        if (sizeX<4 || sizeY<4)
+        {
+            return;
+        }
+
+        const uint8_t* pSourceData = pSource->GetData();
+
+        auto mapFormat = pMap->GetFormat();
+        bool isUncompressed = ( mapFormat == EImageFormat::IF_L_UBYTE );
+
+        if ( isUncompressed )
+		{
+			switch (pResult->GetFormat())
+			{
+
+			case EImageFormat::IF_L_UBYTE:
+			{
+				//for ( int y=0; y<sizeY; ++y )
+				const auto& ProcessRow = [
+					pSourceData, pMap, pResult, sizeX, sizeY
+				] (int32 y)
+				{
+					const uint8_t* pMapData = pMap->GetData() + y * sizeX;
+					uint8_t* pResultData = pResult->GetData() + y * sizeX;
+
+					for (int x = 0; x < sizeX; ++x)
+					{
+						int dx, dy;
+						MutableDecodeOffset(*pMapData, dx, dy);
+						check(x + dx >= 0 && x + dx < sizeX);
+						check(y + dy >= 0 && y + dy < sizeY);
+
+						int offset = ((y + dy) * sizeX + (x + dx));
+						memcpy(pResultData, &pSourceData[offset], 1);
+
+						++pResultData;
+						++pMapData;
+					}
+
+				};
+
+				ParallelFor(sizeY, ProcessRow);
+				break;
+			}
+
+			case EImageFormat::IF_RGB_UBYTE:
+			{
+				//for (int y = 0; y < sizeY; ++y)
+				const auto& ProcessRow = [
+					pSourceData, pMap, pResult, sizeX, sizeY
+				] (int32 y)
+				{
+					const uint8_t* pMapData = pMap->GetData() + y * sizeX;
+					uint8_t* pResultData = pResult->GetData() + y * sizeX*3;
+
+					for (int x = 0; x < sizeX; ++x)
+					{
+						int dx, dy;
+						MutableDecodeOffset(*pMapData, dx, dy);
+						check(x + dx >= 0 && x + dx < sizeX);
+						check(y + dy >= 0 && y + dy < sizeY);
+
+						int offset = ((y + dy) * sizeX + (x + dx)) * 3;
+						memcpy(pResultData, &pSourceData[offset], 3);
+
+						pResultData += 3;
+						++pMapData;
+					}
+
+				};
+
+				ParallelFor(sizeY, ProcessRow);
+				break;
+			}
+
+			default:
+			{
+				// Generic
+				int pixelSize = GetImageFormatData( pResult->GetFormat() ).m_bytesPerBlock;
+				//for ( int y=0; y<sizeY; ++y )
+				const auto& ProcessRow = [
+					pSourceData, pMap, pResult, sizeX, sizeY, pixelSize
+				] (int32 y)
+				{
+					const uint8_t* pMapData = pMap->GetData() + y * sizeX;
+					uint8_t* pResultData = pResult->GetData() + y * sizeX * pixelSize;
+
+					for (int x = 0; x < sizeX; ++x)
+					{
+						int dx, dy;
+						MutableDecodeOffset(*pMapData, dx, dy);
+						check(x + dx >= 0 && x + dx < sizeX);
+						check(y + dy >= 0 && y + dy < sizeY);
+
+						int offset = ((y + dy) * sizeX + (x + dx)) * pixelSize;
+						memcpy(pResultData, &pSourceData[offset], pixelSize);
+
+						pResultData += pixelSize;
+						++pMapData;
+					}
+
+				};
+
+				ParallelFor(sizeY, ProcessRow);
+			}
+
+			}
+		}
+
+        else if (mapFormat==EImageFormat::IF_L_UBYTE_RLE )
+		{
+            int pixelSize = GetImageFormatData( pResult->GetFormat() ).m_bytesPerBlock;
+
+			const uint8_t* pMapData = pMap->GetData();
+			uint8_t* pResultData = pResult->GetData();
+
+			// TODO: ParallelFor this
+
+            for (int lod=0;lod<pMap->GetLODCount(); ++lod)
+            {
+                pMapData += 4+sizeY*sizeof(uint32_t);
+
+                for ( int y=0; y<sizeY; ++y )
+                {
+                    int x=0;
+                    while ( x<sizeX )
+                    {
+                        // Decode header
+                        uint16_t equal = *(const uint16_t*)pMapData;
+                        pMapData += 2;
+
+                        uint8_t different = *pMapData;
+                        ++pMapData;
+
+                        uint8_t equalPixel = *pMapData;
+                        ++pMapData;
+
+                        // Equal pixels
+                        {
+                            int dx,dy;
+                            MutableDecodeOffset( equalPixel, dx, dy );
+                            dx += x;
+                            dy += y;
+
+							if (dx >= 0 && dx < sizeX && dy >= 0 && dy < sizeY)
+							{
+								int offset = (dy * sizeX + dx) * pixelSize;
+								memcpy(pResultData, &pSourceData[offset], equal * pixelSize);
+							}
+
+                            pResultData += equal*pixelSize;
+                            x += equal;
+                        }
+
+
+                        // Different pixels
+                        for ( int i=0; i<different; ++i )
+                        {
+                            int dx,dy;
+                            MutableDecodeOffset( pMapData[i], dx, dy );
+                            dx += x;
+                            dy += y;
+
+							if (dx>=0 && dx<sizeX && dy>=0 && dy<sizeY)
+							{
+								int offset = (dy * sizeX + dx) * pixelSize;
+								memcpy(pResultData, &pSourceData[offset], pixelSize);
+							}
+
+                            pResultData += pixelSize;
+                            ++x;
+                        }
+
+                        pMapData += different;
+                    }
+                }
+
+                sizeX = FMath::DivideAndRoundUp(sizeX,2);
+                sizeY = FMath::DivideAndRoundUp(sizeY,2);
+            }
+		}
+
+		else
+		{
+			checkf( false, TEXT("Unsupported mask format.") );
+		}
+
+
+	}
+
+
+}
diff --git a/Engine/Plugins/Experimental/Mutable/Source/MutableRuntime/Private/OpImageDodge.h b/Engine/Plugins/Experimental/Mutable/Source/MutableRuntime/Private/OpImageDodge.h
new file mode 100644
--- /dev/null
+++ b/Engine/Plugins/Experimental/Mutable/Source/MutableRuntime/Private/OpImageDodge.h
@@ -0,0 +1,54 @@
+// Copyright Epic Games, Inc. All Rights Reserved.
+
+#pragma once
+
+#include "ImagePrivate.h"
+
+
+namespace mu
+{
+
+	inline unsigned DodgeChannelMasked( unsigned base, unsigned blended, unsigned mask )
+	{
+		// R = Base / (1-Blend)
+		unsigned dodge = ( base << 8 ) / ( 256 - blended );
+		unsigned masked = ( ( ( 255 - mask ) * base ) + ( mask * dodge ) ) >> 8;
+		return masked;
+	}
+
+	inline unsigned DodgeChannel( unsigned base, unsigned blended )
+	{
+		// R = Base / (1-Blend)
+		unsigned dodge = ( base << 8 ) / ( 256 - blended );
+		return dodge;
+	}
+
+
+    inline void ImageDodge( Image* pResult, const Image* pBase, const Image* pMask, vec3<float> col )
+	{
+        ImageLayerColour<DodgeChannelMasked,DodgeChannel, true>( pResult, pBase, pMask, col );
+	}
+
+
+    inline void ImageDodge( Image* pResult, const Image* pBase, const Image* pMask,
+                            const Image* pBlended,
+                            bool applyToAlpha )
+	{
+        ImageLayer<DodgeChannelMasked,DodgeChannel, true>
+                ( pResult, pBase, pMask, pBlended, applyToAlpha );
+	}
+
+
+    inline void ImageDodge( Image* pResult, const Image* pBase, vec3<float> col )
+	{
+        ImageLayerColour<DodgeChannel, true>( pResult, pBase, col );
+	}
+
+
+    inline void ImageDodge( Image* pResult, const Image* pBase, const Image* pBlended,
+                            bool applyToAlpha )
+	{
+        ImageLayer<DodgeChannel, true>( pResult, pBase, pBlended, applyToAlpha );
+	}
+
+}
diff --git a/Engine/Plugins/Experimental/Mutable/Source/MutableRuntime/Private/OpImageGradient.h b/Engine/Plugins/Experimental/Mutable/Source/MutableRuntime/Private/OpImageGradient.h
new file mode 100644
--- /dev/null
+++ b/Engine/Plugins/Experimental/Mutable/Source/MutableRuntime/Private/OpImageGradient.h
@@ -0,0 +1,43 @@
+// Copyright Epic Games, Inc. All Rights Reserved.
+
+#pragma once
+
+#include "ImagePrivate.h"
+#include "Platform.h"
+
+
+
+namespace mu
+{
+
+	inline ImagePtr ImageGradient( const vec3<float>& c0,
+								   const vec3<float>& c1,
+								   int sizeX, int sizeY )
+	{
+		ImagePtr pDest = new Image((uint16_t)sizeX, (uint16_t)sizeY, 1, EImageFormat::IF_RGB_UBYTE );
+
+        uint8_t* pDestBuf = pDest->GetData();
+
+		for ( int i=0; i<sizeX; ++i )
+		{
+			float delta = float(i)/float(sizeX-1);
+			vec3<float> c = c0 * (1.0f-delta) + c1 * delta;
+
+            uint8_t colour[4];
+            colour[0] = (uint8_t)std::max( 0, std::min( 255, int(c[0]*255.0f) ) );
+            colour[1] = (uint8_t)std::max( 0, std::min( 255, int(c[1]*255.0f) ) );
+            colour[2] = (uint8_t)std::max( 0, std::min( 255, int(c[2]*255.0f) ) );
+			colour[3] = 255;
+
+			for ( int j=0; j<sizeY; ++j )
+			{
+				pDestBuf[ (sizeX*j+i)*3 + 0 ] = colour[0];
+				pDestBuf[ (sizeX*j+i)*3 + 1 ] = colour[1];
+				pDestBuf[ (sizeX*j+i)*3 + 2 ] = colour[2];
+			}
+		}
+
+		return pDest;
+	}
+
+}
diff --git a/Engine/Plugins/Experimental/Mutable/Source/MutableRuntime/Private/OpImageGrow.h b/Engine/Plugins/Experimental/Mutable/Source/MutableRuntime/Private/OpImageGrow.h
new file mode 100644
--- /dev/null
+++ b/Engine/Plugins/Experimental/Mutable/Source/MutableRuntime/Private/OpImageGrow.h
@@ -0,0 +1,124 @@
+// Copyright Epic Games, Inc. All Rights Reserved.
+
+#pragma once
+
+#include "ImagePrivate.h"
+#include "Platform.h"
+#include "ConvertData.h"
+
+
+namespace mu
+{
+
+	//---------------------------------------------------------------------------------------------
+	//! Grow the image as if it was a bitmap
+	//---------------------------------------------------------------------------------------------
+	inline void ImageGrow( Image* pImage )
+	{
+		check( pImage->GetFormat()== EImageFormat::IF_L_UBYTE );
+
+		MUTABLE_CPUPROFILER_SCOPE(ImageGrow1)
+
+		int sizeX = pImage->GetSizeX();
+		int sizeY = pImage->GetSizeY();
+		check( sizeX>0 && sizeY>0 );
+
+		ImagePtrConst pSource = pImage->Clone();
+
+		int rowSize = sizeX*GetImageFormatData( pImage->GetFormat() ).m_bytesPerBlock;
+
+		for ( int y=0; y<sizeY; ++y )
+		{
+			int y0 = std::max( 0, std::min( sizeY-1, y-1 ) );
+			int y1 = std::max( 0, std::min( sizeY-1, y   ) );
+			int y2 = std::max( 0, std::min( sizeY-1, y+1 ) );
+
+            const uint8_t* pS0 = pSource->GetData() + y0 * rowSize;
+            const uint8_t* pS1 = pSource->GetData() + y1 * rowSize;
+            const uint8_t* pS2 = pSource->GetData() + y2 * rowSize;
+
+            uint8_t* pR = pImage->GetData() + y1 * rowSize;
+
+			if ( pS0[0] || pS2[0] || pS1[1] )
+			{
+				pR[0] = 255;
+			}
+
+			if ( pS0[sizeX-1] || pS2[sizeX-1] || pS1[sizeX-2] )
+			{
+				pR[sizeX-1] = 255;
+			}
+
+			for ( int x=1; x<sizeX-1; ++x )
+			{
+				if ( pS0[x] || pS2[x] || pS1[x-1] || pS1[x+1] )
+				{
+					pR[x] = 255;
+				}
+			}
+
+		}
+
+	}
+
+
+	//---------------------------------------------------------------------------------------------
+	//! Grow the image copying border colours, by using an external mask as reference
+	//---------------------------------------------------------------------------------------------
+	inline void ImageGrow( Image* pImage, const Image* pMask )
+	{
+		MUTABLE_CPUPROFILER_SCOPE(ImageGrow2)
+
+		check( pMask->GetFormat()== EImageFormat::IF_L_UBYTE );
+
+		int sizeX = pImage->GetSizeX();
+		int sizeY = pImage->GetSizeY();
+		check( sizeX>0 && sizeY>0 );
+
+        if (sizeX<2 || sizeY<2)
+        {
+            return;
+        }
+
+		int pixelSize = GetImageFormatData( pImage->GetFormat() ).m_bytesPerBlock;
+		int rowSize = sizeX*pixelSize;
+
+		for ( int y=0; y<sizeY; ++y )
+		{
+            const uint8_t* pS0 = y>0 ? pMask->GetData() + (y-1) * sizeX : nullptr;
+            const uint8_t* pS1 = pMask->GetData() + y * sizeX;
+            const uint8_t* pS2 = y<sizeY-1 ? pMask->GetData() + (y+1) * sizeX : nullptr;
+
+            const uint8_t* pR0 = y>0 ? pImage->GetData() + (y-1) * rowSize : nullptr;
+                  uint8_t* pR1 = pImage->GetData() + y * rowSize;
+            const uint8_t* pR2 = y<sizeY-1 ? pImage->GetData() + (y+1) * rowSize : nullptr;
+
+			for ( int x=0; x<sizeX; ++x )
+			{
+				if (pS1[x])
+				{
+					// Nothing to do.
+				}
+				else if ( y>0 && pS0[x] )
+				{
+					memcpy( &pR1[x*pixelSize], &pR0[x*pixelSize], pixelSize );
+				}
+				else if ( y<sizeY-1 && pS2[x] )
+				{
+					memcpy( &pR1[x*pixelSize], &pR2[x*pixelSize], pixelSize );
+				}
+				else if ( x<sizeX-1 && pS1[x+1] )
+				{
+					memcpy( &pR1[x*pixelSize], &pR1[(x+1)*pixelSize], pixelSize );
+				}
+				else if ( x>0 && pS1[x-1] )
+				{
+					memcpy( &pR1[x*pixelSize], &pR1[(x-1)*pixelSize], pixelSize );
+				}
+			}
+
+		}
+
+	}
+
+}
diff --git a/Engine/Plugins/Experimental/Mutable/Source/MutableRuntime/Private/OpImageHardLight.h b/Engine/Plugins/Experimental/Mutable/Source/MutableRuntime/Private/OpImageHardLight.h
new file mode 100644
--- /dev/null
+++ b/Engine/Plugins/Experimental/Mutable/Source/MutableRuntime/Private/OpImageHardLight.h
@@ -0,0 +1,61 @@
+// Copyright Epic Games, Inc. All Rights Reserved.
+
+#pragma once
+
+#include "ImagePrivate.h"
+
+namespace mu
+{
+
+	inline unsigned HardLightChannel( unsigned base, unsigned blended )
+	{
+		// gimp-like
+
+
+		// photoshop-like
+		// if (Blend > ½) R = 1 - (1-Base) × (1-2×(Blend-½))
+		// if (Blend <= ½) R = Base × (2×Blend)
+		unsigned hardlight = blended > 128
+				? 255 - ( ( ( 255 - base ) * ( 255 - 2*(blended - 128 ) ) ) >> 8 )
+				: ( base * (2 * blended) ) >> 8;
+
+		hardlight = std::min( 255u, hardlight );
+
+		return hardlight;
+	}
+
+	inline unsigned HardLightChannelMasked( unsigned base, unsigned blended, unsigned mask )
+	{
+		unsigned hardlight = HardLightChannel( base, blended );
+
+		unsigned masked = ( ( ( 255 - mask ) * base ) + ( mask * hardlight ) ) >> 8;
+		return masked;
+	}
+
+
+    inline void ImageHardLight( Image* pResult, const Image* pBase, vec3<float> col )
+	{
+        ImageLayerColour<HardLightChannel, true>( pResult, pBase, col );
+	}
+
+
+    inline void ImageHardLight( Image* pResult, const Image* pBase, const Image* pBlended,
+                                bool applyToAlpha )
+	{
+        ImageLayer<HardLightChannel, true>( pResult, pBase, pBlended, applyToAlpha );
+	}
+
+
+    inline void ImageHardLight( Image* pResult, const Image* pBase, const Image* pMask, vec3<float> col )
+	{
+        ImageLayerColour<HardLightChannelMasked,HardLightChannel,true>( pResult, pBase, pMask, col );
+	}
+
+
+    inline void ImageHardLight( Image* pResult, const Image* pBase, const Image* pMask, const Image* pBlended,
+                                bool applyToAlpha )
+	{
+        ImageLayer<HardLightChannelMasked,HardLightChannel,true>( pResult, pBase, pMask, pBlended, applyToAlpha );
+	}
+
+}
diff --git a/Engine/Plugins/Experimental/Mutable/Source/MutableRuntime/Private/OpImageInterpolate.h b/Engine/Plugins/Experimental/Mutable/Source/MutableRuntime/Private/OpImageInterpolate.h
new file mode 100644
--- /dev/null
+++ b/Engine/Plugins/Experimental/Mutable/Source/MutableRuntime/Private/OpImageInterpolate.h
@@ -0,0 +1,172 @@
+// Copyright Epic Games, Inc. All Rights Reserved.
+
+#pragma once
+
+#include "ImagePrivate.h"
+#include "Platform.h"
+
+namespace mu
+{
+
+	//---------------------------------------------------------------------------------------------
+	//!
+	//---------------------------------------------------------------------------------------------
+    inline void ImageInterpolate( Image* pDest,
+                                  const Image* pA,
+                                  const Image* pB,
+                                  float factor )
+	{
+		MUTABLE_CPUPROFILER_SCOPE(ImageInterpolate)
+
+        check( pA->GetSizeX() == pDest->GetSizeX() );
+        check( pA->GetSizeY() == pDest->GetSizeY() );
+        check( pA->GetFormat() == pDest->GetFormat() );
+        check( pA->GetSizeX() == pB->GetSizeX() );
+        check( pA->GetSizeY() == pB->GetSizeY() );
+        check( pA->GetFormat() == pB->GetFormat() );
+
+		// Clamp the factor
+		factor = std::max( 0.0f, std::min( 1.0f, factor ) );
+
+        uint8_t* pDestBuf = pDest->GetData();
+        const uint8_t* pABuf = pA->GetData();
+        const uint8_t* pBBuf = pB->GetData();
+
+		// Generic implementation
+		int pixelCount = (int)pA->CalculatePixelCount();
+
+		switch ( pA->GetFormat() )
+		{
+		case EImageFormat::IF_L_UBYTE:
+		{
+            uint32_t w_8 = (uint32_t)(factor*255);
+			for ( int i=0; i<pixelCount; ++i )
+			{
+                uint32_t a_8 = *pABuf++;
+                uint32_t b_8 = *pBBuf++;
+                uint32_t i_16 = a_8 * (255-w_8) + b_8 * w_8;
+                *pDestBuf++ = (uint8_t) ( i_16>>8 );
+			}
+			break;
+		}
+
+		case EImageFormat::IF_RGB_UBYTE:
+		{
+            uint32_t w_8 = (uint32_t)(factor*255);
+			for ( int i=0; i<pixelCount*3; ++i )
+			{
+                uint32_t a_8 = *pABuf++;
+                uint32_t b_8 = *pBBuf++;
+                uint32_t i_16 = a_8 * (255-w_8) + b_8 * w_8;
+                *pDestBuf++ = (uint8_t) ( i_16>>8 );
+			}
+			break;
+		}
+
+        case EImageFormat::IF_BGRA_UBYTE:
+        case EImageFormat::IF_RGBA_UBYTE:
+        {
+            uint32_t w_8 = (uint32_t)(factor*255);
+			for ( int i=0; i<pixelCount*4; ++i )
+			{
+                uint32_t a_8 = *pABuf++;
+                uint32_t b_8 = *pBBuf++;
+                uint32_t i_16 = a_8 * (255-w_8) + b_8 * w_8;
+                *pDestBuf++ = (uint8_t)(i_16>>8);
+			}
+			break;
+		}
+
+		default:
+			check(false);
+		}
+
+	}
+
+
+	//---------------------------------------------------------------------------------------------
+	//!
+	//---------------------------------------------------------------------------------------------
+	inline ImagePtr ImageInterpolate( const Image* p0, const Image* p1, const Image* p2,
+									  float factor1, float factor2 )
+	{
+		MUTABLE_CPUPROFILER_SCOPE(ImageInterpolate2)
+
+		check( p0->GetSizeX() == p1->GetSizeX() );
+		check( p0->GetSizeY() == p1->GetSizeY() );
+		check( p0->GetFormat() == p1->GetFormat() );
+		check( p0->GetSizeX() == p2->GetSizeX() );
+		check( p0->GetSizeY() == p2->GetSizeY() );
+		check( p0->GetFormat() == p2->GetFormat() );
+
+		// Clamp the factor
+		factor1 = std::max( 0.0f, std::min( 1.0f, factor1 ) );
+		factor2 = std::max( 0.0f, std::min( 1.0f-factor1, factor2 ) );
+		float factor0 = std::max( 0.0f, 1.0f-factor1-factor2 );
+
+        ImagePtr pDest = new Image( p0->GetSizeX(), p0->GetSizeY(),
+                                    p0->GetLODCount(),
+                                    p0->GetFormat() );
+
+        uint8_t* pDestBuf = pDest->GetData();
+        const uint8_t* p0Buf = p0->GetData();
+        const uint8_t* p1Buf = p1->GetData();
+        const uint8_t* p2Buf = p2->GetData();
+
+		// Generic implementation
+		int pixelCount = (int)p0->CalculatePixelCount();
+
+        uint32_t w0_8 = (uint32_t)(factor0*255);
+        uint32_t w1_8 = (uint32_t)(factor1*255);
+        uint32_t w2_8 = (uint32_t)(factor2*255);
+
+		switch ( p0->GetFormat() )
+		{
+		case EImageFormat::IF_L_UBYTE:
+		{
+			for ( int i=0; i<pixelCount; ++i )
+			{
+                uint32_t t0_8 = *p0Buf++;
+                uint32_t t1_8 = *p1Buf++;
+                uint32_t t2_8 = *p2Buf++;
+                uint32_t i_16 = t0_8 * w0_8 + t1_8 * w1_8 + t2_8 * w2_8;
+                *pDestBuf++ = (uint8_t) ( i_16>>8 );
+			}
+			break;
+		}
+
+		case EImageFormat::IF_RGB_UBYTE:
+		{
+			for ( int i=0; i<pixelCount*3; ++i )
+			{
+                uint32_t t0_8 = *p0Buf++;
+                uint32_t t1_8 = *p1Buf++;
+                uint32_t t2_8 = *p2Buf++;
+                uint32_t i_16 = t0_8 * w0_8 + t1_8 * w1_8 + t2_8 * w2_8;
+                *pDestBuf++ = (uint8_t) ( i_16>>8 );
+			}
+			break;
+		}
+
+        case EImageFormat::IF_BGRA_UBYTE:
+        case EImageFormat::IF_RGBA_UBYTE:
+        {
+			for ( int i=0; i<pixelCount*4; ++i )
+			{
+                uint32_t t0_8 = *p0Buf++;
+                uint32_t t1_8 = *p1Buf++;
+                uint32_t t2_8 = *p2Buf++;
+                uint32_t i_16 = t0_8 * w0_8 + t1_8 * w1_8 + t2_8 * w2_8;
+                *pDestBuf++ = (uint8_t) (i_16>>8);
+			}
+			break;
+		}
+
+		default:
+			check(false);
+		}
+
+		return pDest;
+	}
+
+}
diff --git a/Engine/Plugins/Experimental/Mutable/Source/MutableRuntime/Private/OpImageInvert.h b/Engine/Plugins/Experimental/Mutable/Source/MutableRuntime/Private/OpImageInvert.h
new file mode 100644
--- /dev/null
+++ b/Engine/Plugins/Experimental/Mutable/Source/MutableRuntime/Private/OpImageInvert.h
@@ -0,0 +1,68 @@
+// Copyright Epic Games, Inc. All Rights Reserved.
+
+#pragma once
+
+#include "ImagePrivate.h"
+#include "Platform.h"
+
+namespace mu
+{
+
+	//---------------------------------------------------------------------------------------------
+	//!
+	//---------------------------------------------------------------------------------------------
+	inline ImagePtr ImageInvert(const Image* pA)
+	{
+		MUTABLE_CPUPROFILER_SCOPE(ImageInvert)
+
+		ImagePtr pDest = new Image(pA->GetSizeX(), pA->GetSizeY(), pA->GetLODCount(), pA->GetFormat());
+
+		uint8_t* pDestBuf = pDest->GetData();
+		const uint8_t* pABuf = pA->GetData();
+
+		//Generic implementation
+		int pixelCount = (int)pA->CalculatePixelCount();
+
+		switch (pA->GetFormat())
+		{
+		case EImageFormat::IF_L_UBYTE:
+		{
+			for (int i = 0; i < pixelCount; ++i)
+			{
+				pDestBuf[i] = 255 - pABuf[i];
+			}
+			break;
+		}
+		
+		case EImageFormat::IF_RGB_UBYTE:
+		{
+			for (int i = 0; i < pixelCount; ++i)
+			{
+				for (int c = 0; c < 3; ++c)
+				{
+					pDestBuf[i * 3 + c] = 255 - pABuf[i * 3 + c];
+				}
+			}
+			break;
+		}
+		case EImageFormat::IF_RGBA_UBYTE:
+		case EImageFormat::IF_BGRA_UBYTE:
+		{
+			for (int i = 0; i < pixelCount; ++i)
+			{
+				for (int c = 0; c < 3; ++c)
+				{
+					pDestBuf[i * 4 + c] = 255 - pABuf[i * 4 + c];
+				}
+
+				pDestBuf[i * 4 + 3] = pABuf[i * 4 + 3];
+			}
+			break;
+		}
+		default:
+			check(false);			
+		}
+
+		return pDest;
+	}
+}
