[mutable] Move Mutable plugin from the restricted folder to the experimental plugins folder.

diff --git a/Engine/Plugins/Experimental/Mutable/Source/MutableRuntime/Private/OpImageLayer.h b/Engine/Plugins/Experimental/Mutable/Source/MutableRuntime/Private/OpImageLayer.h
new file mode 100644
--- /dev/null
+++ b/Engine/Plugins/Experimental/Mutable/Source/MutableRuntime/Private/OpImageLayer.h
@@ -0,0 +1,1729 @@
+// Copyright Epic Games, Inc. All Rights Reserved.
+
+#pragma once
+
+#include "ImagePrivate.h"
+#include "MutableRuntime/Private/ImageRLE.h"
+#include "Async/ParallelFor.h"
+
+#include "Templates/UnrealTemplate.h"
+
+namespace mu
+{
+
+
+	//---------------------------------------------------------------------------------------------
+	//! Apply a blending function to an image with a colour source
+	//---------------------------------------------------------------------------------------------
+	template< unsigned (*BLEND_FUNC)(unsigned,unsigned),
+			  bool CLAMP >
+    inline void ImageLayerColour( Image* pResult, const Image* pBase, vec3<float> col )
+	{
+        check( pResult->GetFormat() == pBase->GetFormat() );
+        check( pResult->GetSizeX() == pBase->GetSizeX() );
+        check( pResult->GetSizeY() == pBase->GetSizeY() );
+        check( pResult->GetLODCount() == pBase->GetLODCount() );
+
+		EImageFormat baseFormat = pBase->GetFormat();
+
+        uint8* pDestBuf = pResult->GetData();
+        const uint8* pBaseBuf = pBase->GetData();
+
+		unsigned top[3];
+		top[0] = (unsigned)(255 * col[0]);
+		top[1] = (unsigned)(255 * col[1]);
+		top[2] = (unsigned)(255 * col[2]);
+
+		// Generic implementation
+		int32 pixelCount = pBase->CalculatePixelCount();
+
+		constexpr int PixelCountConcurrencyThreshold = 0xff;
+
+		switch (baseFormat)
+		{
+		case EImageFormat::IF_RGB_UBYTE:
+		{
+			const auto& ProcessPixel = [
+				pBaseBuf, pDestBuf, top
+			] (uint32 i)
+			{
+				for (int c = 0; c < 3; ++c)
+				{
+					unsigned base = pBaseBuf[3 * i + c];
+					unsigned result = BLEND_FUNC(base, top[c]);
+					if (CLAMP)
+					{
+						pDestBuf[3 * i + c] = (uint8)std::min(255u, result);
+					}
+					else
+					{
+						pDestBuf[3 * i + c] = (uint8)result;
+					}
+				}
+			};
+
+			if (pixelCount > PixelCountConcurrencyThreshold)
+			{
+				ParallelFor(pixelCount, ProcessPixel);
+			}
+			else
+			{
+				for ( int p=0; p< pixelCount; ++p )
+				{
+					ProcessPixel(p);
+				}
+			}
+			break;
+		}
+		case EImageFormat::IF_RGBA_UBYTE:
+		{
+			const auto& ProcessPixel = [
+				pBaseBuf, pDestBuf, top
+			] (uint32 i)
+			{
+				for (int c = 0; c < 3; ++c)
+				{
+					unsigned base = pBaseBuf[4 * i + c];
+					unsigned result = BLEND_FUNC(base, top[c]);
+					if (CLAMP)
+					{
+						pDestBuf[4 * i + c] = (uint8)std::min(255u, result);
+					}
+					else
+					{
+						pDestBuf[4 * i + c] = (uint8)result;
+					}
+				}
+				pDestBuf[4 * i + 3] = pBaseBuf[4 * i + 3];
+			};
+
+			if (pixelCount > PixelCountConcurrencyThreshold)
+			{
+				ParallelFor(pixelCount, ProcessPixel);
+			}
+			else
+			{
+				for (int p = 0; p < pixelCount; ++p)
+				{
+					ProcessPixel(p);
+				}
+			}
+			break;
+		}
+		case EImageFormat::IF_BGRA_UBYTE:
+		{
+			const auto& ProcessPixel = [
+					pBaseBuf, pDestBuf, top
+				] (uint32 i)
+				{
+					for (int c = 0; c < 3; ++c)
+					{
+						unsigned base = pBaseBuf[4 * i + c];
+						unsigned result = BLEND_FUNC(base, top[2 - c]);
+						if (CLAMP)
+						{
+							pDestBuf[4 * i + c] = (uint8)std::min(255u, result);
+						}
+						else
+						{
+							pDestBuf[4 * i + c] = (uint8)result;
+						}
+					}
+					pDestBuf[4 * i + 3] = pBaseBuf[4 * i + 3];
+				};
+
+			if (pixelCount > PixelCountConcurrencyThreshold)
+			{
+				ParallelFor(pixelCount, ProcessPixel);
+			}
+			else
+			{
+				for (int p = 0; p < pixelCount; ++p)
+				{
+					ProcessPixel(p);
+				}
+			}
+			break;
+		}
+		case EImageFormat::IF_L_UBYTE:
+		{
+			const auto& ProcessPixel = [
+					pBaseBuf, pDestBuf, top
+				] (uint32 i)
+				{
+					unsigned base = pBaseBuf[i];
+					unsigned result = BLEND_FUNC(base, top[0]);
+					if (CLAMP)
+					{
+						pDestBuf[i] = (uint8)std::min(255u, result);
+					}
+					else
+					{
+						pDestBuf[i] = (uint8)result;
+					}
+				};
+
+
+			if (pixelCount > PixelCountConcurrencyThreshold)
+			{
+				ParallelFor(pixelCount, ProcessPixel);
+			}
+			else
+			{
+				for (int p = 0; p < pixelCount; ++p)
+				{
+					ProcessPixel(p);
+				}
+			}
+			break;
+		}
+
+		default:
+			checkf(false, TEXT("Unsupported format."));
+			break;
+		}
+    }
+
+
+	//---------------------------------------------------------------------------------------------
+	template< unsigned (*BLEND_FUNC_MASKED)(unsigned,unsigned,unsigned),
+			  unsigned (*BLEND_FUNC)(unsigned,unsigned),
+			  bool CLAMP,
+			  int NC >
+    inline void BufferLayerColourFormat( uint8* pStartDestBuf,
+								   const Image* pBase,
+								   const Image* pMask,
+								   vec3<float> col )
+	{
+		unsigned top[3];
+		top[0] = (unsigned)(255 * col[0]);
+		top[1] = (unsigned)(255 * col[1]);
+		top[2] = (unsigned)(255 * col[2]);
+
+		uint8* pDestBuf = pStartDestBuf;
+        const uint8* pMaskBuf = pMask->GetData();
+        const uint8* pBaseBuf = pBase->GetData();
+		EImageFormat maskFormat = pMask->GetFormat();
+
+		int32 LODCount = pBase->GetLODCount();
+
+        bool isUncompressed = ( maskFormat == EImageFormat::IF_L_UBYTE );
+
+        if ( isUncompressed )
+        {
+            int32 pixelCount = pBase->CalculatePixelCount();
+            for ( int i=0; i<pixelCount; ++i )
+            {
+                unsigned mask = pMaskBuf[i];
+                for ( int c=0; c<NC && c<3; ++c )
+                {
+                    unsigned base = pBaseBuf[NC*i+c];
+                    unsigned result = BLEND_FUNC_MASKED( base, top[c], mask );
+                    if ( CLAMP )
+                    {
+                        pDestBuf[NC*i+c] = (uint8)FMath::Min( 255u, result );
+                    }
+                    else
+                    {
+                        pDestBuf[NC*i+c] = (uint8)result;
+                    }
+                }
+
+                constexpr bool isNC4 = (NC==4);
+                if ( isNC4 )
+                {
+                    pDestBuf[NC*i+3] = pBaseBuf[NC*i+3];
+                }
+            }
+        }
+        else if ( maskFormat==EImageFormat::IF_L_UBYTE_RLE )
+        {
+            int rows = pBase->GetSizeY();
+            int width = pBase->GetSizeX();
+
+			// debug test
+			ImagePtr TempDecodedMask = new Image(pMask->GetSizeX(), pMask->GetSizeY(), pMask->GetLODCount(), EImageFormat::IF_L_UBYTE);
+			UncompressRLE_L(pMask, TempDecodedMask.get());
+
+            for (int32 lod=0;lod<LODCount; ++lod)
+            {
+				// Skip mip size and row sizes.
+                pMaskBuf += sizeof(uint32) +rows*sizeof(uint32);
+
+                for ( int r=0; r<rows; ++r )
+                {
+                    const uint8* pDestRowEnd = pDestBuf + width*NC;
+                    while ( pDestBuf!=pDestRowEnd )
+                    {
+                        // Decode header
+						uint16 equal = 0;
+						FMemory::Memmove(&equal, pMaskBuf, sizeof(uint16));
+                        pMaskBuf += 2;
+
+                        uint8 different = *pMaskBuf;
+                        ++pMaskBuf;
+
+                        uint8 equalPixel = *pMaskBuf;
+                        ++pMaskBuf;
+
+                        // Equal pixels
+						check(pDestBuf + NC * equal <= pStartDestBuf + pBase->GetDataSize());
+                        if ( equalPixel==255 )
+                        {
+                            for ( int i=0; i<equal; ++i )
+                            {
+                                for ( int c=0; c<NC && c<3; ++c )
+                                {
+                                    unsigned base = pBaseBuf[NC*i+c];
+                                    unsigned result = BLEND_FUNC( base, top[c] );
+                                    if ( CLAMP )
+                                    {
+                                        pDestBuf[NC*i+c] = (uint8)FMath::Min( 255u, result );
+                                    }
+                                    else
+                                    {
+                                        pDestBuf[NC*i+c] = (uint8)result;
+                                    }
+                                }
+
+                                constexpr bool isNC4 = (NC==4);
+                                if ( isNC4 )
+                                {
+                                    pDestBuf[NC*i+3] = pBaseBuf[NC*i+3];
+                                }
+                            }
+                        }
+                        else if ( equalPixel>0 )
+                        {
+                            for ( int i=0; i<equal; ++i )
+                            {
+                                for ( int c=0; c<NC && c<3; ++c )
+                                {
+                                    unsigned base = pBaseBuf[NC*i+c];
+                                    unsigned result = BLEND_FUNC_MASKED( base, top[c], equalPixel );
+                                    if ( CLAMP )
+                                    {
+                                        pDestBuf[NC*i+c] = (uint8)FMath::Min( 255u, result );
+                                    }
+                                    else
+                                    {
+                                        pDestBuf[NC*i+c] = (uint8)result;
+                                    }
+                                }
+
+                                constexpr bool isNC4 = (NC==4);
+                                if ( isNC4 )
+                                {
+                                    pDestBuf[NC*i+3] = pBaseBuf[NC*i+3];
+                                }
+                            }
+                        }
+                        else
+                        {
+                            // It could happen if xxxxxOnBase
+                            if (pDestBuf!=pBaseBuf)
+                            {
+                                FMemory::Memmove( pDestBuf, pBaseBuf, NC*equal );
+                            }
+                        }
+                        pDestBuf += NC*equal;
+                        pBaseBuf += NC*equal;
+
+                        // Different pixels
+						check(pDestBuf + NC * different <= pStartDestBuf + pBase->GetDataSize());
+                        for ( int i=0; i<different; ++i )
+                        {
+                            for ( int c=0; c<NC && c<3; ++c )
+                            {
+                                unsigned mask = pMaskBuf[i];
+                                unsigned base = pBaseBuf[NC*i+c];
+                                unsigned result = BLEND_FUNC_MASKED( base, top[c], mask );
+                                if ( CLAMP )
+                                {
+                                    pDestBuf[NC*i+c] = (uint8)FMath::Min( 255u, result );
+                                }
+                                else
+                                {
+                                    pDestBuf[NC*i+c] = (uint8)result;
+                                }
+                            }
+
+                            constexpr bool isNC4 = (NC==4);
+                            if ( isNC4 )
+                            {
+                                pDestBuf[NC*i+3] = pBaseBuf[NC*i+3];
+                            }
+                        }
+
+                        pDestBuf += NC*different;
+                        pBaseBuf += NC*different;
+                        pMaskBuf += different;
+                    }
+                }
+
+                rows = FMath::DivideAndRoundUp(rows,2);
+                width = FMath::DivideAndRoundUp(width,2);
+            }
+        }
+        else
+        {
+            checkf( false, TEXT("Unsupported mask format.") );
+        }
+	}
+
+
+	//---------------------------------------------------------------------------------------------
+	template< unsigned (*BLEND_FUNC_MASKED)(unsigned,unsigned,unsigned),
+			  unsigned (*BLEND_FUNC)(unsigned,unsigned),
+			  bool CLAMP
+			  >
+    inline void BufferLayerColour( uint8* pDestBuf,
+								   const Image* pBase,
+								   const Image* pMask,
+								   vec3<float> col )
+	{
+		EImageFormat baseFormat = pBase->GetFormat();
+		if ( baseFormat==EImageFormat::IF_RGB_UBYTE )
+		{
+			BufferLayerColourFormat<BLEND_FUNC_MASKED, BLEND_FUNC, CLAMP, 3>( pDestBuf, pBase, pMask, col );
+		}
+        else if ( baseFormat==EImageFormat::IF_RGBA_UBYTE )
+        {
+            BufferLayerColourFormat<BLEND_FUNC_MASKED, BLEND_FUNC, CLAMP, 4>( pDestBuf, pBase, pMask, col );
+        }
+        else if ( baseFormat==EImageFormat::IF_BGRA_UBYTE )
+        {
+            float temp = col[0];
+            col[0] = col[2];
+            col[2] = temp;
+            BufferLayerColourFormat<BLEND_FUNC_MASKED, BLEND_FUNC, CLAMP, 4>( pDestBuf, pBase, pMask, col );
+        }
+        else if ( baseFormat==EImageFormat::IF_L_UBYTE )
+		{
+			BufferLayerColourFormat<BLEND_FUNC_MASKED, BLEND_FUNC, CLAMP, 1>( pDestBuf, pBase, pMask, col );
+		}
+		else
+		{
+			checkf( false, TEXT("Unsupported format.") );
+		}
+	}
+
+	//---------------------------------------------------------------------------------------------
+	//! Apply a blending function to an image with a colour source and a mask
+	//---------------------------------------------------------------------------------------------
+	template< unsigned (*BLEND_FUNC_MASKED)(unsigned,unsigned,unsigned),
+			  unsigned (*BLEND_FUNC)(unsigned,unsigned),
+			  bool CLAMP >
+    inline void ImageLayerColour( Image* pResult,
+                                  const Image* pBase,
+                                  const Image* pMask,
+                                  vec3<float> col )
+	{
+        check( pResult->GetFormat() == pBase->GetFormat() );
+        check( pResult->GetSizeX() == pBase->GetSizeX() );
+        check( pResult->GetSizeY() == pBase->GetSizeY() );
+        check( pResult->GetLODCount() == pBase->GetLODCount() );
+        check( pBase->GetSizeX() == pMask->GetSizeX() );
+		check( pBase->GetSizeY() == pMask->GetSizeY() );
+		check( pMask->GetFormat() == EImageFormat::IF_L_UBYTE
+						||
+						pMask->GetFormat() == EImageFormat::IF_L_UBYTE_RLE );
+
+        bool valid =
+                ( pResult->GetFormat() == pBase->GetFormat() )
+                &&
+                ( pResult->GetSizeX() == pBase->GetSizeX() )
+                &&
+                ( pResult->GetSizeY() == pBase->GetSizeY() )
+                &&
+                ( pResult->GetLODCount() == pBase->GetLODCount() )
+                &&
+                ( pBase->GetSizeX() == pMask->GetSizeX() )
+                &&
+                ( pBase->GetSizeY() == pMask->GetSizeY() )
+                &&
+                ( pMask->GetFormat() == EImageFormat::IF_L_UBYTE
+                                        ||
+                                        pMask->GetFormat() == EImageFormat::IF_L_UBYTE_RLE );
+        if (valid)
+        {
+            BufferLayerColour<BLEND_FUNC_MASKED, BLEND_FUNC, CLAMP>( pResult->GetData(), pBase, pMask, col );
+        }
+	}
+
+
+	//---------------------------------------------------------------------------------------------
+	//! Apply a blending function to an image with a colour source on the base image itself
+	//---------------------------------------------------------------------------------------------
+	template< unsigned (*BLEND_FUNC_MASKED)(unsigned,unsigned,unsigned),
+			  unsigned (*BLEND_FUNC)(unsigned,unsigned),
+			  bool CLAMP >
+	inline void ImageLayerColourOnBase( Image* pBase, const Image* pMask, vec3<float> col )
+	{
+		check( pBase->GetSizeX() == pMask->GetSizeX() );
+		check( pBase->GetSizeY() == pMask->GetSizeY() );
+		check( pMask->GetFormat() == EImageFormat::IF_L_UBYTE
+						||
+						pMask->GetFormat() == EImageFormat::IF_L_UBYTE_RLE );
+
+		BufferLayerColour<BLEND_FUNC_MASKED,BLEND_FUNC,CLAMP>( pBase->GetData(), pBase, pMask, col );
+	}
+
+
+	//---------------------------------------------------------------------------------------------
+	//! Apply a blending function to an image with another image as blending layer
+	//---------------------------------------------------------------------------------------------
+	template< unsigned (*BLEND_FUNC)(unsigned,unsigned), bool CLAMP, int NC >
+    inline void BufferLayer(uint8* pDestBuf, const Image* pBase, const Image* pBlended,
+                            bool applyToAlpha )
+	{
+		check( pBase->GetSizeX() == pBlended->GetSizeX() );
+		check( pBase->GetSizeY() == pBlended->GetSizeY() );
+		check( pBase->GetFormat() == pBlended->GetFormat() );
+
+        const uint8* pBaseBuf = pBase->GetData();
+        const uint8* pBlendedBuf = pBlended->GetData();
+
+		// Generic implementation
+		int32 pixelCount = pBase->CalculatePixelCount();
+
+		ParallelFor(pixelCount,
+			[
+				pBaseBuf, pBlendedBuf, pDestBuf, applyToAlpha
+			] (uint32 i)
+			{
+				if (applyToAlpha && NC > 3)
+				{
+					for (int c = 0; c < 3; ++c)
+					{
+						unsigned base = pBaseBuf[NC * i + c];
+						unsigned blended = pBlendedBuf[NC * i + c];
+						unsigned result = BLEND_FUNC(base, blended);
+						if (CLAMP)
+						{
+							pDestBuf[NC * i + c] = (uint8)FMath::Min(255u, result);
+						}
+						else
+						{
+							pDestBuf[NC * i + c] = (uint8)result;
+						}
+					}
+
+					pDestBuf[NC * i + 3] = pBaseBuf[NC * i + 3];
+				}
+				else
+				{
+					for (int c = 0; c < NC; ++c)
+					{
+						unsigned base = pBaseBuf[NC * i + c];
+						unsigned blended = pBlendedBuf[NC * i + c];
+						unsigned result = BLEND_FUNC(base, blended);
+						if (CLAMP)
+						{
+							pDestBuf[NC * i + c] = (uint8)FMath::Min(255u, result);
+						}
+						else
+						{
+							pDestBuf[NC * i + c] = (uint8)result;
+						}
+					}
+				}
+			});
+	}
+
+
+	//---------------------------------------------------------------------------------------------
+	//! Apply a blending function to an image with another image as blending layer
+	//---------------------------------------------------------------------------------------------
+	template< unsigned (*BLEND_FUNC)(unsigned,unsigned), bool CLAMP >
+    inline void ImageLayer( Image* pResult, const Image* pBase, const Image* pBlended,
+                            bool applyToAlpha )
+	{
+        check( pResult->GetFormat() == pBase->GetFormat() );
+        check( pResult->GetSizeX() == pBase->GetSizeX() );
+        check( pResult->GetSizeY() == pBase->GetSizeY() );
+        check( pResult->GetLODCount() == pBase->GetLODCount() );
+        check( pBase->GetSizeX() == pBlended->GetSizeX() );
+		check( pBase->GetSizeY() == pBlended->GetSizeY() );
+		check( pBase->GetFormat() == pBlended->GetFormat() );
+        check( pResult->GetLODCount() <= pBlended->GetLODCount() );
+
+		EImageFormat baseFormat = pBase->GetFormat();
+
+        uint8* pDestBuf = pResult->GetData();
+
+		if ( baseFormat==EImageFormat::IF_RGB_UBYTE )
+		{
+            BufferLayer< BLEND_FUNC, CLAMP, 3 > ( pDestBuf, pBase, pBlended, applyToAlpha );
+		}
+        else if ( baseFormat==EImageFormat::IF_RGBA_UBYTE || baseFormat==EImageFormat::IF_BGRA_UBYTE )
+		{
+            BufferLayer< BLEND_FUNC, CLAMP, 4 > ( pDestBuf, pBase, pBlended, applyToAlpha );
+		}
+		else if ( baseFormat==EImageFormat::IF_L_UBYTE )
+		{
+            BufferLayer< BLEND_FUNC, CLAMP, 1 > ( pDestBuf, pBase, pBlended, applyToAlpha );
+		}
+		else
+		{
+			checkf( false, TEXT("Unsupported format.") );
+		}
+	}
+
+
+	//---------------------------------------------------------------------------------------------
+	//! Apply a blending function to an image with another image as blending layer on the base image
+	//---------------------------------------------------------------------------------------------
+	template< unsigned (*BLEND_FUNC)(unsigned,unsigned), bool CLAMP >
+    inline void ImageLayerOnBase( Image* pBase, const Image* pBlended,
+                                  bool applyToAlpha )
+	{
+		check( pBase->GetSizeX() == pBlended->GetSizeX() );
+		check( pBase->GetSizeY() == pBlended->GetSizeY() );
+		check( pBase->GetFormat() == pBlended->GetFormat() );
+
+		EImageFormat baseFormat = pBase->GetFormat();
+        uint8* pDestBuf = pBase->GetData();
+
+		if ( baseFormat==EImageFormat::IF_RGB_UBYTE )
+		{
+            BufferLayer< BLEND_FUNC, CLAMP, 3 > ( pDestBuf, pBase, pBlended, applyToAlpha );
+		}
+        else if ( baseFormat==EImageFormat::IF_RGBA_UBYTE || baseFormat==EImageFormat::IF_BGRA_UBYTE )
+		{
+            BufferLayer< BLEND_FUNC, CLAMP, 4 > ( pDestBuf, pBase, pBlended, applyToAlpha );
+		}
+		else if ( baseFormat==EImageFormat::IF_L_UBYTE )
+		{
+            BufferLayer< BLEND_FUNC, CLAMP, 1 > ( pDestBuf, pBase, pBlended, applyToAlpha );
+		}
+		else
+		{
+			checkf( false, TEXT("Unsupported format.") );
+		}
+	}
+
+
+	//---------------------------------------------------------------------------------------------
+	template< unsigned (*BLEND_FUNC_MASKED)(unsigned,unsigned,unsigned),
+			  unsigned (*BLEND_FUNC)(unsigned,unsigned),
+			  bool CLAMP,
+			  int NC >
+    inline void BufferLayerFormat( uint8* pStartDestBuf,
+                                   const Image* pBase,
+                                   const Image* pMask,
+                                   const Image* pBlend,
+                                   bool applyToAlpha )
+	{
+        check(pBase->GetSizeX()==pMask->GetSizeX() && pBase->GetSizeY()==pMask->GetSizeY());
+        check(pBase->GetSizeX()==pBlend->GetSizeX() && pBase->GetSizeY()==pBlend->GetSizeY());
+        check(pBase->GetLODCount() <= pMask->GetLODCount() );
+        check(pBase->GetLODCount() <= pBlend->GetLODCount() );
+
+		uint8* pDestBuf = pStartDestBuf;
+        const uint8* pMaskBuf = pMask->GetData();
+        const uint8* pBaseBuf = pBase->GetData();
+        const uint8* pBlendedBuf = pBlend->GetData();
+
+		EImageFormat maskFormat = pMask->GetFormat();
+        bool isUncompressed = ( maskFormat == EImageFormat::IF_L_UBYTE );
+
+		// The base determines the number of lods to process.
+		int32 LODCount = pBase->GetLODCount();
+
+        if ( isUncompressed )
+		{
+			int32 pixelCount = pBase->CalculatePixelCount();
+
+            if (!applyToAlpha && NC>3)
+            {
+				ParallelFor(pixelCount,
+					[
+						pBaseBuf, pBlendedBuf, pMaskBuf, pDestBuf
+					] (uint32 i)
+					{
+						unsigned mask = pMaskBuf[i];
+						for (int c = 0; c < NC && c < 3; ++c)
+						{
+							unsigned base = pBaseBuf[NC * i + c];
+							unsigned blended = pBlendedBuf[NC * i + c];
+							unsigned result = BLEND_FUNC_MASKED(base, blended, mask);
+							if (CLAMP)
+							{
+								pDestBuf[NC * i + c] = (uint8)FMath::Min(255u, result);
+							}
+							else
+							{
+								pDestBuf[NC * i + c] = (uint8)result;
+							}
+						}
+
+						pDestBuf[NC * i + 3] = pBaseBuf[NC * i + 3];
+					});
+            }
+            else
+            {
+				ParallelFor(pixelCount,
+					[
+						pBaseBuf, pBlendedBuf, pMaskBuf, pDestBuf
+					] (uint32 i)
+					{
+						unsigned mask = pMaskBuf[i];
+						for (int c = 0; c < NC; ++c)
+						{
+							unsigned base = pBaseBuf[NC * i + c];
+							unsigned blended = pBlendedBuf[NC * i + c];
+							unsigned result = BLEND_FUNC_MASKED(base, blended, mask);
+							if (CLAMP)
+							{
+								pDestBuf[NC * i + c] = (uint8)FMath::Min(255u, result);
+							}
+							else
+							{
+								pDestBuf[NC * i + c] = (uint8)result;
+							}
+						}
+					});
+            }
+		}
+        else if ( maskFormat==EImageFormat::IF_L_UBYTE_RLE )
+		{
+			int rows = pBase->GetSizeY();
+			int width = pBase->GetSizeX();
+
+            for (int lod=0;lod<LODCount; ++lod)
+            {
+                pMaskBuf += 4+rows*sizeof(uint32);
+
+				// \todo: See how to handle mask buf here
+				//ParallelFor(rows,
+				//	[
+				//		pBaseBuf, pBlendedBuf, pMaskBuf, pDestBuf, width, applyToAlpha
+				//	] (uint32 r)
+				for (int r = 0; r < rows; ++r)
+				{						
+					//UE_LOG(LogMutableCore, Warning, "row: %d", r);
+
+					const uint8* pDestRowEnd = pDestBuf + width*NC;
+					while ( pDestBuf!=pDestRowEnd )
+					{
+						// Decode header
+						uint16 equal = *(const uint16*)pMaskBuf;
+						pMaskBuf += 2;
+
+						uint8 different = *pMaskBuf;
+						++pMaskBuf;
+
+						uint8 equalPixel = *pMaskBuf;
+						++pMaskBuf;
+
+						//UE_LOG(LogMutableCore, Warning, "    block: %d %d", (int)equal, (int)different );
+
+						// Equal pixels
+						check(pDestBuf + NC * equal <= pStartDestBuf + pBase->GetDataSize());
+						if ( equalPixel==255 )
+						{
+							for ( int i=0; i<equal; ++i )
+							{
+								if (!applyToAlpha && NC>3)
+								{
+									for ( int c=0; c<3; ++c )
+									{
+										unsigned base = pBaseBuf[NC*i+c];
+										unsigned blended = pBlendedBuf[NC*i+c];
+										unsigned result = BLEND_FUNC( base, blended );
+										if ( CLAMP )
+										{
+											pDestBuf[NC*i+c] = (uint8)FMath::Min( 255u, result );
+										}
+										else
+										{
+											pDestBuf[NC*i+c] = (uint8)result;
+										}
+									}
+
+									pDestBuf[NC*i+3] = (uint8)pBaseBuf[NC*i+3];
+								}
+								else
+								{
+									for ( int c=0; c<NC; ++c )
+									{
+										unsigned base = pBaseBuf[NC*i+c];
+										unsigned blended = pBlendedBuf[NC*i+c];
+										unsigned result = BLEND_FUNC( base, blended );
+										if ( CLAMP )
+										{
+											pDestBuf[NC*i+c] = (uint8)FMath::Min( 255u, result );
+										}
+										else
+										{
+											pDestBuf[NC*i+c] = (uint8)result;
+										}
+									}
+								}
+							}
+						}
+						else if ( equalPixel>0 )
+						{
+							for ( int i=0; i<equal; ++i )
+							{
+								if (!applyToAlpha && NC>3)
+								{
+									for ( int c=0; c<3; ++c )
+									{
+										unsigned base = pBaseBuf[NC*i+c];
+										unsigned blended = pBlendedBuf[NC*i+c];
+										unsigned result = BLEND_FUNC_MASKED( base, blended, equalPixel );
+										if ( CLAMP )
+										{
+											pDestBuf[NC*i+c] = (uint8)FMath::Min( 255u, result );
+										}
+										else
+										{
+											pDestBuf[NC*i+c] = (uint8)result;
+										}
+									}
+
+									pDestBuf[NC*i+3] = (uint8)pBaseBuf[NC*i+3];
+								}
+								else
+								{
+									for ( int c=0; c<NC; ++c )
+									{
+										unsigned base = pBaseBuf[NC*i+c];
+										unsigned blended = pBlendedBuf[NC*i+c];
+										unsigned result = BLEND_FUNC_MASKED( base, blended, equalPixel );
+										if ( CLAMP )
+										{
+											pDestBuf[NC*i+c] = (uint8)FMath::Min( 255u, result );
+										}
+										else
+										{
+											pDestBuf[NC*i+c] = (uint8)result;
+										}
+									}
+								}
+							}
+						}
+						else
+						{
+							// It could happen if xxxxxOnBase
+							if (pDestBuf!=pBaseBuf)
+							{
+								FMemory::Memmove( pDestBuf, pBaseBuf, NC*equal );
+							}
+						}
+						pDestBuf += NC*equal;
+						pBaseBuf += NC*equal;
+						pBlendedBuf += NC*equal;
+
+						// Different pixels
+						check(pDestBuf + NC * different <= pStartDestBuf + pBase->GetDataSize());
+						for ( int i=0; i<different; ++i )
+						{
+							if (!applyToAlpha && NC>3)
+							{
+								for ( int c=0; c<3; ++c )
+								{
+									unsigned mask = pMaskBuf[i];
+									unsigned base = pBaseBuf[NC*i+c];
+									unsigned blended = pBlendedBuf[NC*i+c];
+									unsigned result = BLEND_FUNC_MASKED( base, blended, mask );
+									if ( CLAMP )
+									{
+										pDestBuf[NC*i+c] = (uint8)std::min( 255u, result );
+									}
+									else
+									{
+										pDestBuf[NC*i+c] = (uint8)result;
+									}
+								}
+
+								pDestBuf[NC*i+3] = (uint8)pBaseBuf[NC*i+3];
+							}
+							else
+							{
+								for ( int c=0; c<NC; ++c )
+								{
+									unsigned mask = pMaskBuf[i];
+									unsigned base = pBaseBuf[NC*i+c];
+									unsigned blended = pBlendedBuf[NC*i+c];
+									unsigned result = BLEND_FUNC_MASKED( base, blended, mask );
+									if ( CLAMP )
+									{
+										pDestBuf[NC*i+c] = (uint8)std::min( 255u, result );
+									}
+									else
+									{
+										pDestBuf[NC*i+c] = (uint8)result;
+									}
+								}
+							}
+						}
+
+						pDestBuf += NC*different;
+						pBaseBuf += NC*different;
+						pBlendedBuf += NC*different;
+						pMaskBuf += different;
+					}
+				}
+
+                rows = FMath::DivideAndRoundUp(rows,2);
+                width = FMath::DivideAndRoundUp(width,2);
+            }
+		}
+		else
+		{
+			checkf( false, TEXT("Unsupported mask format.") );
+		}
+	}
+
+
+	//---------------------------------------------------------------------------------------------
+	template< unsigned (*BLEND_FUNC)(unsigned,unsigned),
+			  bool CLAMP,
+			  int NC >
+    inline void BufferLayerFormatStrideNoAlpha( uint8* pDestBuf,
+		int stride,
+		const Image* pMask,
+		const Image* pBlend, int32 LODCount )
+	{
+        const uint8* pMaskBuf = pMask->GetData();
+        const uint8* pBlendedBuf = pBlend->GetData();
+
+		EImageFormat maskFormat = pMask->GetFormat();
+        bool isUncompressed = ( maskFormat == EImageFormat::IF_L_UBYTE );
+
+        if ( isUncompressed )
+		{
+			int rowCount = pBlend->GetSizeY();
+			int pixelCount = pBlend->GetSizeX();
+			for ( int j=0; j<rowCount; ++j )
+			{
+				for ( int i=0; i<pixelCount; ++i )
+				{
+					unsigned mask = *pMaskBuf;
+					if (mask)
+					{
+						for ( int c=0; c<NC; ++c )
+						{
+							unsigned base = *pDestBuf;
+							unsigned blended = *pBlendedBuf;
+							unsigned result = BLEND_FUNC( base, blended );
+							if ( CLAMP )
+							{
+								*pDestBuf = (uint8)FMath::Min( 255u, result );
+							}
+							else
+							{
+								*pDestBuf = (uint8)result;
+							}
+							++pDestBuf;
+							++pBlendedBuf;
+						}
+					}
+					else
+					{
+						pDestBuf+=NC;
+						pBlendedBuf+=NC;
+					}
+					++pMaskBuf;
+				}
+
+				pDestBuf+=stride;
+			}
+		}
+        else if ( maskFormat==EImageFormat::IF_L_UBIT_RLE )
+		{
+			int rows = pMask->GetSizeY();
+			int width = pMask->GetSizeX();
+
+            for (int lod=0;lod< LODCount; ++lod)
+            {
+                pMaskBuf += 4+rows*sizeof(uint32);
+
+                for ( int r=0; r<rows; ++r )
+                {
+                    const uint8* pDestRowEnd = pDestBuf + width*NC;
+                    while ( pDestBuf!=pDestRowEnd )
+                    {
+                        // Decode header
+                        uint16 zeros = *(const uint16*)pMaskBuf;
+                        pMaskBuf += 2;
+
+                        uint16 ones = *(const uint16*)pMaskBuf;
+                        pMaskBuf += 2;
+
+                        // Skip
+                        pDestBuf += zeros*NC;
+                        pBlendedBuf += zeros*NC;
+
+                        // Copy
+                        FMemory::Memmove( pDestBuf, pBlendedBuf, ones*NC );
+
+                        pDestBuf += NC*ones;
+                        pBlendedBuf += NC*ones;
+                    }
+
+                    pDestBuf += stride;
+                }
+
+                rows = FMath::DivideAndRoundUp(rows,2);
+                width = FMath::DivideAndRoundUp(width,2);
+            }
+		}
+		else
+		{
+			checkf( false, TEXT("Unsupported mask format.") );
+		}
+	}
+
+
+	//---------------------------------------------------------------------------------------------
+	template< unsigned (*BLEND_FUNC_MASKED)(unsigned,unsigned,unsigned),
+			  unsigned (*BLEND_FUNC)(unsigned,unsigned),
+			  bool CLAMP >
+    inline void BufferLayer( uint8* pDestBuf,
+							 const Image* pBase,
+							 const Image* pMask,
+                             const Image* pBlend,
+                             bool applyToAlpha )
+	{
+		if ( pBase->GetFormat()==EImageFormat::IF_RGB_UBYTE )
+		{
+            BufferLayerFormat< BLEND_FUNC_MASKED, BLEND_FUNC, CLAMP, 3 >
+                    ( pDestBuf, pBase, pMask, pBlend, applyToAlpha ) ;
+		}
+        else if ( pBase->GetFormat()==EImageFormat::IF_RGBA_UBYTE || pBase->GetFormat()==EImageFormat::IF_BGRA_UBYTE )
+		{
+            BufferLayerFormat< BLEND_FUNC_MASKED, BLEND_FUNC, CLAMP, 4 >
+                    ( pDestBuf, pBase, pMask, pBlend, applyToAlpha ) ;
+		}
+		else if ( pBase->GetFormat()==EImageFormat::IF_L_UBYTE )
+		{
+            BufferLayerFormat< BLEND_FUNC_MASKED, BLEND_FUNC, CLAMP, 1 >
+                    ( pDestBuf, pBase, pMask, pBlend, applyToAlpha ) ;
+		}
+        else
+		{
+			checkf( false, TEXT("Unsupported format.") );
+		}
+	}
+
+
+	//---------------------------------------------------------------------------------------------
+	template< unsigned (*BLEND_FUNC)(unsigned,unsigned),
+			  bool CLAMP >
+    inline void BufferLayerStrideNoAlpha( uint8* pDestBuf, int stride, const Image* pMask, const Image* pBlend, int32 LODCount )
+	{
+		if ( pBlend->GetFormat()==EImageFormat::IF_RGB_UBYTE )
+		{
+			BufferLayerFormatStrideNoAlpha< BLEND_FUNC, CLAMP, 3 >
+					( pDestBuf, stride, pMask, pBlend, LODCount ) ;
+		}
+        else if ( pBlend->GetFormat()==EImageFormat::IF_RGBA_UBYTE || pBlend->GetFormat()==EImageFormat::IF_BGRA_UBYTE )
+		{
+			BufferLayerFormatStrideNoAlpha< BLEND_FUNC, CLAMP, 4 >
+					( pDestBuf, stride, pMask, pBlend, LODCount) ;
+		}
+		else if ( pBlend->GetFormat()==EImageFormat::IF_L_UBYTE )
+		{
+			BufferLayerFormatStrideNoAlpha< BLEND_FUNC, CLAMP, 1 >
+					( pDestBuf, stride, pMask, pBlend, LODCount) ;
+		}
+		else
+		{
+			checkf( false, TEXT("Unsupported format.") );
+		}
+	}
+
+
+	//---------------------------------------------------------------------------------------------
+	//! Apply a blending function to an image with another image as blending layer with a mask
+	//---------------------------------------------------------------------------------------------
+	template< unsigned (*BLEND_FUNC_MASKED)(unsigned,unsigned,unsigned),
+			  unsigned (*BLEND_FUNC)(unsigned,unsigned),
+			  bool CLAMP >
+    inline void ImageLayer( Image* pResult,
+                            const Image* pBase,
+                            const Image* pMask,
+                            const Image* pBlended,
+                            bool applyToAlpha )
+	{
+        check( pResult->GetFormat() == pBase->GetFormat() );
+        check( pResult->GetSizeX() == pBase->GetSizeX() );
+        check( pResult->GetSizeY() == pBase->GetSizeY() );
+        check( pResult->GetLODCount() == pBase->GetLODCount() );
+		check( pBase->GetSizeX() == pMask->GetSizeX() );
+		check( pBase->GetSizeY() == pMask->GetSizeY() );
+		check( pBase->GetSizeX() == pBlended->GetSizeX() );
+		check( pBase->GetSizeY() == pBlended->GetSizeY() );
+		check( pBase->GetFormat() == pBlended->GetFormat() );
+		check( pMask->GetFormat() == EImageFormat::IF_L_UBYTE
+						||
+						pMask->GetFormat() == EImageFormat::IF_L_UBYTE_RLE );
+        check( pBase->GetLODCount() <= pMask->GetLODCount() );
+        check( pBase->GetLODCount() <= pBlended->GetLODCount() );
+
+        BufferLayer<BLEND_FUNC_MASKED,BLEND_FUNC,CLAMP>( pResult->GetData(), pBase, pMask,
+                                                         pBlended, applyToAlpha );
+	}
+
+
+	//---------------------------------------------------------------------------------------------
+	//! Apply a blending function to an image with another image as blending layer, on the base
+	//! image itself
+	//---------------------------------------------------------------------------------------------
+	template< unsigned (*BLEND_FUNC_MASKED)(unsigned,unsigned,unsigned),
+			  unsigned (*BLEND_FUNC)(unsigned,unsigned),
+			  bool CLAMP >
+    inline void ImageLayerOnBase( Image* pBase, const Image* pMask, const Image* pBlended,
+                                  bool applyToAlpha )
+	{
+		check( pBase->GetSizeX() == pMask->GetSizeX() );
+		check( pBase->GetSizeY() == pMask->GetSizeY() );
+		check( pBase->GetSizeX() == pBlended->GetSizeX() );
+		check( pBase->GetSizeY() == pBlended->GetSizeY() );
+		check( pBase->GetFormat() == pBlended->GetFormat() );
+		check( pMask->GetFormat() == EImageFormat::IF_L_UBYTE
+						||
+						pMask->GetFormat() == EImageFormat::IF_L_UBYTE_RLE );
+        check( pBase->GetLODCount() <= pMask->GetLODCount() );
+        check( pBase->GetLODCount() <= pBlended->GetLODCount() );
+
+        BufferLayer<BLEND_FUNC_MASKED,BLEND_FUNC,CLAMP>( pBase->GetData(), pBase, pMask,
+                                                         pBlended, applyToAlpha );
+	}
+
+
+	//---------------------------------------------------------------------------------------------
+	//! Apply a blending function to an image with another image as blending layer, on a subrect of
+	//! the base image.
+	//! \warning this method applies the blending function to the alpha channel too
+	//! \warning this method uses the mask as a binary mask (>0)
+	//---------------------------------------------------------------------------------------------
+	template< unsigned (*BLEND_FUNC)(unsigned,unsigned),
+			  bool CLAMP >
+	inline void ImageLayerOnBaseNoAlpha( Image* pBase,
+										 const Image* pMask,
+										 const Image* pBlended,
+										 const box< vec2<int> >& rect )
+	{
+		check( pBase->GetSizeX() >= rect.min[0]+rect.size[0] );
+		check( pBase->GetSizeY() >= rect.min[1]+rect.size[1] );
+		check( pMask->GetSizeX() == pBlended->GetSizeX() );
+		check( pMask->GetSizeY() == pBlended->GetSizeY() );
+		check( pBase->GetFormat() == pBlended->GetFormat() );
+		check( pMask->GetFormat() == EImageFormat::IF_L_UBYTE
+						||
+						pMask->GetFormat() == EImageFormat::IF_L_UBYTE_RLE
+						||
+						pMask->GetFormat() == EImageFormat::IF_L_UBIT_RLE );
+        check( pBase->GetLODCount() <= pMask->GetLODCount() );
+        check( pBase->GetLODCount() <= pBlended->GetLODCount() );
+
+		int pixelSize = GetImageFormatData( pBase->GetFormat() ).m_bytesPerBlock;
+
+		int start = ( pBase->GetSizeX()*rect.min[1]+rect.min[0] ) * pixelSize;
+		int stride = ( pBase->GetSizeX() - rect.size[0] ) * pixelSize;
+		BufferLayerStrideNoAlpha<BLEND_FUNC,CLAMP>( pBase->GetData()+start, stride, pMask, pBlended, pBase->GetLODCount());
+	}
+
+
+	//---------------------------------------------------------------------------------------------
+	template< int NC >
+    inline void BufferTableColourFormat( uint8* pStartDestBuf,
+										 const Image* pBase,
+										 const Image* pMask,
+                                         const uint8* pTable0,
+                                         const uint8* pTable1,
+                                         const uint8* pTable2 )
+	{
+		uint8* pDestBuf = pStartDestBuf;
+        const uint8* pMaskBuf = pMask->GetData();
+        const uint8* pBaseBuf = pBase->GetData();
+		EImageFormat maskFormat = pMask->GetFormat();
+
+        bool isUncompressed = ( maskFormat == EImageFormat::IF_L_UBYTE );
+
+		// The base determines the lods to process.
+		int32 LODCount = pBase->GetLODCount();
+
+        if ( isUncompressed )
+		{
+			int pixelCount = (int)pBase->CalculatePixelCount();
+			for ( int i=0; i<pixelCount; ++i )
+			{
+				unsigned mask = pMaskBuf[i];
+
+				if (mask)
+				{
+					if (mask<255)
+					{
+                        pDestBuf[NC*i+0] = (uint8)( ( ( ( 255 - mask ) * pBaseBuf[NC*i+0] )
+                                             + ( mask * pTable0[ pBaseBuf[NC*i+0] ] ) ) >> 8 );
+                        pDestBuf[NC*i+1] = (uint8)( ( ( ( 255 - mask ) * pBaseBuf[NC*i+1] )
+                                             + ( mask * pTable1[ pBaseBuf[NC*i+1] ] ) ) >> 8 );
+                        pDestBuf[NC*i+2] = (uint8)( ( ( ( 255 - mask ) * pBaseBuf[NC*i+2] )
+                                             + ( mask * pTable2[ pBaseBuf[NC*i+2] ] ) ) >> 8 );
+					}
+					else
+					{
+						pDestBuf[NC*i+0] = pTable0[ pBaseBuf[NC*i+0] ];
+						pDestBuf[NC*i+1] = pTable1[ pBaseBuf[NC*i+1] ];
+						pDestBuf[NC*i+2] = pTable2[ pBaseBuf[NC*i+2] ];
+					}
+				}
+				else
+				{
+					pDestBuf[NC*i+0] = pBaseBuf[NC*i+0];
+					pDestBuf[NC*i+1] = pBaseBuf[NC*i+1];
+					pDestBuf[NC*i+2] = pBaseBuf[NC*i+2];
+				}
+
+                constexpr bool isNC4 = (NC==4);
+                if ( isNC4 )
+				{
+					pDestBuf[NC*i+3] = pBaseBuf[NC*i+3];
+				}
+			}
+		}
+		else if ( maskFormat==EImageFormat::IF_L_UBYTE_RLE )
+		{
+			int rows = pBase->GetSizeY();
+			int width = pBase->GetSizeX();
+
+            for (int32 lod=0;lod< LODCount; ++lod)
+            {
+                pMaskBuf += 4+rows*sizeof(uint32);
+
+                for ( int r=0; r<rows; ++r )
+                {
+                    const uint8* pDestRowEnd = pDestBuf + width*NC;
+                    while ( pDestBuf!=pDestRowEnd )
+                    {
+                        // Decode header
+                        uint16 equal = *(const uint16*)pMaskBuf;
+                        pMaskBuf += 2;
+
+                        uint8 different = *pMaskBuf;
+                        ++pMaskBuf;
+
+                        unsigned equalPixel = *pMaskBuf;
+                        ++pMaskBuf;
+
+                        // Equal pixels
+						check(pDestBuf + NC * equal <= pStartDestBuf + pBase->GetDataSize());
+                        if ( equalPixel==255 )
+                        {
+                            for ( int i=0; i<equal; ++i )
+                            {
+                                pDestBuf[NC*i+0] = pTable0[ pBaseBuf[NC*i+0] ];
+                                pDestBuf[NC*i+1] = pTable1[ pBaseBuf[NC*i+1] ];
+                                pDestBuf[NC*i+2] = pTable2[ pBaseBuf[NC*i+2] ];
+
+                                constexpr bool isNC4 = (NC==4);
+                                if ( isNC4 )
+                                {
+                                    pDestBuf[NC*i+3] = pBaseBuf[NC*i+3];
+                                }
+                            }
+                        }
+                        else if ( equalPixel>0 )
+                        {
+                            for ( int i=0; i<equal; ++i )
+                            {
+                                pDestBuf[NC*i+0] = (uint8)( ( ( ( 255 - equalPixel ) * pBaseBuf[NC*i+0] )
+                                                     + ( equalPixel * pTable0[ pBaseBuf[NC*i+0] ] ) ) >> 8 );
+                                pDestBuf[NC*i+1] = (uint8)( ( ( ( 255 - equalPixel ) * pBaseBuf[NC*i+1] )
+                                                     + ( equalPixel * pTable1[ pBaseBuf[NC*i+1] ] ) ) >> 8 );
+                                pDestBuf[NC*i+2] = (uint8)( ( ( ( 255 - equalPixel ) * pBaseBuf[NC*i+2] )
+                                                     + ( equalPixel * pTable2[ pBaseBuf[NC*i+2] ] ) ) >> 8 );
+
+                                constexpr bool isNC4 = (NC==4);
+                                if ( isNC4 )
+                                {
+                                    pDestBuf[NC*i+3] = pBaseBuf[NC*i+3];
+                                }
+                            }
+                        }
+                        else
+                        {
+                            // It could happen if xxxxxOnBase
+                            if (pDestBuf!=pBaseBuf)
+                            {
+								FMemory::Memmove( pDestBuf, pBaseBuf, NC*equal );
+                            }
+                        }
+                        pDestBuf += NC*equal;
+                        pBaseBuf += NC*equal;
+
+                        // Different pixels
+						check(pDestBuf + NC * different <= pStartDestBuf + pBase->GetDataSize());
+                        for ( int i=0; i<different; ++i )
+                        {
+                            unsigned mask = pMaskBuf[i];
+                            pDestBuf[NC*i+0] = (uint8)( ( ( ( 255 - mask ) * pBaseBuf[NC*i+0] )
+                                                 + ( mask * pTable0[ pBaseBuf[NC*i+0] ] ) ) >> 8 );
+                            pDestBuf[NC*i+1] = (uint8)( ( ( ( 255 - mask ) * pBaseBuf[NC*i+1] )
+                                                 + ( mask * pTable1[ pBaseBuf[NC*i+1] ] ) ) >> 8 );
+                            pDestBuf[NC*i+2] = (uint8)( ( ( ( 255 - mask ) * pBaseBuf[NC*i+2] )
+                                                 + ( mask * pTable2[ pBaseBuf[NC*i+2] ] ) ) >> 8 );
+
+                            constexpr bool isNC4 = (NC==4);
+                            if ( isNC4 )
+                            {
+                                pDestBuf[NC*i+3] = pBaseBuf[NC*i+3];
+                            }
+                        }
+
+                        pDestBuf += NC*different;
+                        pBaseBuf += NC*different;
+                        pMaskBuf += different;
+                    }
+                }
+
+                rows = FMath::DivideAndRoundUp(rows,2);
+                width = FMath::DivideAndRoundUp(width,2);
+            }
+		}
+		else
+		{
+			checkf( false, TEXT("Unsupported mask format.") );
+		}
+	}
+
+
+	//---------------------------------------------------------------------------------------------
+    inline void ImageTable( Image* pResult,
+                                const Image* pBase,
+								const Image* pMask,
+                                const uint8* pTable0,
+                                const uint8* pTable1,
+                                const uint8* pTable2
+								)
+	{
+        check( pResult->GetFormat() == pBase->GetFormat() );
+        check( pResult->GetSizeX() == pBase->GetSizeX() );
+        check( pResult->GetSizeY() == pBase->GetSizeY() );
+        check( pResult->GetLODCount() == pBase->GetLODCount() );
+        check( pBase->GetSizeX() == pMask->GetSizeX() );
+		check( pBase->GetSizeY() == pMask->GetSizeY() );
+		check( pMask->GetFormat() == EImageFormat::IF_L_UBYTE
+						||
+						pMask->GetFormat() == EImageFormat::IF_L_UBYTE_RLE );
+        check( pBase->GetLODCount() <= pMask->GetLODCount() );
+
+		EImageFormat baseFormat = pBase->GetFormat();
+		if ( baseFormat==EImageFormat::IF_RGB_UBYTE )
+		{
+            BufferTableColourFormat<3>( pResult->GetData(), pBase, pMask, pTable0, pTable1, pTable2 );
+		}
+        else if ( baseFormat==EImageFormat::IF_RGBA_UBYTE )
+        {
+            BufferTableColourFormat<4>( pResult->GetData(), pBase, pMask, pTable0, pTable1, pTable2 );
+        }
+        else if ( baseFormat==EImageFormat::IF_BGRA_UBYTE )
+        {
+            BufferTableColourFormat<4>( pResult->GetData(), pBase, pMask, pTable2, pTable1, pTable0 );
+        }
+        else
+		{
+			checkf( false, TEXT("Unsupported format.") );
+		}
+	}
+
+
+	template<size_t NC>
+	inline uint32 PackPixel(const uint8* pPixel)
+	{
+		static_assert(NC > 0 && NC <= 4);
+
+		uint32 PixelPack = 0;
+
+		// The compiler should be able to optimize this given that NC is a constant expression
+		FMemory::Memcpy( &PixelPack, pPixel, NC );
+
+		return PixelPack;
+	}
+
+	template<size_t NC>
+	inline void UnpackPixel(uint8* const pPixel, uint32 pixelData )
+	{
+		static_assert(NC > 0 && NC <= 4);
+
+		// The compiler should be able to optimize this given that NC is a constant expression
+		FMemory::Memcpy( pPixel, &pixelData, NC );
+	}
+
+	template< 
+			uint32 (*BLEND_FUNC)(uint32, uint32), 
+			size_t NC >
+	inline void BufferLayerCombineColour(uint8* pDestBuf, const Image* pBase, vec3<float> col)
+	{
+		//static_assert(NC > 0 && NC <= 4);
+
+		const uint8* pBaseBuf = pBase->GetData();
+
+		uint32 top = 
+			(uint32)(255.0f * col[0]) << 0 |
+			(uint32)(255.0f * col[1]) << 8 |
+			(uint32)(255.0f * col[2]) << 16;
+
+		int32 pixelCount = pBase->CalculatePixelCount();
+
+		for (int i = 0; i < pixelCount; ++i)
+		{
+			uint32 base = PackPixel<NC>(&pBaseBuf[NC * i]);
+
+			uint32 result = BLEND_FUNC(base, top);
+
+			UnpackPixel<NC>(&pDestBuf[NC * i], result);
+		}
+	}
+
+	template< 
+			uint32 (*BLEND_FUNC)(uint32, uint32), 
+			size_t NC >
+	inline void BufferLayerCombine(uint8* pDestBuf, const Image* pBase, const Image* pBlended)
+	{
+		//static_assert(NC > 0 && NC <= 4);
+
+		check(pBase->GetSizeX() == pBlended->GetSizeX());
+		check(pBase->GetSizeY() == pBlended->GetSizeY());
+		check(pBase->GetFormat() == pBlended->GetFormat());
+
+		const uint8* pBaseBuf = pBase->GetData();
+		const uint8* pBlendedBuf = pBlended->GetData();
+
+		int32 pixelCount = pBase->CalculatePixelCount();
+
+		for (int i = 0; i < pixelCount; ++i)
+		{
+			const uint32 base = PackPixel<NC>(&pBaseBuf[NC * i]);
+			const uint32 blend = PackPixel<NC>(&pBlendedBuf[NC * i]);
+
+			const uint32 result = BLEND_FUNC(base, blend);
+
+			UnpackPixel<NC>(&pDestBuf[NC * i], result);
+		}
+	}
+
+	template< 
+			uint32 (*BLEND_FUNC_MASKED)(uint32, uint32, uint32), 
+			size_t NC >
+	inline void BufferLayerCombine(uint8* pDestBuf, const Image* pBase, const Image* pMask, const Image* pBlended)
+	{
+		//static_assert(NC > 0 && NC <= 4);
+
+		check(pBase->GetSizeX() == pBlended->GetSizeX());
+		check(pBase->GetSizeY() == pBlended->GetSizeY());
+		check(pBase->GetSizeX() == pMask->GetSizeX());
+		check(pBase->GetSizeY() == pMask->GetSizeY());
+		check(pBase->GetFormat() == pBlended->GetFormat());
+
+		const uint8* pBaseBuf = pBase->GetData();
+		const uint8* pMaskBuf = pMask->GetData();
+		const uint8* pBlendedBuf = pBlended->GetData();
+
+		int32 pixelCount = pBase->CalculatePixelCount();
+
+		for (int i = 0; i < pixelCount; ++i)
+		{
+			const uint32 base = PackPixel<NC>(&pBaseBuf[NC * i]);
+			const uint32 blend = PackPixel<NC>(&pBlendedBuf[NC * i]);
+			const uint32 mask = PackPixel<1>(&pMaskBuf[i]);
+
+			const uint32 result = BLEND_FUNC_MASKED(base, blend, mask);
+
+			UnpackPixel<NC>(&pDestBuf[NC * i], result);
+		}
+	}
+
+	template< 
+			uint32 (*BLEND_FUNC_MASKED)(uint32, uint32, uint32), 
+			size_t NC >
+	inline void BufferLayerCombineColour(uint8* pDestBuf, const Image* pBase, const Image* pMask, vec3<float> col)
+	{
+		//static_assert(NC > 0 && NC <= 4);
+
+		check(pBase->GetSizeX() == pMask->GetSizeX());
+		check(pBase->GetSizeY() == pMask->GetSizeY());
+
+
+		const uint8* pBaseBuf = pBase->GetData();
+		const uint8* pMaskBuf = pMask->GetData();
+
+		uint32 top = 
+			(uint32)(255.0f * col[0]) << 0 |
+			(uint32)(255.0f * col[1]) << 8 |
+			(uint32)(255.0f * col[2]) << 16;
+
+
+		int32 pixelCount = pBase->CalculatePixelCount();
+
+		for (int i = 0; i < pixelCount; ++i)
+		{
+			const uint32 base = PackPixel<NC>(&pBaseBuf[NC * i]);
+			const uint32 mask = PackPixel<1>(&pMaskBuf[i]);
+
+			const uint32 result = BLEND_FUNC_MASKED(base, top, mask);
+
+			UnpackPixel<NC>(&pDestBuf[NC * i], result);
+		}
+	}
+
+	template< uint32 (*BLEND_FUNC)(uint32, uint32) >
+	inline void ImageLayerCombine(Image* pResult, const Image* pBase, const Image* pBlended)
+	{
+		check(pResult->GetFormat() == pBase->GetFormat());
+		check(pResult->GetSizeX() == pBase->GetSizeX());
+		check(pResult->GetSizeY() == pBase->GetSizeY());
+		check(pResult->GetLODCount() == pBase->GetLODCount());
+		check(pBase->GetSizeX() == pBlended->GetSizeX());
+		check(pBase->GetSizeY() == pBlended->GetSizeY());
+		check(pBase->GetFormat() == pBlended->GetFormat());
+		check(pResult->GetLODCount() <= pBlended->GetLODCount());
+
+		const EImageFormat baseFormat = pBase->GetFormat();
+
+		uint8* pDestBuf = pResult->GetData();
+		if (baseFormat == EImageFormat::IF_L_UBYTE)
+		{
+			BufferLayerCombine<BLEND_FUNC, 1 >(pDestBuf, pBase, pBlended);
+		}
+		else if (baseFormat == EImageFormat::IF_RGB_UBYTE)
+		{
+			BufferLayerCombine< BLEND_FUNC, 3 >(pDestBuf, pBase, pBlended);
+		}
+		else if (baseFormat == EImageFormat::IF_RGBA_UBYTE || baseFormat == EImageFormat::IF_BGRA_UBYTE)
+		{
+			// \todo: pass swizzle template argument if BGRA_UBYTE, not yet supported.
+			BufferLayerCombine< BLEND_FUNC, 4 >(pDestBuf, pBase, pBlended);
+		}
+		else
+		{
+			checkf(false, TEXT("Unsupported format."));
+		}
+	}
+
+	template< 
+		uint32 (*BLEND_FUNC)(uint32, uint32), 
+		uint32 (*BLEND_FUNC_MASKED)(uint32, uint32, uint32) >
+	inline void ImageLayerCombine(Image* pResult, const Image* pBase, const Image* pMask, const Image* pBlended)
+	{
+		check(pResult->GetFormat() == pBase->GetFormat());
+		check(pResult->GetSizeX() == pBase->GetSizeX());
+		check(pResult->GetSizeY() == pBase->GetSizeY());
+		check(pResult->GetLODCount() == pBase->GetLODCount());
+		check(pBase->GetSizeX() == pBlended->GetSizeX());
+		check(pBase->GetSizeY() == pBlended->GetSizeY());
+		check(pBase->GetFormat() == pBlended->GetFormat());
+		check(pResult->GetLODCount() <= pBlended->GetLODCount());
+
+		const EImageFormat baseFormat = pBase->GetFormat();
+
+		uint8* pDestBuf = pResult->GetData();
+
+		if (pMask->GetFormat() != EImageFormat::IF_L_UBYTE)
+		{
+			checkf(false, TEXT("Unsupported mask format."));
+
+			BufferLayerCombine< BLEND_FUNC, 1 >(pDestBuf, pBase, pBlended);
+		}
+
+		if (baseFormat == EImageFormat::IF_L_UBYTE)
+		{
+			BufferLayerCombine< BLEND_FUNC_MASKED, 1 >(pDestBuf, pBase, pMask, pBlended);
+		}
+		else if (baseFormat == EImageFormat::IF_RGB_UBYTE)
+		{
+			BufferLayerCombine< BLEND_FUNC_MASKED, 3 >(pDestBuf, pBase, pMask, pBlended);
+		}
+		else if (baseFormat == EImageFormat::IF_RGBA_UBYTE || baseFormat == EImageFormat::IF_BGRA_UBYTE)
+		{
+			// \todo: pass swizzle template argument if BGRA_UBYTE, not yet supported.
+			BufferLayerCombine< BLEND_FUNC_MASKED, 4 >(pDestBuf, pBase, pMask, pBlended);
+		}
+		else
+		{
+			checkf(false, TEXT("Unsupported format."));
+		}
+	}
+
+	template< uint32 (*BLEND_FUNC)(uint32, uint32)>
+	inline void ImageLayerCombineColour(Image* pResult, const Image* pBase, vec3<float> col)
+	{
+		check(pResult->GetFormat() == pBase->GetFormat());
+		check(pResult->GetSizeX() == pBase->GetSizeX());
+		check(pResult->GetSizeY() == pBase->GetSizeY());
+		check(pResult->GetLODCount() == pBase->GetLODCount());
+
+		const EImageFormat baseFormat = pBase->GetFormat();
+
+		uint8* pDestBuf = pResult->GetData();
+
+		if (baseFormat == EImageFormat::IF_L_UBYTE)
+		{
+			BufferLayerCombineColour< BLEND_FUNC, 1 >(pDestBuf, pBase, col);
+		}
+		else if (baseFormat == EImageFormat::IF_RGB_UBYTE)
+		{
+			BufferLayerCombineColour< BLEND_FUNC, 3 >(pDestBuf, pBase, col);
+		}
+		else if (baseFormat == EImageFormat::IF_RGBA_UBYTE || baseFormat == EImageFormat::IF_BGRA_UBYTE)
+		{
+			// \todo: pass swizzle template argument if BGRA_UBYTE, not yet supported.
+ 			BufferLayerCombineColour< BLEND_FUNC, 4 >(pDestBuf, pBase, col);
+		}
+		else
+		{
+			checkf(false, TEXT("Unsupported format."));
+		}
+	}
+
+	template< 
+		uint32 (*BLEND_FUNC)(uint32, uint32),
+		uint32 (*BLEND_FUNC_MASKED)(uint32, uint32, uint32) >
+	inline void ImageLayerCombineColour(Image* pResult, const Image* pBase, const Image* pMask, vec3<float> col)
+	{
+		check(pResult->GetFormat() == pBase->GetFormat());
+		check(pResult->GetSizeX() == pBase->GetSizeX());
+		check(pResult->GetSizeY() == pBase->GetSizeY());
+		check(pResult->GetLODCount() == pBase->GetLODCount());
+
+		const EImageFormat baseFormat = pBase->GetFormat();
+
+		uint8* pDestBuf = pResult->GetData();
+
+		if (pMask->GetFormat() != EImageFormat::IF_L_UBYTE)
+		{
+			checkf(false, TEXT("Unsupported mask format."));
+
+			BufferLayerCombineColour< BLEND_FUNC, 1 >(pDestBuf, pBase, col);
+		}
+
+		if (baseFormat == EImageFormat::IF_L_UBYTE)
+		{
+			BufferLayerCombineColour< BLEND_FUNC_MASKED, 1 >(pDestBuf, pBase, pMask, col);
+		}
+		else if (baseFormat == EImageFormat::IF_RGB_UBYTE)
+		{
+			BufferLayerCombineColour< BLEND_FUNC_MASKED, 3 >(pDestBuf, pBase, pMask, col);
+		}
+		else if (baseFormat == EImageFormat::IF_RGBA_UBYTE || baseFormat == EImageFormat::IF_BGRA_UBYTE)
+		{
+			// \todo: pass swizzle template argument if BGRA_UBYTE, not yet supported.
+ 			BufferLayerCombineColour< BLEND_FUNC_MASKED, 4 >(pDestBuf, pBase, pMask, col);
+		}
+		else
+		{
+			checkf(false, TEXT("Unsupported format."));
+		}
+	}
+
+	template< size_t NCBase, size_t NCBlend, class ImageCombineFn >
+	inline void BufferLayerCombineFunctor(uint8* pDestBuf, const Image* pBase, const Image* pBlended, ImageCombineFn&& ImageCombine)
+	{
+		//static_assert(NCBase > 0 && NCBase <= 4);
+		//static_assert(NCBlend > 0 && NCBlend <= 4);
+
+		check(pBase->GetSizeX() == pBlended->GetSizeX());
+		check(pBase->GetSizeY() == pBlended->GetSizeY());
+		check(pBase->GetLODCount() <= pBlended->GetLODCount());
+
+		const uint8* pBaseBuf = pBase->GetData();
+		const uint8* pBlendedBuf = pBlended->GetData();
+
+		int32 pixelCount = pBase->CalculatePixelCount();
+
+		for (int i = 0; i < pixelCount; ++i)
+		{
+			const uint32 base = PackPixel<NCBase>(&pBaseBuf[NCBase * i]);
+			const uint32 blend = PackPixel<NCBlend>(&pBlendedBuf[NCBlend * i]);
+
+			const uint32 result = ImageCombine(base, blend);
+
+			UnpackPixel<NCBase>(&pDestBuf[NCBase* i], result);
+		}
+	}
+	// Same functionality as above, in this case we use a functor which allows to pass user data. 
+	template< class ImageCombineFn >
+	inline void ImageLayerCombineFunctor(Image* pResult, const Image* pBase, const Image* pBlended, ImageCombineFn&& ImageCombine)
+	{
+		check(pResult->GetFormat() == pBase->GetFormat());
+		check(pResult->GetSizeX() == pBase->GetSizeX());
+		check(pResult->GetSizeY() == pBase->GetSizeY());
+		check(pResult->GetLODCount() == pBase->GetLODCount());
+		check(pBase->GetSizeX() == pBlended->GetSizeX());
+		check(pBase->GetSizeY() == pBlended->GetSizeY());
+		check(pResult->GetLODCount() <= pBlended->GetLODCount());
+
+		const EImageFormat baseFormat = pBase->GetFormat();
+		const EImageFormat blendFormat = pBlended->GetFormat();
+
+		uint8* pDestBuf = pResult->GetData();
+		if (baseFormat == EImageFormat::IF_L_UBYTE )
+		{
+			if (blendFormat == EImageFormat::IF_L_UBYTE )
+			{
+				BufferLayerCombineFunctor< 1, 1 >(pDestBuf, pBase, pBlended, Forward<ImageCombineFn>(ImageCombine));
+			}
+			else if (blendFormat == EImageFormat::IF_RGB_UBYTE )
+			{
+				BufferLayerCombineFunctor< 1, 3 >(pDestBuf, pBase, pBlended, Forward<ImageCombineFn>(ImageCombine));
+			}
+			else if (blendFormat == EImageFormat::IF_RGBA_UBYTE || blendFormat == EImageFormat::IF_BGRA_UBYTE )
+			{
+				BufferLayerCombineFunctor< 1, 4 >(pDestBuf, pBase, pBlended, Forward<ImageCombineFn>(ImageCombine));
+			}
+			else
+			{
+				checkf(false, TEXT("Unsupported format."));
+			}
+		}
+		else if (baseFormat == EImageFormat::IF_RGB_UBYTE)
+		{
+			if (blendFormat == EImageFormat::IF_L_UBYTE )
+			{
+				BufferLayerCombineFunctor< 3, 1 >(pDestBuf, pBase, pBlended, Forward<ImageCombineFn>(ImageCombine));
+			}
+			else if (blendFormat == EImageFormat::IF_RGB_UBYTE )
+			{
+				BufferLayerCombineFunctor< 3, 3 >(pDestBuf, pBase, pBlended, Forward<ImageCombineFn>(ImageCombine));
+			}
+			else if (blendFormat == EImageFormat::IF_RGBA_UBYTE || blendFormat == EImageFormat::IF_BGRA_UBYTE )
+			{
+				BufferLayerCombineFunctor< 3, 4 >(pDestBuf, pBase, pBlended, Forward<ImageCombineFn>(ImageCombine));
+			}
+			else
+			{
+				checkf(false, TEXT("Unsupported format."));
+			}
+		}
+		else if (baseFormat == EImageFormat::IF_RGBA_UBYTE || baseFormat == EImageFormat::IF_BGRA_UBYTE)
+		{
+			if (blendFormat == EImageFormat::IF_L_UBYTE )
+			{
+				BufferLayerCombineFunctor< 4, 1 >(pDestBuf, pBase, pBlended, Forward<ImageCombineFn>(ImageCombine));
+			}
+			else if (blendFormat == EImageFormat::IF_RGB_UBYTE )
+			{
+				BufferLayerCombineFunctor< 4, 3 >(pDestBuf, pBase, pBlended, Forward<ImageCombineFn>(ImageCombine));
+			}
+			else if (blendFormat == EImageFormat::IF_RGBA_UBYTE || blendFormat == EImageFormat::IF_BGRA_UBYTE )
+			{
+				BufferLayerCombineFunctor< 4, 4 >(pDestBuf, pBase, pBlended, Forward<ImageCombineFn>(ImageCombine));
+			}
+			else
+			{
+				checkf(false, TEXT("Unsupported format."));
+			}
+		}
+		else
+		{
+			checkf(false, TEXT("Unsupported format."));
+		}
+	}
+
+}
diff --git a/Engine/Plugins/Experimental/Mutable/Source/MutableRuntime/Private/OpImageLuminance.h b/Engine/Plugins/Experimental/Mutable/Source/MutableRuntime/Private/OpImageLuminance.h
new file mode 100644
--- /dev/null
+++ b/Engine/Plugins/Experimental/Mutable/Source/MutableRuntime/Private/OpImageLuminance.h
@@ -0,0 +1,61 @@
+// Copyright Epic Games, Inc. All Rights Reserved.
+
+#pragma once
+
+#include "ImagePrivate.h"
+#include "Platform.h"
+
+namespace mu
+{
+
+	inline ImagePtr ImageLuminance( const Image* pA )
+	{
+        ImagePtr pDest = new Image( pA->GetSizeX(), pA->GetSizeY(), pA->GetLODCount(), EImageFormat::IF_L_UBYTE );
+
+        uint8_t* pDestBuf = pDest->GetData();
+        const uint8_t* pABuf = pA->GetData();
+
+		// Generic implementation
+		int pixelCount = (int)pA->CalculatePixelCount();
+
+		switch ( pA->GetFormat() )
+		{
+		case EImageFormat::IF_RGB_UBYTE:
+		{
+			for ( int i=0; i<pixelCount; ++i )
+			{
+                uint32_t l_16 = 76 * pABuf[3*i+0] + 150 * pABuf[3*i+1] + 29 * pABuf[3*i+2];
+                pDestBuf[i] = (uint8_t)std::min( l_16>>8, 255u );
+			}
+			break;
+		}
+
+        case EImageFormat::IF_RGBA_UBYTE:
+        {
+            for ( int i=0; i<pixelCount; ++i )
+            {
+                uint32_t l_16 = 76 * pABuf[4*i+0] + 150 * pABuf[4*i+1] + 29 * pABuf[4*i+2];
+                pDestBuf[i] = (uint8_t)std::min( l_16>>8, 255u );
+            }
+            break;
+        }
+
+        case EImageFormat::IF_BGRA_UBYTE:
+        {
+            for ( int i=0; i<pixelCount; ++i )
+            {
+                uint32_t l_16 = 76 * pABuf[4*i+2] + 150 * pABuf[4*i+1] + 29 * pABuf[4*i+0];
+                pDestBuf[i] = (uint8_t)std::min( l_16>>8, 255u );
+            }
+            break;
+        }
+
+		default:
+			check(false);
+			break;
+		}
+
+		return pDest;
+	}
+
+}
diff --git a/Engine/Plugins/Experimental/Mutable/Source/MutableRuntime/Private/OpImageMipmap.h b/Engine/Plugins/Experimental/Mutable/Source/MutableRuntime/Private/OpImageMipmap.h
new file mode 100644
--- /dev/null
+++ b/Engine/Plugins/Experimental/Mutable/Source/MutableRuntime/Private/OpImageMipmap.h
@@ -0,0 +1,716 @@
+// Copyright Epic Games, Inc. All Rights Reserved.
+
+#pragma once
+
+#include "ImagePrivate.h"
+#include "OpImagePixelFormat.h"
+#include "Async/ParallelFor.h"
+
+namespace mu
+{
+
+	namespace OpImageMipmap_Detail
+	{
+
+		struct SharpenKernelStorage8
+		{
+			constexpr static size_t SIZE = 8;
+			float m_storage[SIZE * SIZE];
+		};
+
+		template<size_t ITERS, size_t N>
+		inline void BlurKernel(float(&inOutKernelStorage)[N])
+		{
+			static_assert(N > 2);
+
+			float temp[N];
+
+			float* inBuf = &inOutKernelStorage[0];
+			float* outBuf = &temp[0];
+
+			for (int32_t i = 0; i < ITERS; ++i)
+			{
+				constexpr float oneOverThree = 1.0f / 3.0f;
+
+				outBuf[0] = (inBuf[0] + inBuf[1]) * oneOverThree;
+
+				for (int32_t j = 1; j < N - 1; ++j)
+				{
+					outBuf[j] = (inBuf[j - 1] + inBuf[j] + inBuf[j + 1]) * oneOverThree;
+				}
+
+				outBuf[N - 1] = (inBuf[N - 2] + inBuf[N - 1]) * oneOverThree;
+
+				// Ping-pong buffers.
+				std::swap(outBuf, inBuf);
+			}
+
+			if (ITERS % 2 == 1)
+			{
+				for (uint32_t i = 0; i < N; ++i)
+				{
+					inOutKernelStorage[i] = temp[i];
+				}
+			}
+		}
+
+		template<size_t N>
+		inline void BuildGaussianBlurKernel(float(&kernel)[N], float variance)
+		{
+			constexpr float sqrtTau = 2.5066282746310002f;
+			const auto normalDist = [sdSqrtTauInv = 1.0f / (std::sqrt(variance) * sqrtTau),
+				twoVInv = 1.0f / (2.0f * variance)]
+				(float x)
+				-> float
+			{
+				return std::exp(-x * x * twoVInv) * sdSqrtTauInv;
+			};
+
+			const float center = static_cast<float>(N) * 0.5f;
+			float sum = 0;
+			for (int32_t i = 0; i < N; ++i)
+			{
+				const float dist = normalDist(static_cast<float>(i) - center + 0.5f);
+				kernel[i] = dist;
+				sum += dist;
+			}
+
+			const float normFactor = 1.0f / sum;
+			for (int32_t i = 0; i < N; ++i)
+			{
+				kernel[i] *= normFactor;
+			}
+		}
+
+		template<size_t N>
+		inline void BuildSharpenKernel(float(&kernel)[N], float factor)
+		{
+			for (uint32_t i = 0; i < N; ++i)
+			{
+				kernel[i] = 0.0f;
+			}
+
+			kernel[(N >> 1) - 1] = factor * 0.5f;
+			kernel[(N >> 1)] = factor * 0.5f;
+		}
+
+		template<size_t N>
+		inline void BuildKernel2DFromSeparable(float(&kernel1D)[N], SharpenKernelStorage8& outKernel)
+		{
+			static_assert(N == SharpenKernelStorage8::SIZE);
+
+			for (int32_t y = 0; y < N; ++y)
+			{
+				for (int32_t x = 0; x < N; ++x)
+				{
+					outKernel.m_storage[y * N + x] = kernel1D[y] * kernel1D[x];
+				}
+			}
+		}
+
+		inline SharpenKernelStorage8 MakeMipGaussianSharpenKernel(float factor)
+		{
+			SharpenKernelStorage8 kernel;
+
+			float kernel1DStorage0[SharpenKernelStorage8::SIZE];
+			float kernel1DStorage1[SharpenKernelStorage8::SIZE];
+
+			// negative factors indicate blur where the -factor is the variance.
+			if (factor < 0)
+			{
+				BuildGaussianBlurKernel(kernel1DStorage0, -factor);
+
+				BuildKernel2DFromSeparable(kernel1DStorage0, kernel);
+
+				return kernel;
+			}
+
+			BuildSharpenKernel(kernel1DStorage0, 1.0f + factor);
+			BuildSharpenKernel(kernel1DStorage1, -factor);
+
+			BlurKernel<1>(kernel1DStorage0);
+			BlurKernel<3>(kernel1DStorage1);
+
+			for (int32_t i = 0; i < SharpenKernelStorage8::SIZE; ++i)
+			{
+				kernel1DStorage0[i] += kernel1DStorage1[i];
+			}
+
+			BuildKernel2DFromSeparable(kernel1DStorage0, kernel);
+
+			return kernel;
+		}
+
+		template<int32_t PIXEL_SIZE, typename CHANNEL_TYPE, EAddressMode AD_MODE = EAddressMode::AM_NONE>
+		class ImageAccessor
+		{
+			const vec2<int32_t> m_dim;
+			CHANNEL_TYPE const* const m_pImageBuf;
+
+		public:
+			ImageAccessor(const CHANNEL_TYPE* const pImageBuf, vec2<uint32_t> dim)
+				: m_dim(dim)
+				, m_pImageBuf(pImageBuf)
+			{
+			}
+
+			uint8_t operator()(int32_t x, int32_t y, int32_t c) const
+			{
+				if constexpr (AD_MODE == EAddressMode::AM_WRAP)
+				{
+					// (y & (m_dim.y() - 1)) * m_dim.y() + (x & (m_dim.x() - 1))
+					return m_pImageBuf[((y % m_dim.y()) * m_dim.x() + (x % m_dim.x())) * PIXEL_SIZE + c];
+				}
+				else if (AD_MODE == EAddressMode::AM_CLAMP)
+				{
+					return m_pImageBuf[(mu::clamp(y, 0, m_dim.y() - 1) * m_dim.x() +
+						mu::clamp(x, 0, m_dim.x() - 1))
+						* PIXEL_SIZE + c];
+				}
+				else if (AD_MODE == EAddressMode::AM_BLACK_BORDER)
+				{
+					return ((x < m_dim.x()) && (y < m_dim.y()) && (x >= 0) && (y >= 0))
+						? m_pImageBuf[(y * m_dim.x() + x) * PIXEL_SIZE + c]
+						: 0;
+				}
+				else
+				{
+					return m_pImageBuf[(y * m_dim.x() + x) * PIXEL_SIZE + c];
+				}
+			}
+		};
+
+
+		template<int32_t PIXEL_SIZE, EAddressMode AD, typename PIXEL_TYPE>
+		void GenerateMipSharpenedRegion(
+			const SharpenKernelStorage8& kernel,
+			const ImageAccessor<PIXEL_SIZE, PIXEL_TYPE, AD>& sourceAccessor,
+			PIXEL_TYPE* pDest, vec2<int32_t> destSize,
+			vec2<int32_t> regionStart, vec2<int32_t> regionEnd)
+		{
+			// Only kernels of size 8 supported.
+			constexpr int32_t kernelSize = SharpenKernelStorage8::SIZE;
+			static_assert(kernelSize == 8);
+
+			constexpr int32_t kernelCenter = (kernelSize >> 1) - 1;
+			for (int32_t y = regionStart.y(); y < regionEnd.y(); ++y)
+			{
+				for (int32_t x = regionStart.x(); x < regionEnd.x(); ++x)
+				{
+					float stagingStorage[kernelSize * kernelSize * PIXEL_SIZE];
+
+					const vec2<uint32_t> sourceCoord = vec2<uint32_t>(x << 1, y << 1);
+
+
+					// Convert Image kernel stamp to float. This conversion will be done multiple times.
+					// and should be done once before applying the kernel, or operate directly in unorm8/snorm8.
+					for (int32_t ky = 0; ky < kernelSize; ++ky)
+					{
+						for (int32_t kx = 0; kx < kernelSize; ++kx)
+						{
+							for (int32_t c = 0; c < PIXEL_SIZE; ++c)
+							{
+								stagingStorage[(ky * kernelSize + kx) + (kernelSize * kernelSize * c)] = static_cast<float>(
+									sourceAccessor(sourceCoord.x() + kx - kernelCenter,
+										sourceCoord.y() + ky - kernelCenter,
+										c)) / 255.0f;
+							}
+						}
+					}
+
+					// Apply kernel. Although the kernel is separable, the cost would not differ much from the not separable one.
+					// That is because to make it separable a partial convolution of the source is needed. 
+					// Since here we are only applying to half the image and the kernel used is small (8x8) the benefits are not 
+					// clear.
+
+					for (int32_t c = 0; c < PIXEL_SIZE; ++c)
+					{
+						float sampleValue = 0.0f;
+						float sum = 0.0f;
+						for (int32_t k = 0; k < kernelSize * kernelSize; ++k)
+						{
+							sampleValue += kernel.m_storage[k] * stagingStorage[k + (kernelSize * kernelSize * c)];
+							sum += kernel.m_storage[k];
+						}
+
+						pDest[(y * destSize.x() + x) * PIXEL_SIZE + c] = static_cast<uint8_t>(mu::clamp(sampleValue, 0.0f, 1.0f) * 255.0f);
+					}
+				}
+			}
+		}
+
+		template<size_t PIXEL_SIZE, EAddressMode AD_MODE>
+		inline void GenerateMipmapUint8Sharpen(int mips,
+			const uint8_t* pSource, uint8_t* pDest,
+			vec2<int32_t> sourceSize,
+			float sharpeningFactor)
+		{
+			SharpenKernelStorage8 kernel = MakeMipGaussianSharpenKernel(sharpeningFactor);
+
+			size_t totalDestSize = 0;
+
+			for (; mips >= 0; --mips)
+			{
+				ImageAccessor<PIXEL_SIZE, uint8_t, AD_MODE> sourceSampler(pSource, sourceSize);
+				ImageAccessor<PIXEL_SIZE, uint8_t, EAddressMode::AM_NONE> sourceSamplerAddressModeNone(pSource, sourceSize);
+
+				vec2<int32_t> destSize = vec2<int32_t>(FMath::DivideAndRoundUp(sourceSize.x(), 2), FMath::DivideAndRoundUp(sourceSize.y(), 2));
+
+				// Core image, will never sample outside the image so we don't need to care about address mode.
+				// Only sample source pixels where the kernel can be fully applied.
+				// Border cases are treated separately.
+				GenerateMipSharpenedRegion(
+					kernel,
+					sourceSamplerAddressModeNone,
+					pDest, destSize,
+					vec2<int32_t>(2, 2), vec2<int32_t>(destSize.x() - 2, destSize.y() - 2));
+
+				// Horizontal borders.
+				GenerateMipSharpenedRegion(
+					kernel,
+					sourceSampler,
+					pDest, destSize,
+					vec2<int32_t>(0, 0), vec2<int32_t>(destSize.x(), std::min(2, destSize.y())));
+
+				GenerateMipSharpenedRegion(
+					kernel,
+					sourceSampler,
+					pDest, destSize,
+					vec2<int32_t>(0, std::max(destSize.y() - 2, 0)), vec2<int32_t>(destSize.x(), destSize.y()));
+
+				// Vertical borders
+				GenerateMipSharpenedRegion(
+					kernel,
+					sourceSampler,
+					pDest, destSize,
+					vec2<int32_t>(0, std::min(2, destSize.y())), vec2<int32_t>(std::min(2, destSize.x()), destSize.y() - 2));
+
+				GenerateMipSharpenedRegion(
+					kernel,
+					sourceSampler,
+					pDest, destSize,
+					vec2<int32_t>(destSize.x() - 2, std::min(2, destSize.y())), vec2<int32_t>(destSize.x(), destSize.y() - 2));
+
+				sourceSize = destSize;
+
+				pSource = pDest;
+				pDest = pDest + destSize.x() * destSize.y() * PIXEL_SIZE;
+
+				totalDestSize += sourceSize.x() * sourceSize.y() * PIXEL_SIZE;
+			}
+		}
+
+		template<int32_t PIXEL_SIZE>
+		inline void GenerateMipmapUint8Unfiltered(
+			int mips,
+			const uint8_t* pSource, uint8_t* pDest,
+			vec2<int> sourceSize)
+		{
+			for (; mips >= 0; --mips)
+			{
+				vec2<int32_t> destSize = vec2<int32_t>(FMath::DivideAndRoundUp(sourceSize.x(), 2), FMath::DivideAndRoundUp(sourceSize.y(), 2));
+
+				for (int32_t y = 0; y < destSize.y(); ++y)
+				{
+					for (int32_t x = 0; x < destSize.x(); ++x)
+					{
+						for (int32_t c = 0; c < PIXEL_SIZE; ++c)
+						{
+							pDest[(y * destSize.x() + x) * PIXEL_SIZE + c] =
+								pSource[((y << 1) * sourceSize.x() + (x << 1)) * PIXEL_SIZE + c];
+						}
+					}
+				}
+
+				sourceSize = destSize;
+				pSource = pDest;
+				pDest = pDest + destSize.x() * destSize.y() * PIXEL_SIZE;
+			}
+		}
+
+		template<int32_t PIXEL_SIZE>
+		inline void GenerateMipmapsUint8SimpleAverage(
+			int mips,
+			const uint8_t* pSource, uint8_t* pDest,
+			vec2<int> sourceSize)
+		{
+			const uint8_t* pMipSource = pSource;
+			uint8_t* pMipDest = pDest;
+
+			vec2<int> destSize = sourceSize;
+			for (int m = 0; m < mips; ++m)
+			{
+				check(destSize[0] > 1 || destSize[1] > 1);
+
+				sourceSize = destSize;
+
+				int fullColumns = destSize[0] / 2;
+				bool strayColumn = (destSize[0] % 2) != 0;
+				int fullRows = destSize[1] / 2;
+				bool strayRow = (destSize[1] % 2) != 0;
+
+				destSize[0] = FMath::DivideAndRoundUp(destSize[0], 2);
+				destSize[1] = FMath::DivideAndRoundUp(destSize[1], 2);
+
+				int sourceStride = sourceSize[0] * PIXEL_SIZE;
+				int destStride = destSize[0] * PIXEL_SIZE;
+
+				const auto& ProcessRow = [
+					pMipDest, pMipSource, fullColumns, strayColumn, sourceStride, destStride
+				] (uint32 y)
+				{
+					const uint8_t* pSourceRow0 = pMipSource + 2 * y * sourceStride;
+					const uint8_t* pSourceRow1 = pSourceRow0 + sourceStride;
+					uint8_t* pDestRow = pMipDest + y * destStride;
+
+					for (int x = 0; x < fullColumns; ++x)
+					{
+						for (int c = 0; c < PIXEL_SIZE; ++c)
+						{
+							int p = pSourceRow0[c] + pSourceRow0[PIXEL_SIZE + c] + pSourceRow1[c] +
+								pSourceRow1[PIXEL_SIZE + c];
+							pDestRow[c] = (uint8_t)(p >> 2);
+						}
+
+						pSourceRow0 += 2 * PIXEL_SIZE;
+						pSourceRow1 += 2 * PIXEL_SIZE;
+						pDestRow += PIXEL_SIZE;
+					}
+
+					if (strayColumn)
+					{
+						for (int c = 0; c < PIXEL_SIZE; ++c)
+						{
+							int p = pSourceRow0[c] + pSourceRow1[c];
+							pDestRow[c] = (uint8_t)(p >> 1);
+						}
+					}
+				};
+
+					constexpr int PixelConcurrencyThreshold = 0xffff;
+					if (destSize[0] * destSize[1] < PixelConcurrencyThreshold)
+					{
+						for (int y = 0; y < fullRows; ++y)
+						{
+							ProcessRow(y);
+						}
+					}
+					else
+					{
+						ParallelFor(fullRows, ProcessRow);
+					}
+
+					if (strayRow)
+					{
+						const uint8_t* pSourceRow0 = pMipSource + 2 * fullRows * sourceStride;
+						const uint8_t* pSourceRow1 = pSourceRow0 + sourceStride;
+						uint8_t* pDestRow = pMipDest + fullRows * destStride;
+
+						for (int x = 0; x < fullColumns; ++x)
+						{
+							for (int c = 0; c < PIXEL_SIZE; ++c)
+							{
+								int p = pSourceRow0[c] + pSourceRow0[PIXEL_SIZE + c];
+								pDestRow[c] = (uint8_t)(p >> 1);
+							}
+
+							pSourceRow0 += 2 * PIXEL_SIZE;
+							pDestRow += PIXEL_SIZE;
+						}
+
+						if (strayColumn)
+						{
+							for (int c = 0; c < PIXEL_SIZE; ++c)
+							{
+								pDestRow[c] = pSourceRow0[c];
+							}
+						}
+					}
+
+					// Reset the source pointer for the next mip, to use the dest that we have just
+					// generated.
+					pMipSource = pMipDest;
+					pMipDest += destSize[0] * destSize[1] * PIXEL_SIZE;
+			}
+		}
+
+	} // namespace OpImageMipmap_Detail
+
+
+		//---------------------------------------------------------------------------------------------
+		//! Generate the mipmaps for byte-based images of whatever number of channels.
+		//! \param mips number of additional levels to build from the source.
+		//---------------------------------------------------------------------------------------------
+	template<int32_t PIXEL_SIZE>
+	inline void GenerateMipmapsUint8(int mips,
+		const uint8_t* pSource, uint8_t* pDest,
+		vec2<int32_t> sourceSize,
+		const FMipmapGenerationSettings& settings)
+	{
+		using namespace OpImageMipmap_Detail;
+
+		switch (settings.m_filterType)
+		{
+		case EMipmapFilterType::MFT_SimpleAverage:
+		{
+			GenerateMipmapsUint8SimpleAverage<PIXEL_SIZE>(mips, pSource, pDest, sourceSize);
+			break;
+		}
+		case EMipmapFilterType::MFT_Sharpen:
+		{
+			switch (settings.m_addressMode)
+			{
+			case EAddressMode::AM_CLAMP:
+			{
+				GenerateMipmapUint8Sharpen<PIXEL_SIZE, EAddressMode::AM_CLAMP>(
+					mips, pSource, pDest, sourceSize, settings.m_sharpenFactor);
+				break;
+			}
+			case EAddressMode::AM_WRAP:
+			{
+				GenerateMipmapUint8Sharpen<PIXEL_SIZE, EAddressMode::AM_WRAP>(
+					mips, pSource, pDest, sourceSize, settings.m_sharpenFactor);
+				break;
+			}
+			case EAddressMode::AM_BLACK_BORDER:
+			{
+				GenerateMipmapUint8Sharpen<PIXEL_SIZE, EAddressMode::AM_BLACK_BORDER>(
+					mips, pSource, pDest, sourceSize, settings.m_sharpenFactor);
+				break;
+			}
+			default:
+			{
+				GenerateMipmapUint8Sharpen<PIXEL_SIZE, EAddressMode::AM_NONE>(
+					mips, pSource, pDest, sourceSize, settings.m_sharpenFactor);
+				break;
+			}
+			}
+
+			break;
+		}
+		case EMipmapFilterType::MFT_Unfiltered:
+		{
+			GenerateMipmapUint8Unfiltered<PIXEL_SIZE>(mips, pSource, pDest, sourceSize);
+			break;
+		}
+		default:
+		{
+			check(false);
+			break;
+		}
+		}
+	}
+
+	struct SCRATCH_IMAGE_MIPMAP
+	{
+		ImagePtr pUncompressed;
+		ImagePtr pUncompressedMips;
+		ImagePtr pCompressedMips;
+	};
+
+
+	//---------------------------------------------------------------------------------------------
+	inline void ImageMipmap_PrepareScratch(Image* pDest, const Image* pBase, int levelCount,
+		SCRATCH_IMAGE_MIPMAP* scratch)
+	{
+		int startLevel = pBase->GetLODCount() - 1;
+
+		check(pDest->GetLODCount() == levelCount);
+		check(pDest->GetSizeX() == pBase->GetSizeX());
+		check(pDest->GetSizeY() == pBase->GetSizeY());
+		check(pDest->GetFormat() == pBase->GetFormat());
+		(void)pDest;
+
+		switch (pBase->GetFormat())
+		{
+		case EImageFormat::IF_L_UBYTE:
+		case EImageFormat::IF_RGB_UBYTE:
+		case EImageFormat::IF_RGBA_UBYTE:
+		case EImageFormat::IF_BGRA_UBYTE:
+			break;
+
+			// Bad cases: we need to decompress, mip and then recompress. It may be necessary to
+			// generate the latests mips after composing, or in unoptimised code.
+			// TODO: Make sure that the code optimisations avoid this cases generating separate
+			// operations to generate the mip tail.
+		case EImageFormat::IF_BC1:
+		case EImageFormat::IF_BC2:
+		case EImageFormat::IF_BC3:
+		case EImageFormat::IF_BC4:
+		case EImageFormat::IF_BC5:
+		case EImageFormat::IF_ASTC_4x4_RGB_LDR:
+		case EImageFormat::IF_ASTC_4x4_RGBA_LDR:
+		case EImageFormat::IF_ASTC_4x4_RG_LDR:
+		{
+			// Uncompress the last mip that we already have
+			vec2<int> uncompressedSize = pBase->CalculateMipSize(startLevel);
+			scratch->pUncompressed = new Image(
+				(uint16_t)uncompressedSize[0],
+				(uint16_t)uncompressedSize[1],
+				1,
+				EImageFormat::IF_RGBA_UBYTE);
+
+			// Generate the mipmaps from there on
+			scratch->pUncompressedMips = new Image(
+				(uint16_t)uncompressedSize[0],
+				(uint16_t)uncompressedSize[1],
+				levelCount - startLevel,
+				EImageFormat::IF_RGBA_UBYTE);
+
+			// Compress the mipmapped image
+			scratch->pCompressedMips = new Image(
+				(uint16_t)uncompressedSize[0],
+				(uint16_t)uncompressedSize[1],
+				scratch->pUncompressedMips->GetLODCount(),
+				pBase->GetFormat());
+
+			break;
+		}
+
+		case EImageFormat::IF_L_UBIT_RLE:
+		case EImageFormat::IF_L_UBYTE_RLE:
+		{
+			// Uncompress the last mip that we already have
+			vec2<int> uncompressedSize = pBase->CalculateMipSize(startLevel);
+			scratch->pUncompressed = new Image(
+				(uint16_t)uncompressedSize[0],
+				(uint16_t)uncompressedSize[1],
+				1,
+				EImageFormat::IF_L_UBYTE);
+
+			// Generate the mipmaps from there on
+			scratch->pUncompressedMips = new Image(
+				(uint16_t)uncompressedSize[0],
+				(uint16_t)uncompressedSize[1],
+				levelCount - startLevel,
+				EImageFormat::IF_L_UBYTE);
+
+			// Compress the mipmapped image
+			scratch->pCompressedMips = new Image(
+				(uint16_t)uncompressedSize[0],
+				(uint16_t)uncompressedSize[1],
+				scratch->pUncompressedMips->GetLODCount(),
+				pBase->GetFormat());
+
+			// Preallocate ample memory for the compressed data
+			scratch->pCompressedMips->m_data.SetNum(scratch->pUncompressedMips->GetDataSize());
+
+			// Preallocate ample memory for the destination data
+			pDest->m_data.SetNum(pBase->GetDataSize()
+				+
+				scratch->pUncompressedMips->GetDataSize());
+
+			break;
+		}
+
+		default:
+			checkf(false, TEXT("Format not implemented in mipmap generation."));
+		}
+	}
+
+
+	//---------------------------------------------------------------------------------------------
+	inline void ImageMipmap(int imageCompressionQuality, Image* pDest, const Image* pBase,
+		int levelCount, SCRATCH_IMAGE_MIPMAP* scratch,
+		FMipmapGenerationSettings& settings)
+	{
+		int startLevel = pBase->GetLODCount() - 1;
+
+		check(!(pBase->m_flags & Image::IF_CANNOT_BE_SCALED));
+
+		check(pDest->GetLODCount() == levelCount);
+		check(pDest->GetSizeX() == pBase->GetSizeX());
+		check(pDest->GetSizeY() == pBase->GetSizeY());
+		check(pDest->GetFormat() == pBase->GetFormat());
+
+		// Copy already-existing mips
+
+		// Calculate the data size, since the base could have more data allocated.
+		int32 mipDataSize = pBase->CalculateDataSize();
+		check(pDest->GetDataSize() >= mipDataSize);
+		FMemory::Memcpy(pDest->GetData(), pBase->GetData(), mipDataSize);
+
+		int mipsToBuild = levelCount - startLevel - 1;
+		if (!mipsToBuild)
+		{
+			return;
+		}
+
+		uint8_t* pSourceBuf = pDest->GetMipData(startLevel);
+		uint8_t* pDestBuf = pDest->GetMipData(startLevel + 1);
+
+		vec2<int> sourceSize = pDest->CalculateMipSize(startLevel);
+
+		switch (pBase->GetFormat())
+		{
+		case EImageFormat::IF_L_UBYTE:
+			GenerateMipmapsUint8<1>(levelCount - startLevel - 1, pSourceBuf, pDestBuf, sourceSize, settings);
+			break;
+
+		case EImageFormat::IF_RGB_UBYTE:
+			GenerateMipmapsUint8<3>(levelCount - startLevel - 1, pSourceBuf, pDestBuf, sourceSize, settings);
+			break;
+
+		case EImageFormat::IF_BGRA_UBYTE:
+		case EImageFormat::IF_RGBA_UBYTE:
+			GenerateMipmapsUint8<4>(levelCount - startLevel - 1, pSourceBuf, pDestBuf, sourceSize, settings);
+			break;
+
+			// Bad cases: we need to decompress, mip and then recompress. It may be necessary to
+			// generate the latests mips after composing, or in unoptimised code.
+			// TODO: Make sure that the code optimisations avoid this cases generating separate
+			// operations to generate the mip tail.
+		case EImageFormat::IF_BC1:
+		case EImageFormat::IF_BC2:
+		case EImageFormat::IF_BC3:
+		case EImageFormat::IF_BC4:
+		case EImageFormat::IF_BC5:
+		case EImageFormat::IF_ASTC_4x4_RGB_LDR:
+		case EImageFormat::IF_ASTC_4x4_RGBA_LDR:
+		case EImageFormat::IF_ASTC_4x4_RG_LDR:
+		case EImageFormat::IF_L_UBIT_RLE:
+		case EImageFormat::IF_L_UBYTE_RLE:
+		{
+			// Uncompress the last mip that we already have
+			bool bSuccess = ImagePixelFormatInPlace(imageCompressionQuality, scratch->pUncompressed.get(), pBase, startLevel);
+			check(bSuccess);
+
+			// Generate the mipmaps from there on
+			ImageMipmap(imageCompressionQuality, scratch->pUncompressedMips.get(),
+				scratch->pUncompressed.get(), levelCount - startLevel, 0, settings);
+
+			// Compress the mipmapped image
+			bSuccess = ImagePixelFormatInPlace(imageCompressionQuality, scratch->pCompressedMips.get(), scratch->pUncompressedMips.get());
+			int32 ExcessDataSize = scratch->pCompressedMips->GetDataSize();
+			while (!bSuccess)
+			{
+				// Bad case: this should almost never happen.
+				MUTABLE_CPUPROFILER_SCOPE(Mipmap_Recompression_OutOfSpace);
+
+				scratch->pCompressedMips->m_data.SetNum(ExcessDataSize);
+				bSuccess = ImagePixelFormatInPlace(imageCompressionQuality, scratch->pCompressedMips.get(), scratch->pUncompressedMips.get());
+				ExcessDataSize *= 4;
+			}
+
+			int32 FinalDestSize = pSourceBuf - pDest->GetData() + scratch->pCompressedMips->GetDataSize();
+
+			if (FinalDestSize > pDest->GetDataSize())
+			{
+				// Bad case: this should almost never happen.
+				pDest->m_data.SetNum(FinalDestSize);
+				pSourceBuf = pDest->GetMipData(startLevel);
+			}
+
+			// TODO: We are recompressing the source level, which we could avoid.
+			FMemory::Memcpy(pSourceBuf,
+				scratch->pCompressedMips->GetData(),
+				scratch->pCompressedMips->GetDataSize());
+			break;
+		}
+
+		default:
+			checkf(false, TEXT("Format not implemented in mipmap generation."));
+		}
+
+	}
+}
diff --git a/Engine/Plugins/Experimental/Mutable/Source/MutableRuntime/Private/OpImageMultiply.h b/Engine/Plugins/Experimental/Mutable/Source/MutableRuntime/Private/OpImageMultiply.h
new file mode 100644
--- /dev/null
+++ b/Engine/Plugins/Experimental/Mutable/Source/MutableRuntime/Private/OpImageMultiply.h
@@ -0,0 +1,72 @@
+// Copyright Epic Games, Inc. All Rights Reserved.
+
+#pragma once
+
+#include "ImagePrivate.h"
+#include "OpImageLayer.h"
+#include "Platform.h"
+
+namespace mu
+{
+
+	inline unsigned MultiplyChannelMasked( unsigned base, unsigned blended, unsigned mask )
+	{
+		unsigned multiply = ( base * blended ) >> 8;
+		unsigned masked = ( ( ( 255 - mask ) * base ) + ( mask * multiply ) ) >> 8;
+		return masked;
+	}
+
+	inline unsigned MultiplyChannel( unsigned base, unsigned blended )
+	{
+		unsigned multiply = ( base * blended ) >> 8;
+		return multiply;
+	}
+
+
+    inline void ImageMultiply( Image* pResult, const Image* pBase, const Image* pMask, vec3<float> col )
+	{
+        ImageLayerColour<MultiplyChannelMasked,MultiplyChannel,false>( pResult, pBase, pMask, col );
+	}
+
+
+    inline void ImageMultiply( Image* pResult, const Image* pBase, const Image* pMask,
+                               const Image* pBlended,
+                               bool applyToAlpha )
+	{
+        return ImageLayer<MultiplyChannelMasked,MultiplyChannel,false>
+                ( pResult, pBase, pMask, pBlended, applyToAlpha );
+	}
+
+
+    inline void ImageMultiply( Image* pResult, const Image* pBase, vec3<float> col )
+	{
+        return ImageLayerColour<MultiplyChannel,false>( pResult, pBase, col );
+	}
+
+
+    inline void ImageMultiply( Image* pResult, const Image* pBase, const Image* pBlended,
+                               bool applyToAlpha )
+	{
+        return ImageLayer<MultiplyChannel,false>( pResult, pBase, pBlended, applyToAlpha );
+	}
+
+
+	inline void ImageMultiplyOnBase( Image* pBase, const Image* pMask, vec3<float> col )
+	{
+		ImageLayerColourOnBase<MultiplyChannelMasked,MultiplyChannel, false>( pBase, pMask, col );
+	}
+
+    inline void ImageMultiplyOnBase( Image* pBase, const Image* pMask, const Image* pBlended,
+                                     bool applyToAlpha )
+	{
+        ImageLayerOnBase<MultiplyChannelMasked,MultiplyChannel, false>
+                ( pBase, pMask, pBlended, applyToAlpha );
+	}
+
+    inline void ImageMultiplyOnBase( Image* pBase, const Image* pBlended,
+                                     bool applyToAlpha )
+	{
+        ImageLayerOnBase<MultiplyChannel, false>( pBase, pBlended, applyToAlpha );
+	}
+
+}
diff --git a/Engine/Plugins/Experimental/Mutable/Source/MutableRuntime/Private/OpImageNormalCombine.h b/Engine/Plugins/Experimental/Mutable/Source/MutableRuntime/Private/OpImageNormalCombine.h
new file mode 100644
--- /dev/null
+++ b/Engine/Plugins/Experimental/Mutable/Source/MutableRuntime/Private/OpImageNormalCombine.h
@@ -0,0 +1,83 @@
+// Copyright Epic Games, Inc. All Rights Reserved.
+
+#pragma once
+
+#include "ImagePrivate.h"
+#include "OpImageLayer.h"
+
+namespace mu
+{
+	// Combine layer oprations get packed color.
+	inline unsigned CombineNormal(unsigned base, unsigned blended)
+	{
+		// See /Engine/Functions/Engine_MaterialFunctions02/Utility/BlendAngleCorrectedNormals.BlendAngleCorrectedNormals
+		// for the source of the effect.
+
+		const float baseRf = (static_cast<float>((base >> 0)  & 0xFF) / 255.0f) * 2.0f - 1.0f;
+		const float baseGf = (static_cast<float>((base >> 8)  & 0xFF) / 255.0f) * 2.0f - 1.0f;
+		const float baseBf = (static_cast<float>((base >> 16) & 0xFF) / 255.0f) * 2.0f; //One added to the b channel.
+
+		const float blendedRf = (static_cast<float>((blended >> 0)  & 0xFF) / 255.0f) * 2.0f - 1.0f;
+		const float blendedGf = (static_cast<float>((blended >> 8)  & 0xFF) / 255.0f) * 2.0f - 1.0f;
+		const float blendedBf = (static_cast<float>((blended >> 16) & 0xFF) / 255.0f) * 2.0f - 1.0f;
+
+		vec3f a = vec3f(baseRf, baseGf, baseBf);
+		vec3f b = vec3f(-blendedRf, -blendedGf, blendedBf);
+
+		// The original does not normalize but if not done results don't look good due to signal clipping.
+		vec3f n = normalise_approx(a*dot(a, b) - b*a.z());
+
+		return 
+			(mu::clamp(static_cast<unsigned>((n.x() + 1.0f) * 0.5f * 255.0f), 0u, 255u) << 0) |
+			(mu::clamp(static_cast<unsigned>((n.y() + 1.0f) * 0.5f * 255.0f), 0u, 255u) << 8) |
+			(mu::clamp(static_cast<unsigned>((n.z() + 1.0f) * 0.5f * 255.0f), 0u, 255u) << 16);
+	}
+
+	inline unsigned CombineNormalMasked(unsigned base, unsigned blended, unsigned mask)
+	{
+		// See /Engine/Functions/Engine_MaterialFunctions02/Utility/BlendAngleCorrectedNormals.BlendAngleCorrectedNormals
+		// for the source of the effect.
+
+		const float baseRf = (static_cast<float>((base >> 0)  & 0xFF) / 255.0f) * 2.0f - 1.0f;
+		const float baseGf = (static_cast<float>((base >> 8)  & 0xFF) / 255.0f) * 2.0f - 1.0f;
+		const float baseBf = (static_cast<float>((base >> 16) & 0xFF) / 255.0f) * 2.0f - 1.0f;
+
+		const float blendedRf = (static_cast<float>((blended >> 0)  & 0xFF) / 255.0f) * 2.0f - 1.0f;
+		const float blendedGf = (static_cast<float>((blended >> 8)  & 0xFF) / 255.0f) * 2.0f - 1.0f;
+		const float blendedBf = (static_cast<float>((blended >> 16) & 0xFF) / 255.0f) * 2.0f - 1.0f;
+
+		const float maskRf = (static_cast<float>(mask & 0xFF) / 255.0f);
+
+		vec3f a = vec3f(baseRf, baseGf, baseBf + 1.0f);
+		vec3f b = vec3f(-blendedRf, -blendedGf, blendedBf);
+
+		// The original does not normalize but if not done results don't look good due to signal clipping.
+		vec3f n = mu::normalise_approx( 
+					mu::lerp( vec3f(baseRf, baseGf, baseBf), a*dot(a, b) - b*a.z(), maskRf ) );
+
+		return 
+			(mu::clamp(static_cast<unsigned>((n.x() + 1.0f) * 0.5f * 255.0f), 0u, 255u) << 0) |
+			(mu::clamp(static_cast<unsigned>((n.y() + 1.0f) * 0.5f * 255.0f), 0u, 255u) << 8) |
+			(mu::clamp(static_cast<unsigned>((n.z() + 1.0f) * 0.5f * 255.0f), 0u, 255u) << 16);
+	}
+
+	inline void ImageNormalCombine(Image* pResult, const Image* pBase, vec3<float> col)
+	{
+		ImageLayerCombineColour<CombineNormal>(pResult, pBase, col);
+	}
+
+	inline void ImageNormalCombine(Image* pResult, const Image* pBase, const Image* pBlended)
+	{
+		ImageLayerCombine<CombineNormal>(pResult, pBase, pBlended);
+	}
+
+	inline void ImageNormalCombine(Image* pResult, const Image* pBase, const Image* pMask, vec3<float> col)
+	{
+		ImageLayerCombineColour<CombineNormal, CombineNormalMasked>(pResult, pBase, pMask, col);
+	}
+
+	inline void ImageNormalCombine(Image* pResult, const Image* pBase, const Image* pMask, const Image* pBlended)
+	{
+		ImageLayerCombine<CombineNormal, CombineNormalMasked>(pResult, pBase, pMask, pBlended);
+	}
+}
diff --git a/Engine/Plugins/Experimental/Mutable/Source/MutableRuntime/Private/OpImageOverlay.h b/Engine/Plugins/Experimental/Mutable/Source/MutableRuntime/Private/OpImageOverlay.h
new file mode 100644
--- /dev/null
+++ b/Engine/Plugins/Experimental/Mutable/Source/MutableRuntime/Private/OpImageOverlay.h
@@ -0,0 +1,61 @@
+// Copyright Epic Games, Inc. All Rights Reserved.
+
+#pragma once
+
+#include "ImagePrivate.h"
+#include "OpImageHardLight.h"
+
+
+namespace mu
+{
+	inline unsigned OverlayChannelMasked( unsigned base, unsigned blended, unsigned mask )
+	{
+		// Photoshop like
+		//return HardLightChannel( blended, base, mask );
+
+		// gimp-like
+		unsigned overlay = ( base * ( base + ( ( 2 * blended * ( 255 - base) ) >> 8  ) ) ) >> 8;
+		unsigned masked = ( ( ( 255 - mask ) * base ) + ( mask * overlay ) ) >> 8;
+		return masked;
+	}
+
+
+	inline unsigned OverlayChannel( unsigned base, unsigned blended )
+	{
+		// Photoshop like
+		//return HardLightChannel( blended, base, mask );
+
+		// gimp-like
+		unsigned overlay = ( base * ( base + ( ( 2 * blended * ( 255 - base) ) >> 8  ) ) ) >> 8;
+		return overlay;
+	}
+
+
+    inline void ImageOverlay( Image* pResult, const Image* pBase, const Image* pMask, vec3<float> col )
+	{
+        ImageLayerColour<OverlayChannelMasked,OverlayChannel,true>( pResult, pBase, pMask, col );
+	}
+
+
+    inline void ImageOverlay( Image* pResult, const Image* pBase, const Image* pMask,
+                              const Image* pBlended,
+                              bool applyToAlpha )
+	{
+        ImageLayer<OverlayChannelMasked,OverlayChannel,true>
+                ( pResult, pBase, pMask, pBlended, applyToAlpha );
+	}
+
+
+    inline void ImageOverlay( Image* pResult, const Image* pBase, vec3<float> col )
+	{
+        ImageLayerColour<OverlayChannel,true>( pResult, pBase, col );
+	}
+
+
+    inline void ImageOverlay( Image* pResult, const Image* pBase, const Image* pBlended,
+                              bool applyToAlpha )
+	{
+        ImageLayer<OverlayChannel,true>( pResult, pBase, pBlended, applyToAlpha );
+	}
+
+}
