[mutable] Move Mutable plugin from the restricted folder to the experimental plugins folder.

diff --git a/Engine/Plugins/Experimental/Mutable/Source/MutableRuntime/Private/OpImagePixelFormat.cpp b/Engine/Plugins/Experimental/Mutable/Source/MutableRuntime/Private/OpImagePixelFormat.cpp
new file mode 100644
--- /dev/null
+++ b/Engine/Plugins/Experimental/Mutable/Source/MutableRuntime/Private/OpImagePixelFormat.cpp
@@ -0,0 +1,1188 @@
+// Copyright Epic Games, Inc. All Rights Reserved.
+
+
+#include "OpImagePixelFormat.h"
+
+#include "ImageRLE.h"
+#include "BlockCompression/BlockCompression.h"
+
+#include "MutableTrace.h"
+
+using namespace mu;
+
+
+//-------------------------------------------------------------------------------------------------
+ImagePtr mu::ImagePixelFormat( int imageCompressionQuality, const Image* pBase,
+                               EImageFormat targetFormat, int onlyLOD )
+{
+	MUTABLE_CPUPROFILER_SCOPE(ImagePixelFormat);
+
+    vec2<int> resultSize;
+    int resultLODCount = 0;
+
+    if (onlyLOD==-1)
+    {
+        resultSize[0] = pBase->GetSizeX();
+        resultSize[1] = pBase->GetSizeY();
+        resultLODCount = pBase->GetLODCount();
+    }
+    else
+    {
+        resultSize = pBase->CalculateMipSize( onlyLOD );
+        resultLODCount = 1;
+    }
+
+    ImagePtr result = new Image( (uint16_t)resultSize[0], (uint16_t)resultSize[1], resultLODCount, targetFormat );
+	result->m_flags = pBase->m_flags;
+	bool bSuccess = ImagePixelFormatInPlace( imageCompressionQuality, result.get(), pBase, onlyLOD );
+	int32 OriginalDataSize = FMath::Max(result->m_data.Num(), pBase->m_data.Num());
+	int32 ExcessDataSize = OriginalDataSize * 4;
+	while (!bSuccess)
+    {
+		MUTABLE_CPUPROFILER_SCOPE(Recompression_OutOfSpace);
+
+        // Bad case, where the RLE compressed data requires more memory than the uncompressed data.
+		// We need to support it anyway for small mips or scaled images.
+		result->m_data.SetNum(ExcessDataSize);
+		bSuccess = ImagePixelFormatInPlace(imageCompressionQuality, result.get(), pBase, onlyLOD);
+		ExcessDataSize *= 4;
+	}
+    return result;
+}
+
+
+//-------------------------------------------------------------------------------------------------
+bool mu::ImagePixelFormatInPlace( int imageCompressionQuality, Image* pResult, const Image* pBase,
+                                  int onlyLOD )
+{
+	MUTABLE_CPUPROFILER_SCOPE(ImagePixelFormatInPlace);
+
+    bool success = true;
+
+    vec2<int> resultSize;
+    int resultLODCount = 0;
+    const uint8_t* pBaseBuf = 0;
+    int pixelCount = 0;
+    int baseLOD = 0;
+
+    if (onlyLOD==-1)
+    {
+        pixelCount = (int)pBase->CalculatePixelCount();
+        pBaseBuf = pBase->GetData();
+        resultSize[0] = pBase->GetSizeX();
+        resultSize[1] = pBase->GetSizeY();
+        resultLODCount = pBase->GetLODCount();
+        baseLOD = 0;
+    }
+    else
+    {
+        pixelCount = (int)pBase->CalculatePixelCount( onlyLOD );
+        pBaseBuf = pBase->GetMipData( onlyLOD );
+        resultSize = pBase->CalculateMipSize( onlyLOD );
+        resultLODCount = 1;
+        baseLOD = onlyLOD;
+    }
+
+    check( pResult->GetSizeX()==resultSize[0] && pResult->GetSizeY()==resultSize[1] );
+    check( pResult->GetLODCount()==resultLODCount );
+
+    if ( pResult->GetFormat() == pBase->GetFormat()
+         && onlyLOD==-1 )
+    {
+        // Shouldn't really happen
+        pResult->m_data = pBase->m_data;
+    }
+    else
+    {
+        switch ( pResult->GetFormat() )
+        {
+
+        case EImageFormat::IF_L_UBYTE:
+        {
+            uint8_t* pDestBuf = pResult->GetData();
+
+            switch ( pBase->GetFormat() )
+            {
+
+            case EImageFormat::IF_L_UBYTE:
+            {
+				FMemory::Memcpy( pDestBuf, pBaseBuf, pixelCount );
+                break;
+            }
+
+            case EImageFormat::IF_L_UBYTE_RLE:
+            {
+                if( onlyLOD == -1 )
+                {
+                    UncompressRLE_L( pBase, pResult );
+                }
+                else
+                {
+                    UncompressRLE_L( resultSize[0], resultSize[1], pBaseBuf, pDestBuf );
+                }
+                break;
+            }
+
+            case EImageFormat::IF_L_UBIT_RLE:
+            {
+                if( onlyLOD == -1 )
+                {
+                    UncompressRLE_L1( pBase, pResult );
+                }
+                else
+                {
+                    UncompressRLE_L1( resultSize[0], resultSize[1], pBaseBuf, pDestBuf );
+                }
+                break;
+            }
+
+            case EImageFormat::IF_RGB_UBYTE:
+            {
+                for ( int i=0; i<pixelCount; ++i )
+                {
+                    unsigned result = 76 * pBaseBuf[3*i+0]
+                            + 150 * pBaseBuf[3*i+1]
+                            + 29 * pBaseBuf[3*i+2];
+                    pDestBuf[i] = (uint8_t)std::min( 255u, result >> 8 );
+                }
+                break;
+            }
+
+            case EImageFormat::IF_RGB_UBYTE_RLE:
+            {
+                ImagePtr pTempBase =
+                    ImagePixelFormat( imageCompressionQuality, pBase, EImageFormat::IF_RGB_UBYTE );
+                if (onlyLOD==-1)
+                {
+                    pBaseBuf = pTempBase->GetData();
+                }
+                else
+                {
+                    pBaseBuf = pTempBase->GetMipData( onlyLOD );
+                }
+                for ( int i=0; i<pixelCount; ++i )
+                {
+                    unsigned result = 76 * pBaseBuf[3*i+0]
+                            + 150 * pBaseBuf[3*i+1]
+                            + 29 * pBaseBuf[3*i+2];
+                    pDestBuf[i] = (uint8_t)std::min( 255u, result >> 8 );
+                }
+                break;
+            }
+
+            case EImageFormat::IF_RGBA_UBYTE:
+            {
+                for ( int i=0; i<pixelCount; ++i )
+                {
+                    unsigned result = 76 * pBaseBuf[4*i+0]
+                            + 150 * pBaseBuf[4*i+1]
+                            + 29 * pBaseBuf[4*i+2];
+                    pDestBuf[i] = (uint8_t)std::min( 255u, result >> 8 );
+                }
+                break;
+            }
+
+            case EImageFormat::IF_RGBA_UBYTE_RLE:
+            {
+                ImagePtr pTempBase =
+                    ImagePixelFormat( imageCompressionQuality, pBase, EImageFormat::IF_RGBA_UBYTE );
+                if (onlyLOD==-1)
+                {
+                    pBaseBuf = pTempBase->GetData();
+                }
+                else
+                {
+                    pBaseBuf = pTempBase->GetMipData( onlyLOD );
+                }
+                for ( int i=0; i<pixelCount; ++i )
+                {
+                    unsigned result = 76 * pBaseBuf[4*i+0]
+                            + 150 * pBaseBuf[4*i+1]
+                            + 29 * pBaseBuf[4*i+2];
+                    pDestBuf[i] = (uint8_t)std::min( 255u, result >> 8 );
+                }
+                break;
+            }
+
+            case EImageFormat::IF_BGRA_UBYTE:
+            {
+                for ( int i=0; i<pixelCount; ++i )
+                {
+                    unsigned result = 76 * pBaseBuf[4*i+2]
+                            + 150 * pBaseBuf[4*i+1]
+                            + 29 * pBaseBuf[4*i+0];
+                    pDestBuf[i] = (uint8_t)std::min( 255u, result >> 8 );
+                }
+                break;
+            }
+
+            case EImageFormat::IF_BC1:
+            case EImageFormat::IF_BC2:
+            case EImageFormat::IF_BC3:
+            {
+                for (int m = 0; m < resultLODCount; ++m)
+                {
+                    ImagePtr pTempBase =
+                        ImagePixelFormat( imageCompressionQuality, pBase, EImageFormat::IF_RGB_UBYTE );
+                    if (onlyLOD == -1)
+                    {
+                        pBaseBuf = pTempBase->GetData();
+                    }
+                    else
+                    {
+                        pBaseBuf = pTempBase->GetMipData(onlyLOD);
+                    }
+                    for (int i = 0; i < pixelCount; ++i)
+                    {
+                        unsigned result = 76 * pBaseBuf[3 * i + 0] + 150 * pBaseBuf[3 * i + 1] +
+                                          29 * pBaseBuf[3 * i + 2];
+                        pDestBuf[i] = (uint8_t)std::min(255u, result >> 8);
+                    }
+                }
+                break;
+            }
+
+            default:
+                // Case not implemented
+                check( false );
+
+            }
+
+            break;
+        }
+
+        case EImageFormat::IF_L_UBYTE_RLE:
+        {
+            switch ( pBase->GetFormat() )
+            {
+
+            case EImageFormat::IF_L_UBYTE:
+            {
+                check( onlyLOD == -1 );
+
+                // Try to compress
+				if (!pResult->m_data.Num())
+				{
+					// Allocate memory for the compressed data. TODO: Smaller?
+					pResult->m_data.SetNum(pBase->m_data.Num());
+				}
+
+                success = CompressRLE_L( pBase, pResult );
+
+#ifdef MUTABLE_DEBUG_RLE		
+				if (success)
+                {
+                    // verify
+                    ImagePtr pTest = ImagePixelFormat( imageCompressionQuality, pResult, EImageFormat::IF_L_UBYTE );
+					check(!FMemory::Memcmp(pTest->GetData(), pBase->GetData(), pBase->GetDataSize()));
+                }
+#endif
+
+                break;
+            }
+
+            case EImageFormat::IF_RGB_UBYTE:
+            case EImageFormat::IF_RGBA_UBYTE:
+            case EImageFormat::IF_BGRA_UBYTE:
+            {
+                check( onlyLOD == -1 );
+                ImagePtr pTemp = ImagePixelFormat( imageCompressionQuality, pBase, EImageFormat::IF_L_UBYTE );
+ 
+				// Try to compress
+				if (!pResult->m_data.Num())
+				{
+					// Allocate memory for the compressed data. TODO: Smaller?
+					pResult->m_data.SetNum(pTemp->m_data.Num());
+				}
+
+                success = CompressRLE_L( pTemp.get(), pResult );
+
+#ifdef MUTABLE_DEBUG_RLE		
+				if (success)
+                {
+                    // verify
+                    ImagePtr pTest = ImagePixelFormat( imageCompressionQuality, pResult, EImageFormat::IF_L_UBYTE );
+                    check( !FMemory::Memcmp(pTest->GetData(), pTemp->GetData(), pTemp->GetDataSize() ) );
+                }
+#endif
+
+                break;
+            }
+
+            default:
+                // Case not implemented
+                check( false );
+
+            }
+
+            break;
+        }
+
+        case EImageFormat::IF_L_UBIT_RLE:
+        {
+            switch ( pBase->GetFormat() )
+            {
+
+            case EImageFormat::IF_L_UBYTE:
+            {
+                check( onlyLOD == -1 );
+
+				// Allocate memory for the compressed data. TODO: Smaller?
+				if (!pResult->m_data.Num())
+				{
+					// Allocate memory for the compressed data. TODO: Smaller?
+					pResult->m_data.SetNum(pBase->m_data.Num());
+				}
+				
+				success = CompressRLE_L1( pBase, pResult );
+                
+#ifdef MUTABLE_DEBUG_RLE		
+				if (success)
+                {
+					// Verify
+                    ImagePtr pTest = ImagePixelFormat( imageCompressionQuality, pResult, EImageFormat::IF_L_UBYTE );
+                    check( !FMemory::Memcmp(pTest->GetData(), pBase->GetData(), pBase->GetDataSize() ) );
+                }
+#endif
+
+                break;
+            }
+
+            case EImageFormat::IF_RGB_UBYTE:
+            case EImageFormat::IF_RGBA_UBYTE:
+            case EImageFormat::IF_BGRA_UBYTE:
+            {
+                check( onlyLOD == -1 );
+                ImagePtr pTemp = ImagePixelFormat( imageCompressionQuality, pBase, EImageFormat::IF_L_UBYTE );
+				success = CompressRLE_L1( pTemp.get(), pResult );
+
+#ifdef MUTABLE_DEBUG_RLE		
+				if (success)
+				{
+					// Verify
+					ImagePtr pTest = ImagePixelFormat(imageCompressionQuality, pResult, EImageFormat::IF_L_UBYTE);
+					check(!FMemory::Memcmp(pTest->GetData(), pTemp->GetData(), pTemp->GetDataSize()));
+				}
+#endif
+
+                break;
+            }
+
+            default:
+                // Case not implemented
+                check( false );
+
+            }
+
+            break;
+        }
+
+        case EImageFormat::IF_RGB_UBYTE:
+        {
+            uint8_t* pDestBuf = pResult->GetData();
+
+            switch ( pBase->GetFormat() )
+            {
+
+            case EImageFormat::IF_L_UBYTE:
+            {
+                for ( int i=0; i<pixelCount; ++i )
+                {
+                    pDestBuf[i*3+0] = pBaseBuf[i];
+                    pDestBuf[i*3+1] = pBaseBuf[i];
+                    pDestBuf[i*3+2] = pBaseBuf[i];
+                }
+                break;
+            }
+
+            case EImageFormat::IF_RGB_UBYTE:
+            {
+				FMemory::Memcpy( pDestBuf, pBaseBuf, 3*pixelCount );
+                break;
+            }
+
+            case EImageFormat::IF_RGBA_UBYTE:
+            {
+                for ( int i=0; i<pixelCount; ++i )
+                {
+                    pDestBuf[i*3+0] = pBaseBuf[i*4+0];
+                    pDestBuf[i*3+1] = pBaseBuf[i*4+1];
+                    pDestBuf[i*3+2] = pBaseBuf[i*4+2];
+                }
+                break;
+            }
+
+            case EImageFormat::IF_BGRA_UBYTE:
+            {
+                for ( int i=0; i<pixelCount; ++i )
+                {
+                    pDestBuf[i*3+0] = pBaseBuf[i*4+2];
+                    pDestBuf[i*3+1] = pBaseBuf[i*4+1];
+                    pDestBuf[i*3+2] = pBaseBuf[i*4+0];
+                }
+                break;
+            }
+
+            case EImageFormat::IF_RGB_UBYTE_RLE:
+            {
+                UncompressRLE_RGB( pBase, pResult );
+                break;
+            }
+
+            case EImageFormat::IF_RGBA_UBYTE_RLE:
+            {
+                check( onlyLOD == -1 );
+                ImagePtr pTemp = ImagePixelFormat( imageCompressionQuality, pBase, EImageFormat::IF_RGBA_UBYTE );
+                const uint8_t* pTempBuf = pTemp->GetData();
+                for ( int i=0; i<pixelCount; ++i )
+                {
+                    pDestBuf[i*3+0] = pTempBuf[i*4+0];
+                    pDestBuf[i*3+1] = pTempBuf[i*4+1];
+                    pDestBuf[i*3+2] = pTempBuf[i*4+2];
+                }
+                break;
+            }
+
+            case EImageFormat::IF_BC1:
+            {
+                for ( int m=0; m<resultLODCount; ++m )
+                {
+                    vec2<int> mipSize = pResult->CalculateMipSize( m );
+                    miro::BC1_to_RGB( mipSize[0], mipSize[1],
+                                        pBase->GetMipData(baseLOD+m),
+                                        pResult->GetMipData(m) );
+                }
+                break;
+            }
+
+            case EImageFormat::IF_BC2:
+            {
+                for ( int m=0; m<resultLODCount; ++m )
+                {
+                    vec2<int> mipSize = pResult->CalculateMipSize( m );
+                    miro::BC2_to_RGB( mipSize[0], mipSize[1],
+                                        pBase->GetMipData(baseLOD+m),
+                                        pResult->GetMipData(m) );
+                }
+                break;
+            }
+
+            case EImageFormat::IF_BC3:
+            {
+                for ( int m=0; m<resultLODCount; ++m )
+                {
+                    vec2<int> mipSize = pResult->CalculateMipSize( m );
+                    miro::BC3_to_RGB( mipSize[0], mipSize[1],
+                                        pBase->GetMipData(baseLOD+m),
+                                        pResult->GetMipData(m) );
+                }
+                break;
+            }
+
+            case EImageFormat::IF_BC4:
+            {
+                for ( int m=0; m<resultLODCount; ++m )
+                {
+                    vec2<int> mipSize = pResult->CalculateMipSize( m );
+                    miro::BC4_to_RGB( mipSize[0], mipSize[1],
+                                      pBase->GetMipData(baseLOD+m),
+                                      pResult->GetMipData(m) );
+                }
+                break;
+            }
+
+            case EImageFormat::IF_BC5:
+            {
+                for ( int m=0; m<resultLODCount; ++m )
+                {
+                    vec2<int> mipSize = pResult->CalculateMipSize( m );
+                    miro::BC5_to_RGB( mipSize[0], mipSize[1],
+                                      pBase->GetMipData(baseLOD+m),
+                                      pResult->GetMipData(m) );
+                }
+                break;
+            }
+
+            case EImageFormat::IF_ASTC_4x4_RGB_LDR:
+            {
+                for ( int m=0; m<resultLODCount; ++m )
+                {
+                    vec2<int> mipSize = pResult->CalculateMipSize( m );
+                    miro::ASTC4x4RGBL_to_RGB( mipSize[0], mipSize[1],
+                                      pBase->GetMipData(baseLOD+m),
+                                      pResult->GetMipData(m) );
+                }
+                break;
+            }
+
+            case EImageFormat::IF_ASTC_4x4_RGBA_LDR:
+            {
+                for ( int m=0; m<resultLODCount; ++m )
+                {
+                    vec2<int> mipSize = pResult->CalculateMipSize( m );
+                    miro::ASTC4x4RGBAL_to_RGB( mipSize[0], mipSize[1],
+                                      pBase->GetMipData(baseLOD+m),
+                                      pResult->GetMipData(m) );
+                }
+                break;
+            }
+
+            case EImageFormat::IF_ASTC_4x4_RG_LDR:
+            {
+                for ( int m=0; m<resultLODCount; ++m )
+                {
+                    vec2<int> mipSize = pResult->CalculateMipSize( m );
+                    miro::ASTC4x4RGL_to_RGB( mipSize[0], mipSize[1],
+                                      pBase->GetMipData(baseLOD+m),
+                                      pResult->GetMipData(m) );
+                }
+                break;
+            }
+
+            default:
+                // Case not implemented
+                check( false );
+
+            }
+
+            break;
+        }
+
+        case EImageFormat::IF_BGRA_UBYTE:
+        {
+            // TODO: Optimise
+            pResult->m_format = EImageFormat::IF_RGBA_UBYTE;
+            bool bSuccess = ImagePixelFormatInPlace( imageCompressionQuality, pResult, pBase, onlyLOD );
+			check(bSuccess);
+            pResult->m_format = EImageFormat::IF_BGRA_UBYTE;
+
+            uint8_t* pDestBuf = pResult->GetData();
+            for ( int i=0; i<pixelCount; ++i )
+            {
+                uint8_t temp = pDestBuf[i*4+0];
+                pDestBuf[i*4+0] = pDestBuf[i*4+2];
+                pDestBuf[i*4+2] = temp;
+            }
+            break;
+        }
+
+        case EImageFormat::IF_RGBA_UBYTE:
+        {
+            uint8_t* pDestBuf = pResult->GetData();
+
+            switch ( pBase->GetFormat() )
+            {
+
+            case EImageFormat::IF_RGBA_UBYTE_RLE:
+            {
+                check( onlyLOD == -1 );
+                UncompressRLE_RGBA( pBase, pResult );
+                break;
+            }
+
+            case EImageFormat::IF_L_UBYTE:
+            {
+                for ( int i=0; i<pixelCount; ++i )
+                {
+                    pDestBuf[i*4+0] = pBaseBuf[i];
+                    pDestBuf[i*4+1] = pBaseBuf[i];
+                    pDestBuf[i*4+2] = pBaseBuf[i];
+                    pDestBuf[i*4+3] = 255;
+                }
+                break;
+            }
+
+            case EImageFormat::IF_RGB_UBYTE:
+            {
+                for ( int i=0; i<pixelCount; ++i )
+                {
+                    pDestBuf[i*4+0] = pBaseBuf[i*3+0];
+                    pDestBuf[i*4+1] = pBaseBuf[i*3+1];
+                    pDestBuf[i*4+2] = pBaseBuf[i*3+2];
+                    pDestBuf[i*4+3] = 255;
+                }
+                break;
+            }
+
+            case EImageFormat::IF_RGBA_UBYTE:
+            {
+				FMemory::Memcpy( pDestBuf, pBaseBuf, 4*pixelCount );
+                break;
+            }
+
+            case EImageFormat::IF_BGRA_UBYTE:
+            {
+                for ( int i=0; i<pixelCount; ++i )
+                {
+                    pDestBuf[i*4+0] = pBaseBuf[i*4+2];
+                    pDestBuf[i*4+1] = pBaseBuf[i*4+1];
+                    pDestBuf[i*4+2] = pBaseBuf[i*4+0];
+                    pDestBuf[i*4+3] = pBaseBuf[i*4+3];
+                }
+                break;
+            }
+
+            case EImageFormat::IF_BC1:
+            {
+                for ( int m=0; m<resultLODCount; ++m )
+                {
+                    vec2<int> mipSize = pResult->CalculateMipSize( m );
+                    miro::BC1_to_RGBA( mipSize[0], mipSize[1],
+                                        pBase->GetMipData(baseLOD+m),
+                                        pResult->GetMipData(m) );
+                }
+                break;
+            }
+
+            case EImageFormat::IF_BC2:
+            {
+                for ( int m=0; m<resultLODCount; ++m )
+                {
+                    vec2<int> mipSize = pResult->CalculateMipSize( m );
+                    miro::BC2_to_RGBA( mipSize[0], mipSize[1],
+                                        pBase->GetMipData(baseLOD+m),
+                                        pResult->GetMipData(m) );
+                }
+                break;
+            }
+
+            case EImageFormat::IF_BC3:
+            {
+                for ( int m=0; m<resultLODCount; ++m )
+                {
+                    vec2<int> mipSize = pResult->CalculateMipSize( m );
+                    miro::BC3_to_RGBA( mipSize[0], mipSize[1],
+                                        pBase->GetMipData(baseLOD+m),
+                                        pResult->GetMipData(m) );
+                }
+                break;
+            }
+
+            case EImageFormat::IF_BC4:
+            {
+                for ( int m=0; m<resultLODCount; ++m )
+                {
+                    vec2<int> mipSize = pResult->CalculateMipSize( m );
+                    miro::BC4_to_RGBA( mipSize[0], mipSize[1],
+                                       pBase->GetMipData(baseLOD+m),
+                                       pResult->GetMipData(m) );
+                }
+                break;
+            }
+
+            case EImageFormat::IF_BC5:
+            {
+                for ( int m=0; m<resultLODCount; ++m )
+                {
+                    vec2<int> mipSize = pResult->CalculateMipSize( m );
+                    miro::BC5_to_RGBA( mipSize[0], mipSize[1],
+                                        pBase->GetMipData(baseLOD+m),
+                                        pResult->GetMipData(m) );
+                }
+                break;
+            }
+
+            case EImageFormat::IF_ASTC_4x4_RGB_LDR:
+            {
+                for ( int m=0; m<resultLODCount; ++m )
+                {
+                    vec2<int> mipSize = pResult->CalculateMipSize( m );
+                    miro::ASTC4x4RGBL_to_RGBA( mipSize[0], mipSize[1],
+                                        pBase->GetMipData(baseLOD+m),
+                                        pResult->GetMipData(m) );
+                }
+                break;
+            }
+
+            case EImageFormat::IF_ASTC_4x4_RGBA_LDR:
+            {
+                for ( int m=0; m<resultLODCount; ++m )
+                {
+                    vec2<int> mipSize = pResult->CalculateMipSize( m );
+                    miro::ASTC4x4RGBAL_to_RGBA( mipSize[0], mipSize[1],
+                                        pBase->GetMipData(baseLOD+m),
+                                        pResult->GetMipData(m) );
+                }
+                break;
+            }
+
+            case EImageFormat::IF_ASTC_4x4_RG_LDR:
+            {
+                for ( int m=0; m<resultLODCount; ++m )
+                {
+                    vec2<int> mipSize = pResult->CalculateMipSize( m );
+                    miro::ASTC4x4RGL_to_RGBA( mipSize[0], mipSize[1],
+                                        pBase->GetMipData(baseLOD+m),
+                                        pResult->GetMipData(m) );
+                }
+                break;
+            }
+
+			case EImageFormat::IF_RGB_UBYTE_RLE:
+			{
+				check(onlyLOD == -1);
+				ImagePtr pTemp = ImagePixelFormat(imageCompressionQuality, pBase, EImageFormat::IF_RGB_UBYTE);
+				const uint8_t* pTempData = pTemp->GetData();
+				for (int i = 0; i < pixelCount; ++i)
+				{
+					pDestBuf[i * 4 + 0] = pTempData[i * 3 + 0];
+					pDestBuf[i * 4 + 1] = pTempData[i * 3 + 1];
+					pDestBuf[i * 4 + 2] = pTempData[i * 3 + 2];
+					pDestBuf[i * 4 + 3] = 255;
+				}
+				break;
+			}
+
+			case EImageFormat::IF_L_UBYTE_RLE:
+			{
+				check(onlyLOD == -1);
+				ImagePtr pTemp = ImagePixelFormat(imageCompressionQuality, pBase, EImageFormat::IF_L_UBYTE);
+				const uint8_t* pTempData = pTemp->GetData();
+				for (int i = 0; i < pixelCount; ++i)
+				{
+					pDestBuf[i * 4 + 0] = pTempData[i];
+					pDestBuf[i * 4 + 1] = pTempData[i];
+					pDestBuf[i * 4 + 2] = pTempData[i];
+					pDestBuf[i * 4 + 3] = 255;
+				}
+				break;
+			}
+
+            default:
+                // Case not implemented
+                check( false );
+
+            }
+
+            break;
+        }
+
+
+        case EImageFormat::IF_RGBA_UBYTE_RLE:
+        {
+            switch ( pBase->GetFormat() )
+            {
+
+            case EImageFormat::IF_RGBA_UBYTE:
+            {
+                check( onlyLOD == -1 );
+                CompressRLE_RGBA( pBase, pResult );
+
+                // TODO: TEST BLEH
+                //                ImagePtr pTest = ImagePixelFormat( imageCompressionQuality,
+                //                pResult, EImageFormat::IF_RGBA_UBYTE ); check(
+                //                !memcmp(pTest->GetData(), pBase->GetData(),
+                //                pBase->GetDataSize() ) );
+                break;
+            }
+
+            case EImageFormat::IF_RGB_UBYTE:
+            {
+                check( onlyLOD == -1 );
+
+                // \todo: optimise
+                ImagePtr pTemp = ImagePixelFormat( imageCompressionQuality, pBase, EImageFormat::IF_RGBA_UBYTE );
+                CompressRLE_RGBA( pTemp.get(), pResult );
+
+                // TODO: TEST BLEH
+                //                ImagePtr pTest = ImagePixelFormat( imageCompressionQuality,
+                //                pResult, EImageFormat::IF_RGBA_UBYTE ); check(
+                //                !memcmp(pTest->GetData(), pBase->GetData(),
+                //                pBase->GetDataSize() ) );
+                break;
+            }
+
+            case EImageFormat::IF_RGB_UBYTE_RLE:
+            {
+                check( onlyLOD == -1 );
+
+                // \todo: optimise
+                ImagePtr pTemp1 = ImagePixelFormat( imageCompressionQuality, pBase, EImageFormat::IF_RGB_UBYTE );
+                ImagePtr pTemp2 =
+                    ImagePixelFormat( imageCompressionQuality, pTemp1.get(), EImageFormat::IF_RGBA_UBYTE );
+                CompressRLE_RGBA( pTemp2.get(), pResult );
+                break;
+            }
+
+            default:
+                // Case not implemented
+                check( false );
+
+            }
+
+            break;
+        }
+
+
+        case EImageFormat::IF_RGB_UBYTE_RLE:
+        {
+            switch ( pBase->GetFormat() )
+            {
+
+            case EImageFormat::IF_RGB_UBYTE:
+            {
+                check( onlyLOD == -1 );
+                CompressRLE_RGB( pBase, pResult );
+
+                // TODO: TEST BLEH
+                // ImagePtr pTest = ImagePixelFormat( imageCompressionQuality, pDest.get(),
+                // EImageFormat::IF_RGB_UBYTE ); check( !memcmp(pTest->GetData(), pBase->GetData(),
+                // pBase->GetDataSize() ) );
+                break;
+            }
+
+            case EImageFormat::IF_RGBA_UBYTE_RLE:
+            {
+                check( onlyLOD == -1 );
+
+                // \todo: optimise
+                ImagePtr pTemp1 = ImagePixelFormat( imageCompressionQuality, pBase, EImageFormat::IF_RGBA_UBYTE );
+                ImagePtr pTemp2 =
+                    ImagePixelFormat( imageCompressionQuality, pTemp1.get(), EImageFormat::IF_RGB_UBYTE );
+                CompressRLE_RGB( pTemp2.get(), pResult );
+                break;
+            }
+
+            default:
+                // Case not implemented
+                check( false );
+
+            }
+
+            break;
+        }
+
+        case EImageFormat::IF_BC1:
+        {
+            switch ( pBase->GetFormat() )
+            {
+            case EImageFormat::IF_RGB_UBYTE:
+                for ( int m=0; m<resultLODCount; ++m )
+                {
+                    vec2<int> mipSize = pResult->CalculateMipSize( m );
+                    miro::RGB_to_BC1(
+                        mipSize[0], mipSize[1], pBase->GetMipData( baseLOD + m ),
+                        pResult->GetMipData( m ), imageCompressionQuality );
+                }
+                break;
+            case EImageFormat::IF_RGBA_UBYTE:
+                for ( int m=0; m<resultLODCount; ++m )
+                {
+                    vec2<int> mipSize = pResult->CalculateMipSize( m );
+                    miro::RGBA_to_BC1(
+                        mipSize[0], mipSize[1], pBase->GetMipData( baseLOD + m ),
+                        pResult->GetMipData( m ), imageCompressionQuality );
+                }
+                break;
+
+            case EImageFormat::IF_L_UBYTE:
+            {
+                ImagePtr pTemp = ImagePixelFormat( imageCompressionQuality, pBase, EImageFormat::IF_RGB_UBYTE );
+                for ( int m=0; m<resultLODCount; ++m )
+                {
+                    vec2<int> mipSize = pResult->CalculateMipSize( m );
+                    miro::RGB_to_BC1(
+                        mipSize[0], mipSize[1], pTemp->GetMipData( baseLOD + m ),
+                        pResult->GetMipData( m ), imageCompressionQuality );
+                }
+                break;
+            }
+
+            case EImageFormat::IF_BC3:
+            {
+                ImagePtr pTemp = ImagePixelFormat( imageCompressionQuality, pBase, EImageFormat::IF_RGBA_UBYTE );
+                for (int m = 0; m < resultLODCount; ++m)
+                {
+                    vec2<int> mipSize = pResult->CalculateMipSize(m);
+                    miro::RGBA_to_BC1(
+                        mipSize[0], mipSize[1], pTemp->GetMipData( baseLOD + m ),
+                        pResult->GetMipData( m ), imageCompressionQuality );
+                }
+                break;
+            }
+
+            default:
+                // Case not implemented
+                check( false );
+
+            }
+
+            break;
+        }
+
+        case EImageFormat::IF_BC2:
+        {
+            switch ( pBase->GetFormat() )
+            {
+            case EImageFormat::IF_RGBA_UBYTE:
+                for ( int m=0; m<resultLODCount; ++m )
+                {
+                    vec2<int> mipSize = pResult->CalculateMipSize( m );
+                    miro::RGBA_to_BC2(
+                        mipSize[0], mipSize[1], pBase->GetMipData( baseLOD + m ),
+                        pResult->GetMipData( m ), imageCompressionQuality );
+                }
+                break;
+
+            default:
+                // Case not implemented
+                check( false );
+
+            }
+
+            break;
+        }
+
+        case EImageFormat::IF_BC3:
+        {
+            switch ( pBase->GetFormat() )
+            {
+            case EImageFormat::IF_RGBA_UBYTE:
+                for ( int m=0; m<resultLODCount; ++m )
+                {
+                    vec2<int> mipSize = pResult->CalculateMipSize( m );
+                    miro::RGBA_to_BC3(
+                        mipSize[0], mipSize[1], pBase->GetMipData( baseLOD + m ),
+                        pResult->GetMipData( m ), imageCompressionQuality );
+                }
+                break;
+
+            case EImageFormat::IF_RGB_UBYTE:          
+                for ( int m=0; m<resultLODCount; ++m )
+                {
+                    vec2<int> mipSize = pResult->CalculateMipSize( m );
+                    miro::RGB_to_BC3(
+                        mipSize[0], mipSize[1], pBase->GetMipData( baseLOD + m ),
+                        pResult->GetMipData( m ), imageCompressionQuality );
+                }
+                break;
+
+            case EImageFormat::IF_BC1:
+            {
+                MUTABLE_CPUPROFILER_SCOPE(BC1toBC3);
+
+                for ( int m=0; m<resultLODCount; ++m )
+                {
+                    vec2<int> mipSize = pResult->CalculateMipSize( m );
+                    miro::BC1_to_BC3(
+                        mipSize[0], mipSize[1], pBase->GetMipData( baseLOD + m ),
+                        pResult->GetMipData( m ), imageCompressionQuality );
+                }
+                break;
+            }
+
+            default:
+                // Case not implemented
+                check( false );
+
+            }
+
+            break;
+        }
+
+        case EImageFormat::IF_BC4:
+        {
+            switch ( pBase->GetFormat() )
+            {
+            case EImageFormat::IF_L_UBYTE:
+                for ( int m=0; m<resultLODCount; ++m )
+                {
+                    vec2<int> mipSize = pResult->CalculateMipSize( m );
+                    miro::L_to_BC4(
+                        mipSize[0], mipSize[1], pBase->GetMipData( baseLOD + m ),
+                        pResult->GetMipData( m ), imageCompressionQuality );
+                }
+                break;
+
+            default:
+                // Case not implemented
+                check( false );
+
+            }
+
+            break;
+        }
+
+        case EImageFormat::IF_BC5:
+        {
+            switch ( pBase->GetFormat() )
+            {
+            case EImageFormat::IF_RGBA_UBYTE:
+                for ( int m=0; m<resultLODCount; ++m )
+                {
+                    vec2<int> mipSize = pResult->CalculateMipSize( m );
+                    miro::RGBA_to_BC5(
+                        mipSize[0], mipSize[1], pBase->GetMipData( baseLOD + m ),
+                        pResult->GetMipData( m ), imageCompressionQuality );
+                }
+                break;
+
+            case EImageFormat::IF_RGB_UBYTE:
+                for ( int m=0; m<resultLODCount; ++m )
+                {
+                    vec2<int> mipSize = pResult->CalculateMipSize( m );
+                    miro::RGB_to_BC5(
+                        mipSize[0], mipSize[1], pBase->GetMipData( baseLOD + m ),
+                        pResult->GetMipData( m ), imageCompressionQuality );
+                }
+                break;
+
+            default:
+                // Case not implemented
+                check( false );
+
+            }
+
+            break;
+        }
+
+        case EImageFormat::IF_ASTC_4x4_RGB_LDR:
+        {
+            switch ( pBase->GetFormat() )
+            {
+            case EImageFormat::IF_RGBA_UBYTE:
+                for ( int m=0; m<resultLODCount; ++m )
+                {
+                    vec2<int> mipSize = pResult->CalculateMipSize( m );
+                    miro::RGBA_to_ASTC4x4RGBL(
+                        mipSize[0], mipSize[1], pBase->GetMipData( baseLOD + m ),
+                        pResult->GetMipData( m ), imageCompressionQuality );
+                }
+                break;
+
+            case EImageFormat::IF_RGB_UBYTE:
+                for ( int m=0; m<resultLODCount; ++m )
+                {
+                    vec2<int> mipSize = pResult->CalculateMipSize( m );
+                    miro::RGB_to_ASTC4x4RGBL(
+                        mipSize[0], mipSize[1], pBase->GetMipData( baseLOD + m ),
+                        pResult->GetMipData( m ), imageCompressionQuality );
+                }
+                break;
+
+            case EImageFormat::IF_ASTC_4x4_RGBA_LDR:
+                for ( int m=0; m<resultLODCount; ++m )
+                {
+                    vec2<int> mipSize = pResult->CalculateMipSize( m );
+                    miro::ASTC4x4RGBAL_to_ASTC4x4RGBL( mipSize[0], mipSize[1],
+                                       pBase->GetMipData(baseLOD+m),
+                                       pResult->GetMipData(m) );
+                }
+                break;
+
+			case EImageFormat::IF_L_UBYTE:
+				for (int m = 0; m < resultLODCount; ++m)
+				{
+					vec2<int> mipSize = pResult->CalculateMipSize(m);
+					miro::L_to_ASTC4x4RGBL(mipSize[0], mipSize[1],
+						pBase->GetMipData(baseLOD + m),
+						pResult->GetMipData(m), imageCompressionQuality);
+				}
+				break;
+
+            default:
+                // Case not implemented
+                check( false );
+
+            }
+
+            break;
+        }
+
+        case EImageFormat::IF_ASTC_4x4_RGBA_LDR:
+        {
+            switch ( pBase->GetFormat() )
+            {
+            case EImageFormat::IF_RGBA_UBYTE:
+                for ( int m=0; m<resultLODCount; ++m )
+                {
+                    vec2<int> mipSize = pResult->CalculateMipSize( m );
+                    miro::RGBA_to_ASTC4x4RGBAL(
+                        mipSize[0], mipSize[1], pBase->GetMipData( baseLOD + m ),
+                        pResult->GetMipData( m ), imageCompressionQuality );
+                }
+                break;
+
+            case EImageFormat::IF_RGB_UBYTE:
+                for ( int m=0; m<resultLODCount; ++m )
+                {
+                    vec2<int> mipSize = pResult->CalculateMipSize( m );
+                    miro::RGB_to_ASTC4x4RGBAL(
+                        mipSize[0], mipSize[1], pBase->GetMipData( baseLOD + m ),
+                        pResult->GetMipData( m ), imageCompressionQuality );
+                }
+                break;
+
+            case EImageFormat::IF_ASTC_4x4_RGB_LDR:
+                for ( int m=0; m<resultLODCount; ++m )
+                {
+                    vec2<int> mipSize = pResult->CalculateMipSize( m );
+                    miro::ASTC4x4RGBL_to_ASTC4x4RGBAL( mipSize[0], mipSize[1],
+                                       pBase->GetMipData(baseLOD+m),
+                                       pResult->GetMipData(m) );
+                }
+                break;
+
+            default:
+                // Case not implemented
+                check( false );
+
+            }
+
+            break;
+        }
+
+        case EImageFormat::IF_ASTC_4x4_RG_LDR:
+        {
+            switch ( pBase->GetFormat() )
+            {
+            case EImageFormat::IF_RGBA_UBYTE:
+                for ( int m=0; m<resultLODCount; ++m )
+                {
+                    vec2<int> mipSize = pResult->CalculateMipSize( m );
+                    miro::RGBA_to_ASTC4x4RGL(
+                        mipSize[0], mipSize[1], pBase->GetMipData( baseLOD + m ),
+                        pResult->GetMipData( m ), imageCompressionQuality );
+                }
+                break;
+
+            case EImageFormat::IF_RGB_UBYTE:
+                for ( int m=0; m<resultLODCount; ++m )
+                {
+                    vec2<int> mipSize = pResult->CalculateMipSize( m );
+                    miro::RGB_to_ASTC4x4RGL(
+                        mipSize[0], mipSize[1], pBase->GetMipData( baseLOD + m ),
+                        pResult->GetMipData( m ), imageCompressionQuality );
+                }
+                break;
+
+            case EImageFormat::IF_ASTC_4x4_RG_LDR:
+                for ( int m=0; m<resultLODCount; ++m )
+                {
+                    vec2<int> mipSize = pResult->CalculateMipSize( m );
+                    // Hack that actually works because of block size.
+                    miro::ASTC4x4RGBAL_to_ASTC4x4RGBL( mipSize[0], mipSize[1],
+                                       pBase->GetMipData(baseLOD+m),
+                                       pResult->GetMipData(m) );
+                }
+                break;
+
+            default:
+                // Case not implemented
+                check( false );
+
+            }
+
+            break;
+        }
+
+        default:
+            // Case not implemented
+            check( false );
+        }
+    }
+
+    return success;
+}
+
+
diff --git a/Engine/Plugins/Experimental/Mutable/Source/MutableRuntime/Private/OpImagePixelFormat.h b/Engine/Plugins/Experimental/Mutable/Source/MutableRuntime/Private/OpImagePixelFormat.h
new file mode 100644
--- /dev/null
+++ b/Engine/Plugins/Experimental/Mutable/Source/MutableRuntime/Private/OpImagePixelFormat.h
@@ -0,0 +1,27 @@
+// Copyright Epic Games, Inc. All Rights Reserved.
+
+#pragma once
+
+#include "ImagePrivate.h"
+
+
+namespace mu
+{
+    //---------------------------------------------------------------------------------------------
+    //! Convert an image to another pixel format.
+    //! \warning Not all format conversions are implemented.
+    //! \param onlyLOD If different than -1, only the specified lod level will be converted in the
+    //! returned image.
+    //! \return false if the conversion failed, usually because not enough memory was allocated in
+    //!     the result. This is only checked for RLE compression.
+    //---------------------------------------------------------------------------------------------
+	MUTABLERUNTIME_API extern bool ImagePixelFormatInPlace( int quality, Image* pResult, const Image* pBase,
+                                         int onlyLOD = -1 );
+
+    //---------------------------------------------------------------------------------------------
+    //! Wrapper for the above method that allocates the destination image.
+    //---------------------------------------------------------------------------------------------
+	MUTABLERUNTIME_API extern ImagePtr ImagePixelFormat( int quality, const Image* pBase, EImageFormat targetFormat,
+                                      int onlyLOD = -1 );
+
+}
