[mutable] Move Mutable plugin from the restricted folder to the experimental plugins folder.

diff --git a/Engine/Plugins/Experimental/Mutable/Source/MutableRuntime/Private/OpImageProject.cpp b/Engine/Plugins/Experimental/Mutable/Source/MutableRuntime/Private/OpImageProject.cpp
new file mode 100644
--- /dev/null
+++ b/Engine/Plugins/Experimental/Mutable/Source/MutableRuntime/Private/OpImageProject.cpp
@@ -0,0 +1,2742 @@
+// Copyright Epic Games, Inc. All Rights Reserved.
+
+
+#include "OpImageProject.h"
+#include "OpMeshClipWithMesh.h"
+
+#include "Async/ParallelFor.h"
+#include "Platform.h"
+#include "ImagePrivate.h"
+#include "MeshPrivate.h"
+#include "Raster.h"
+#include "ConvertData.h"
+#include "MutableMath.h"
+#include "Octree.h"
+
+#include <queue>
+#include <unordered_set>
+
+
+using namespace mu;
+
+
+namespace
+{
+
+	class RaycastPixelProcessor
+	{
+	public:
+
+		RaycastPixelProcessor( const PROJECTOR& projector,
+							   const Image* pSource,
+                               const uint8_t* pTargetData,
+                               const uint8_t* pMaskData,
+							   float fadeStart,
+							   float fadeEnd )
+		{
+			projector.GetDirectionSideUp( m_direction, m_side, m_up );
+
+			m_position[0] = projector.position[0];
+			m_position[1] = projector.position[1];
+			m_position[2] = projector.position[2];
+
+			m_scale[0] = projector.scale[0];
+            m_scale[1] = projector.scale[1];
+            m_scale[2] = projector.scale[2];
+
+			m_pSource = pSource;
+
+            m_fadeStartCos = cosf(fadeStart);
+            m_fadeEndCos = cosf(fadeEnd);
+
+			m_pTargetData = pTargetData;
+			m_pMaskData = pMaskData;
+
+			EImageFormat format = m_pSource->GetFormat();
+			m_pixelSize = GetImageFormatData( format ).m_bytesPerBlock;
+
+		}
+
+
+		//-----------------------------------------------------------------------------------------
+        inline void ProcessPixel( unsigned char* pBufferPos, float varying[4] ) const
+		{
+			float factor = 1;
+
+            // depth clamp
+            if (varying[2]<0.0f || varying[2]>1.0f)
+            {
+                factor = 0;
+            }
+
+            if ( factor>0 && m_pMaskData )
+			{
+                uint8_t maskFactor = m_pMaskData[ (pBufferPos-m_pTargetData) / m_pixelSize ];
+				factor = float(maskFactor) / 255.0f;
+			}
+
+			if ( factor>0 )
+			{
+//				vec3<float> normal( varying[2], varying[3], varying[4] );
+//				normal = normalise_approx( normal );
+
+//				float angleCos = dot( normal, m_direction * -1.0f );
+                float angleCos = varying[3];
+
+                if ( angleCos<m_fadeStartCos && angleCos>m_fadeEndCos )
+				{
+                    factor *= ( angleCos - m_fadeEndCos ) / ( m_fadeStartCos - m_fadeEndCos );
+				}
+                else if ( angleCos<=m_fadeEndCos )
+				{
+					factor = 0;
+				}
+			}
+
+
+			if ( factor>0 )
+			{
+				// Transform the origin to the projected image texture coordinates
+				float u = varying[0];
+				float v = varying[1];
+
+				if ( u>=0 && u<1 && v>=0 && v<1 )
+				{
+
+                    const uint8_t* pPixel = m_pSource->GetData();
+					// TODO: clamp?
+					pPixel += ( m_pSource->GetSizeX() * int(m_pSource->GetSizeY() * v)
+								+ int( m_pSource->GetSizeX() * u ) )
+							* m_pixelSize;
+
+					// Write result
+					switch ( m_pSource->GetFormat() )
+					{
+					case EImageFormat::IF_L_UBYTE:
+                        pBufferPos[0] = uint8_t( pPixel[0] * factor );
+						break;
+
+					case EImageFormat::IF_RGB_UBYTE:
+                        pBufferPos[0] = uint8_t( pPixel[0] * factor );
+                        pBufferPos[1] = uint8_t( pPixel[1] * factor );
+                        pBufferPos[2] = uint8_t( pPixel[2] * factor );
+						break;
+
+                    case EImageFormat::IF_BGRA_UBYTE:
+                    case EImageFormat::IF_RGBA_UBYTE:
+                        pBufferPos[0] = uint8_t( pPixel[0] * factor );
+                        pBufferPos[1] = uint8_t( pPixel[1] * factor );
+                        pBufferPos[2] = uint8_t( pPixel[2] * factor );
+                        pBufferPos[3] = uint8_t( pPixel[3] * factor );
+						break;
+
+					default:
+						// Not implemented
+						check( false );
+						break;
+					}
+				}
+			}
+		}
+
+
+		const Image* m_pSource;
+		vec3<float> m_direction;
+		vec3<float> m_up;
+		vec3<float> m_side;
+		vec3<float> m_position;
+        vec3<float> m_scale;
+
+        //! Cosine of the fading angle range
+        float m_fadeStartCos, m_fadeEndCos;
+
+        const uint8_t* m_pTargetData;
+        const uint8_t* m_pMaskData;
+		int m_pixelSize;
+
+	};
+
+
+	//---------------------------------------------------------------------------------------------
+    template<int PIXEL_SIZE,bool HAS_MASK=true>
+	class RaycastPixelProcessor_UBYTE
+	{
+	public:
+
+		RaycastPixelProcessor_UBYTE( const PROJECTOR& projector,
+									 const Image* pSource,
+                                     const uint8_t* pTargetData,
+                                     const uint8_t* pMaskData,
+									 float fadeStart,
+									 float fadeEnd )
+		{
+			projector.GetDirectionSideUp( m_direction, m_side, m_up );
+
+			m_position[0] = projector.position[0];
+			m_position[1] = projector.position[1];
+			m_position[2] = projector.position[2];
+
+			m_scale[0] = projector.scale[0];
+            m_scale[1] = projector.scale[1];
+            m_scale[2] = projector.scale[2];
+
+			m_pSourceData = pSource->GetData();
+			m_sourceSizeX = pSource->GetSizeX();
+			m_sourceSizeY = pSource->GetSizeY();
+
+            m_fadeStartCos = cosf(fadeStart);
+            m_fadeEndCos = cosf(fadeEnd);
+
+			m_pTargetData = pTargetData;
+			m_pMaskData = pMaskData;
+
+			check( GetImageFormatData( pSource->GetFormat() ).m_bytesPerBlock
+							==
+							PIXEL_SIZE );
+		}
+
+
+		//-----------------------------------------------------------------------------------------
+        inline void ProcessPixel( unsigned char* pBufferPos, float varying[4] ) const
+		{
+            int32_t factor_8 = 256;
+
+            // depth clamp
+            if (varying[2]<0.0f || varying[2]>1.0f)
+            {
+                factor_8 = 0;
+            }
+
+//				vec3<float> normal( varying[2], varying[3], varying[4] );
+//				normal = normalise_approx( normal );
+
+//				float angleCos = dot( normal, m_direction * -1.0f );
+            float angleCos = varying[3];
+
+            if (HAS_MASK)
+            {
+                if ( factor_8>0 && m_pMaskData )
+                {
+                    uint8_t maskFactor = m_pMaskData[ (pBufferPos-m_pTargetData) / PIXEL_SIZE ];
+                    factor_8 = maskFactor;
+                }
+
+                if ( factor_8>0 )
+                {
+                    if ( angleCos<m_fadeStartCos && angleCos>m_fadeEndCos )
+                    {
+                        factor_8 = ( factor_8
+                                     *
+                                     int32_t( 256.0f * ( angleCos - m_fadeEndCos ) / ( m_fadeStartCos - m_fadeEndCos ) )
+                                     ) >> 8;
+                    }
+                    else if ( angleCos<=m_fadeEndCos )
+                    {
+                        factor_8 = 0;
+                    }
+                }
+            }
+            else
+            {
+                if ( angleCos<m_fadeStartCos && angleCos>m_fadeEndCos )
+                {
+                    factor_8 = int32_t( 256.0f * ( angleCos - m_fadeEndCos ) / ( m_fadeStartCos - m_fadeEndCos ) );
+                }
+                else if ( angleCos<=m_fadeEndCos )
+                {
+                    factor_8 = 0;
+                }
+            }
+
+
+			if ( factor_8>0 )
+			{
+				float u = varying[0];
+				float v = varying[1];
+
+				if ( u>=0 && u<1 && v>=0 && v<1 )
+				{
+
+                    const uint8_t* pPixel = m_pSourceData;
+					// TODO: clamp?
+					pPixel += ( m_sourceSizeX * int(m_sourceSizeY * v)
+								+ int( m_sourceSizeX * u ) )
+							* PIXEL_SIZE;
+
+					// Write result
+					for ( int i=0; i<PIXEL_SIZE; ++i )
+					{
+						pBufferPos[i] = uint8_t( ( pPixel[i] * factor_8 ) >> 8 );
+					}
+				}
+			}
+		}
+
+
+        const uint8_t* m_pSourceData;
+		int m_sourceSizeX, m_sourceSizeY;
+
+		vec3<float> m_direction;
+		vec3<float> m_up;
+		vec3<float> m_side;
+		vec3<float> m_position;
+        vec3<float> m_scale;
+
+        //! Cosine of the fading angle range
+        float m_fadeStartCos, m_fadeEndCos;
+
+        const uint8_t* m_pTargetData;
+        const uint8_t* m_pMaskData;
+	};
+
+
+
+    class RasterProjectedPixelProcessor
+    {
+    public:
+
+        RasterProjectedPixelProcessor( const Image* pSource,
+                               const uint8_t* pTargetData,
+                               const uint8_t* pMaskData,
+                               float fadeStart,
+                               float fadeEnd )
+        {
+            m_pSource = pSource;
+
+            m_fadeStartCos = cosf(fadeStart);
+            m_fadeEndCos = cosf(fadeEnd);
+
+            m_pTargetData = pTargetData;
+            m_pMaskData = pMaskData;
+
+			EImageFormat format = m_pSource->GetFormat();
+            m_pixelSize = GetImageFormatData( format ).m_bytesPerBlock;
+
+        }
+
+
+        //-----------------------------------------------------------------------------------------
+        inline void ProcessPixel( unsigned char* pBufferPos, float varying[4] ) const
+        {
+            float factor = 1;
+
+            // depth clamp
+            if (varying[2]<0.0f || varying[2]>1.0f)
+            {
+                factor = 0;
+            }
+
+            if ( factor>0 && m_pMaskData )
+            {
+                uint8_t maskFactor = m_pMaskData[ (pBufferPos-m_pTargetData) / m_pixelSize ];
+                factor = float(maskFactor) / 255.0f;
+            }
+
+            if ( factor>0 )
+            {
+                float angleCos = varying[3];
+
+                if ( angleCos<m_fadeStartCos && angleCos>m_fadeEndCos )
+                {
+                    factor *= ( angleCos - m_fadeEndCos ) / ( m_fadeStartCos - m_fadeEndCos );
+                }
+                else if ( angleCos<=m_fadeEndCos )
+                {
+                    factor = 0;
+                }
+            }
+
+
+            if ( factor>0 )
+            {
+                // Transform the origin to the projected image texture coordinates
+                float u = varying[0];
+                float v = varying[1];
+
+                if ( u>=0 && u<1 && v>=0 && v<1 )
+                {
+
+                    const uint8_t* pPixel = m_pSource->GetData();
+                    // TODO: clamp?
+                    pPixel += ( m_pSource->GetSizeX() * int(m_pSource->GetSizeY() * v)
+                                + int( m_pSource->GetSizeX() * u ) )
+                            * m_pixelSize;
+
+                    // Write result
+                    switch ( m_pSource->GetFormat() )
+                    {
+                    case EImageFormat::IF_L_UBYTE:
+                        pBufferPos[0] = uint8_t( pPixel[0] * factor );
+                        break;
+
+                    case EImageFormat::IF_RGB_UBYTE:
+                        pBufferPos[0] = uint8_t( pPixel[0] * factor );
+                        pBufferPos[1] = uint8_t( pPixel[1] * factor );
+                        pBufferPos[2] = uint8_t( pPixel[2] * factor );
+                        break;
+
+                    case EImageFormat::IF_BGRA_UBYTE:
+                    case EImageFormat::IF_RGBA_UBYTE:
+                        pBufferPos[0] = uint8_t( pPixel[0] * factor );
+                        pBufferPos[1] = uint8_t( pPixel[1] * factor );
+                        pBufferPos[2] = uint8_t( pPixel[2] * factor );
+                        pBufferPos[3] = uint8_t( pPixel[3] * factor );
+                        break;
+
+                    default:
+                        // Not implemented
+                        check( false );
+                        break;
+                    }
+                }
+            }
+        }
+
+
+        const Image* m_pSource;
+
+        //! Cosine of the fading angle range
+        float m_fadeStartCos, m_fadeEndCos;
+
+        const uint8_t* m_pTargetData;
+        const uint8_t* m_pMaskData;
+        int m_pixelSize;
+
+    };
+
+
+    //---------------------------------------------------------------------------------------------
+    template<int PIXEL_SIZE,bool HAS_MASK=true>
+    class RasterProjectedPixelProcessor_UBYTE
+    {
+    public:
+
+        RasterProjectedPixelProcessor_UBYTE( const Image* pSource,
+                                     const uint8_t* pTargetData,
+                                     const uint8_t* pMaskData,
+                                     float fadeStart,
+                                     float fadeEnd )
+        {
+            m_pSourceData = pSource->GetData();
+            m_sourceSizeX = pSource->GetSizeX();
+            m_sourceSizeY = pSource->GetSizeY();
+
+            m_fadeStartCos = cosf(fadeStart);
+            m_fadeEndCos = cosf(fadeEnd);
+
+            m_pTargetData = pTargetData;
+            m_pMaskData = pMaskData;
+
+            check( GetImageFormatData( pSource->GetFormat() ).m_bytesPerBlock
+                            ==
+                            PIXEL_SIZE );
+        }
+
+
+        //-----------------------------------------------------------------------------------------
+        inline void ProcessPixel( unsigned char* pBufferPos, float varying[4] ) const
+        {
+            int32_t factor_8 = 256;
+
+            // depth clamp
+            if (varying[2]<0.0f || varying[2]>1.0f)
+            {
+                factor_8 = 0;
+            }
+
+            float angleCos = varying[3];
+
+            if (HAS_MASK)
+            {
+                if ( factor_8>0 && m_pMaskData )
+                {
+                    uint8_t maskFactor = m_pMaskData[ (pBufferPos-m_pTargetData) / PIXEL_SIZE ];
+                    factor_8 = maskFactor;
+                }
+
+                if ( factor_8>0 )
+                {
+                    if ( angleCos<m_fadeStartCos && angleCos>m_fadeEndCos )
+                    {
+                        factor_8 = ( factor_8
+                                     *
+                                     int32_t( 256.0f * ( angleCos - m_fadeEndCos ) / ( m_fadeStartCos - m_fadeEndCos ) )
+                                     ) >> 8;
+                    }
+                    else if ( angleCos<=m_fadeEndCos )
+                    {
+                        factor_8 = 0;
+                    }
+                }
+            }
+            else
+            {
+                if ( angleCos<m_fadeStartCos && angleCos>m_fadeEndCos )
+                {
+                    factor_8 = int32_t( 256.0f * ( angleCos - m_fadeEndCos ) / ( m_fadeStartCos - m_fadeEndCos ) );
+                }
+                else if ( angleCos<=m_fadeEndCos )
+                {
+                    factor_8 = 0;
+                }
+            }
+
+
+            if ( factor_8>0 )
+            {
+                float u = varying[0];
+                float v = varying[1];
+
+                if ( u>=0 && u<1 && v>=0 && v<1 )
+                {
+
+                    const uint8_t* pPixel = m_pSourceData;
+                    // TODO: clamp?
+                    pPixel += ( m_sourceSizeX * int(m_sourceSizeY * v)
+                                + int( m_sourceSizeX * u ) )
+                            * PIXEL_SIZE;
+
+                    // Write result
+                    for ( int i=0; i<PIXEL_SIZE; ++i )
+                    {
+                        pBufferPos[i] = uint8_t( ( pPixel[i] * factor_8 ) >> 8 );
+                    }
+                }
+            }
+        }
+
+
+        const uint8_t* m_pSourceData;
+        int m_sourceSizeX, m_sourceSizeY;
+
+        //! Cosine of the fading angle range
+        float m_fadeStartCos, m_fadeEndCos;
+
+        const uint8_t* m_pTargetData;
+        const uint8_t* m_pMaskData;
+    };
+
+
+    //---------------------------------------------------------------------------------------------
+    template<int PIXEL_SIZE,bool HAS_MASK=true>
+    class RasterCylindricalProjectedPixelProcessor_UBYTE
+    {
+    public:
+
+        RasterCylindricalProjectedPixelProcessor_UBYTE( const Image* pSource,
+                                     const uint8_t* pTargetData,
+                                     const uint8_t* pMaskData,
+                                     float fadeStart,
+                                     float fadeEnd,
+                                                        float projectionAngle )
+        {
+            m_pSourceData = pSource->GetData();
+            m_sourceSizeX = pSource->GetSizeX();
+            m_sourceSizeY = pSource->GetSizeY();
+
+            m_fadeStartCos = cosf(fadeStart);
+            m_fadeEndCos = cosf(fadeEnd);
+            m_projectionAngle = projectionAngle;
+
+            m_pTargetData = pTargetData;
+            m_pMaskData = pMaskData;
+
+            check( GetImageFormatData( pSource->GetFormat() ).m_bytesPerBlock
+                            ==
+                            PIXEL_SIZE );
+        }
+
+
+        //-----------------------------------------------------------------------------------------
+        inline void ProcessPixel( unsigned char* pBufferPos, float varying[4] ) const
+        {
+            int32_t factor_8 = 256;
+
+            // Position in unit cylinder space
+            float x_cyl = varying[0];
+            float y_cyl = varying[1];
+            float z_cyl = varying[2];
+            float r = y_cyl*y_cyl+z_cyl*z_cyl;
+
+            // depth clamp
+            if (r>1.0f)
+            {
+                factor_8 = 0;
+            }
+
+            float angleCos = varying[3];
+
+            if (HAS_MASK)
+            {
+                if ( factor_8>0 && m_pMaskData )
+                {
+                    uint8_t maskFactor = m_pMaskData[ (pBufferPos-m_pTargetData) / PIXEL_SIZE ];
+                    factor_8 = maskFactor;
+                }
+
+                if ( factor_8>0 )
+                {
+                    if ( angleCos<m_fadeStartCos && angleCos>m_fadeEndCos )
+                    {
+                        factor_8 = ( factor_8
+                                     *
+                                     int32_t( 256.0f * ( angleCos - m_fadeEndCos ) / ( m_fadeStartCos - m_fadeEndCos ) )
+                                     ) >> 8;
+                    }
+                    else if ( angleCos<=m_fadeEndCos )
+                    {
+                        factor_8 = 0;
+                    }
+                }
+            }
+            else
+            {
+                if ( angleCos<m_fadeStartCos && angleCos>m_fadeEndCos )
+                {
+                    factor_8 = int32_t( 256.0f * ( angleCos - m_fadeEndCos ) / ( m_fadeStartCos - m_fadeEndCos ) );
+                }
+                else if ( angleCos<=m_fadeEndCos )
+                {
+                    factor_8 = 0;
+                }
+            }
+
+
+            if ( factor_8>0 )
+            {
+                // Project
+                float u = 0.5f + atan2f(z_cyl,-y_cyl) / m_projectionAngle;
+                float v = x_cyl;
+
+                if ( u>=0 && u<1 && v>=0 && v<1 )
+                {
+
+                    const uint8_t* pPixel = m_pSourceData;
+                    // TODO: clamp?
+                    pPixel += ( m_sourceSizeX * int(m_sourceSizeY * v)
+                                + int( m_sourceSizeX * u ) )
+                            * PIXEL_SIZE;
+
+                    // Write result
+                    for ( int i=0; i<PIXEL_SIZE; ++i )
+                    {
+                        pBufferPos[i] = uint8_t( ( pPixel[i] * factor_8 ) >> 8 );
+                    }
+                }
+            }
+        }
+
+
+        const uint8_t* m_pSourceData;
+        int m_sourceSizeX, m_sourceSizeY;
+
+        //! Cosine of the fading angle range
+        float m_fadeStartCos, m_fadeEndCos;
+
+        float m_projectionAngle;
+
+        const uint8_t* m_pTargetData;
+        const uint8_t* m_pMaskData;
+    };
+
+
+} // anon namespace
+
+
+
+//-------------------------------------------------------------------------------------------------
+void ImageRasterProjected_Generic( const Mesh* pMesh,
+                           Image* pImage,
+                           const Image* pSource,
+                           const Image* pMask,
+                           float fadeStart,
+                           float fadeEnd,
+                           int layout,
+                           int block )
+{
+	MUTABLE_CPUPROFILER_SCOPE(ImageRasterProjected_Generic);
+
+	EImageFormat format = pImage->GetFormat();
+    int pixelSize = GetImageFormatData( format ).m_bytesPerBlock;
+
+    int sizeX = pImage->GetSizeX();
+    int sizeY = pImage->GetSizeY();
+
+    RasterProjectedPixelProcessor pixelProc( pSource,
+                                     pImage->GetData(),
+                                     pMask ? pMask->GetData() : 0,
+                                     fadeStart, fadeEnd );
+
+    // Get the vertices
+    int vertexCount = pMesh->GetVertexCount();
+    vector< RasterVertex<4> > vertices( vertexCount );
+
+    UntypedMeshBufferIteratorConst texIt( pMesh->GetVertexBuffers(), MBS_TEXCOORDS, layout );
+    UntypedMeshBufferIteratorConst posIt( pMesh->GetVertexBuffers(), MBS_POSITION, 0 );
+    UntypedMeshBufferIteratorConst norIt( pMesh->GetVertexBuffers(), MBS_NORMAL, 0 );
+    for ( int v=0; v<vertexCount; ++v )
+    {
+        float uv[2]={0.0f,0.0f};
+        ConvertData( 0, uv, MBF_FLOAT32, texIt.ptr(), texIt.GetFormat() );
+        ConvertData( 1, uv, MBF_FLOAT32, texIt.ptr(), texIt.GetFormat() );
+
+        vertices[v].x = uv[0] * sizeX;
+        vertices[v].y = uv[1] * sizeY;
+
+        vec3f pos={0.0f,0.0f,0.0f};
+        ConvertData( 0, &pos[0], MBF_FLOAT32, posIt.ptr(), posIt.GetFormat() );
+        ConvertData( 1, &pos[0], MBF_FLOAT32, posIt.ptr(), posIt.GetFormat() );
+        ConvertData( 2, &pos[0], MBF_FLOAT32, posIt.ptr(), posIt.GetFormat() );
+        vertices[v].interpolators[0] = pos[0];
+        vertices[v].interpolators[1] = pos[1];
+        vertices[v].interpolators[2] = pos[2];
+
+        float nor=0.0f;
+        ConvertData( 0, &nor, MBF_FLOAT32, norIt.ptr(), norIt.GetFormat() );
+        vertices[v].interpolators[3] = nor;
+
+        ++texIt;
+        ++posIt;
+        ++norIt;
+    }
+
+    // Get the indices
+    int faceCount = pMesh->GetFaceCount();
+    vector<int> indices( faceCount*3 );
+
+    UntypedMeshBufferIteratorConst indIt( pMesh->GetIndexBuffers(), MBS_VERTEXINDEX, 0 );
+    for ( int i=0; i<faceCount*3; ++i )
+    {
+        uint32_t index=0;
+        ConvertData( 0, &index, MBF_UINT32, indIt.ptr(), indIt.GetFormat() );
+
+        indices[i] = index;
+        ++indIt;
+    }
+
+    // Get the block per face
+    vector<int> blocks( vertexCount );
+
+    UntypedMeshBufferIteratorConst bloIt( pMesh->GetVertexBuffers(), MBS_LAYOUTBLOCK, layout );
+    if (bloIt.ptr())
+    {
+        for ( int i=0; i<vertexCount; ++i )
+        {
+            uint16_t index=0;
+            ConvertData( 0, &index, MBF_UINT16, bloIt.ptr(), bloIt.GetFormat() );
+
+            blocks[i] = index;
+            ++bloIt;
+        }
+    }
+
+    if ( block<0 )
+    {
+        // We are rastering the entire layout in an image, so we need to apply the block->layout
+        // transform to the UVs
+        const Layout* pLayout = layout<pMesh->GetLayoutCount() ? pMesh->GetLayout( layout ) : nullptr;
+        if (pLayout && bloIt.ptr())
+        {
+            vector< box< vec2<float> > > transforms( pLayout->GetBlockCount() );
+            for ( int b=0; b<pLayout->GetBlockCount(); ++b )
+            {
+                FIntPoint grid = pLayout->GetGridSize();
+
+                box< vec2<int> > blockRect;
+                pLayout->GetBlock( b, &blockRect.min[0], &blockRect.min[1], &blockRect.size[0], &blockRect.size[1] );
+
+                box< vec2<float> > rect;
+                rect.min[0] = ( (float)blockRect.min[0] ) / (float) grid[0];
+                rect.min[1] = ( (float)blockRect.min[1] ) / (float) grid[1];
+                rect.size[0] = ( (float)blockRect.size[0] ) / (float) grid[0];
+                rect.size[1] = ( (float)blockRect.size[1] ) / (float) grid[1];
+                transforms[b] = rect;
+            }
+
+            for ( int i=0; i<vertexCount; ++i )
+            {
+                int relBlock = pLayout->FindBlock( blocks[i] );
+                vertices[i].x = vertices[i].x * transforms[ relBlock ].size[0]
+                        + transforms[ relBlock ].min[0] * sizeX;
+                vertices[i].y = vertices[i].y * transforms[ relBlock ].size[1]
+                        + transforms[ relBlock ].min[1] * sizeY;
+            }
+        }
+
+        // Raster all the faces
+        for ( int f=0; f<faceCount; ++f )
+        {
+            Triangle( pImage->GetData(),
+                      sizeX, sizeY,
+                      pixelSize,
+                      vertices[indices[f*3+0]],
+                      vertices[indices[f*3+1]],
+                      vertices[indices[f*3+2]],
+                      pixelProc,
+                      false );
+        }
+    }
+    else
+    {
+        // Raster only the faces in the selected block
+        for ( int f=0; f<faceCount; ++f )
+        {
+            if ( blocks[ indices[f*3+0] ]==block )
+            {
+                Triangle( pImage->GetData(),
+                          sizeX, sizeY,
+                          pixelSize,
+                          vertices[indices[f*3+0]],
+                          vertices[indices[f*3+1]],
+                          vertices[indices[f*3+2]],
+                          pixelProc,
+                          false );
+            }
+        }
+    }
+}
+
+
+//-------------------------------------------------------------------------------------------------
+//! This format is the one we assumee the meshes optimised for planar and cylindrical projection
+//! will have.
+//! See CreateMeshOptimisedForProjection
+struct OPTIMISED_VERTEX
+{
+    vec2<float> uv;
+    vec3<float> pos;
+    vec3<float> nor;
+};
+
+static_assert( sizeof(OPTIMISED_VERTEX)==32, "UNEXPECTED_STRUCT_SIZE" );
+
+
+//-------------------------------------------------------------------------------------------------
+//! This format is the one we assumee the meshes optimised for wrapping projection will have.
+//! See CreateMeshOptimisedForWrappingProjection
+struct OPTIMISED_VERTEX_WRAPPING
+{
+    vec2<float> uv;
+    vec3<float> pos;
+    vec3<float> nor;
+    uint32_t layoutBlock;
+};
+
+static_assert( sizeof(OPTIMISED_VERTEX_WRAPPING)==36, "UNEXPECTED_STRUCT_SIZE" );
+
+
+//-------------------------------------------------------------------------------------------------
+template<class PIXEL_PROCESSOR>
+void ImageRasterProjected_Optimised( const Mesh* pMesh,
+                             Image* pImage,
+                             PIXEL_PROCESSOR& pixelProc,
+                             float fadeEnd,
+                             SCRATCH_IMAGE_PROJECT* scratch )
+{
+	MUTABLE_CPUPROFILER_SCOPE(ImageRasterProjected_Optimised);
+
+    if (!pMesh || !pMesh->GetFaceCount())
+    {
+        return;
+    }
+
+	EImageFormat format = pImage->GetFormat();
+    int pixelSize = GetImageFormatData( format ).m_bytesPerBlock;
+
+    int sizeX = pImage->GetSizeX();
+    int sizeY = pImage->GetSizeY();
+
+    // Get the vertices
+    int vertexCount = pMesh->GetVertexCount();
+
+    check( (int)scratch->vertices.size()==vertexCount );
+    check( (int)scratch->culledVertex.size()==vertexCount );
+
+    check( pMesh->GetVertexBuffers().GetElementSize(0)==sizeof(OPTIMISED_VERTEX) );
+    auto pVertices = reinterpret_cast<const OPTIMISED_VERTEX*>( pMesh->GetVertexBuffers().GetBufferData(0) );
+
+    float fadeEndCos = cosf(fadeEnd);
+    for ( int v=0; v<vertexCount; ++v )
+    {
+        scratch->vertices[v].x = pVertices[v].uv[0] * sizeX;
+        scratch->vertices[v].y = pVertices[v].uv[1] * sizeY;
+
+        // TODO: No need to copy all. use scratch for the rest only.
+        scratch->vertices[v].interpolators[0] = pVertices[v].pos[0];
+        scratch->vertices[v].interpolators[1] = pVertices[v].pos[1];
+        scratch->vertices[v].interpolators[2] = pVertices[v].pos[2];
+        scratch->vertices[v].interpolators[3] = pVertices[v].nor[0];
+        scratch->culledVertex[v] = pVertices[v].nor[0] < fadeEndCos;
+    }
+
+    // Get the indices
+    check( pMesh->GetIndexBuffers().GetElementSize(0)==4 );
+    auto pIndices = reinterpret_cast<const uint32_t*>( pMesh->GetIndexBuffers().GetBufferData(0) );
+
+
+    // The mesh is supposed to contain only the faces in the selected block
+    int faceCount = pMesh->GetFaceCount();
+    //for ( int f=0; f<faceCount; ++f )
+	ParallelFor(faceCount, [pIndices, scratch, sizeX, sizeY, pixelSize, pImage, pixelProc](int f)
+		{
+			int i0 = pIndices[f * 3 + 0];
+			int i1 = pIndices[f * 3 + 1];
+			int i2 = pIndices[f * 3 + 2];
+
+			// TODO: This optimisation may remove projection in the center of the face, if the angle
+			// range is small. Make it optional or more sophisticated (cross product may help).
+			if (!scratch->culledVertex[i0] ||
+				!scratch->culledVertex[i1] ||
+				!scratch->culledVertex[i2])
+			{
+				Triangle(pImage->GetData(),
+					sizeX, sizeY,
+					pixelSize,
+					scratch->vertices[i0],
+					scratch->vertices[i1],
+					scratch->vertices[i2],
+					pixelProc,
+					false);
+			}
+		});
+
+}
+
+
+//-------------------------------------------------------------------------------------------------
+template<class PIXEL_PROCESSOR>
+void ImageRasterProjected_OptimisedWrapping( const Mesh* pMesh,
+                             Image* pImage,
+                             PIXEL_PROCESSOR& pixelProc,
+                             float fadeEnd,
+                             int block,
+                             SCRATCH_IMAGE_PROJECT* scratch )
+{
+	MUTABLE_CPUPROFILER_SCOPE(ImageRasterProjected_OptimisedWrapping);
+
+    if (!pMesh || !pMesh->GetFaceCount())
+    {
+        return;
+    }
+
+	EImageFormat format = pImage->GetFormat();
+    int pixelSize = GetImageFormatData( format ).m_bytesPerBlock;
+
+    int sizeX = pImage->GetSizeX();
+    int sizeY = pImage->GetSizeY();
+
+    // Get the vertices
+    int vertexCount = pMesh->GetVertexCount();
+
+    check( (int)scratch->vertices.size()==vertexCount );
+    check( (int)scratch->culledVertex.size()==vertexCount );
+
+    check( pMesh->GetVertexBuffers().GetElementSize(0)==sizeof(OPTIMISED_VERTEX_WRAPPING) );
+    auto pVertices = reinterpret_cast<const OPTIMISED_VERTEX_WRAPPING*>( pMesh->GetVertexBuffers().GetBufferData(0) );
+
+    float fadeEndCos = cosf(fadeEnd);
+    for ( int v=0; v<vertexCount; ++v )
+    {
+        scratch->vertices[v].x = pVertices[v].uv[0] * sizeX;
+        scratch->vertices[v].y = pVertices[v].uv[1] * sizeY;
+
+        // TODO: No need to copy all. use scratch for the rest only.
+        scratch->vertices[v].interpolators[0] = pVertices[v].pos[0];
+        scratch->vertices[v].interpolators[1] = pVertices[v].pos[1];
+        scratch->vertices[v].interpolators[2] = pVertices[v].pos[2];
+        scratch->vertices[v].interpolators[3] = pVertices[v].nor[0];
+        scratch->culledVertex[v] = pVertices[v].nor[0] < fadeEndCos;
+
+        // Cull vertices that don't belong to the current layout block.
+        if (pVertices[v].layoutBlock!=uint32_t(block))
+        {
+            scratch->culledVertex[v] = true;
+        }
+    }
+
+    // Get the indices
+    check( pMesh->GetIndexBuffers().GetElementSize(0)==4 );
+    auto pIndices = reinterpret_cast<const uint32_t*>( pMesh->GetIndexBuffers().GetBufferData(0) );
+
+
+    // The mesh is supposed to contain only the faces in the selected block
+    int faceCount = pMesh->GetFaceCount();
+    for ( int f=0; f<faceCount; ++f )
+    {
+        int i0 = pIndices[f*3+0];
+        int i1 = pIndices[f*3+1];
+        int i2 = pIndices[f*3+2];
+
+        // TODO: This optimisation may remove projection in the center of the face, if the angle
+        // range is small. Make it optional or more sophisticated (cross product may help).
+        if ( !scratch->culledVertex[i0] ||
+             !scratch->culledVertex[i1] ||
+             !scratch->culledVertex[i2] )
+        {
+            Triangle( pImage->GetData(),
+                      sizeX, sizeY,
+                      pixelSize,
+                      scratch->vertices[i0],
+                      scratch->vertices[i1],
+                      scratch->vertices[i2],
+                      pixelProc,
+                      false );
+        }
+    }
+
+}
+
+
+//-------------------------------------------------------------------------------------------------
+void mu::ImageRasterProjectedPlanar( const Mesh* pMesh,
+                                     Image* pImage,
+                                     const Image* pSource,
+                                     const Image* pMask,
+                                     float fadeStart,
+                                     float fadeEnd,
+                                     int layout,
+                                     int block,
+                                     SCRATCH_IMAGE_PROJECT* scratch )
+{
+    check( !pMask || pMask->GetSizeX() == pImage->GetSizeX() );
+    check( !pMask || pMask->GetSizeY() == pImage->GetSizeY() );
+    check( !pMask || pMask->GetFormat()==EImageFormat::IF_L_UBYTE );
+
+	MUTABLE_CPUPROFILER_SCOPE(ImageProject);
+
+    if ( ( pMesh->m_staticFormatFlags & (1<<SMF_PROJECT) )
+         )
+    {
+        // Mesh-optimised version
+        if ( pSource->GetFormat() == EImageFormat::IF_RGB_UBYTE )
+        {
+            RasterProjectedPixelProcessor_UBYTE<3> pixelProc(  pSource,
+                                                       pImage->GetData(),
+                                                       pMask ? pMask->GetData() : 0,
+                                                       fadeStart, fadeEnd );
+
+            ImageRasterProjected_Optimised( pMesh, pImage, pixelProc, fadeEnd, scratch );
+        }
+
+        else if ( pSource->GetFormat() == EImageFormat::IF_RGBA_UBYTE ||
+                  pSource->GetFormat() == EImageFormat::IF_BGRA_UBYTE )
+        {
+            if (pMask!=nullptr)
+            {
+                RasterProjectedPixelProcessor_UBYTE<4,true> pixelProc( pSource,
+                                                           pImage->GetData(),
+                                                           pMask->GetData(),
+                                                           fadeStart, fadeEnd );
+                ImageRasterProjected_Optimised( pMesh, pImage, pixelProc, fadeEnd, scratch );
+            }
+            else
+            {
+                RasterProjectedPixelProcessor_UBYTE<4,false> pixelProc( pSource,
+                                                           pImage->GetData(),
+                                                           nullptr,
+                                                           fadeStart, fadeEnd );
+                ImageRasterProjected_Optimised( pMesh, pImage, pixelProc, fadeEnd, scratch );
+            }
+
+        }
+
+        else if ( pSource->GetFormat() == EImageFormat::IF_L_UBYTE )
+        {
+            RasterProjectedPixelProcessor_UBYTE<1> pixelProc( pSource,
+                                                     pImage->GetData(),
+                                                     pMask ? pMask->GetData() : 0,
+                                                     fadeStart, fadeEnd );
+
+            ImageRasterProjected_Optimised( pMesh, pImage, pixelProc, fadeEnd, scratch );
+        }
+
+        else
+        {
+            RasterProjectedPixelProcessor pixelProc( pSource,
+                                             pImage->GetData(),
+                                             pMask ? pMask->GetData() : 0,
+                                             fadeStart, fadeEnd );
+
+            ImageRasterProjected_Optimised( pMesh, pImage, pixelProc, fadeEnd, scratch );
+        }
+    }
+
+    else
+    {
+        check(false);
+        // Generic version
+        // TODO: It will allocate temp memory.
+        ImageRasterProjected_Generic( pMesh, pImage, pSource, pMask,
+                              fadeStart, fadeEnd,
+                              layout, block );
+    }
+
+}
+
+
+//-------------------------------------------------------------------------------------------------
+void mu::ImageRasterProjectedWrapping( const Mesh* pMesh,
+                                       Image* pImage,
+                                       const Image* pSource,
+                                       const Image* pMask,
+                                       float fadeStart,
+                                       float fadeEnd,
+                                       int layout,
+                                       int block,
+                                       SCRATCH_IMAGE_PROJECT* scratch )
+{
+    check( !pMask || pMask->GetSizeX() == pImage->GetSizeX() );
+    check( !pMask || pMask->GetSizeY() == pImage->GetSizeY() );
+    check( !pMask || pMask->GetFormat()==EImageFormat::IF_L_UBYTE );
+
+	MUTABLE_CPUPROFILER_SCOPE(ImageProjectWrapping);
+
+    if ( ( pMesh->m_staticFormatFlags & (1<<SMF_PROJECTWRAPPING) )
+         )
+    {
+        // Mesh-optimised version
+        if ( pSource->GetFormat() == EImageFormat::IF_RGB_UBYTE )
+        {
+            RasterProjectedPixelProcessor_UBYTE<3> pixelProc(  pSource,
+                                                       pImage->GetData(),
+                                                       pMask ? pMask->GetData() : 0,
+                                                       fadeStart, fadeEnd );
+
+            ImageRasterProjected_OptimisedWrapping( pMesh, pImage, pixelProc, fadeEnd, block, scratch );
+        }
+
+        else if ( pSource->GetFormat() == EImageFormat::IF_RGBA_UBYTE ||
+                  pSource->GetFormat() == EImageFormat::IF_BGRA_UBYTE )
+        {
+            if (pMask!=nullptr)
+            {
+                RasterProjectedPixelProcessor_UBYTE<4,true> pixelProc( pSource,
+                                                           pImage->GetData(),
+                                                           pMask->GetData(),
+                                                           fadeStart, fadeEnd );
+                ImageRasterProjected_OptimisedWrapping( pMesh, pImage, pixelProc, fadeEnd, block, scratch );
+            }
+            else
+            {
+                RasterProjectedPixelProcessor_UBYTE<4,false> pixelProc( pSource,
+                                                           pImage->GetData(),
+                                                           nullptr,
+                                                           fadeStart, fadeEnd );
+                ImageRasterProjected_OptimisedWrapping( pMesh, pImage, pixelProc, fadeEnd, block, scratch );
+            }
+
+        }
+
+        else if ( pSource->GetFormat() == EImageFormat::IF_L_UBYTE )
+        {
+            RasterProjectedPixelProcessor_UBYTE<1> pixelProc( pSource,
+                                                     pImage->GetData(),
+                                                     pMask ? pMask->GetData() : 0,
+                                                     fadeStart, fadeEnd );
+
+            ImageRasterProjected_OptimisedWrapping( pMesh, pImage, pixelProc, fadeEnd, block, scratch );
+        }
+
+        else
+        {
+            RasterProjectedPixelProcessor pixelProc( pSource,
+                                             pImage->GetData(),
+                                             pMask ? pMask->GetData() : 0,
+                                             fadeStart, fadeEnd );
+
+            ImageRasterProjected_OptimisedWrapping( pMesh, pImage, pixelProc, fadeEnd, block, scratch );
+        }
+    }
+
+    else
+    {
+        check(false);
+        // Generic version
+        // TODO: It will allocate temp memory.
+        ImageRasterProjected_Generic( pMesh, pImage, pSource, pMask,
+                              fadeStart, fadeEnd,
+                              layout, block );
+    }
+
+}
+
+
+//-------------------------------------------------------------------------------------------------
+void mu::ImageRasterProjectedCylindrical( const Mesh* pMesh,
+                                          Image* pImage,
+                                          const Image* pSource,
+                                          const Image* pMask,
+                                          float fadeStart,
+                                          float fadeEnd,
+                                          int /*layout*/,
+                                          float projectionAngle,
+                                          SCRATCH_IMAGE_PROJECT* scratch )
+{
+    check( !pMask || pMask->GetSizeX() == pImage->GetSizeX() );
+    check( !pMask || pMask->GetSizeY() == pImage->GetSizeY() );
+    check( !pMask || pMask->GetFormat()==EImageFormat::IF_L_UBYTE );
+
+	MUTABLE_CPUPROFILER_SCOPE(ImageProjectCylindrical);
+
+    if ( ( pMesh->m_staticFormatFlags & (1<<SMF_PROJECT) )
+         )
+    {
+        // Mesh-optimised version
+        if ( pSource->GetFormat() == EImageFormat::IF_RGB_UBYTE )
+        {
+            RasterCylindricalProjectedPixelProcessor_UBYTE<3> pixelProc(  pSource,
+                                                       pImage->GetData(),
+                                                       pMask ? pMask->GetData() : 0,
+                                                       fadeStart, fadeEnd,
+                                                                          projectionAngle);
+
+            ImageRasterProjected_Optimised( pMesh, pImage, pixelProc, fadeEnd, scratch );
+        }
+
+        else if ( pSource->GetFormat() == EImageFormat::IF_RGBA_UBYTE ||
+                  pSource->GetFormat() == EImageFormat::IF_BGRA_UBYTE )
+        {
+            if (pMask!=nullptr)
+            {
+                RasterCylindricalProjectedPixelProcessor_UBYTE<4,true> pixelProc( pSource,
+                                                           pImage->GetData(),
+                                                           pMask->GetData(),
+                                                           fadeStart, fadeEnd,
+                                                                                  projectionAngle );
+                ImageRasterProjected_Optimised( pMesh, pImage, pixelProc, fadeEnd, scratch );
+            }
+            else
+            {
+                RasterCylindricalProjectedPixelProcessor_UBYTE<4,false> pixelProc( pSource,
+                                                           pImage->GetData(),
+                                                           nullptr,
+                                                           fadeStart, fadeEnd,
+                                                                                   projectionAngle );
+                ImageRasterProjected_Optimised( pMesh, pImage, pixelProc, fadeEnd, scratch );
+            }
+
+        }
+
+        else if ( pSource->GetFormat() == EImageFormat::IF_L_UBYTE )
+        {
+            RasterCylindricalProjectedPixelProcessor_UBYTE<1> pixelProc( pSource,
+                                                     pImage->GetData(),
+                                                     pMask ? pMask->GetData() : 0,
+                                                     fadeStart, fadeEnd,
+                                                                         projectionAngle );
+
+            ImageRasterProjected_Optimised( pMesh, pImage, pixelProc, fadeEnd, scratch );
+        }
+
+        else
+        {
+            check(false);
+//            RasterCylindricalProjectedPixelProcessor pixelProc( pSource,
+//                                             pImage->GetData(),
+//                                             pMask ? pMask->GetData() : 0,
+//                                             fadeStart, fadeEnd,
+//                                                                projectionAngle );
+
+//            ImageRasterProjected_Optimised( pMesh, pImage, pixelProc, fadeEnd, block, scratch );
+        }
+    }
+    else
+    {
+        check(false);
+        // Generic version
+        // TODO: It will allocate temp memory.
+//        ImageRasterCylindricalProjected_Generic( pMesh, pImage, pSource, pMask,
+//                              fadeStart, fadeEnd,
+//                                                 projectionAngle,
+//                              layout, block );
+    }
+
+}
+
+
+//#define DEBUG_PROJECTION 1
+
+#ifdef DEBUG_PROJECTION
+PRAGMA_DISABLE_OPTIMIZATION
+
+#undef assert
+int* assert_aux = 0;
+#define assert(x) if((x) == 0) assert_aux[0] = 1;
+#endif
+
+//---------------------------------------------------------------------------------------------
+//! Create a map from vertices into vertices, collapsing vertices that have the same position
+//---------------------------------------------------------------------------------------------
+inline void MeshCreateCollapsedVertexMap( const Mesh* pMesh,
+                                    mu::vector<int>& collapsedVertexMap,
+                                    mu::vector<vec3f>& vertices,
+									std::multimap<int, int>& collapsedVertsMap
+                                    )
+{
+	MUTABLE_CPUPROFILER_SCOPE(CreateCollapseMap);
+
+    int vcount = pMesh->GetVertexCount();
+    collapsedVertexMap.resize(vcount);
+    vertices.resize(vcount);
+
+    const FMeshBufferSet& MBSPriv = pMesh->GetVertexBuffers();
+    for (int32 b = 0; b < MBSPriv.m_buffers.Num(); ++b)
+    {
+        for (int32 c = 0; c < MBSPriv.m_buffers[b].m_channels.Num(); ++c)
+        {
+            MESH_BUFFER_SEMANTIC sem = MBSPriv.m_buffers[b].m_channels[c].m_semantic;
+            int semIndex = MBSPriv.m_buffers[b].m_channels[c].m_semanticIndex;
+
+            auto it = UntypedMeshBufferIteratorConst(pMesh->GetVertexBuffers(), sem, semIndex);
+
+			box<vec3f> meshBoundingBox;
+			Octree* pOctree = nullptr;
+
+            switch (sem)
+            {
+            case MBS_POSITION:
+
+                // First create a cache of the vertices in the vertices array
+                for (int v = 0; v < vcount; ++v)
+                {
+                    vec3f vertex(0.0f, 0.0f, 0.0f);
+                    for (int i = 0; i < 3; ++i)
+                    {
+                        ConvertData(i, &vertex[0], MBF_FLOAT32, it.ptr(), it.GetFormat());
+                    }
+
+                    vertices[v] = vertex;
+
+                    ++it;
+                }
+				
+				//Initializing Octree boundingbox
+				meshBoundingBox.min = vertices[0];
+				meshBoundingBox.size = vec3f(0.0f, 0.0f, 0.0f);
+				for (size_t i = 1; i < vertices.size(); ++i)
+				{
+					meshBoundingBox.Bound(vertices[i]);
+				}
+
+				//Initializing and filling Octree
+				pOctree = new Octree(meshBoundingBox.min, meshBoundingBox.min + meshBoundingBox.size,0.001f);
+				for (size_t itr = 0; itr < vertices.size(); ++itr)
+				{
+					pOctree->InsertElement(vertices[itr], int(itr));
+				}
+
+				// Create map to store which vertices are the same (collapse nearby vertices)
+				for (size_t itr = 0; itr < vertices.size(); ++itr)
+				{
+					int collapsed_candidate_v_idx = pOctree->GetNearests(vertices[itr], vert_collapse_eps, int(itr));
+					collapsedVertexMap[itr] = collapsed_candidate_v_idx;
+
+					if (collapsed_candidate_v_idx != int(itr))
+					{
+						collapsedVertsMap.insert(std::make_pair<>(collapsed_candidate_v_idx, (int)itr));
+					}
+				}
+                break;
+
+            default:
+                break;
+            }
+
+			if (pOctree)
+			{
+				delete pOctree;
+				pOctree = nullptr;
+			}
+        }
+    }
+}
+
+
+struct AdjacentFaces
+{
+	int faces[3];
+	int newVertices[3];
+	bool changesUVIsland[3];
+
+	AdjacentFaces()
+	{
+		faces[0] = -1;
+		faces[1] = -1;
+		faces[2] = -1;
+
+		newVertices[0] = -1;
+		newVertices[1] = -1;
+		newVertices[2] = -1;
+
+		changesUVIsland[0] = false;
+		changesUVIsland[1] = false;
+		changesUVIsland[2] = false;
+	}
+
+	void addConnectedFace(int newConnectedFace, int newVertex, bool changesUVIslandParam)
+	{
+#ifdef DEBUG_PROJECTION
+		assert( newConnectedFace >= 0 );
+#endif
+
+		for (int i = 0; i < 3; ++i)
+		{
+			if (faces[i] == newConnectedFace)
+			{
+				return;
+			}
+		}
+
+		for (int i = 0; i < 3; ++i)
+		{
+			if (faces[i] == -1)
+			{
+				faces[i] = newConnectedFace;
+				newVertices[i] = newVertex;
+				changesUVIsland[i] = changesUVIslandParam;
+				break;
+			}
+		}
+	}
+};
+
+
+void PlanarlyProjectVertex(const vec3f& unfoldedPosition, vec4f& projectedPosition,
+	const PROJECTOR& projector, const vec3f& projectorPosition, const vec3f& projectorDirection, const vec3f& s, const vec3f& u)
+{
+	float x = dot( unfoldedPosition - projectorPosition, s ) / projector.scale[0] + 0.5f;
+	float y = dot( unfoldedPosition - projectorPosition, u ) / projector.scale[1] + 0.5f;
+	y = 1.0f - y;
+	float z = dot( unfoldedPosition - projectorPosition, projectorDirection ) / projector.scale[2];
+
+	bool inside = x>=0.0f && x<=1.0f && y>=0.0f && y<=1.0 && z>=0.0f && z<=1.0f;
+	projectedPosition[0] = x;
+	projectedPosition[1] = y;
+	projectedPosition[2] = z;
+	projectedPosition[3] = inside ? 1.0f : 0.0f;
+}
+
+
+vec2f ChangeBase2D(const vec2f& origPosition, const vec2f& newOrigin, const vec2f& newBaseX, const vec2f& newBaseY)
+{
+	float x = dot(origPosition - newOrigin, newBaseX) / powf(length(newBaseX), 2.f) + 0.5f;
+	float y = dot(origPosition - newOrigin, newBaseY) / powf(length(newBaseY), 2.f) + 0.5f;
+	//y = 1.0f - y;
+
+	return vec2f(x, y);
+}
+
+
+// Compute barycentric coordinates (u, v, w) for
+// point p with respect to triangle (a, b, c)
+void GetBarycentricCoords(vec3f p, vec3f a, vec3f b, vec3f c, float &u, float &v, float &w)
+{
+    vec3f v0 = b - a, v1 = c - a, v2 = p - a;
+    float d00 = dot(v0, v0);
+    float d01 = dot(v0, v1);
+    float d11 = dot(v1, v1);
+    float d20 = dot(v2, v0);
+    float d21 = dot(v2, v1);
+    float denom = d00 * d11 - d01 * d01;
+    v = (d11 * d20 - d01 * d21) / denom;
+    w = (d00 * d21 - d01 * d20) / denom;
+    u = 1.0f - v - w;
+}
+
+void GetBarycentricCoords(vec2f p, vec2f a, vec2f b, vec2f c, float &u, float &v, float &w)
+{
+	vec2f v0 = b - a, v1 = c - a, v2 = p - a;
+    float den = v0.x() * v1.y() - v1.x() * v0.y();
+#ifdef DEBUG_PROJECTION
+	assert(den != 0.f);
+#endif
+    v = (v2.x() * v1.y() - v1.x() * v2.y()) / den;
+    w = (v0.x() * v2.y() - v2.x() * v0.y()) / den;
+    u = 1.0f - v - w;
+}
+
+
+float getTriangleArea(vec2f& a, vec2f& b, vec2f& c)
+{
+	float area = a.x() * (b.y() - c.y()) + b.x() * (c.y() - a.y()) + c.x() * (a.y() - b.y());
+	return area / 2.f;
+}
+
+
+float getTriangleRatio(const vec2f& a, const vec2f& b, const vec2f& c)
+{
+	float lenSide1 = length(b - a);
+	float lenSide2 = length(c - a);
+	float lenSide3 = length(b - c);
+
+	float maxLen = max(max(lenSide1, lenSide2), lenSide3);
+	float minLen = min(min(lenSide1, lenSide2), lenSide3);
+
+	return maxLen / minLen;
+}
+
+struct NeighborFace
+{
+	int neighborFace;
+	int newVertex;
+	int previousFace;
+	int numUVIslandChanges = 0;
+	int step;
+	bool changesUVIsland;
+
+	friend bool operator <(const NeighborFace& a, const NeighborFace& b)
+	{
+		//if (a.numUVIslandChanges == b.numUVIslandChanges)
+		//{
+			return a.step > b.step;
+		//}
+
+		//return a.numUVIslandChanges > b.numUVIslandChanges;
+	}
+};
+
+
+void getEdgeHorizontalLength( int edgeVert0, int edgeVert1, int opposedVert,
+                              const OPTIMISED_VERTEX_WRAPPING* pVertices,
+                              float &out_uvSpaceLen, float &out_objSpaceLen,
+                              float& out_midEdgePointFraction)
+{
+	const int oldVertices[2] = { edgeVert0, edgeVert1 };
+
+	const vec2f& edge0_oldUVSpace = pVertices[oldVertices[0]].uv;
+	const vec2f& edge1_oldUVSpace = pVertices[oldVertices[1]].uv;
+	const vec2f& opposedVert_oldUVSpace = pVertices[opposedVert].uv;
+
+	vec2f edgeVector_oldUVSpace = edge1_oldUVSpace - edge0_oldUVSpace;
+	vec2f sideVector_oldUVSpace = opposedVert_oldUVSpace - edge0_oldUVSpace;
+
+	float edgeVector_oldUVSpaceLen = length(edgeVector_oldUVSpace);
+	float dotEdgeSideVectors_oldUVSpace = dot(edgeVector_oldUVSpace, sideVector_oldUVSpace);
+	float midEdgePointFraction_oldUVSpace = (dotEdgeSideVectors_oldUVSpace / powf(edgeVector_oldUVSpaceLen, 2));
+	vec2f edgeVectorProj_oldUVSpace = edgeVector_oldUVSpace * midEdgePointFraction_oldUVSpace;
+	vec2f midEdgePoint_oldUVSpace = edge0_oldUVSpace + edgeVectorProj_oldUVSpace;
+
+	vec2f perpEdgeVector_oldUVSpace = opposedVert_oldUVSpace - midEdgePoint_oldUVSpace;
+#ifdef DEBUG_PROJECTION
+	assert(fabs(dot(perpEdgeVector_oldUVSpace, edgeVector_oldUVSpace)) < 0.1f);
+#endif
+	float perpEdgeVector_oldUVSpace_Len = length(perpEdgeVector_oldUVSpace);
+	out_uvSpaceLen = perpEdgeVector_oldUVSpace_Len;
+	out_midEdgePointFraction = midEdgePointFraction_oldUVSpace;
+
+	// Do the same in object space to be able to extract the scale of the original uv space
+	const vec3f& edge0_objSpace = pVertices[oldVertices[0]].pos;
+	const vec3f& edge1_objSpace = pVertices[oldVertices[1]].pos;
+	const vec3f& opposedVert_objSpace = pVertices[opposedVert].pos;
+
+	vec3f edgeVector_objSpace = edge1_objSpace - edge0_objSpace;
+	vec3f sideVector_objSpace = opposedVert_objSpace - edge0_objSpace;
+
+	float edgeVector_objSpaceLen = length(edgeVector_objSpace);
+	float dotEdgeSideVectors_objSpace = dot(edgeVector_objSpace, sideVector_objSpace);
+	float midEdgePointFraction_objSpace = (dotEdgeSideVectors_objSpace / powf(edgeVector_objSpaceLen, 2));
+	vec3f edgeVectorProj_objSpace = edgeVector_objSpace * midEdgePointFraction_objSpace;
+	vec3f midEdgePoint_objSpace = edge0_objSpace + edgeVectorProj_objSpace;
+
+	vec3f perpEdgeVector_objSpace = opposedVert_objSpace - midEdgePoint_objSpace;
+#ifdef DEBUG_PROJECTION
+	assert(fabs(dot(perpEdgeVector_objSpace, edgeVector_objSpace)) < 0.1f);
+#endif
+	float perpEdgeVector_objSpace_Len = length(perpEdgeVector_objSpace);
+	out_objSpaceLen = perpEdgeVector_objSpace_Len;
+}
+
+
+bool testPointsAreInOppositeSidesOfEdge( const vec2f& pointA, const vec2f& pointB,
+                                         const vec2f& edge0, const vec2f& edge1 )
+{
+	vec2f edge = edge1 - edge0;
+	vec2f perp_edge = vec2f(-edge.y(), edge.x());
+	vec2f AVertexVector = pointA - edge0;
+	vec2f BVertexVector = pointB - edge0;
+	float dotAVertexVector = dot(AVertexVector, perp_edge);
+	float dotBVertexVector = dot(BVertexVector, perp_edge);
+	
+	return dotAVertexVector * dotBVertexVector < 0.f;
+}
+
+
+//-------------------------------------------------------------------------------------------------
+#pragma pack(push,1)
+struct PROJECTED_VERTEX
+{   
+    float pos0 = 0, pos1 = 0, pos2 = 0;
+    uint32_t mask3 = 0;
+
+	inline vec2<float> xy() const { return vec2<float>(pos0,pos1); }
+};
+#pragma pack(pop)
+static_assert( sizeof(PROJECTED_VERTEX)==16, "Unexpected struct size" );
+
+
+//-------------------------------------------------------------------------------------------------
+void MeshProject_Optimised_Planar( const OPTIMISED_VERTEX* pVertices, int vertexCount,
+                                   const uint32_t* pIndices, int faceCount,
+                                   const vec3f& projectorPosition, const vec3f& projectorDirection,
+                                   const vec3f& projectorSide, const vec3f& projectorUp,
+                                   const vec3f& projectorScale,
+                                   OPTIMISED_VERTEX* pResultVertices, int& currentVertex,
+                                   uint32_t* pResultIndices, int& currentIndex  )
+{
+	MUTABLE_CPUPROFILER_SCOPE(MeshProject_Optimised_Planar)
+
+    vector<int32_t> oldToNewVertex(vertexCount,-1);
+
+    vector<PROJECTED_VERTEX> projectedPositions(vertexCount);
+
+    for ( int v=0; v<vertexCount; ++v )
+    {
+        float x = dot( pVertices[v].pos-projectorPosition, projectorSide ) / projectorScale[0] + 0.5f;
+        float y = dot( pVertices[v].pos-projectorPosition, projectorUp ) / projectorScale[1] + 0.5f;
+        y = 1.0f-y;
+        float z = dot( pVertices[v].pos-projectorPosition, projectorDirection ) / projectorScale[2];
+
+        // Plane mask with bits for each plane discarding the vertex
+        uint32_t planeMask =
+                ((x<0.0f)<<0) |
+                ((x>1.0f)<<1) |
+                ((y<0.0f)<<2) |
+                ((y>1.0f)<<3) |
+                ((z<0.0f)<<4) |
+                ((z>1.0f)<<5);
+        projectedPositions[v].pos0 = x;
+        projectedPositions[v].pos1 = y;
+        projectedPositions[v].pos2 = z;
+        projectedPositions[v].mask3 = planeMask;
+    }
+
+    // Iterate the faces
+    for ( int f=0; f<faceCount; ++f )
+    {
+        int i0 = pIndices[f*3+0];
+        int i1 = pIndices[f*3+1];
+        int i2 = pIndices[f*3+2];
+
+        // Approximate test: discard the triangle if any of the 6 planes entirely discards all the vertices.
+        // This will let some triangles through that could be discarded with a precise test.
+        bool discarded = (
+                projectedPositions[i0].mask3 &
+                projectedPositions[i1].mask3 &
+                projectedPositions[i2].mask3 )
+                !=
+                0;
+
+        if ( !discarded )
+        {
+            // This face is required.
+            for (int v=0;v<3;++v)
+            {
+                uint32_t i = pIndices[f*3+v];
+                if (oldToNewVertex[i]<0)
+                {
+                    pResultVertices[currentVertex] = pVertices[i];
+
+                    pResultVertices[currentVertex].pos[0] = projectedPositions[i].pos0;
+                    pResultVertices[currentVertex].pos[1] = projectedPositions[i].pos1;
+                    pResultVertices[currentVertex].pos[2] = projectedPositions[i].pos2;
+
+                    // Normal is actually the fade factor
+                    float angleCos = dot( pVertices[i].nor, projectorDirection * -1.0f );
+                    pResultVertices[currentVertex].nor[0] = angleCos;
+                    pResultVertices[currentVertex].nor[1] = angleCos;
+                    pResultVertices[currentVertex].nor[2] = angleCos;
+
+                    oldToNewVertex[i] = currentVertex++;
+                }
+
+                pResultIndices[currentIndex++] = oldToNewVertex[i];
+            }
+        }
+    }
+}
+
+
+//-------------------------------------------------------------------------------------------------
+void MeshProject_Optimised_Cylindrical( const OPTIMISED_VERTEX* pVertices, int vertexCount,
+                                        const uint32_t* pIndices, int faceCount,
+                                        const vec3f& projectorPosition, const vec3f& projectorDirection,
+                                        const vec3f& projectorSide, const vec3f& projectorUp,
+                                        const vec3f& projectorScale,
+                                        OPTIMISED_VERTEX* pResultVertices, int& currentVertex,
+                                        uint32_t* pResultIndices, int& currentIndex  )
+{
+	MUTABLE_CPUPROFILER_SCOPE(MeshProject_Optimised_Cylindrical)
+
+	vector<int32_t> oldToNewVertex(vertexCount,-1);
+
+    vector<PROJECTED_VERTEX> projectedPositions(vertexCount);
+
+    // TODO: support for non uniform scale?
+    float radius = projectorScale[1];
+    float height = projectorScale[0];
+    mat3f worldToCylinder( projectorDirection, projectorSide, projectorUp );
+    //worldToCylinder = worldToCylinder.GetTransposed();
+
+    for ( int v=0; v<vertexCount; ++v )
+    {
+        // Cylinder is along the X axis
+
+        // Project
+        vec3f relPos = pVertices[v].pos - projectorPosition;
+        vec3f vertexPos_cylinder = worldToCylinder * relPos;
+
+        // This final projection needs to be done per pixel
+        float x = vertexPos_cylinder.x() / height;
+        float r2 = vertexPos_cylinder.y()*vertexPos_cylinder.y()
+                + vertexPos_cylinder.z()*vertexPos_cylinder.z();
+        projectedPositions[v].pos0 = x;
+        projectedPositions[v].pos1 = vertexPos_cylinder.y() / radius;
+        projectedPositions[v].pos2 = vertexPos_cylinder.z() / radius;
+        uint32_t planeMask =
+                ((x<0.0f)<<0) |
+                ((x>1.0f)<<1) |
+                ((r2>=(radius*radius))<<2);
+        projectedPositions[v].mask3 = planeMask;
+    }
+
+    // Iterate the faces
+    for ( int f=0; f<faceCount; ++f )
+    {
+        int i0 = pIndices[f*3+0];
+        int i1 = pIndices[f*3+1];
+        int i2 = pIndices[f*3+2];
+
+        // Approximate test: discard the triangle if any of the 3 "planes" (top, bottom and radius)
+        // entirely discards all the vertices.
+        // This will let some triangles through that could be discarded with a precise test.
+        // Also, some big triangles can be wrongly discarded by the radius test, since it is not
+        // really a plane.
+        bool discarded = (
+                projectedPositions[i0].mask3 &
+                projectedPositions[i1].mask3 &
+                projectedPositions[i2].mask3 )
+                !=
+                0;
+
+        if ( !discarded )
+        {
+            // This face is required.
+            for (int v=0;v<3;++v)
+            {
+                uint32_t i = pIndices[f*3+v];
+                if (oldToNewVertex[i]<0)
+                {
+                    pResultVertices[currentVertex] = pVertices[i];
+
+                    pResultVertices[currentVertex].pos[0] = projectedPositions[i].pos0;
+                    pResultVertices[currentVertex].pos[1] = projectedPositions[i].pos1;
+                    pResultVertices[currentVertex].pos[2] = projectedPositions[i].pos2;
+
+                    // Normal is actually the fade factor
+                    //vec3 vertexCylinderDirection =
+                    float angleCos = 1.0f; //dot( pVertices[i].nor, vertexCylinderDirection * -1.0f );
+                    pResultVertices[currentVertex].nor[0] = angleCos;
+                    pResultVertices[currentVertex].nor[1] = angleCos;
+                    pResultVertices[currentVertex].nor[2] = angleCos;
+
+                    oldToNewVertex[i] = currentVertex++;
+                }
+
+                pResultIndices[currentIndex++] = oldToNewVertex[i];
+            }
+        }
+    }
+}
+
+
+//-------------------------------------------------------------------------------------------------
+void MeshProject_Optimised_Wrapping( const Mesh* pMesh,
+                                     const vec3f& projectorPosition, const vec3f& projectorDirection,
+                                     const vec3f& projectorSide, const vec3f& projectorUp,
+                                     const vec3f& projectorScale,
+                                     OPTIMISED_VERTEX_WRAPPING* pResultVertices, int& currentVertex,
+                                     uint32_t* pResultIndices, int& currentIndex )
+{
+	MUTABLE_CPUPROFILER_SCOPE(MeshProject_Optimised_Wrapping)
+
+    // Get the vertices
+    int vertexCount = pMesh->GetVertexCount();
+    check( pMesh->GetVertexBuffers().GetElementSize(0)==sizeof(OPTIMISED_VERTEX_WRAPPING) );
+    auto pVertices = reinterpret_cast<const OPTIMISED_VERTEX_WRAPPING*>( pMesh->GetVertexBuffers().GetBufferData(0) );
+
+    // Get the indices
+    check( pMesh->GetIndexBuffers().GetElementSize(0)==4 );
+    auto pIndices = reinterpret_cast<const uint32_t*>( pMesh->GetIndexBuffers().GetBufferData(0) );
+    int faceCount = pMesh->GetFaceCount();
+
+    vector<PROJECTED_VERTEX> projectedPositions(vertexCount);
+
+    // Iterate the faces and trace a ray to find the origin face of the projection
+    const float maxDist = 100000.f;
+    float min_t = maxDist;
+    float rayLength = maxDist;
+    int intersectedFace = -1;
+    std::set<int> processedVertices;
+    vec3f projectionPlaneNormal;
+    vec3f out_intersection;
+
+    vector<AdjacentFaces> faceConnectivity(faceCount);
+    std::set<int> processedFaces;
+    std::set<int> discardedWrapAroundFaces;
+    vector<int> faceStep(faceCount, 0);
+
+    // Map vertices to the one they are collapsed to because they are very similar, if they aren't collapsed then they are mapped to themselves
+    vector<int> collapsedVertexMap;
+    vector<vec3f> vertices;
+    std::multimap<int, int> collapsedVertsMap; // Maps a collapsed vertex to all the vertices that collapse to it
+    MeshCreateCollapsedVertexMap(pMesh, collapsedVertexMap, vertices, collapsedVertsMap);
+
+    // Used to speed up connectivity building
+    std::multimap<int, int> vertToFacesMap;
+
+    for (int f = 0; f < faceCount; ++f)
+    {
+        int i0 = pIndices[f * 3 + 0];
+        int i1 = pIndices[f * 3 + 1];
+        int i2 = pIndices[f * 3 + 2];
+
+        vertToFacesMap.insert(std::pair<int, int>(collapsedVertexMap[i0], f));
+        vertToFacesMap.insert(std::pair<int, int>(collapsedVertexMap[i1], f));
+        vertToFacesMap.insert(std::pair<int, int>(collapsedVertexMap[i2], f));
+        //vertToFacesMap.insert(std::pair<int, int>(i0, f));
+        //vertToFacesMap.insert(std::pair<int, int>(i1, f));
+        //vertToFacesMap.insert(std::pair<int, int>(i2, f));
+    }
+
+    // Trace a ray in the projection direction to find the face that will be projected planarly and be the root of the unfolding
+    // Also build face connectivity information
+    for (int f = 0; f < faceCount; ++f)
+    {
+        int i0 = pIndices[f * 3 + 0];
+        int i1 = pIndices[f * 3 + 1];
+        int i2 = pIndices[f * 3 + 2];
+
+        vec3f rayStart = projectorPosition;
+        vec3f rayEnd = projectorPosition + projectorDirection * maxDist;
+        vec3f aux_out_intersection;
+        int out_intersected_vert, out_intersected_edge_v0, out_intersected_edge_v1;
+        float t;
+
+        rayLength = length(rayEnd - rayStart);
+
+        bool intersects = rayIntersectsFace2(rayStart, rayEnd, pVertices[i0].pos, pVertices[i1].pos, pVertices[i2].pos,
+            aux_out_intersection, out_intersected_vert, out_intersected_edge_v0, out_intersected_edge_v1, t);
+
+        if (intersects && t < min_t)
+        {
+            intersectedFace = f;
+            min_t = t;
+            out_intersection = aux_out_intersection;
+
+            vec3<float> v0 = pVertices[i0].pos;
+            vec3<float> v1 = pVertices[i1].pos;
+            vec3<float> v2 = pVertices[i2].pos;
+            projectionPlaneNormal = normalise(cross<float>(v1 - v0, v2 - v0));
+        }
+
+        // Face connectivity info
+        for (int i = 0; i < 3; ++i)
+        {
+            int v = collapsedVertexMap[pIndices[f * 3 + i]];
+            //int v = pIndices[f * 3 + i];
+
+            for (auto it = vertToFacesMap.find(v); it != vertToFacesMap.end(); ++it)
+            {
+                if (it->first != v) // The return iterator points at the first instance of v, but will iterate to the end of the map
+                {
+                    break;
+                }
+
+                int f2 = it->second;
+
+                if (f != f2 && f2 >= 0)
+                {
+                    int commonVertices = 0;
+                    int commonVerticesDifferentIsland = 0;
+                    bool commonMask[3] = { false, false, false };
+
+                    for (int ii = 0; ii < 3; ++ii)
+                    {
+                        int f_vertex_orig = pIndices[f * 3 + ii];
+                        int f_vertex = collapsedVertexMap[f_vertex_orig];
+                        //int f_vertex = pIndices[f * 3 + ii];
+
+                        for (int j = 0; j < 3; ++j)
+                        {
+                            int f2_vertex_orig = pIndices[f2 * 3 + j];
+                            int f2_vertex = collapsedVertexMap[f2_vertex_orig];
+                            //int f2_vertex = pIndices[f2 * 3 + j];
+
+                            if (f_vertex != -1 && f_vertex == f2_vertex)
+                            {
+                                commonVertices++;
+                                commonMask[j] = true;
+
+                                if (f_vertex_orig != f2_vertex_orig)
+                                {
+                                    commonVerticesDifferentIsland++;
+                                }
+                            }
+                        }
+                    }
+#ifdef DEBUG_PROJECTION
+                    assert(commonVerticesDifferentIsland <= commonVertices);
+#endif
+                    if (commonVertices == 2)
+                    {
+                        int newVertex = -1;
+
+                        for (int j = 0; j < 3; ++j)
+                        {
+                            if (!commonMask[j])
+                            {
+                                newVertex = collapsedVertexMap[pIndices[f2 * 3 + j]];
+                                //newVertex = pIndices[f2 * 3 + j];
+                                break;
+                            }
+                        }
+#ifdef DEBUG_PROJECTION
+                        assert(newVertex >= 0);
+#endif
+                        faceConnectivity[f].addConnectedFace(f2, newVertex, commonVerticesDifferentIsland == 2);
+                    }
+                }
+            }
+        }
+    }
+
+    // New projector located perpendicularly up from the hit face
+    PROJECTOR projector2;
+    projector2.direction[0] = projectionPlaneNormal[0];
+    projector2.direction[1] = projectionPlaneNormal[1];
+    projector2.direction[2] = projectionPlaneNormal[2];
+    vec3f auxPosition = out_intersection - projectionPlaneNormal * rayLength * min_t;
+    //vec3f auxPosition = vec3f( projector.position[0], projector.position[1], projector.position[2] );
+    projector2.position[0] = auxPosition[0];
+    projector2.position[1] = auxPosition[1];
+    projector2.position[2] = auxPosition[2];
+
+    vec3f auxUp = projectorUp;
+    float test = fabs(dot(auxUp, projectionPlaneNormal));
+
+    if (test < 0.9f)
+    {
+        vec3f auxSide = normalise(cross(projectionPlaneNormal, auxUp));
+        auxUp = cross(auxSide, projectionPlaneNormal);
+    }
+    else
+    {
+        auxUp = cross(projectionPlaneNormal, projectorSide);
+    }
+
+    projector2.up[0] = auxUp[0];
+    projector2.up[1] = auxUp[1];
+    projector2.up[2] = auxUp[2];
+
+    projector2.scale[0] = projectorScale[0];
+    projector2.scale[1] = projectorScale[1];
+    projector2.scale[2] = projectorScale[2];
+    vec3f projectorPosition2, projectorDirection2, s2, u2;
+    projectorPosition2 = vec3f( projector2.position[0], projector2.position[1], projector2.position[2] );
+    projector2.GetDirectionSideUp( projectorDirection2, s2, u2 );
+    float maxDistSquared = powf(projector2.scale[0], 2.f) + powf(projector2.scale[1], 2.f);
+
+    // Do a BFS walk of the mesh, unfolding a face at each step
+    if (intersectedFace >= 0)
+    {
+        std::priority_queue<NeighborFace> pendingFaces; // Queue of pending face + new vertex
+        std::unordered_set<int> pendingFacesUnique; // Used to quickly check uniqueness in the pendingFaces queue
+        //vec3f hitFaceProjectedNormal;
+        bool hitFaceHasPositiveArea = false;
+
+        NeighborFace neighborFace;
+        neighborFace.neighborFace = intersectedFace;
+        neighborFace.newVertex = -1; // -1 because all the vertices of the face are new
+        neighborFace.previousFace = -1;
+        neighborFace.numUVIslandChanges = 0;
+        neighborFace.step = 0;
+        neighborFace.changesUVIsland = false;
+        pendingFaces.push(neighborFace);
+        pendingFacesUnique.insert(intersectedFace);
+
+        int step = 0;
+        float totalUVAreaCovered = 0.f;
+        vector<int> oldVertices;
+        oldVertices.reserve(3);
+
+        while (!pendingFaces.empty())
+        {
+            NeighborFace currentFaceStruct = pendingFaces.top();
+            int currentFace = currentFaceStruct.neighborFace;
+#ifdef DEBUG_PROJECTION
+            int newVertexFromQueue = currentFaceStruct.newVertex;
+#endif
+            pendingFaces.pop();
+            pendingFacesUnique.erase(currentFace);
+
+            processedFaces.insert(currentFace);
+            faceStep[currentFace] = step;
+            float currentTriangleArea = -2.f;
+
+            // Process the face's vertices
+            if (currentFace == intersectedFace)
+            {
+#ifdef DEBUG_PROJECTION
+                assert(newVertexFromQueue == -1);
+#endif
+
+                vec3f outIntersectionBaricentric;
+                int i0 = pIndices[currentFace * 3 + 0];
+                int i1 = pIndices[currentFace * 3 + 1];
+                int i2 = pIndices[currentFace * 3 + 2];
+                vec3<float> v3_0 = pVertices[i0].pos;
+                vec3<float> v3_1 = pVertices[i1].pos;
+                vec3<float> v3_2 = pVertices[i2].pos;
+                GetBarycentricCoords(out_intersection, v3_0, v3_1, v3_2, outIntersectionBaricentric[0], outIntersectionBaricentric[1], outIntersectionBaricentric[2]);
+                vec2f outIntersectionUV = pVertices[i0].uv * outIntersectionBaricentric[0] + pVertices[i1].uv * outIntersectionBaricentric[1] + pVertices[i2].uv * outIntersectionBaricentric[2];
+
+                vec4f proj4D_v0, proj4D_v1, proj4D_v2;
+                PlanarlyProjectVertex(pVertices[i0].pos, proj4D_v0, projector2, projectorPosition2, projectorDirection2, s2, u2);
+                PlanarlyProjectVertex(pVertices[i1].pos, proj4D_v1, projector2, projectorPosition2, projectorDirection2, s2, u2);
+                PlanarlyProjectVertex(pVertices[i2].pos, proj4D_v2, projector2, projectorPosition2, projectorDirection2, s2, u2);
+
+                vec2f proj_v0 = vec2f(proj4D_v0[0], proj4D_v0[1]);
+                vec2f proj_v1 = vec2f(proj4D_v1[0], proj4D_v1[1]);
+                vec2f proj_v2 = vec2f(proj4D_v2[0], proj4D_v2[1]);
+
+                float proj_TriangleArea = getTriangleArea(proj_v0, proj_v1, proj_v2);
+
+                vec3f BaricentricBaseU;
+                vec3f BaricentricBaseV;
+                vec3f BaricentricOrigin;
+
+                GetBarycentricCoords(vec2f(1.f, 0.f), proj_v0, proj_v1, proj_v2, BaricentricBaseU[0], BaricentricBaseU[1], BaricentricBaseU[2]);
+                GetBarycentricCoords(vec2f(0.f, 1.f), proj_v0, proj_v1, proj_v2, BaricentricBaseV[0], BaricentricBaseV[1], BaricentricBaseV[2]);
+                GetBarycentricCoords(vec2f(0.f, 0.f), proj_v0, proj_v1, proj_v2, BaricentricOrigin[0], BaricentricOrigin[1], BaricentricOrigin[2]);
+
+                vec2f OrigUVs_BaseU_EndPoint = pVertices[i0].uv * BaricentricBaseU[0] + pVertices[i1].uv * BaricentricBaseU[1] + pVertices[i2].uv * BaricentricBaseU[2];
+                vec2f OrigUVs_BaseV_EndPoint = pVertices[i0].uv * BaricentricBaseV[0] + pVertices[i1].uv * BaricentricBaseV[1] + pVertices[i2].uv * BaricentricBaseV[2];
+                vec2f OrigUVs_Origin = pVertices[i0].uv * BaricentricOrigin[0] + pVertices[i1].uv * BaricentricOrigin[1] + pVertices[i2].uv * BaricentricOrigin[2];
+
+                vec2f OrigUVs_BaseU = OrigUVs_BaseU_EndPoint - OrigUVs_Origin;
+                vec2f OrigUVs_BaseV_Aux = OrigUVs_BaseV_EndPoint - OrigUVs_Origin;
+
+                OrigUVs_BaseU = normalise(OrigUVs_BaseU);
+                OrigUVs_BaseV_Aux = normalise(OrigUVs_BaseV_Aux);
+
+                // Generate perp vector from OrigUVs_BaseU by flipping coords and one sign. Using OrigUVs_BaseV directly sometimes doesn't give an orthogonal basis
+                vec2f OrigUVs_BaseV = vec2f(copysign(OrigUVs_BaseU.y(), OrigUVs_BaseV_Aux.x()), copysign(OrigUVs_BaseU.x(), OrigUVs_BaseV_Aux.y()));
+
+#ifdef DEBUG_PROJECTION
+                assert(fabs(dot(OrigUVs_BaseU, OrigUVs_BaseV)) < 0.3f);
+#endif
+
+                OrigUVs_Origin = outIntersectionUV;
+                GetBarycentricCoords(outIntersectionUV, pVertices[i0].uv, pVertices[i1].uv, pVertices[i2].uv, outIntersectionBaricentric[0], outIntersectionBaricentric[1], outIntersectionBaricentric[2]);
+
+                for (int i = 0; i < 3; ++i)
+                {
+                    int v = pIndices[currentFace * 3 + i];
+#ifdef DEBUG_PROJECTION
+                    assert(v >= 0 && v < vertexCount);
+#endif
+
+                    int collapsedVert = collapsedVertexMap[v];
+
+                    processedVertices.insert(collapsedVert);
+
+                    //PlanarlyProjectVertex(pVertices[collapsedVert].pos, projectedPositions[collapsedVert], projector2, projectorPosition2, projectorDirection2, s2, u2);
+                    vec2f projectedVertex = ChangeBase2D(pVertices[v].uv, OrigUVs_Origin, OrigUVs_BaseU, OrigUVs_BaseV);
+
+                    projectedPositions[collapsedVert].pos0 = projectedVertex[0];
+                    projectedPositions[collapsedVert].pos1 = projectedVertex[1];
+                    projectedPositions[collapsedVert].pos2 = 0.f;
+                    projectedPositions[collapsedVert].mask3 = 1;
+                }
+
+                vec2<float> v0 = projectedPositions[collapsedVertexMap[pIndices[currentFace * 3 + 0]]].xy();
+                vec2<float> v1 = projectedPositions[collapsedVertexMap[pIndices[currentFace * 3 + 1]]].xy();
+                vec2<float> v2 = projectedPositions[collapsedVertexMap[pIndices[currentFace * 3 + 2]]].xy();
+                currentTriangleArea = getTriangleArea(v0, v1, v2);
+                float triangleAreaFactor = sqrt(fabs(proj_TriangleArea) / fabs(currentTriangleArea));
+
+                vec2f origin = v0 * outIntersectionBaricentric[0] + v1 * outIntersectionBaricentric[1] + v2 * outIntersectionBaricentric[2];
+#ifdef DEBUG_PROJECTION
+                assert(length(origin - vec2f(0.5f, 0.5f)) < 0.001f);
+#endif
+
+                //hitFaceProjectedNormal = (cross<float>(v1 - v0, v2 - v0));
+                hitFaceHasPositiveArea = currentTriangleArea >= 0.f;
+
+                for (int i = 0; i < 3; ++i)
+                {
+                    int v = pIndices[currentFace * 3 + i];
+                    int collapsedVert = collapsedVertexMap[v];
+
+                    projectedPositions[collapsedVert].pos0 = ((projectedPositions[collapsedVert].pos0 - 0.5f) * triangleAreaFactor) + 0.5f;
+                    projectedPositions[collapsedVert].pos1 = ((projectedPositions[collapsedVert].pos1 - 0.5f) * triangleAreaFactor) + 0.5f;
+                    projectedPositions[collapsedVert].pos2 = 0.f;
+                    projectedPositions[collapsedVert].mask3 = projectedPositions[collapsedVert].pos0 >= 0.0f && projectedPositions[collapsedVert].pos0 <= 1.0f &&
+                                                           projectedPositions[collapsedVert].pos1 >= 0.0f && projectedPositions[collapsedVert].pos1 <= 1.0f;
+
+                    // Copy the new info to all the vertices that collapse to the same vertex
+                    for (auto it = collapsedVertsMap.find(collapsedVert); it != collapsedVertsMap.end(); ++it)
+                    {
+                        if (it->first != collapsedVert)
+                        {
+                            break;
+                        }
+
+                        int otherVert = it->second;
+                        processedVertices.insert(otherVert);
+
+                        projectedPositions[otherVert] = projectedPositions[collapsedVert];
+                    }
+                }
+
+                v0 = projectedPositions[collapsedVertexMap[pIndices[currentFace * 3 + 0]]].xy();
+                v1 = projectedPositions[collapsedVertexMap[pIndices[currentFace * 3 + 1]]].xy();
+                v2 = projectedPositions[collapsedVertexMap[pIndices[currentFace * 3 + 2]]].xy();
+                currentTriangleArea = getTriangleArea(v0, v1, v2);
+#ifdef DEBUG_PROJECTION
+                assert(fabs(proj_TriangleArea - currentTriangleArea) / proj_TriangleArea < 0.1f);
+#endif
+                origin = v0 * outIntersectionBaricentric[0] + v1 * outIntersectionBaricentric[1] + v2 * outIntersectionBaricentric[2];
+#ifdef DEBUG_PROJECTION
+                assert(length(origin - vec2f(0.5f, 0.5f)) < 0.001f);
+#endif
+            }
+            else
+            {
+                // Unwrap the rest of the faces
+#ifdef DEBUG_PROJECTION
+                assert(newVertexFromQueue >= 0);
+#endif
+                int newVertex = -1;
+                //int oldVertices[2] = { -1, -1 };
+                oldVertices.clear();
+                bool reverseOrder = false;
+
+                for (int i = 0; i < 3; ++i)
+                {
+                    int v = pIndices[currentFace * 3 + i];
+
+                    if (processedVertices.count(collapsedVertexMap[v]) == 0)
+                    {
+                        newVertex = v;
+
+                        if (i == 1)
+                        {
+                            reverseOrder = true;
+                        }
+                    }
+                    else
+                    {
+                        //int index = oldVertices[0] == -1 ? 0 : 1;
+                        //oldVertices[index] = v;
+                        //oldVertices.push_back(collapsedVertexMap[v]);
+                        oldVertices.push_back(v);
+                    }
+                }
+
+                if (reverseOrder && oldVertices.size() == 2)
+                {
+                    int aux = oldVertices[0];
+                    oldVertices[0] = oldVertices[1];
+                    oldVertices[1] = aux;
+                }
+
+                if (newVertex >= 0)
+                {
+#ifdef DEBUG_PROJECTION
+                    assert(newVertex >= 0);
+                    //assert(collapsedVertexMap[newVertex] == newVertexFromQueue);
+                    assert(oldVertices[0] >= 0 && oldVertices[1] >= 0);
+                    assert(oldVertices.size() == 2);
+                    assert(newVertex < vertexCount && oldVertices[0] < vertexCount && oldVertices[1] < vertexCount);
+#endif
+
+                    int previousFace = currentFaceStruct.previousFace;
+#ifdef DEBUG_PROJECTION
+                    assert(previousFace >= 0);
+                    assert(processedFaces.count(previousFace) == 1);
+#endif
+                    // Previous face indices
+                    int pi0 = pIndices[previousFace * 3 + 0];
+                    int pi1 = pIndices[previousFace * 3 + 1];
+                    int pi2 = pIndices[previousFace * 3 + 2];
+
+                    // Previous face uv coords
+                    vec2<float> pv0;
+                    vec2<float> pv1;
+                    vec2<float> pv2;
+
+                    // Proj vertices from previous triangle
+                    vec2<float> pv0_proj;
+                    vec2<float> pv1_proj;
+                    vec2<float> pv2_proj;
+
+                    int oldVertices_previousFace[2] = { -1, -1 };
+                    int oldVertex = -1;
+
+                    for (int i = 0; i < 3; ++i)
+                    {
+                        int v = pIndices[previousFace * 3 + i];
+                        int collapsed_v = collapsedVertexMap[v];
+
+                        if (collapsed_v == collapsedVertexMap[oldVertices[0]])
+                        {
+                            oldVertices_previousFace[0] = v;
+                        }
+                        else if (collapsed_v == collapsedVertexMap[oldVertices[1]])
+                        {
+                            oldVertices_previousFace[1] = v;
+                        }
+                        else
+                        {
+                            oldVertex = v;
+                        }
+                    }
+
+                    if (currentFaceStruct.changesUVIsland)
+                    {
+                        // It's crossing a uv island border, so the previous uvs need to be converted to the new islands uv space
+#ifdef DEBUG_PROJECTION
+                        assert(oldVertex >= 0 && oldVertices_previousFace[0] >= 0 && oldVertices_previousFace[1] >= 0);
+                        assert(oldVertex != newVertex && oldVertices_previousFace[0] != oldVertices[0] && oldVertices_previousFace[1] != oldVertices[1]);
+                        assert(collapsedVertexMap[oldVertices_previousFace[0]] == collapsedVertexMap[oldVertices[0]]);
+                        assert(collapsedVertexMap[oldVertices_previousFace[1]] == collapsedVertexMap[oldVertices[1]]);
+#endif
+                        // Compute the horizontal lengths in the previous island old (original) uv space, and in obj space
+                        float previousFace_uvSpaceLen, previousFace_objSpaceLen, midEdgePointFraction_previousFace;
+                        getEdgeHorizontalLength(oldVertices_previousFace[0], oldVertices_previousFace[1], oldVertex, pVertices, previousFace_uvSpaceLen, previousFace_objSpaceLen, midEdgePointFraction_previousFace);
+
+                        float currentFace_uvSpaceLen, currentFace_objSpaceLen, midEdgePointFraction_currentFace;
+                        getEdgeHorizontalLength(oldVertices[0], oldVertices[1], newVertex, pVertices, currentFace_uvSpaceLen, currentFace_objSpaceLen, midEdgePointFraction_currentFace);
+
+                        float objSpaceToCurrentUVSpaceConversion = currentFace_uvSpaceLen / currentFace_objSpaceLen;
+                        float previousFaceWidthInCurrentUVspace = previousFace_objSpaceLen * objSpaceToCurrentUVSpaceConversion;
+
+                        const vec2f& edge0_currentFace_oldUVSpace = pVertices[oldVertices[0]].uv;
+                        const vec2f& edge1_currentFace_oldUVSpace = pVertices[oldVertices[1]].uv;
+                        const vec2f& newVertex_currentFace_oldUVSpace = pVertices[newVertex].uv;
+
+                        vec2f edge_currentFace_oldUVSpace = edge1_currentFace_oldUVSpace - edge0_currentFace_oldUVSpace;
+                        vec2f midpoint_currentFace_oldUVSpace = edge0_currentFace_oldUVSpace + edge_currentFace_oldUVSpace * midEdgePointFraction_currentFace;
+                        vec2f edgePerpVectorNorm = normalise(midpoint_currentFace_oldUVSpace - newVertex_currentFace_oldUVSpace);
+                        vec2f test_edge_currentFace_oldUVSpaceNorm = normalise(edge_currentFace_oldUVSpace);
+                        (void)test_edge_currentFace_oldUVSpaceNorm;
+
+#ifdef DEBUG_PROJECTION
+                        vec2f sideVector_oldUVSpace = newVertex_currentFace_oldUVSpace - edge0_currentFace_oldUVSpace;
+                        float edgeVector_oldUVSpaceLen = length(edge_currentFace_oldUVSpace);
+                        float dotEdgeSideVectors_oldUVSpace = dot(edge_currentFace_oldUVSpace, sideVector_oldUVSpace);
+                        float midEdgePointFraction_oldUVSpace = (dotEdgeSideVectors_oldUVSpace / powf(edgeVector_oldUVSpaceLen, 2));
+                        assert(fabs(midEdgePointFraction_oldUVSpace - midEdgePointFraction_currentFace) < 0.01f);
+                        assert(fabs(dot(edgePerpVectorNorm, test_edge_currentFace_oldUVSpaceNorm)) < 0.1f);
+#endif
+
+                        vec2f midpoint_previousFace_oldUVSpace = edge0_currentFace_oldUVSpace + edge_currentFace_oldUVSpace * midEdgePointFraction_previousFace;
+                        vec2f oldVertex_currentFace_oldUVSpace = midpoint_previousFace_oldUVSpace + edgePerpVectorNorm * previousFaceWidthInCurrentUVspace;
+
+                        pv0 = edge0_currentFace_oldUVSpace;
+                        pv1 = edge1_currentFace_oldUVSpace;
+                        pv2 = oldVertex_currentFace_oldUVSpace;
+
+                        pv0_proj = projectedPositions[collapsedVertexMap[oldVertices[0]]].xy();
+                        pv1_proj = projectedPositions[collapsedVertexMap[oldVertices[1]]].xy();
+                        pv2_proj = projectedPositions[collapsedVertexMap[oldVertex]].xy();
+
+#ifdef DEBUG_PROJECTION
+                        assert(testPointsAreInOppositeSidesOfEdge(oldVertex_currentFace_oldUVSpace, newVertex_currentFace_oldUVSpace, edge0_currentFace_oldUVSpace, edge1_currentFace_oldUVSpace));
+#endif
+                    }
+                    else
+                    {
+                        // If there's no UV island border, just use the previous uvs
+                        pv0 = pVertices[pi0].uv;
+                        pv1 = pVertices[pi1].uv;
+                        pv2 = pVertices[pi2].uv;
+
+                        pv0_proj = projectedPositions[pi0].xy();
+                        pv1_proj = projectedPositions[pi1].xy();
+                        pv2_proj = projectedPositions[pi2].xy();
+                    }
+
+                    vec2<float> v2 = pVertices[newVertex].uv;
+
+                    // New vertex baricentric coords in respect to old triangle
+                    float a, b, c;
+                    GetBarycentricCoords(v2, pv0, pv1, pv2, a, b, c);
+                    vec2f newVertex_proj = pv0_proj * a + pv1_proj * b + pv2_proj * c;
+#ifdef DEBUG_PROJECTION
+                    vec2f v2Test = pv0 * a + pv1 * b + pv2 * c;
+                    float v2TestDist = length(v2 - v2Test);
+                    assert(v2TestDist < 0.001f);
+
+                    float a2, b2, c2;
+                    GetBarycentricCoords(newVertex_proj, pv0_proj, pv1_proj, pv2_proj, a2, b2, c2);
+                    assert(fabs(a - a2) < 0.5f);
+                    assert(fabs(b - b2) < 0.5f);
+                    assert(fabs(c - c2) < 0.5f);
+
+                    assert(testPointsAreInOppositeSidesOfEdge(projectedPositions[collapsedVertexMap[oldVertex]].xy(), newVertex_proj,
+                        projectedPositions[collapsedVertexMap[oldVertices[0]]].xy(), projectedPositions[collapsedVertexMap[oldVertices[1]]].xy()));
+#endif
+
+                    int collapsedNewVertex = collapsedVertexMap[newVertex];
+                    processedVertices.insert(collapsedNewVertex);
+
+                    projectedPositions[collapsedNewVertex].pos0 = newVertex_proj[0];
+                    projectedPositions[collapsedNewVertex].pos1 = newVertex_proj[1];
+                    projectedPositions[collapsedNewVertex].pos2 = projectedPositions[oldVertices[0]].pos2;
+                    projectedPositions[collapsedNewVertex].mask3 = newVertex_proj[0] >= 0.0f && newVertex_proj[0] <= 1.0f &&
+                                                                newVertex_proj[1] >= 0.0f && newVertex_proj[1] <= 1.0f;
+
+                    // Copy the new info to all the vertices that collapse to the same vertex
+                    for (auto it = collapsedVertsMap.find(collapsedNewVertex); it != collapsedVertsMap.end(); ++it)
+                    {
+                        if (it->first != collapsedNewVertex)
+                        {
+                            break;
+                        }
+
+                        int otherVert = it->second;
+                        processedVertices.insert(otherVert);
+
+                        projectedPositions[otherVert] = projectedPositions[collapsedNewVertex];
+                    }
+
+                    //float oldRatio = getTriangleRatio(pVertices[newVertex].uv, pVertices[oldVertices[0]].uv, pVertices[oldVertices[1]].uv);
+                    //float newRatio = getTriangleRatio(newVertex_proj, projectedPositions[collapsedVertexMap[oldVertices[0]]].xy(), projectedPositions[collapsedVertexMap[oldVertices[1]]].xy());
+
+                    //if (fabs(newRatio - oldRatio) > 31.f)
+                    //{
+                    //	break;
+                    //}
+                }
+                else
+                {
+#ifdef DEBUG_PROJECTION
+                    assert(oldVertices.size() == 3);
+#endif
+                }
+
+                int i0 = collapsedVertexMap[pIndices[currentFace * 3 + 0]];
+                int i1 = collapsedVertexMap[pIndices[currentFace * 3 + 1]];
+                int i2 = collapsedVertexMap[pIndices[currentFace * 3 + 2]];
+
+                vec2<float> v0 = projectedPositions[i0].xy();
+                vec2<float> v1 = projectedPositions[i1].xy();
+                vec2<float> v2 = projectedPositions[i2].xy();
+                currentTriangleArea = getTriangleArea(v0, v1, v2);
+                bool currentTriangleHasPositiveArea = currentTriangleArea >= 0.f;
+                //vec3f currentFaceNormal = (cross<float>(v1 - v0, v2 - v0));
+
+                //if (hitFaceProjectedNormal.z() * currentFaceNormal.z() < 0)
+                if(hitFaceHasPositiveArea != currentTriangleHasPositiveArea) // Is the current face wound in the opposite direction?
+                {
+                    discardedWrapAroundFaces.insert(currentFace); // If so, discard it since it's probably a wrap-around face
+                }
+            }
+
+            bool anyVertexInUVSpace = false;
+
+            for (int i = 0; i < 3; ++i)
+            {
+                int v = collapsedVertexMap[pIndices[currentFace * 3 + i]];
+
+                if (projectedPositions[v].mask3 == 1) // Is it inside?
+                {
+                    anyVertexInUVSpace = true;
+                    break;
+                }
+            }
+
+            bool anyVertexInObjSpaceRange = false;
+
+            for (int i = 0; i < 3; ++i)
+            {
+                int v = pIndices[currentFace * 3 + i];
+                vec3f r = pVertices[v].pos - out_intersection;
+                float squaredDist = dot(r, r);
+
+                if (squaredDist <= maxDistSquared / 4.f) // Is it inside?
+                {
+                    anyVertexInObjSpaceRange = true;
+                    break;
+                }
+            }
+
+            if (anyVertexInUVSpace && anyVertexInObjSpaceRange && discardedWrapAroundFaces.count(currentFace) == 0)
+            {
+#ifdef DEBUG_PROJECTION
+                assert(currentTriangleArea != -2.f);
+#endif
+                totalUVAreaCovered += fabs(currentTriangleArea);
+
+                for(int i = 0; i < 3; ++i)
+                {
+                    int neighborFace2 = faceConnectivity[currentFace].faces[i];
+                    int newVertex = faceConnectivity[currentFace].newVertices[i];
+                    bool changesUVIsland = faceConnectivity[currentFace].changesUVIsland[i];
+
+                    if (neighborFace2 >= 0 && processedFaces.count(neighborFace2) == 0 && pendingFacesUnique.count(neighborFace2) == 0)
+                    {
+                        NeighborFace neighborFaceStruct;
+                        neighborFaceStruct.neighborFace = neighborFace2;
+                        neighborFaceStruct.newVertex = newVertex;
+                        neighborFaceStruct.previousFace = currentFace;
+                        neighborFaceStruct.numUVIslandChanges =  currentFaceStruct.changesUVIsland ? currentFaceStruct.numUVIslandChanges + 1
+                                                                                                   : currentFaceStruct.numUVIslandChanges;
+                        neighborFaceStruct.step = step;
+                        neighborFaceStruct.changesUVIsland = changesUVIsland;
+                        pendingFaces.push(neighborFaceStruct);
+                        pendingFacesUnique.insert(neighborFace2);
+                    }
+                }
+            }
+            else
+            {
+                discardedWrapAroundFaces.insert(currentFace);
+            }
+
+            //if(step == 1000)
+            //{
+            //	break;
+            //}
+
+            if (totalUVAreaCovered > 1.5f)
+            {
+                break;
+            }
+
+            step++;
+        }
+    }
+
+    vector<int32_t> oldToNewVertex(vertexCount,-1);
+
+    // Add the projected face
+    for(int f : processedFaces)
+    {
+        if (discardedWrapAroundFaces.count(f) > 0)
+        {
+            continue;
+        }
+
+#ifdef DEBUG_PROJECTION
+        int i0 = pIndices[f * 3 + 0];
+        int i1 = pIndices[f * 3 + 1];
+        int i2 = pIndices[f * 3 + 2];
+
+        assert(processedVertices.count(i0));
+        assert(processedVertices.count(i1));
+        assert(processedVertices.count(i2));
+
+        assert(i0 >= 0 && i0 < vertexCount);
+        assert(i1 >= 0 && i1 < vertexCount);
+        assert(i2 >= 0 && i2 < vertexCount);
+#endif
+
+        //// TODO: This test is wrong and may fail for big triangles! Do proper triangle-quad culling.
+        //if ( projectedPositions[i0][3] > 0.0f ||
+        //     projectedPositions[i1][3] > 0.0f ||
+        //     projectedPositions[i2][3] > 0.0f )
+        {
+            // This face is required.
+            for (int v = 0; v < 3; ++v)
+            {
+                int i = pIndices[f * 3 + v];
+#ifdef DEBUG_PROJECTION
+                assert(i >= 0 && i < vertexCount);
+#endif
+                if (oldToNewVertex[i] < 0)
+                {
+                    pResultVertices[currentVertex] = pVertices[i];
+
+                    pResultVertices[currentVertex].pos[0] = projectedPositions[i].pos0;
+                    pResultVertices[currentVertex].pos[1] = projectedPositions[i].pos1;
+                    pResultVertices[currentVertex].pos[2] = projectedPositions[i].pos2;
+
+                    // Normal is actually the fade factor
+                    int step = faceStep[f];
+                    const float maxGradient = 10.f;
+                    float stepGradient = step / maxGradient;
+                    stepGradient = stepGradient > maxGradient ? maxGradient : stepGradient;
+                    float angleCos = stepGradient; //1.f; // dot(pVertices[i].nor, projectorDirection * -1.0f);
+                    pResultVertices[currentVertex].nor[0] = angleCos;
+                    pResultVertices[currentVertex].nor[1] = angleCos;
+                    pResultVertices[currentVertex].nor[2] = angleCos;
+
+                    oldToNewVertex[i] = currentVertex++;
+                }
+
+                pResultIndices[currentIndex++] = oldToNewVertex[i];
+            }
+        }
+    }
+
+#ifdef DEBUG_PROJECTION
+    std::ofstream outfile;
+    outfile.open("[REDACTED_USER_HOME]/Documents/uvs.obj", std::ios::out | std::ios::trunc);
+
+    //for (int i = 0; i < currentVertex; ++i)
+    //{
+    //	LogDebug("v %f %f %f\n", pResultVertices[i].pos[0], pResultVertices[i].pos[1], pResultVertices[i].pos[2]);
+    //}
+
+    //for (int i = 0; i < currentIndex; ++i)
+    //{
+    //	LogDebug("f %d\n", pResultIndices[i]);
+    //}
+
+    if (outfile)
+    {
+        for (int i = 0; i < currentVertex; ++i)
+        {
+            outfile << "v " << pResultVertices[i].pos[0] << " " << pResultVertices[i].pos[1] << " " << pResultVertices[i].pos[2] << "\n";
+        }
+
+        for (int i = 0; i < currentVertex; ++i)
+        {
+            outfile << "vt " << pResultVertices[i].pos[0] << " " << 1.f - pResultVertices[i].pos[1] << "\n";
+        }
+
+        for (int i = 0; i < currentIndex / 3; ++i)
+        {
+            int i0 = pResultIndices[3 * i] + 1;
+            int i1 = pResultIndices[3 * i + 1] + 1;
+            int i2 = pResultIndices[3 * i + 2] + 1;
+
+            outfile << "f " << i0 << "/" << i0 << " " << i1 << "/" << i1 << " " << i2 << "/" << i2 << "\n";
+
+            vec2f v0 = pResultVertices[i0].pos.xy();
+            vec2f v1 = pResultVertices[i1].pos.xy();
+            vec2f v2 = pResultVertices[i2].pos.xy();
+
+            outfile << "# Triangle " << i << ", area = " << getTriangleArea(v0, v1, v2) << "\n";
+        }
+
+        outfile << std::endl;
+
+        outfile.close();
+    }
+#endif
+
+}
+
+
+//-------------------------------------------------------------------------------------------------
+MeshPtr MeshProject_Optimised( const Mesh* pMesh,
+                               const PROJECTOR& projector )
+{
+	MUTABLE_CPUPROFILER_SCOPE(MeshProject_Optimised);
+
+    vec3f projectorPosition,projectorDirection,projectorSide,projectorUp;
+    projectorPosition = vec3f( projector.position[0], projector.position[1], projector.position[2] );
+    projector.GetDirectionSideUp( projectorDirection, projectorSide, projectorUp );
+    vec3f projectorScale( projector.scale[0], projector.scale[1], projector.scale[2] );
+
+    // Create with worse case, shrink later.
+    int vertexCount = pMesh->GetVertexCount();
+    int indexCount = pMesh->GetIndexCount();
+    MeshPtr pResult;
+    int currentVertex = 0;
+    int currentIndex = 0;
+
+    // At this point the code generation ensures we are working with layout 0
+    const int layout = 0;
+
+    switch (projector.type)
+    {
+
+    case PROJECTOR_TYPE::PLANAR:
+    case PROJECTOR_TYPE::CYLINDRICAL:
+    {
+        pResult = CreateMeshOptimisedForProjection( layout );
+        pResult->GetVertexBuffers().SetElementCount(vertexCount);
+        pResult->GetIndexBuffers().SetElementCount(indexCount);
+        auto pResultIndices = reinterpret_cast<uint32_t*>( pResult->GetIndexBuffers().GetBufferData(0) );
+        auto pResultVertices = reinterpret_cast<OPTIMISED_VERTEX*>( pResult->GetVertexBuffers().GetBufferData(0) );
+
+        // Get the vertices
+        check( pMesh->GetVertexBuffers().GetElementSize(0)==sizeof(OPTIMISED_VERTEX) );
+        auto pVertices = reinterpret_cast<const OPTIMISED_VERTEX*>( pMesh->GetVertexBuffers().GetBufferData(0) );
+
+        // Get the indices
+        check( pMesh->GetIndexBuffers().GetElementSize(0)==4 );
+        auto pIndices = reinterpret_cast<const uint32_t*>( pMesh->GetIndexBuffers().GetBufferData(0) );
+        int faceCount = pMesh->GetFaceCount();
+
+        if (projector.type==PROJECTOR_TYPE::PLANAR)
+        {
+            MeshProject_Optimised_Planar( pVertices, vertexCount,
+                                          pIndices, faceCount,
+                                          projectorPosition, projectorDirection,
+                                          projectorSide, projectorUp,
+                                          projectorScale,
+                                          pResultVertices, currentVertex,
+                                          pResultIndices, currentIndex  );
+        }
+
+        else
+        {
+            MeshProject_Optimised_Cylindrical( pVertices, vertexCount,
+                                               pIndices, faceCount,
+                                               projectorPosition, projectorDirection,
+                                               projectorSide, projectorUp,
+                                               projectorScale,
+                                               pResultVertices, currentVertex,
+                                               pResultIndices, currentIndex  );
+        }
+
+        break;
+    }
+
+    case PROJECTOR_TYPE::WRAPPING:
+	{
+        pResult = CreateMeshOptimisedForWrappingProjection( layout );
+        pResult->GetVertexBuffers().SetElementCount(vertexCount);
+        pResult->GetIndexBuffers().SetElementCount(indexCount);
+        auto pResultIndices = reinterpret_cast<uint32_t*>( pResult->GetIndexBuffers().GetBufferData(0) );
+        auto pResultVertices = reinterpret_cast<OPTIMISED_VERTEX_WRAPPING*>( pResult->GetVertexBuffers().GetBufferData(0) );
+
+        MeshProject_Optimised_Wrapping( pMesh,
+                                        projectorPosition, projectorDirection,
+                                        projectorSide, projectorUp,
+                                        projectorScale,
+                                        pResultVertices, currentVertex,
+                                        pResultIndices, currentIndex  );
+
+        break;
+	}
+
+    default:
+        // Projector type not implemented.
+        check(false);
+        break;
+
+    }
+
+    // Shrink result mesh
+    pResult->GetVertexBuffers().SetElementCount(currentVertex);
+    pResult->GetIndexBuffers().SetElementCount(currentIndex);
+    pResult->GetFaceBuffers().SetElementCount(currentIndex/3);
+
+    return pResult;
+}
+#ifdef DEBUG_PROJECTION
+PRAGMA_ENABLE_OPTIMIZATION
+#endif
+
+
+
+//-------------------------------------------------------------------------------------------------
+MeshPtr mu::MeshProject( const Mesh* pMesh,
+                         const PROJECTOR& projector )
+{
+	MUTABLE_CPUPROFILER_SCOPE(MeshProject);
+
+    MeshPtr pResult;
+
+    if ( pMesh->m_staticFormatFlags & (1<<SMF_PROJECT) )
+    {
+        // Mesh-optimised version
+        pResult = MeshProject_Optimised( pMesh, projector );
+    }
+    else if ( pMesh->m_staticFormatFlags & (1<<SMF_PROJECTWRAPPING) )
+    {
+        // Mesh-optimised version for wrapping projectors
+        // \todo: make sure the projector is a wrapping projector
+        pResult = MeshProject_Optimised( pMesh, projector );
+    }
+    else
+    {
+        check(false);
+    }
+
+	if (pResult)
+	{
+		pResult->m_surfaces.SetNum(0);
+		pResult->EnsureSurfaceData();
+		pResult->ResetStaticFormatFlags();
+	}
+
+    return pResult;
+}
+
+
+
+//-------------------------------------------------------------------------------------------------
+MeshPtr mu::CreateMeshOptimisedForProjection( int layout )
+{
+    MeshPtr pFormatMesh = new Mesh();
+    pFormatMesh->GetVertexBuffers().SetBufferCount( 1 );
+    pFormatMesh->GetIndexBuffers().SetBufferCount( 1 );
+
+    MESH_BUFFER_SEMANTIC semantics[3] =	{ MBS_TEXCOORDS,	MBS_POSITION,	MBS_NORMAL };
+    int semanticIndices[3] =			{ 0,				0,				0 };
+    MESH_BUFFER_FORMAT formats[3] =		{ MBF_FLOAT32,		MBF_FLOAT32,	MBF_FLOAT32 };
+    int componentCounts[3] =			{ 2,				3,				3 };
+    int offsets[3] =					{ 0,				8,				20 };
+    semanticIndices[0] = layout;
+    pFormatMesh->GetVertexBuffers().SetBuffer
+            ( 0, 32, 3,
+              semantics, semanticIndices,
+              formats, componentCounts,
+              offsets );
+
+    MESH_BUFFER_SEMANTIC isemantics[1] =	{ MBS_VERTEXINDEX };
+    int isemanticIndices[1] =				{ 0 };
+    MESH_BUFFER_FORMAT iformats[1] =		{ MBF_UINT32 };
+    int icomponentCounts[1] =				{ 1 };
+    int ioffsets[1] =						{ 0 };
+    pFormatMesh->GetIndexBuffers().SetBuffer
+            ( 0, 4, 1,
+              isemantics, isemanticIndices,
+              iformats, icomponentCounts,
+              ioffsets );
+
+    return pFormatMesh;
+}
+
+
+//-------------------------------------------------------------------------------------------------
+MeshPtr mu::CreateMeshOptimisedForWrappingProjection( int layout )
+{
+    MeshPtr pFormatMesh = new Mesh();
+    pFormatMesh->GetVertexBuffers().SetBufferCount( 1 );
+    pFormatMesh->GetIndexBuffers().SetBufferCount( 1 );
+
+    MESH_BUFFER_SEMANTIC semantics[4] =	{ MBS_TEXCOORDS,	MBS_POSITION,	MBS_NORMAL,     MBS_LAYOUTBLOCK };
+    int semanticIndices[4] =			{ 0,				0,				0,              0 };
+    MESH_BUFFER_FORMAT formats[4] =		{ MBF_FLOAT32,		MBF_FLOAT32,	MBF_FLOAT32,    MBF_UINT32 };
+    int componentCounts[4] =			{ 2,				3,				3,              1 };
+    int offsets[4] =					{ 0,				8,				20,             32 };
+    semanticIndices[0] = layout;
+    semanticIndices[3] = layout;
+    pFormatMesh->GetVertexBuffers().SetBuffer
+            ( 0, 36, 4,
+              semantics, semanticIndices,
+              formats, componentCounts,
+              offsets );
+
+    MESH_BUFFER_SEMANTIC isemantics[1] =	{ MBS_VERTEXINDEX };
+    int isemanticIndices[1] =				{ 0 };
+    MESH_BUFFER_FORMAT iformats[1] =		{ MBF_UINT32 };
+    int icomponentCounts[1] =				{ 1 };
+    int ioffsets[1] =						{ 0 };
+    pFormatMesh->GetIndexBuffers().SetBuffer
+            ( 0, 4, 1,
+              isemantics, isemanticIndices,
+              iformats, icomponentCounts,
+              ioffsets );
+
+    return pFormatMesh;
+}
