[mutable] Move Mutable plugin from the restricted folder to the experimental plugins folder.

diff --git a/Engine/Plugins/Experimental/Mutable/Source/MutableRuntime/Private/OpImageProject.h b/Engine/Plugins/Experimental/Mutable/Source/MutableRuntime/Private/OpImageProject.h
new file mode 100644
--- /dev/null
+++ b/Engine/Plugins/Experimental/Mutable/Source/MutableRuntime/Private/OpImageProject.h
@@ -0,0 +1,69 @@
+// Copyright Epic Games, Inc. All Rights Reserved.
+
+#pragma once
+
+#include "Platform.h"
+#include "Raster.h"
+#include "ParametersPrivate.h"
+#include "Image.h"
+#include "Mesh.h"
+
+
+namespace mu
+{
+
+    struct SCRATCH_IMAGE_PROJECT
+    {
+        vector< RasterVertex<4> > vertices;
+        vector<uint8_t> culledVertex;
+    };
+
+
+//    extern void ImageProject( const Mesh* pMesh,
+//                              Image* pTargetImage,
+//                              const Image* pSource,
+//                              const Image* pMask,
+//                              const PROJECTOR& projector,
+//                              float fadeStart,
+//                              float fadeEnd,
+//                              int layout,
+//                              int block,
+//                              SCRATCH_IMAGE_PROJECT* scratch );
+
+    extern void ImageRasterProjectedPlanar( const Mesh* pMesh,
+                                            Image* pTargetImage,
+                                            const Image* pSource,
+                                            const Image* pMask,
+                                            float fadeStart,
+                                            float fadeEnd,
+                                            int layout,
+                                            int block,
+                                            SCRATCH_IMAGE_PROJECT* scratch );
+
+    extern void ImageRasterProjectedCylindrical( const Mesh* pMesh,
+                                      Image* pTargetImage,
+                                      const Image* pSource,
+                                      const Image* pMask,
+                                      float fadeStart,
+                                      float fadeEnd,
+                                      int layout,
+                                      float projectionAngle,
+                                      SCRATCH_IMAGE_PROJECT* scratch );
+
+    extern void ImageRasterProjectedWrapping( const Mesh* pMesh,
+                                      Image* pTargetImage,
+                                      const Image* pSource,
+                                      const Image* pMask,
+                                      float fadeStart,
+                                      float fadeEnd,
+                                      int layout,
+                                      int block,
+                                      SCRATCH_IMAGE_PROJECT* scratch );
+
+    extern MeshPtr MeshProject( const Mesh* pMesh,
+                                const PROJECTOR& projector );
+
+	MUTABLERUNTIME_API extern MeshPtr CreateMeshOptimisedForProjection( int layout );
+	MUTABLERUNTIME_API extern MeshPtr CreateMeshOptimisedForWrappingProjection( int layout );
+
+}
diff --git a/Engine/Plugins/Experimental/Mutable/Source/MutableRuntime/Private/OpImageRasterMesh.h b/Engine/Plugins/Experimental/Mutable/Source/MutableRuntime/Private/OpImageRasterMesh.h
new file mode 100644
--- /dev/null
+++ b/Engine/Plugins/Experimental/Mutable/Source/MutableRuntime/Private/OpImageRasterMesh.h
@@ -0,0 +1,128 @@
+// Copyright Epic Games, Inc. All Rights Reserved.
+
+#pragma once
+
+#include "ImagePrivate.h"
+#include "Platform.h"
+#include "MeshPrivate.h"
+#include "Raster.h"
+#include "ConvertData.h"
+
+
+namespace mu
+{
+
+	class WhitePixelProcessor
+	{
+	public:
+		inline void ProcessPixel( unsigned char* pBufferPos, float[1] ) const
+		{
+			pBufferPos[0] = 255;
+		}
+	};
+
+
+	inline void ImageRasterMesh( const Mesh* pMesh, Image* pImage, int block )
+	{
+		MUTABLE_CPUPROFILER_SCOPE(ImageRasterMesh)
+
+		check( pImage->GetFormat()== EImageFormat::IF_L_UBYTE );
+
+		int sizeX = pImage->GetSizeX();
+		int sizeY = pImage->GetSizeY();
+
+		// Get the vertices
+		int vertexCount = pMesh->GetVertexCount();
+		vector< RasterVertex<1> > vertices( vertexCount );
+
+		UntypedMeshBufferIteratorConst texIt( pMesh->GetVertexBuffers(), MBS_TEXCOORDS, 0 );
+		for ( int v=0; v<vertexCount; ++v )
+		{
+            float uv[2] = {0.0f,0.0f};
+			ConvertData( 0, uv, MBF_FLOAT32, texIt.ptr(), texIt.GetFormat() );
+			ConvertData( 1, uv, MBF_FLOAT32, texIt.ptr(), texIt.GetFormat() );
+
+			vertices[v].x = uv[0] * sizeX;
+			vertices[v].y = uv[1] * sizeY;
+			++texIt;
+		}
+
+		// Get the indices
+		int faceCount = pMesh->GetFaceCount();
+		vector<int> indices( faceCount*3 );
+
+		UntypedMeshBufferIteratorConst indIt( pMesh->GetIndexBuffers(), MBS_VERTEXINDEX, 0 );
+		for ( int i=0; i<faceCount*3; ++i )
+		{
+            uint32_t index=0;
+			ConvertData( 0, &index, MBF_UINT32, indIt.ptr(), indIt.GetFormat() );
+
+			indices[i] = index;
+			++indIt;
+		}
+
+        UntypedMeshBufferIteratorConst bloIt( pMesh->GetVertexBuffers(), MBS_LAYOUTBLOCK, 0 );
+
+        if ( block<0 || bloIt.GetElementSize()==0 )
+		{
+			// Raster all the faces
+            WhitePixelProcessor pixelProc;
+			//for ( int f=0; f<faceCount; ++f )
+			const auto& ProcessFace = [
+				vertices, indices, pImage, sizeX, sizeY, pixelProc
+			] (int32 f)
+			{
+				Triangle(pImage->GetData(),
+					sizeX, sizeY,
+					1,
+					vertices[indices[f * 3 + 0]],
+					vertices[indices[f * 3 + 1]],
+					vertices[indices[f * 3 + 2]],
+					pixelProc,
+					false);
+			};
+
+			ParallelFor(faceCount, ProcessFace);
+		}
+		else
+		{
+			// Raster only the faces in the selected block
+
+			// Get the block per face
+			vector<int> blocks( vertexCount );
+
+			for ( int i=0; i<vertexCount; ++i )
+			{
+                uint16_t index=0;
+				ConvertData( 0, &index, MBF_UINT16, bloIt.ptr(), bloIt.GetFormat() );
+
+				blocks[i] = index;
+				++bloIt;
+			}
+
+            WhitePixelProcessor pixelProc;
+			//for (int f = 0; f < faceCount; ++f)
+			const auto& ProcessFace = [
+				vertices, indices, blocks, block, pImage, sizeX, sizeY, pixelProc
+			] (int32 f)
+			{
+				// TODO: Select faces outside for loop?
+				if (blocks[indices[f * 3 + 0]] == block)
+				{
+					Triangle(pImage->GetData(),
+						sizeX, sizeY,
+						1,
+						vertices[indices[f * 3 + 0]],
+						vertices[indices[f * 3 + 1]],
+						vertices[indices[f * 3 + 2]],
+						pixelProc,
+						false);
+				}
+			};
+
+			ParallelFor(faceCount, ProcessFace);
+		}
+
+	}
+
+}
diff --git a/Engine/Plugins/Experimental/Mutable/Source/MutableRuntime/Private/OpImageResize.h b/Engine/Plugins/Experimental/Mutable/Source/MutableRuntime/Private/OpImageResize.h
new file mode 100644
--- /dev/null
+++ b/Engine/Plugins/Experimental/Mutable/Source/MutableRuntime/Private/OpImageResize.h
@@ -0,0 +1,845 @@
+// Copyright Epic Games, Inc. All Rights Reserved.
+
+#pragma once
+
+#include "ImagePrivate.h"
+#include "OpImagePixelFormat.h"
+#include "Async/ParallelFor.h"
+#include "MutableTrace.h"
+
+
+namespace mu
+{
+
+    //---------------------------------------------------------------------------------------------
+    //! Point-filter image resize
+    //! TODO: Optimise
+    //---------------------------------------------------------------------------------------------
+    inline ImagePtr ImageResize( const Image* pBase, FImageSize destSize )
+    {
+		MUTABLE_CPUPROFILER_SCOPE(ImageResizePoint)
+
+		check( !(pBase->m_flags&Image::IF_CANNOT_BE_SCALED) );
+
+        FImageSize baseSize = FImageSize( pBase->GetSizeX(), pBase->GetSizeY() );
+
+        ImagePtr pDest = new Image( destSize[0], destSize[1], 1, pBase->GetFormat() );
+
+        switch ( pBase->GetFormat() )
+        {
+
+        case EImageFormat::IF_L_UBYTE:
+        {
+            const uint8_t* pBaseBuf = pBase->GetData();
+            uint8_t* pDestBuf = pDest->GetData();
+
+            // Simple nearest pixel implementation
+            for ( int y=0; y<destSize[1]; ++y )
+            {
+                int sy = (y*baseSize[1])/destSize[1];
+
+                for ( int x=0; x<destSize[0]; ++x )
+                {
+                    int sx = (x*baseSize[0])/destSize[0];
+                    *pDestBuf = pBaseBuf[sy*baseSize[0]+sx];
+                    pDestBuf++;
+                }
+            }
+
+            break;
+        }
+
+        case EImageFormat::IF_RGB_UBYTE:
+        {
+            const uint8_t* pBaseBuf = pBase->GetData();
+            uint8_t* pDestBuf = pDest->GetData();
+
+            // Simple nearest pixel implementation
+            for ( int y=0; y<destSize[1]; ++y )
+            {
+                int sy = (y*baseSize[1])/destSize[1];
+
+                for ( int x=0; x<destSize[0]; ++x )
+                {
+                    int sx = (x*baseSize[0])/destSize[0];
+                    pDestBuf[0] = pBaseBuf[(sy*baseSize[0]+sx)*3+0];
+                    pDestBuf[1] = pBaseBuf[(sy*baseSize[0]+sx)*3+1];
+                    pDestBuf[2] = pBaseBuf[(sy*baseSize[0]+sx)*3+2];
+                    pDestBuf+=3;
+                }
+            }
+
+            break;
+        }
+
+        case EImageFormat::IF_BGRA_UBYTE:
+        case EImageFormat::IF_RGBA_UBYTE:
+        {
+            const uint8_t* pBaseBuf = pBase->GetData();
+            uint8_t* pDestBuf = pDest->GetData();
+
+            // Simple nearest pixel implementation
+            for ( int y=0; y<destSize[1]; ++y )
+            {
+                int sy = (y*baseSize[1])/destSize[1];
+
+                for ( int x=0; x<destSize[0]; ++x )
+                {
+                    int sx = (x*baseSize[0])/destSize[0];
+                    pDestBuf[0] = pBaseBuf[(sy*baseSize[0]+sx)*4+0];
+                    pDestBuf[1] = pBaseBuf[(sy*baseSize[0]+sx)*4+1];
+                    pDestBuf[2] = pBaseBuf[(sy*baseSize[0]+sx)*4+2];
+                    pDestBuf[3] = pBaseBuf[(sy*baseSize[0]+sx)*4+3];
+                    pDestBuf+=4;
+                }
+            }
+
+            break;
+        }
+
+        default:
+            // Case not implemented
+            check( false );
+            //mu::Halt();
+        }
+
+        return pDest;
+    }
+
+
+
+    //---------------------------------------------------------------------------------------------
+    //!
+    //---------------------------------------------------------------------------------------------
+    template< int NC >
+    void ImageMagnifyX( Image* pDest, const Image* pBase )
+    {
+        int baseSizeX = pBase->GetSizeX();
+        int destSizeX = pDest->GetSizeX();
+        int sizeY = pBase->GetSizeY();
+
+        uint32_t dx_16 = (uint32_t(baseSizeX)<<16) / destSizeX;
+
+        // Linear filtering
+        //for ( int y=0; y<sizeY; ++y )
+		const auto& ProcessLine = [
+			pDest, pBase, dx_16, baseSizeX, destSizeX
+		] (uint32 y)
+		{
+			uint32_t px_16 = 0;
+			const uint8_t* pBaseBuf = pBase->GetData() + y * baseSizeX * NC;
+			uint8_t* pDestBuf = pDest->GetData() + y * destSizeX * NC;
+
+			for (int x = 0; x < destSizeX; ++x)
+			{
+				uint32_t px = px_16 >> 16;
+				uint32_t epx_16 = px_16 + dx_16;
+
+				if ((px_16 & 0xffff0000) == ((epx_16 - 1) & 0xffff0000))
+				{
+					// One fraction
+					for (int c = 0; c < NC; ++c)
+					{
+						pDestBuf[c] += pBaseBuf[px * NC + c];
+					}
+				}
+				else
+				{
+					// Two fractions
+					uint32_t frac1 = (px_16 & 0xffff);
+					uint32_t frac0 = 0x10000 - frac1;
+
+					for (int c = 0; c < NC; ++c)
+					{
+						pDestBuf[c] = uint8_t( (pBaseBuf[px * NC + c] * frac0 + pBaseBuf[(px + 1) * NC + c] * frac1) >> 16 );
+					}
+
+					++px;
+				}
+
+				px_16 = epx_16;
+				pDestBuf += NC;
+			}
+		};
+
+		ParallelFor(sizeY, ProcessLine);
+    }
+
+
+    inline void ImageMagnifyX( Image* pDest, const Image* pBase )
+    {
+		MUTABLE_CPUPROFILER_SCOPE(ImageMagnifyX)
+
+        check( pDest->GetSizeY() == pBase->GetSizeY() );
+        check( pDest->GetSizeX() > pBase->GetSizeX() );
+
+        switch ( pBase->GetFormat() )
+        {
+
+        case EImageFormat::IF_L_UBYTE:
+        {
+            ImageMagnifyX<1>( pDest, pBase );
+            break;
+        }
+
+        case EImageFormat::IF_RGB_UBYTE:
+        {
+            ImageMagnifyX<3>( pDest, pBase );
+            break;
+        }
+
+        case EImageFormat::IF_BGRA_UBYTE:
+        case EImageFormat::IF_RGBA_UBYTE:
+        {
+            ImageMagnifyX<4>( pDest, pBase );
+            break;
+        }
+
+        default:
+            // Case not implemented
+            check( false );
+            mu::Halt();
+        }
+    }
+
+
+    //---------------------------------------------------------------------------------------------
+    //! General image minimisation
+    //---------------------------------------------------------------------------------------------
+    template< int NC >
+    void ImageMinifyX( Image* pDest, const Image* pBase )
+    {
+        int baseSizeX = pBase->GetSizeX();
+        int destSizeX = pDest->GetSizeX();
+        int sizeY = pBase->GetSizeY();
+
+        uint32_t dx_16 = (uint32_t(baseSizeX)<<16) / destSizeX;
+
+        // Linear filtering
+        //for ( int y=0; y<sizeY; ++y )
+		const auto& ProcessLine = [
+			pDest, pBase, dx_16, baseSizeX, destSizeX
+		] (uint32 y)
+		{
+			const uint8_t* pBaseBuf = pBase->GetData() + y * baseSizeX * NC;
+			uint8_t* pDestBuf = pDest->GetData() + y * destSizeX * NC;
+
+			uint32_t px_16 = 0;
+			for (int x = 0; x < destSizeX; ++x)
+			{
+				uint32_t r_16[NC];
+				for (int c = 0; c < NC; ++c)
+				{
+					r_16[c] = 0;
+				}
+
+				uint32_t epx_16 = px_16 + dx_16;
+				uint32_t px = px_16 >> 16;
+				uint32_t epx = epx_16 >> 16;
+
+				// First fraction
+				uint32_t frac0 = px_16 & 0xffff;
+				if (frac0)
+				{
+					for (int c = 0; c < NC; ++c)
+					{
+						r_16[c] += (0x10000 - frac0) * pBaseBuf[px * NC + c];
+					}
+
+					++px;
+				}
+
+				// Whole pixels
+				while (px < epx)
+				{
+					for (int c = 0; c < NC; ++c)
+					{
+						r_16[c] += uint32_t(pBaseBuf[px * NC + c]) << 16;
+					}
+
+					++px;
+				}
+
+				// Second fraction
+				uint32_t frac1 = epx_16 & 0xffff;
+				if (frac1)
+				{
+					for (int c = 0; c < NC; ++c)
+					{
+						r_16[c] += frac1 * pBaseBuf[px * NC + c];
+					}
+				}
+
+				for (int c = 0; c < NC; ++c)
+				{
+					pDestBuf[c] = (uint8_t)(r_16[c] / dx_16);
+				}
+
+				px_16 = epx_16;
+				pDestBuf += NC;
+			}
+		};
+
+		ParallelFor(sizeY, ProcessLine);
+    }
+
+
+    //---------------------------------------------------------------------------------------------
+    //! Optimised for whole factors
+    //---------------------------------------------------------------------------------------------
+    template< int NC, int FACTOR >
+    void ImageMinifyX_Exact( Image* pDest, const Image* pBase )
+    {
+		int baseSizeX = pBase->GetSizeX();
+        int destSizeX = pDest->GetSizeX();
+        int sizeY = pBase->GetSizeY();
+
+        const uint8_t* pBaseBuf = pBase->GetData();
+        uint8_t* pDestBuf = pDest->GetData();
+
+        // Linear filtering
+		const auto& ProcessLine = [
+			pDest, pBase, baseSizeX, destSizeX
+		] (uint32 y)
+		{
+			const uint8_t* pBaseBuf = pBase->GetData() + y * baseSizeX * NC;
+			uint8_t* pDestBuf = pDest->GetData() + y * destSizeX * NC;
+
+			uint32_t r[NC];
+			for (int x = 0; x < destSizeX; ++x)
+			{
+				for (int c = 0; c < NC; ++c)
+				{
+					r[c] = 0;
+					for (int a = 0; a < FACTOR; ++a)
+					{
+						r[c] += pBaseBuf[a * NC + c];
+					}
+				}
+
+				for (int c = 0; c < NC; ++c)
+				{
+					pDestBuf[c] = (uint8_t)(r[c] / FACTOR);
+				}
+
+				pDestBuf += NC;
+				pBaseBuf += NC * FACTOR;
+			}
+		};
+
+		ParallelFor(sizeY, ProcessLine);
+    }
+
+
+    //---------------------------------------------------------------------------------------------
+    //! Image minify X version hub.
+    //---------------------------------------------------------------------------------------------
+    inline void ImageMinifyX( Image* pDest, const Image* pBase )
+    {
+		MUTABLE_CPUPROFILER_SCOPE(ImageMinifyX)
+
+        check( pDest->GetSizeY() == pBase->GetSizeY() );
+        check( pDest->GetSizeX() < pBase->GetSizeX() );
+
+        switch ( pBase->GetFormat() )
+        {
+
+        case EImageFormat::IF_L_UBYTE:
+        {
+            if ( 2*pDest->GetSizeX()==pBase->GetSizeX() )
+            {
+                // Optimised case
+                ImageMinifyX_Exact<1,2>( pDest, pBase );
+            }
+            else
+            {
+                // Generic case
+                ImageMinifyX<1>( pDest, pBase );
+            }
+            break;
+        }
+
+        case EImageFormat::IF_RGB_UBYTE:
+        {
+            if ( 2*pDest->GetSizeX()==pBase->GetSizeX() )
+            {
+                // Optimised case
+                ImageMinifyX_Exact<3,2>( pDest, pBase );
+            }
+            else
+            {
+                // Generic case
+                ImageMinifyX<3>( pDest, pBase );
+            }
+            break;
+        }
+
+        case EImageFormat::IF_BGRA_UBYTE:
+        case EImageFormat::IF_RGBA_UBYTE:
+        {
+            if ( 2*pDest->GetSizeX()==pBase->GetSizeX() )
+            {
+                // Optimised case
+                ImageMinifyX_Exact<4,2>( pDest, pBase );
+            }
+            else
+            {
+                // Generic case
+                ImageMinifyX<4>( pDest, pBase );
+            }
+            break;
+        }
+
+        default:
+            // Case not implemented
+            check( false );
+            mu::Halt();
+        }
+
+    }
+
+
+
+
+    //---------------------------------------------------------------------------------------------
+    //!
+    //---------------------------------------------------------------------------------------------
+    template< int NC >
+    void ImageMagnifyY( Image* pDest, const Image* pBase )
+    {
+		if (!pBase || !pDest || 
+			!pBase->GetSizeX() || !pBase->GetSizeY() || !pDest->GetSizeX() || !pDest->GetSizeY())
+		{
+			return;
+		}
+			
+		int baseSizeY = pBase->GetSizeY();
+        int destSizeY = pDest->GetSizeY();
+        int sizeX = pBase->GetSizeX();
+
+        size_t rowSize = sizeX * NC;
+
+        // Common case, optimised.
+        if (destSizeY==baseSizeY*2)
+        {
+			//for (int y = 0; y < baseSizeY; ++y)
+			const auto& ProcessLine = [
+				pDest, pBase, rowSize
+			] (uint32 y)
+            {
+				uint8_t* pDestBuf = pDest->GetData()+ 2*y*rowSize;
+				const uint8_t* pBaseBuf = pBase->GetData() + y * rowSize;
+
+                memcpy( pDestBuf, pBaseBuf, rowSize );
+                pDestBuf += rowSize;
+
+                memcpy( pDestBuf, pBaseBuf, rowSize );
+            };
+
+			ParallelFor(baseSizeY, ProcessLine);
+        }
+        else
+        {
+            uint32_t dy_16 = ( uint32_t( baseSizeY ) << 16 ) / destSizeY;
+
+            // Linear filtering
+            // \todo: optimise: swap loops, etc.
+            //for ( int x=0; x<sizeX; ++x )
+			const auto& ProcessColumn = [
+				pDest, pBase, sizeX, destSizeY, dy_16
+			] (uint32 x)
+            {
+                uint32_t py_16 = 0;
+                uint8_t* pDestBuf = pDest->GetData()+x*NC;
+				const uint8_t* pBaseBuf = pBase->GetData();
+
+                for ( int y=0; y<destSizeY; ++y )
+                {
+                    uint32_t py = py_16 >> 16;
+                    uint32_t epy_16 = py_16+dy_16;
+
+                    if ( (py_16 & 0xffff0000) == ((epy_16-1) & 0xffff0000) )
+                    {
+                        // One fraction
+                        for ( int c=0; c<NC; ++c )
+                        {
+                            pDestBuf[c] += pBaseBuf[(py*sizeX+x)*NC+c];
+                        }
+                    }
+                    else
+                    {
+                        // Two fractions
+                        uint32_t frac1 = (py_16 & 0xffff);
+                        uint32_t frac0 = 0x10000 - frac1;
+
+                        for ( int c=0; c<NC; ++c )
+                        {
+                            pDestBuf[c] = (uint8_t)( ( pBaseBuf[(py*sizeX+x)*NC+c] * frac0 +
+                                            pBaseBuf[((py+1)*sizeX+x)*NC+c] * frac1
+                                            ) >> 16 );
+                        }
+
+                        ++py;
+                    }
+
+                    py_16 = epy_16;
+                    pDestBuf+=sizeX*NC;
+                }
+            };
+
+			ParallelFor(sizeX, ProcessColumn);
+        }
+    }
+
+
+    inline void ImageMagnifyY( Image* pDest, const Image* pBase )
+    {
+        check( pDest->GetSizeY() > pBase->GetSizeY() );
+        check( pDest->GetSizeX() == pBase->GetSizeX() );
+
+		MUTABLE_CPUPROFILER_SCOPE(ImageMagnifyY)
+
+        switch ( pBase->GetFormat() )
+        {
+
+        case EImageFormat::IF_L_UBYTE:
+        {
+            ImageMagnifyY<1>( pDest, pBase );
+            break;
+        }
+
+        case EImageFormat::IF_RGB_UBYTE:
+        {
+            ImageMagnifyY<3>( pDest, pBase );
+            break;
+        }
+
+        case EImageFormat::IF_RGBA_UBYTE:
+        case EImageFormat::IF_BGRA_UBYTE:
+        {
+            ImageMagnifyY<4>( pDest, pBase );
+            break;
+        }
+
+        default:
+            // Case not implemented
+            check( false );
+            mu::Halt();
+        }
+    }
+
+
+    //---------------------------------------------------------------------------------------------
+    //!
+    //---------------------------------------------------------------------------------------------
+    template< int NC >
+    void ImageMinifyY( Image* pDest, const Image* pBase )
+    {
+        int baseSizeY = pBase->GetSizeY();
+        int destSizeY = pDest->GetSizeY();
+        int sizeX = pBase->GetSizeX();
+
+        uint32_t dy_16 = (uint32_t(baseSizeY)<<16) / destSizeY;
+
+        const uint8_t* pBaseBuf = pBase->GetData();
+
+        // Linear filtering
+		//for (int x = 0; x < sizeX; ++x)
+		const auto& ProcessColumn = [
+			pDest, pBaseBuf, sizeX, destSizeY, dy_16
+		] (uint32 x) 
+		{
+			uint8_t* pDestBuf = pDest->GetData() + x * NC;
+			uint32_t py_16 = 0;
+			for (int y = 0; y < destSizeY; ++y)
+			{
+				uint32_t r_16[NC];
+				for (int c = 0; c < NC; ++c)
+				{
+					r_16[c] = 0;
+				}
+
+				uint32_t epy_16 = py_16 + dy_16;
+				uint32_t py = py_16 >> 16;
+				uint32_t epy = epy_16 >> 16;
+
+				// First fraction
+				uint32_t frac0 = py_16 & 0xffff;
+				if (frac0)
+				{
+					for (int c = 0; c < NC; ++c)
+					{
+						r_16[c] += (0x10000 - frac0) * pBaseBuf[(py * sizeX + x) * NC + c];
+					}
+
+					++py;
+				}
+
+				// Whole pixels
+				while (py < epy)
+				{
+					for (int c = 0; c < NC; ++c)
+					{
+						r_16[c] += uint32_t(pBaseBuf[(py * sizeX + x) * NC + c]) << 16;
+					}
+
+					++py;
+				}
+
+				// Second fraction
+				uint32_t frac1 = epy_16 & 0xffff;
+				if (frac1)
+				{
+					for (int c = 0; c < NC; ++c)
+					{
+						r_16[c] += frac1 * pBaseBuf[(py * sizeX + x) * NC + c];
+					}
+				}
+
+				for (int c = 0; c < NC; ++c)
+				{
+					pDestBuf[c] = (uint8_t)(r_16[c] / dy_16);
+				}
+
+				py_16 = epy_16;
+				pDestBuf += sizeX * NC;
+			}
+		};
+
+		ParallelFor(sizeX, ProcessColumn);
+
+    }
+
+
+    //---------------------------------------------------------------------------------------------
+    //!
+    //---------------------------------------------------------------------------------------------
+    template< int NC, int FACTOR >
+    void ImageMinifyY_Exact( Image* pDest, const Image* pBase )
+    {
+        int destSizeY = pDest->GetSizeY();
+        int sizeX = pBase->GetSizeX();
+
+
+        // Linear filtering
+        //for ( int y=0; y<destSizeY; ++y )
+		const auto& ProcessRow = [
+			pDest, pBase, sizeX, destSizeY
+		] (uint32 y)
+		{
+			uint8_t* pDestBuf = pDest->GetData() + y*NC*sizeX;
+			const uint8_t* pBaseBuf = pBase->GetData() + y * FACTOR * sizeX * NC;
+
+			for ( int x=0; x<sizeX; ++x )
+            {
+                uint32_t r[NC];
+                for ( int c=0; c<NC; ++c)
+                {
+                    r[c] = 0;
+                }
+
+                // Whole pixels
+                for ( int f=0; f<FACTOR; ++f )
+                {
+                    for ( int c=0; c<NC; ++c)
+                    {
+                        r[c] += pBaseBuf[ sizeX*NC*f + x*NC + c ];
+                    }
+                }
+
+                for ( int c=0; c<NC; ++c)
+                {
+                    pDestBuf[c] = (uint8_t)(r[c]/FACTOR);
+                }
+
+                pDestBuf += NC;
+            }
+            pBaseBuf += FACTOR*sizeX*NC;
+        };
+
+		ParallelFor(destSizeY, ProcessRow);
+    }
+
+
+    //---------------------------------------------------------------------------------------------
+    //! Image minify Y version hub.
+    //---------------------------------------------------------------------------------------------
+    inline void ImageMinifyY( Image* pDest, const Image* pBase )
+    {
+        check( pDest->GetSizeY() < pBase->GetSizeY() );
+        check( pDest->GetSizeX() == pBase->GetSizeX() );
+
+		MUTABLE_CPUPROFILER_SCOPE(ImageMinifyY)
+
+        switch ( pBase->GetFormat() )
+        {
+
+        case EImageFormat::IF_L_UBYTE:
+        {
+            if ( 2*pDest->GetSizeY()==pBase->GetSizeY() )
+            {
+                // Optimised case
+                ImageMinifyY_Exact<1,2>( pDest, pBase );
+            }
+            else
+            {
+                // Generic case
+                ImageMinifyY<1>( pDest, pBase );
+            }
+            break;
+        }
+
+        case EImageFormat::IF_RGB_UBYTE:
+        {
+            if ( 2*pDest->GetSizeY()==pBase->GetSizeY() )
+            {
+                // Optimised case
+                ImageMinifyY_Exact<3,2>( pDest, pBase );
+            }
+            else
+            {
+                // Generic case
+                ImageMinifyY<3>( pDest, pBase );
+            }
+            break;
+        }
+
+        case EImageFormat::IF_RGBA_UBYTE:
+        case EImageFormat::IF_BGRA_UBYTE:
+        {
+            if ( 2*pDest->GetSizeY()==pBase->GetSizeY() )
+            {
+                // Optimised case
+                ImageMinifyY_Exact<4,2>( pDest, pBase );
+            }
+            else
+            {
+                // Generic case
+                ImageMinifyY<4>( pDest, pBase );
+            }
+            break;
+        }
+
+        default:
+            // Case not implemented
+            check( false );
+            //mu::Halt();
+        }
+
+    }
+
+
+    //---------------------------------------------------------------------------------------------
+    //! Bilinear filter image resize.
+    //---------------------------------------------------------------------------------------------
+    inline ImagePtr ImageResizeLinear( int imageCompressionQuality, const Image* pBasePtr,
+                                       FImageSize destSize )
+    {
+		MUTABLE_CPUPROFILER_SCOPE(ImageResizeLinear)
+
+		check(!(pBasePtr->m_flags & Image::IF_CANNOT_BE_SCALED));
+
+        ImagePtrConst pBase = pBasePtr;
+
+        // Shouldn't happen! But if it does...
+        EImageFormat sourceFormat = pBase->GetFormat();
+		EImageFormat uncompressedFormat = GetUncompressedFormat( sourceFormat );
+        if ( sourceFormat!=uncompressedFormat )
+        {
+            pBase = ImagePixelFormat( imageCompressionQuality, pBasePtr, uncompressedFormat );
+        }
+
+        FImageSize baseSize = FImageSize( pBase->GetSizeX(), pBase->GetSizeY() );
+
+        ImagePtr pDest = new Image( destSize[0], destSize[1], 1, pBase->GetFormat() );
+        if (!destSize[0] || !destSize[1] || !baseSize[0] || !baseSize[1])
+        {
+            return pDest;
+        }
+
+        // First resize X
+        ImagePtr pTemp;
+        if ( destSize[0] > baseSize[0] )
+        {
+            pTemp = new Image( destSize[0], baseSize[1], 1, pBase->GetFormat() );
+            ImageMagnifyX( pTemp.get(), pBase.get() );
+        }
+        else if ( destSize[0] < baseSize[0] )
+        {
+            pTemp = new Image( destSize[0], baseSize[1], 1, pBase->GetFormat() );
+            ImageMinifyX( pTemp.get(), pBase.get() );
+        }
+        else
+        {
+            pTemp = pBase->Clone();
+        }
+
+        // Now resize Y
+        if ( destSize[1] > baseSize[1] )
+        {
+            ImageMagnifyY( pDest.get(), pTemp.get() );
+        }
+        else if ( destSize[1] < baseSize[1] )
+        {
+            ImageMinifyY( pDest.get(), pTemp.get() );
+        }
+        else
+        {
+            pDest = pTemp;
+        }
+
+
+        // Reset format if it was changed to scale
+        if ( sourceFormat!=uncompressedFormat )
+        {
+            pDest = ImagePixelFormat( imageCompressionQuality, pDest.get(), sourceFormat );
+        }
+
+        return pDest;
+    }
+
+
+	inline mu::ImagePtr ImageReduceSize( const Image* pImage, const vec2<int32>& size, int32 imageCompressionQuality)
+	{
+		check( size[0] < pImage->GetSizeX() || size[1] < pImage->GetSizeY() );
+
+		check(!(pImage->m_flags & Image::IF_CANNOT_BE_SCALED));
+
+		mu::ImagePtr Result;
+		int32 ResizeByDroppingMips = 0;
+
+		// Is it smaller
+		if (size[0] < pImage->GetSizeX() && size[1] < pImage->GetSizeY()
+			&&
+			// It is a multiple
+			(pImage->GetSizeX() % size[0]) == 0 && (pImage->GetSizeY() % size[1]) == 0
+			&&
+			// It is the same multiple
+			(pImage->GetSizeX() / size[0]) == (pImage->GetSizeY() / size[1]))
+		{
+			ResizeByDroppingMips = FMath::CountTrailingZeros64(pImage->GetSizeX() / size[0]);
+		}
+
+		if (ResizeByDroppingMips > 0 && pImage->GetLODCount() > ResizeByDroppingMips)
+		{
+			Result = pImage->Clone();
+			Result->ReduceLODs(ResizeByDroppingMips);
+		}
+		else
+		{
+			if (IsCompressedFormat(pImage->GetFormat()))
+			{
+				EImageFormat format = GetUncompressedFormat(pImage->GetFormat());
+				Result = ImagePixelFormat(
+					imageCompressionQuality,
+					pImage,
+					format);
+			}
+
+			FImageSize blockSize((uint16_t)size[0], (uint16_t)size[1]);
+			Result = ImageResizeLinear(imageCompressionQuality, pImage, blockSize);
+		}
+
+		return Result;
+	}
+}
diff --git a/Engine/Plugins/Experimental/Mutable/Source/MutableRuntime/Private/OpImageSaturate.h b/Engine/Plugins/Experimental/Mutable/Source/MutableRuntime/Private/OpImageSaturate.h
new file mode 100644
--- /dev/null
+++ b/Engine/Plugins/Experimental/Mutable/Source/MutableRuntime/Private/OpImageSaturate.h
@@ -0,0 +1,90 @@
+// Copyright Epic Games, Inc. All Rights Reserved.
+
+#pragma once
+
+#include "ImagePrivate.h"
+#include "Platform.h"
+
+
+namespace mu
+{
+
+	inline ImagePtr ImageSaturate( const Image* pA, float factor )
+	{
+		// Clamp the factor
+		// TODO: See what happens if we don't
+		factor = std::max( 0.0f, factor );
+        int32_t f_8 = (int32_t)(factor*255);
+
+        ImagePtr pDest = new Image( pA->GetSizeX(), pA->GetSizeY(),
+                                    pA->GetLODCount(),
+                                    pA->GetFormat() );
+
+        uint8_t* pDestBuf = pDest->GetData();
+        const uint8_t* pABuf = pA->GetData();
+
+		// Generic implementation
+		int pixelCount = (int)pA->CalculatePixelCount();
+
+		switch ( pA->GetFormat() )
+		{
+		case EImageFormat::IF_RGB_UBYTE:
+		{
+			for ( int i=0; i<pixelCount; ++i )
+			{
+                int32_t l_16 = 76 * pABuf[3*i+0] + 150 * pABuf[3*i+1] + 29 * pABuf[3*i+2];
+
+				for ( int c=0; c<3; ++c )
+				{
+                    int32_t d_16 = (pABuf[3*i+c]<<8) - l_16;
+                    int32_t r_16 = l_16 + ( (d_16 * f_8) >> 8 );
+                    pDestBuf[3*i+c] = (uint8_t)std::min( (uint32_t)r_16>>8, 255u );
+				}
+			}
+			break;
+		}
+
+        case EImageFormat::IF_RGBA_UBYTE:
+        {
+            for ( int i=0; i<pixelCount; ++i )
+            {
+                int32_t l_16 = 76 * pABuf[4*i+0] + 150 * pABuf[4*i+1] + 29 * pABuf[4*i+2];
+
+                for ( int c=0; c<3; ++c )
+                {
+                    int32_t d_16 = (pABuf[4*i+c]<<8) - l_16;
+                    int32_t r_16 = l_16 + ( (d_16 * f_8) >> 8 );
+                    pDestBuf[4*i+c] = (uint8_t)std::min( (uint32_t)r_16>>8, 255u );
+                }
+
+                pDestBuf[4*i+3] = pABuf[4*i+3];
+            }
+            break;
+        }
+
+        case EImageFormat::IF_BGRA_UBYTE:
+        {
+            for ( int i=0; i<pixelCount; ++i )
+            {
+                int32_t l_16 = 76 * pABuf[4*i+2] + 150 * pABuf[4*i+1] + 29 * pABuf[4*i+0];
+
+                for ( int c=0; c<3; ++c )
+                {
+                    int32_t d_16 = (pABuf[4*i+c]<<8) - l_16;
+                    int32_t r_16 = l_16 + ( (d_16 * f_8) >> 8 );
+                    pDestBuf[4*i+c] = (uint8_t)std::min( (uint32_t)r_16>>8, 255u );
+                }
+
+                pDestBuf[4*i+3] = pABuf[4*i+3];
+            }
+            break;
+        }
+
+		default:
+			check(false);
+		}
+
+		return pDest;
+	}
+
+}
diff --git a/Engine/Plugins/Experimental/Mutable/Source/MutableRuntime/Private/OpImageScreen.h b/Engine/Plugins/Experimental/Mutable/Source/MutableRuntime/Private/OpImageScreen.h
new file mode 100644
--- /dev/null
+++ b/Engine/Plugins/Experimental/Mutable/Source/MutableRuntime/Private/OpImageScreen.h
@@ -0,0 +1,54 @@
+// Copyright Epic Games, Inc. All Rights Reserved.
+
+#pragma once
+
+#include "ImagePrivate.h"
+
+
+namespace mu
+{
+
+	inline unsigned ScreenChannelMasked( unsigned base, unsigned blended, unsigned mask )
+	{
+		// R = 1 - (1-Base) × (1-Blend)
+		unsigned screen = 255 - ( ( ( 255 - base  ) * ( 255 - blended ) ) >> 8 );
+		unsigned masked = ( ( ( 255 - mask ) * base ) + ( mask * screen ) ) >> 8;
+		return masked;
+	}
+
+	inline unsigned ScreenChannel( unsigned base, unsigned blended )
+	{
+		// R = 1 - (1-Base) × (1-Blend)
+		unsigned screen = 255 - ( ( ( 255 - base  ) * ( 255 - blended ) ) >> 8 );
+		return screen;
+	}
+
+
+    inline void ImageScreen( Image* pResult, const Image* pBase, const Image* pMask, vec3<float> col )
+	{
+        ImageLayerColour<ScreenChannelMasked,ScreenChannel,false>( pResult, pBase, pMask, col );
+	}
+
+
+    inline void ImageScreen( Image* pResult, const Image* pBase, const Image* pMask,
+                             const Image* pBlended,
+                             bool applyToAlpha )
+	{
+        ImageLayer<ScreenChannelMasked,ScreenChannel,false>
+                ( pResult, pBase, pMask, pBlended, applyToAlpha );
+	}
+
+
+    inline void ImageScreen( Image* pResult, const Image* pBase, vec3<float> col )
+	{
+        ImageLayerColour<ScreenChannel,false>( pResult, pBase, col );
+	}
+
+
+    inline void ImageScreen( Image* pResult, const Image* pBase, const Image* pBlended,
+                             bool applyToAlpha )
+	{
+        ImageLayer<ScreenChannel,false>( pResult, pBase, pBlended, applyToAlpha );
+	}
+
+}
diff --git a/Engine/Plugins/Experimental/Mutable/Source/MutableRuntime/Private/OpImageSelectColour.h b/Engine/Plugins/Experimental/Mutable/Source/MutableRuntime/Private/OpImageSelectColour.h
new file mode 100644
--- /dev/null
+++ b/Engine/Plugins/Experimental/Mutable/Source/MutableRuntime/Private/OpImageSelectColour.h
@@ -0,0 +1,113 @@
+// Copyright Epic Games, Inc. All Rights Reserved.
+
+#pragma once
+
+#include "ImagePrivate.h"
+#include "Platform.h"
+
+
+namespace mu
+{
+
+	inline ImagePtr ImageSelectColour( const Image* pA, const vec3<float>& colour )
+	{
+		ImagePtr pDest = new Image( pA->GetSizeX(), pA->GetSizeY(), pA->GetLODCount(), EImageFormat::IF_L_UBYTE );
+
+        uint8_t* pDestBuf = pDest->GetData();
+        const uint8_t* pABuf = pA->GetData();
+
+		// Generic implementation
+		int pixelCount = pA->GetSizeX() * pA->GetSizeY();
+
+		switch ( pA->GetFormat() )
+		{
+		case EImageFormat::IF_L_UBYTE:
+		{
+            uint8_t c;
+            c = (uint8_t)(std::max( 0.0f, std::min( 255.0f, colour[0] * 255.0f ) )+0.5f);
+
+			for ( int i=0; i<pixelCount; ++i )
+			{
+				if ( pABuf[i]==c )
+				{
+					pDestBuf[i] = 255;
+				}
+				else
+				{
+					pDestBuf[i] = 0;
+				}
+			}
+			break;
+		}
+
+		case EImageFormat::IF_RGB_UBYTE:
+		{
+            uint8_t c[3];
+            c[0] = (uint8_t)(std::max( 0.0f, std::min( 255.0f, colour[0] * 255.0f ) )+0.5f);
+            c[1] = (uint8_t)(std::max( 0.0f, std::min( 255.0f, colour[1] * 255.0f ) )+0.5f);
+            c[2] = (uint8_t)(std::max( 0.0f, std::min( 255.0f, colour[2] * 255.0f ) )+0.5f);
+
+			for ( int i=0; i<pixelCount; ++i )
+			{
+				if ( pABuf[3*i+0]==c[0] && pABuf[3*i+1]==c[1] && pABuf[3*i+2]==c[2] )
+				{
+					pDestBuf[i] = 255;
+				}
+				else
+				{
+					pDestBuf[i] = 0;
+				}
+			}
+			break;
+		}
+
+        case EImageFormat::IF_RGBA_UBYTE:
+        {
+            uint8_t c[3];
+            c[0] = (uint8_t)(std::max( 0.0f, std::min( 255.0f, colour[0] * 255.0f ) )+0.5f);
+            c[1] = (uint8_t)(std::max( 0.0f, std::min( 255.0f, colour[1] * 255.0f ) )+0.5f);
+            c[2] = (uint8_t)(std::max( 0.0f, std::min( 255.0f, colour[2] * 255.0f ) )+0.5f);
+
+            for ( int i=0; i<pixelCount; ++i )
+            {
+                if ( pABuf[4*i+0]==c[0] && pABuf[4*i+1]==c[1] && pABuf[4*i+2]==c[2] )
+                {
+                    pDestBuf[i] = 255;
+                }
+                else
+                {
+                    pDestBuf[i] = 0;
+                }
+            }
+            break;
+        }
+
+        case EImageFormat::IF_BGRA_UBYTE:
+        {
+            uint8_t c[3];
+            c[0] = (uint8_t)(std::max( 0.0f, std::min( 255.0f, colour[2] * 255.0f ) )+0.5f);
+            c[1] = (uint8_t)(std::max( 0.0f, std::min( 255.0f, colour[1] * 255.0f ) )+0.5f);
+            c[2] = (uint8_t)(std::max( 0.0f, std::min( 255.0f, colour[0] * 255.0f ) )+0.5f);
+
+            for ( int i=0; i<pixelCount; ++i )
+            {
+                if ( pABuf[4*i+0]==c[0] && pABuf[4*i+1]==c[1] && pABuf[4*i+2]==c[2] )
+                {
+                    pDestBuf[i] = 255;
+                }
+                else
+                {
+                    pDestBuf[i] = 0;
+                }
+            }
+            break;
+        }
+
+		default:
+			check(false);
+		}
+
+		return pDest;
+	}
+
+}
diff --git a/Engine/Plugins/Experimental/Mutable/Source/MutableRuntime/Private/OpImageSoftLight.h b/Engine/Plugins/Experimental/Mutable/Source/MutableRuntime/Private/OpImageSoftLight.h
new file mode 100644
--- /dev/null
+++ b/Engine/Plugins/Experimental/Mutable/Source/MutableRuntime/Private/OpImageSoftLight.h
@@ -0,0 +1,139 @@
+// Copyright Epic Games, Inc. All Rights Reserved.
+
+#pragma once
+
+#include "ImagePrivate.h"
+#include "OpImageLayer.h"
+
+namespace mu
+{
+
+	inline unsigned SoftLightChannel( unsigned base, unsigned blended )
+	{
+		// gimp-like
+		unsigned mix = ( base * blended ) >> 8;
+		unsigned screen = 255 - ( ( ( 255 - base  ) * ( 255 - blended ) ) >> 8 );
+		unsigned softlight = ( ( ( 255 - base ) * mix ) + ( base * screen ) ) >> 8;
+		return softlight;
+	}
+
+
+	inline unsigned SoftLightChannelMasked( unsigned base, unsigned blended, unsigned mask )
+	{
+		unsigned softlight = SoftLightChannel( base, blended );
+		unsigned masked = ( ( ( 255 - mask ) * base ) + ( mask * softlight ) ) >> 8;
+		return masked;
+	}
+
+
+    inline void ImageSoftLight( Image* pResult,
+                                const Image* pBase,
+                                const Image* pMask,
+                                vec3<float> col,
+                                vector< std::pair<uint8_t,uint8_t*> >* pTables )
+	{
+        check( pResult->GetFormat() == pBase->GetFormat() );
+        check( pResult->GetSizeX() == pBase->GetSizeX() );
+        check( pResult->GetSizeY() == pBase->GetSizeY() );
+        check( pResult->GetLODCount() == pBase->GetLODCount() );
+
+		EImageFormat format = pBase->GetFormat();
+
+        bool done = false;
+
+        // If we are reusing tables (single-threaded)
+        if ( pTables )
+		{
+			switch ( format )
+			{
+			case EImageFormat::IF_RGB_UBYTE:
+            case EImageFormat::IF_RGBA_UBYTE:
+            case EImageFormat::IF_BGRA_UBYTE:
+            {
+                uint8_t* tables[3];
+
+				for ( int t=0; t<3; ++t )
+				{
+                    uint8_t c = uint8_t( std::min( 255.0f, std::max( 0.0f, col[t]*255.0f ) ) );
+
+					// Try to find an existing table
+					tables[t] = 0;
+					for ( size_t j=0; !tables[t] && j<pTables->size(); ++j )
+					{
+						if ( (*pTables)[j].first==c )
+						{
+							tables[t] = (*pTables)[j].second;
+						}
+					}
+
+					if (!tables[t])
+					{
+                        tables[t] = new uint8_t[256];
+						for ( int i=0; i<256;++i )
+						{
+                            tables[t][i] = (uint8_t)SoftLightChannel( i, c );
+						}
+                        pTables->push_back(std::pair<uint8_t,uint8_t*>( c, tables[t] ) );
+					}
+				}
+
+                ImageTable( pResult, pBase, pMask, tables[0] ,tables[1], tables[2] );
+                done = true;
+				break;
+			}
+
+			default:
+				break;
+			}
+		}
+        else
+        {
+            // If we are going to apply it to many pixels, make a single-use table
+            int pixelCount = pResult->GetSizeX()*pResult->GetSizeY();
+            if ( pixelCount>256*3*2 )
+            {
+                uint8_t tables[256*3];
+                for ( int t=0; t<3; ++t )
+                {
+                    uint8_t c = uint8_t( std::min( 255.0f, std::max( 0.0f, col[t]*255.0f ) ) );
+                    for ( int i=0; i<256;++i )
+                    {
+                        tables[t*256+i] = (uint8_t)SoftLightChannel( i, c );
+                    }
+                }
+
+                ImageTable( pResult, pBase, pMask, &tables[256*0], &tables[256*1], &tables[256*2] );
+                done = true;
+            }
+        }
+
+        if (!done)
+		{
+            ImageLayerColour<SoftLightChannelMasked,SoftLightChannel,false>
+                    ( pResult, pBase, pMask, col );
+		}
+
+	}
+
+
+    inline void ImageSoftLight( Image* pResult, const Image* pBase, const Image* pMask,
+                                const Image* pBlended,
+                                bool applyToAlpha )
+	{
+        ImageLayer<SoftLightChannelMasked,SoftLightChannel,false>( pResult, pBase, pMask, pBlended, applyToAlpha );
+	}
+
+
+    inline void ImageSoftLight( Image* pResult, const Image* pBase, vec3<float> col )
+	{
+        ImageLayerColour<SoftLightChannel,false>( pResult, pBase, col );
+	}
+
+
+    inline void ImageSoftLight( Image* pResult, const Image* pBase, const Image* pBlended,
+                                bool applyToAlpha )
+	{
+        ImageLayer<SoftLightChannel,false>( pResult, pBase, pBlended, applyToAlpha );
+	}
+
+}
diff --git a/Engine/Plugins/Experimental/Mutable/Source/MutableRuntime/Private/OpImageSwizzle.h b/Engine/Plugins/Experimental/Mutable/Source/MutableRuntime/Private/OpImageSwizzle.h
new file mode 100644
--- /dev/null
+++ b/Engine/Plugins/Experimental/Mutable/Source/MutableRuntime/Private/OpImageSwizzle.h
@@ -0,0 +1,165 @@
+// Copyright Epic Games, Inc. All Rights Reserved.
+
+#pragma once
+
+#include "ImagePrivate.h"
+#include "Async/ParallelFor.h"
+#include "Platform.h"
+
+namespace mu
+{
+
+	inline ImagePtr ImageSwizzle
+		(
+			EImageFormat format,
+			const ImagePtrConst pSources[],
+            const uint8_t channels[]
+		)
+	{
+		MUTABLE_CPUPROFILER_SCOPE(ImageSwizzle);
+
+        ImagePtr pDest = new Image( pSources[0]->GetSizeX(), pSources[0]->GetSizeY(),
+                                    pSources[0]->GetLODCount(),
+                                    format );
+
+		// Very slow generic implementations
+        size_t pixelCount = pDest->CalculatePixelCount();
+
+        // Pixelcount should already match, but due to bugs it may not be the case. Try to detect it,
+        // but avoid crashing below:
+        uint16_t numChannels = GetImageFormatData(format).m_channels;
+        for (uint16_t c=0;c<numChannels; ++c)
+        {
+            if (pSources[c])
+            {
+                size_t sourcePixelCount = pSources[c]->CalculatePixelCount();
+                if (pixelCount>sourcePixelCount)
+                {
+                    check(false);
+
+                    // Something went wrong
+                    pixelCount = sourcePixelCount;
+                }
+            }
+        }
+
+		int NumDestChannels = 0;
+
+		switch ( format )
+		{
+		case EImageFormat::IF_L_UBYTE:
+			NumDestChannels = 1;
+			break;
+
+		case EImageFormat::IF_RGB_UBYTE:
+			NumDestChannels = 3;
+			break;
+
+        case EImageFormat::IF_RGBA_UBYTE:
+        case EImageFormat::IF_BGRA_UBYTE:
+			NumDestChannels = 4;
+			break;
+
+        default:
+			check(false);
+		}
+
+		constexpr int PixelCountConcurrencyThreshold = 0xff;
+
+		for (int i = 0; i < NumDestChannels; ++i)
+		{
+			uint8_t* pDestBuf = pDest->GetData() + i;
+
+			if (format == EImageFormat::IF_BGRA_UBYTE)
+			{
+				if (i == 0)
+				{
+					pDestBuf = pDest->GetData() + 2;
+				}
+				else if (i == 2)
+				{
+					pDestBuf = pDest->GetData() + 0;
+				}
+			}
+
+			bool filled = false;
+
+			if (pSources[i])
+			{
+				const uint8_t* pSourceBuf = pSources[i]->GetData() + channels[i];
+
+				switch (pSources[i]->GetFormat())
+				{
+				case EImageFormat::IF_L_UBYTE:
+					if (channels[i] < 1)
+					{
+						//for (size_t p = 0; p < pixelCount; ++p)
+						ParallelFor(pixelCount,[pDestBuf, pSourceBuf, NumDestChannels](int p)
+							{
+								pDestBuf[p * NumDestChannels] = pSourceBuf[p];
+							});
+
+						filled = true;
+					}
+					break;
+
+				case EImageFormat::IF_RGB_UBYTE:
+					if (channels[i] < 3)
+					{
+						ParallelFor(pixelCount, [pDestBuf, pSourceBuf, NumDestChannels](int p)
+							{
+								pDestBuf[p * NumDestChannels] = pSourceBuf[p * 3];
+							});
+						filled = true;
+					}
+					break;
+
+				case EImageFormat::IF_RGBA_UBYTE:
+					if (channels[i] < 4)
+					{
+						ParallelFor(pixelCount, [pDestBuf, pSourceBuf, NumDestChannels](int p)
+							{
+								pDestBuf[p * NumDestChannels] = pSourceBuf[p * 4];
+							});
+						filled = true;
+					}
+					break;
+
+				case EImageFormat::IF_BGRA_UBYTE:
+					if (channels[i] == 0)
+					{
+						pSourceBuf = pSources[i]->GetData() + 2;
+					}
+					else if (channels[i] == 2)
+					{
+						pSourceBuf = pSources[i]->GetData() + 0;
+					}
+					if (channels[i] < 4)
+					{
+						ParallelFor(pixelCount, [pDestBuf, pSourceBuf, NumDestChannels](int p)
+							{
+								pDestBuf[p * NumDestChannels] = pSourceBuf[p * 4];
+							});
+						filled = true;
+					}
+					break;
+
+				default:
+					check(false);
+				}
+			}
+
+			if (!filled)
+			{
+				// Source not set. Clear to 0
+				ParallelFor(pixelCount, [pDestBuf, NumDestChannels](int p)
+					{
+						pDestBuf[p * NumDestChannels] = 0;
+					});
+			}
+		}
+
+		return pDest;
+	}
+
+}
diff --git a/Engine/Plugins/Experimental/Mutable/Source/MutableRuntime/Private/OpImageTransform.h b/Engine/Plugins/Experimental/Mutable/Source/MutableRuntime/Private/OpImageTransform.h
new file mode 100644
--- /dev/null
+++ b/Engine/Plugins/Experimental/Mutable/Source/MutableRuntime/Private/OpImageTransform.h
@@ -0,0 +1,125 @@
+// Copyright Epic Games, Inc. All Rights Reserved.
+
+#pragma once
+
+#include "ImagePrivate.h"
+#include "OpImagePixelFormat.h"
+#include "Async/ParallelFor.h"
+#include "MutableTrace.h"
+
+#include "Math/Vector.h"
+#include "Math/IntVector.h"
+
+namespace mu
+{
+
+    //---------------------------------------------------------------------------------------------
+    //! Image transform, translate, scale, rotate 
+    //---------------------------------------------------------------------------------------------
+
+    template<int32 NC>
+	inline void ImageTransform( Image* pDest, const Image* pBase, FVector2f Offset, FVector2f Scale, float RotationRad)
+    {
+		FIntVector2 BaseSize = FIntVector2(pBase->GetSizeX(), pBase->GetSizeY());
+		FIntVector2 DestSize = FIntVector2(pDest->GetSizeX(), pDest->GetSizeY());
+
+		const FVector2f BaseSizeF = FVector2f(BaseSize.X, BaseSize.Y);
+		const FVector2f DestSizeF = FVector2f(DestSize.X, DestSize.Y);
+
+		float SinRot, CosRot;
+		FMath::SinCos(&SinRot, &CosRot, RotationRad);
+
+		const FVector2f DestNormFactor = FVector2f(1.0f) / DestSizeF;
+		const FVector2f BaseNormFactor = FVector2f(1.0f) / BaseSizeF;
+
+		Scale.X = FMath::IsNearlyZero(Scale.X) ? UE_SMALL_NUMBER : Scale.X;
+		Scale.Y = FMath::IsNearlyZero(Scale.Y) ? UE_SMALL_NUMBER : Scale.Y;
+
+		FVector2f ScaleReciprocal =  FVector2f(1.0f) / Scale; 
+
+		for (int32 Y = 0; Y < DestSize.Y; ++Y)
+		{
+			for (int32 X = 0; X < DestSize.X; ++X)
+			{
+				FVector2f DestUV = (FVector2f(X, Y) + 0.5f) * DestNormFactor - 0.5f;	
+
+				DestUV = FVector2f( CosRot * DestUV.X + SinRot * DestUV.Y,
+									CosRot * DestUV.Y - SinRot * DestUV.X);
+				DestUV *= ScaleReciprocal;
+				DestUV += -Offset + 0.5f;
+
+				const FVector2f Coord = FVector2f(DestUV.X - FMath::Floor(DestUV.X), DestUV.Y - FMath::Floor(DestUV.Y)) * BaseSizeF;
+
+				const FVector2f CoordFloor = FVector2f( FMath::Floor( Coord.X ), FMath::Floor( Coord.Y ) );
+				const FVector2f CoordFract = Coord - CoordFloor;
+
+				const FVector2f UV00 = CoordFloor * BaseNormFactor;
+				const FVector2f UV01 = ( CoordFloor + FVector2f(0.0f, 1.0f) ) * BaseNormFactor;
+				const FVector2f UV10 = ( CoordFloor + FVector2f(1.0f, 0.0f) ) * BaseNormFactor;
+				const FVector2f UV11 = ( CoordFloor + FVector2f(1.0f) ) * BaseNormFactor;
+
+				const FVector2f Coord00 =
+						FVector2f( UV00.X - FMath::Floor(UV00.X), UV00.Y - FMath::Floor(UV00.Y) ) * BaseSizeF - FVector2f(0.5f);
+				const FVector2f Coord01 = 
+						FVector2f( UV01.X - FMath::Floor(UV01.X), UV01.Y - FMath::Floor(UV01.Y) ) * BaseSizeF - FVector2f(0.5f);
+				const FVector2f Coord10 =
+						FVector2f( UV10.X - FMath::Floor(UV10.X), UV10.Y - FMath::Floor(UV10.Y) ) * BaseSizeF - FVector2f(0.5f);
+				const FVector2f Coord11 =
+						FVector2f( UV11.X - FMath::Floor(UV11.X), UV11.Y - FMath::Floor(UV11.Y) ) * BaseSizeF - FVector2f(0.5f);
+
+				const FIntVector2 S00 = FIntVector2(Coord00.X, Coord00.Y);
+				const FIntVector2 S01 = FIntVector2(Coord01.X, Coord01.Y);
+				const FIntVector2 S10 = FIntVector2(Coord10.X, Coord10.Y);
+				const FIntVector2 S11 = FIntVector2(Coord11.X, Coord11.Y);
+
+				for (int32 Channel = 0; Channel < NC; ++Channel)
+				{
+					const float V00 = float(pBase->GetData()[(S00.Y*BaseSize.X + S00.X) * NC + Channel]);
+					const float V01 = float(pBase->GetData()[(S01.Y*BaseSize.X + S01.X) * NC + Channel]);
+					const float V10 = float(pBase->GetData()[(S10.Y*BaseSize.X + S10.X) * NC + Channel]);
+					const float V11 = float(pBase->GetData()[(S11.Y*BaseSize.X + S11.X) * NC + Channel]);
+
+					const float V = FMath::Lerp(
+							FMath::Lerp(V00, V10, CoordFract.X), 
+						    FMath::Lerp(V01, V11, CoordFract.X), CoordFract.Y);
+
+					pDest->GetData()[(Y*DestSize.X + X)*NC + Channel] = static_cast<uint8>(V);
+				}	
+			}
+		}
+    }
+    
+	inline void ImageTransform( Image* pDestImage, const Image* pImage, FVector2f Offset, FVector2f Scale, float RotationRad)
+
+    {
+		MUTABLE_CPUPROFILER_SCOPE(ImageTransform)
+
+		switch (pDestImage->GetFormat())
+		{
+		
+		case EImageFormat::IF_L_UBYTE:
+		{
+			ImageTransform<1>( pDestImage, pImage, Offset, Scale, RotationRad );
+			break;
+		}
+		
+		case EImageFormat::IF_RGB_UBYTE:
+		{
+			ImageTransform<3>( pDestImage, pImage, Offset, Scale, RotationRad );
+			break;
+		}
+
+		case EImageFormat::IF_BGRA_UBYTE:
+		case EImageFormat::IF_RGBA_UBYTE:
+		{
+			ImageTransform<4>( pDestImage, pImage, Offset, Scale, RotationRad );
+			break;
+		}
+
+		default:
+			// Case not implemented
+            check( false );
+		}
+    }
+
+}
diff --git a/Engine/Plugins/Experimental/Mutable/Source/MutableRuntime/Private/OpLayoutPack.h b/Engine/Plugins/Experimental/Mutable/Source/MutableRuntime/Private/OpLayoutPack.h
new file mode 100644
--- /dev/null
+++ b/Engine/Plugins/Experimental/Mutable/Source/MutableRuntime/Private/OpLayoutPack.h
@@ -0,0 +1,834 @@
+// Copyright Epic Games, Inc. All Rights Reserved.
+
+#pragma once
+
+#include "MutableRuntime/Public/Layout.h"
+#include "MutableRuntime/Private/Platform.h"
+
+#include <algorithm>
+
+
+namespace mu
+{
+
+	//---------------------------------------------------------------------------------------------
+    // inline void SimpleLayoutPack( Layout* pResult, const Layout* pSource )
+	// {
+    //     check( pResult->GetBlockCount() == pSource->GetBlockCount() );
+
+	// 	int blockCount = pSource->GetBlockCount();
+	// 	vector< vec2<int> > blocks( blockCount );
+
+	// 	// Look for the maximum block sizes on the layout and the total area
+	// 	int maxX = 0;
+	// 	int maxY = 0;
+	// 	int area = 0;
+	// 	for ( int index=0; index<blockCount; ++index )
+	// 	{
+	// 		box< vec2<int> > b;
+	// 		pSource->GetBlock( index, &b.min[0], &b.min[1], &b.size[0], &b.size[1] );
+
+	// 		maxX = std::max( maxX, b.size[0] );
+	// 		maxY = std::max( maxY, b.size[1] );
+
+	// 		area += b.size[0] * b.size[1];
+
+	// 		blocks[index] = b.size;
+	// 	}
+
+	// 	// Grow until the area is big enough to fit all blocks. We always grow X first, because
+	// 	// in case we cannot pack everything, we will grow Y with the current horizon algorithm.
+	// 	maxX = ceilPow2( maxX );
+	// 	maxY = ceilPow2( maxY );
+	// 	while ( maxX*maxY<area )
+	// 	{
+	// 		if (maxX>maxY)
+	// 		{
+	// 			maxY*=2;
+	// 		}
+	// 		else
+	// 		{
+	// 			maxX*=2;
+	// 		}
+	// 	}
+
+	// 	// Sort by height, area
+	// 	vector<int> sorted;
+	// 	sorted.reserve( blockCount );
+	// 	for ( int index=0; index<blockCount; ++index )
+	// 	{
+	// 		int thisHeight = blocks[index][1];
+	// 		int thisArea = blocks[index][0] * blocks[index][1];
+
+	// 		int p;
+	// 		for ( p=0; p<(int)sorted.size(); ++p )
+	// 		{
+	// 			int pos = sorted[p];
+	// 			int posHeight = blocks[pos][1];
+	// 			int posArea = blocks[pos][0] * blocks[pos][1];
+
+	// 			if ( posHeight<thisHeight || ( posHeight==thisHeight && posArea<thisArea ))
+	// 			{
+	// 				break;
+	// 			}
+	// 		}
+
+	// 		sorted.insert( sorted.begin()+p, index );
+	// 	}
+
+	// 	// Pack with fixed horizontal size
+	// 	vector<int> horizon( maxX, 0 );
+	// 	vector< vec2<int> > positions( blockCount );
+	// 	maxY = 0;
+
+	// 	for ( int p=0; p<(int)sorted.size(); ++p )
+	// 	{
+	// 		int index = sorted[p];
+
+	// 		// Seek for the lowest span where the block fits
+	// 		int currentLevel = std::numeric_limits<int>::max();
+	// 		int currentX = -1;
+	// 		for ( int x=0; x<=maxX-blocks[index][0]; ++x )
+	// 		{
+	// 			int level = 0;
+	// 			for( int xs=x; xs<x+blocks[index][0]; ++xs )
+	// 			{
+	// 				level = std::max( level, horizon[xs] );
+	// 			}
+
+	// 			if (level<currentLevel)
+	// 			{
+	// 				currentLevel = level;
+	// 				currentX = x;
+	// 			}
+
+	// 		}
+
+	// 		check( currentX>=0 && currentX<=maxX-blocks[index][0] );
+
+	// 		// Update horizon
+	// 		for( int xs=currentX; xs<currentX+blocks[index][0]; ++xs )
+	// 		{
+	// 			horizon[xs] = currentLevel+blocks[index][1];
+	// 		}
+
+	// 		// Store
+	// 		positions[ index ] = vec2<int>( currentX, currentLevel );
+	// 		maxY = std::max( maxY, currentLevel+blocks[index][1] );
+	// 	}
+
+    //     // Set data in the result
+	// 	maxY = ceilPow2( maxY );
+	// 	pResult->SetGridSize( maxX, maxY );
+	// 	for ( int index=0; index<blockCount; ++index )
+	// 	{
+	// 		pResult->SetBlock
+	// 			(
+	// 				index,
+	// 				positions[index][0], positions[index][1],
+    //                 blocks[index][0], blocks[index][1]
+	// 			);
+	// 	}
+	// }
+	
+
+	//---------------------------------------------------------------------------------------------
+	//! Block comparison function to establish an order
+	//---------------------------------------------------------------------------------------------
+	struct LAY_BLOCK
+	{
+		LAY_BLOCK()
+		{
+			index = -1;
+		}
+
+		LAY_BLOCK( int i, vec2<int> s, int p=0 )
+		{
+			index = i;
+			size = s;
+			priority = p;
+		}
+
+		int index;
+		vec2<int> size;
+		int priority;
+	};
+
+    inline bool CompareBlocks( const LAY_BLOCK& a, const LAY_BLOCK& b )
+    {
+        if ( a.size[1]>b.size[1] )
+        {
+            return true;
+        }
+        else if ( a.size[1]<b.size[1] )
+        {
+            return false;
+        }
+        else
+        {
+            int areaA = a.size[0]*a.size[1];
+            int areaB = b.size[0]*b.size[1];
+            if (areaA > areaB)
+            {
+                return true;
+            }
+            else if (areaA < areaB)
+            {
+                return false;
+            }
+            else
+            {
+                // This has to be deterministic, and ids are supposed to be unique
+                return a.index<b.index;
+            }
+        }
+    }
+
+	inline bool CompareBlocksPriority(const LAY_BLOCK& a, const LAY_BLOCK& b)
+	{
+		return a.priority > b.priority;
+	}
+
+
+    //---------------------------------------------------------------------------------------------
+    //! More expensive but precise version
+    //---------------------------------------------------------------------------------------------
+    struct SCRATCH_LAYOUT_PACK
+    {
+        vector< vec2<int> > blocks;
+        vector< LAY_BLOCK > sorted;
+        vector< vec2<int> > positions;
+		vector< int > priorities;
+		vector< vec2<int> > reductions;
+    };
+
+
+    inline string DebugGetBlockAt( const SCRATCH_LAYOUT_PACK& scratch,
+                                   const vector<uint8_t>& packedFlag,
+                                   int x, int y )
+    {
+        for ( size_t b=0; b<packedFlag.size(); ++b )
+        {
+            if (packedFlag[b])
+            {
+                auto i = scratch.sorted[b].index;
+                if ( x>=scratch.positions[i][0] &&
+                     x<scratch.positions[i][0]+scratch.sorted[b].size[0] &&
+                     y>=scratch.positions[i][1] &&
+                     y<scratch.positions[i][1]+scratch.sorted[b].size[1]
+                     )
+                {
+                    char temp[2];
+                    temp[1] = 0;
+                    temp[0] = char('a')+char(b);
+                    return temp;
+                }
+            }
+        }
+        return ".";
+    }
+
+
+    // inline void LayoutPack( Layout* pResult,
+    //                         const Layout* pSourceLayout,
+    //                         SCRATCH_LAYOUT_PACK* scratch  )
+    // {
+    //     check( pResult->GetBlockCount() == pSourceLayout->GetBlockCount() );
+
+    //     int blockCount = pSourceLayout->GetBlockCount();
+    //     check( (int)scratch->blocks.size()==blockCount );
+
+    //     // Look for the maximum block sizes on the layout and the total area
+    //     int maxX = 0;
+    //     int maxY = 0;
+    //     int area = 0;
+    //     for ( int index=0; index<blockCount; ++index )
+    //     {
+    //         box< vec2<int> > b;
+    //         pSourceLayout->GetBlock( index, &b.min[0], &b.min[1], &b.size[0], &b.size[1] );
+
+    //         maxX = std::max( maxX, b.size[0] );
+    //         maxY = std::max( maxY, b.size[1] );
+
+    //         area += b.size[0] * b.size[1];
+
+    //         scratch->blocks[index] = b.size;
+    //     }
+
+    //     // Grow until the area is big enough to fit all blocks. We always grow X first, because
+    //     // in case we cannot pack everything, we will grow Y with the current horizon algorithm.
+    //     maxX = ceilPow2( maxX );
+    //     maxY = ceilPow2( maxY );
+    //     while ( maxX*maxY<area )
+    //     {
+    //         if (maxX>maxY)
+    //         {
+    //             maxY*=2;
+    //         }
+    //         else
+    //         {
+    //             maxX*=2;
+    //         }
+    //     }
+
+    //     // Sort by height, area
+    //     check( (int)scratch->sorted.size()==blockCount );
+    //     for ( int index=0; index<blockCount; ++index )
+    //     {
+    //         scratch->sorted[index] = LAY_BLOCK( index, scratch->blocks[index], scratch->priorities[index] );
+    //     }
+    //     std::sort( scratch->sorted.begin(), scratch->sorted.end(), CompareBlocks );
+
+    //     check( (int)scratch->sorted.size()==blockCount );
+
+    //     bool fits = false;
+    //     int iterations = 0;
+
+    //     while (!fits)
+    //     {
+    //         ++iterations;
+
+    //         // Pack with fixed horizontal size
+    //         check( maxX<256 );
+    //         int16_t horizon[256];
+    //         mutable_memset( horizon, 0, 256*sizeof(int16_t) );
+    //         maxY = 0;
+
+    //         for ( size_t p=0; p<scratch->sorted.size(); ++p )
+    //         {
+    //             // Seek for the lowest span where the block fits
+    //             int currentLevel = std::numeric_limits<int>::max();
+    //             int currentX = 0;
+    //             for ( int x=0; x<=maxX-scratch->sorted[p].size[0]; ++x )
+    //             {
+    //                 int level = 0;
+    //                 for( int xs=x; xs<x+scratch->sorted[p].size[0]; ++xs )
+    //                 {
+    //                     level = std::max( level, (int)horizon[xs] );
+    //                 }
+
+    //                 if (level<currentLevel)
+    //                 {
+    //                     currentLevel = level;
+    //                     currentX = x;
+    //                 }
+
+    //             }
+
+    //             check( currentX>=0 && currentX<=maxX-scratch->sorted[p].size[0] );
+
+    //             // Update horizon
+    //             for( int xs=currentX; xs<currentX+scratch->sorted[p].size[0]; ++xs )
+    //             {
+    //                 horizon[xs] = (uint16_t)(currentLevel+scratch->sorted[p].size[1]);
+    //             }
+
+    //             // Store
+    //             scratch->positions[ scratch->sorted[p].index ] = vec2<int>( currentX, currentLevel );
+    //             maxY = std::max( maxY, currentLevel+scratch->sorted[p].size[1] );
+    //         }
+
+    //         maxY = ceilPow2( maxY );
+
+    //         // TODO: Adjust this value
+    //         if ( maxY <= maxX || iterations>5000 )
+    //         {
+    //             fits = true;
+    //         }
+    //         else
+    //         {
+    //             fits = !std::next_permutation( scratch->sorted.begin(), scratch->sorted.end(), CompareBlocks );
+    //         }
+    //     }
+
+    //     // Set data in the result
+    //     pResult->SetGridSize( maxX, maxY );
+    //     for ( int index=0; index<blockCount; ++index )
+    //     {
+    //         pResult->SetBlock
+    //             (
+    //                 index,
+    //                 scratch->positions[index][0], scratch->positions[index][1],
+    //                 scratch->blocks[index][0], scratch->blocks[index][1]
+    //             );
+    //     }
+    // }
+
+
+
+    //---------------------------------------------------------------------------------------------
+    //! Even more expensive but precise version
+    //---------------------------------------------------------------------------------------------
+	inline void ReduceBlock(int blockCount, int* area, int* reverse_it, SCRATCH_LAYOUT_PACK* scratch)
+	{
+		int r_it = (*reverse_it);
+		bool pass = false;
+
+		int oldBlockArea = scratch->sorted[r_it].size[0] * scratch->sorted[r_it].size[1];
+
+		if (scratch->sorted[r_it].size[0] != 1 || scratch->sorted[r_it].size[1] != 1)
+		{
+			if (scratch->reductions[scratch->sorted[r_it].index][0] > scratch->reductions[scratch->sorted[r_it].index][1])
+			{
+				if (scratch->sorted[r_it].size[1] > 1)
+				{
+					scratch->sorted[r_it].size[1] /= 2;
+					scratch->blocks[scratch->sorted[r_it].index][1] /= 2;
+					pass = true;
+				}
+
+				scratch->reductions[scratch->sorted[r_it].index][1] += 1;
+			}
+			else if (scratch->reductions[scratch->sorted[r_it].index][0] < scratch->reductions[scratch->sorted[r_it].index][1])
+			{
+				if (scratch->sorted[r_it].size[0] > 1)
+				{
+					scratch->sorted[r_it].size[0] /= 2;
+					scratch->blocks[scratch->sorted[r_it].index][0] /= 2;
+					pass = true;
+				}
+
+				scratch->reductions[scratch->sorted[r_it].index][0] += 1;
+			}
+			else
+			{
+				if (r_it % 2 == 0)
+				{
+					if (scratch->sorted[r_it].size[0] > 1)
+					{
+						scratch->sorted[r_it].size[0] /= 2;
+						scratch->blocks[scratch->sorted[r_it].index][0] /= 2;
+						pass = true;
+					}
+
+					scratch->reductions[scratch->sorted[r_it].index][0] += 1;
+				}
+				else
+				{
+					if (scratch->sorted[r_it].size[1] > 1)
+					{
+						scratch->sorted[r_it].size[1] /= 2;
+						scratch->blocks[scratch->sorted[r_it].index][1] /= 2;
+						pass = true;
+					}
+
+					scratch->reductions[scratch->sorted[r_it].index][1] += 1;
+				}
+			}
+		}
+		else
+		{
+			pass = true;
+		}
+		
+		int newBlockArea = scratch->sorted[r_it].size[0] * scratch->sorted[r_it].size[1];
+
+		(*area) = (*area) - (oldBlockArea - newBlockArea);
+
+		if (pass)
+		{
+			(*reverse_it) = (*reverse_it) + 1;
+
+			if ((*reverse_it) >= blockCount)
+			{
+				(*reverse_it) = 0;
+			}
+		}
+	}
+
+	inline bool SetPositions(int bestY,int layoutSizeY, int* maxX, int* maxY, SCRATCH_LAYOUT_PACK* scratch, EPackStrategy packStrategy)
+	{
+		bool fits = true;
+
+		// Number of blocks alrady packed
+		size_t packed = 0;
+		vector<uint8_t> packedFlag(scratch->sorted.size(), 0);
+
+		// Pack with fixed horizontal size
+		check(*maxX < 256);
+		if (*maxX > 256) return true;
+
+		int16_t horizon[256];
+		FMemory::Memzero(horizon, 256 * sizeof(int16_t));
+		*maxY = 0;
+
+		int iterations = 0;
+
+		while (packed < scratch->sorted.size() || iterations > 5000)
+		{
+			++iterations;
+
+			int best = -1;
+			int bestX = -1;
+			int bestLevel = -1;
+			int bestWithHole = -1;
+			int bestWithHoleX = -1;
+			int bestWithHoleLevel = -1;
+			for (size_t candidate = 0; candidate < scratch->sorted.size(); ++candidate)
+			{
+				// Skip it if we packed it already
+				if (packedFlag[candidate]) continue;
+
+				auto candidateSizeX = scratch->sorted[candidate].size[0];
+				auto candidateSizeY = scratch->sorted[candidate].size[1];
+
+				// Seek for the lowest span where the block fits
+				int currentLevel = std::numeric_limits<int>::max();
+				int currentX = 0;
+				int currentLevelWithoutHole = std::numeric_limits<int>::max();
+				int currentXWithoutHole = 0;
+				for (int x = 0; x <= *maxX - candidateSizeX; ++x)
+				{
+					int level = 0;
+					for (int xs = x; xs < x + candidateSizeX; ++xs)
+					{
+						level = std::max(level, (int)horizon[xs]);
+					}
+
+					if (level < currentLevel)
+					{
+						currentLevel = level;
+						currentX = x;
+					}
+
+					// Does it make an unfillable hole with the top or side?
+					int minX = std::numeric_limits<int>::max();
+					int minY = std::numeric_limits<int>::max();
+					for (size_t b = 0; b < scratch->sorted.size(); ++b)
+					{
+						if (!packedFlag[b] && b != candidate)
+						{
+							minX = std::min(minX, scratch->sorted[b].size[0]);
+							minY = std::min(minY, scratch->sorted[b].size[1]);
+						}
+					}
+
+					bool hole =
+						// Vertical unfillable gap
+						(
+						(minY != std::numeric_limits<int>::max())
+							&&
+							(currentLevel + candidateSizeY) < bestY
+							&&
+							(currentLevel + minY + candidateSizeY) > bestY
+							)
+						||
+						// Horizontal unfillable gap
+						(
+						(minX != std::numeric_limits<int>::max())
+							&&
+							(currentX + candidateSizeX) < *maxX
+							&&
+							(currentX + minX + candidateSizeX) > *maxX
+							);
+
+					// Does it make a hole with the horizon?
+					for (int xs = x; !hole && xs < x + scratch->sorted[candidate].size[0]; ++xs)
+					{
+						hole = (level > (int)horizon[xs]);
+					}
+
+					if (!hole && level < currentLevelWithoutHole)
+					{
+						currentLevelWithoutHole = level;
+						currentXWithoutHole = x;
+					}
+				}
+
+				if (currentLevelWithoutHole <= currentLevel)
+				{
+					best = int(candidate);
+					bestX = currentXWithoutHole;
+					bestLevel = currentLevelWithoutHole;
+					break;
+				}
+
+				if (bestWithHole < 0)
+				{
+					bestWithHole = int(candidate);
+					bestWithHoleX = currentX;
+					bestWithHoleLevel = currentLevel;
+				}
+			}
+
+			check(best >= 0 || bestWithHole >= 0);
+
+			// If there is no other option, accept leaving a hole.
+			if (best < 0)
+			{
+				best = bestWithHole;
+				bestX = bestWithHoleX;
+				bestLevel = bestWithHoleLevel;
+			}
+
+			if (bestX >= 0)
+			{
+				// Update horizon
+				for (int xs = bestX; xs < bestX + scratch->sorted[best].size[0]; ++xs)
+				{
+					horizon[xs] = (uint16_t)(bestLevel + scratch->sorted[best].size[1]);
+				}
+			}
+
+			// Store
+			scratch->positions[scratch->sorted[best].index] = vec2<int>(bestX, bestLevel);
+			*maxY = std::max(*maxY, bestLevel + scratch->sorted[best].size[1]);
+
+			if (packStrategy == EPackStrategy::FIXED_LAYOUT && *maxY > layoutSizeY)
+			{
+				fits = false;
+				break;
+			}
+
+			packedFlag[best] = 1;
+			packed++;
+
+			//                for (int y=0; y<*maxY; ++y)
+			//                {
+			//                    string line;
+			//                    for (int x=0;x<*maxX;++x)
+			//                    {
+			//                        line += DebugGetBlockAt( *scratch, packedFlag, x, y );
+			//                    }
+			//                    AXE_LOG("layout",Warning,line.c_str());
+			//                }
+			//                AXE_LOG("layout",Warning,"--------------------------------------------------");
+		}
+
+		if (packed < scratch->sorted.size())
+		{
+			fits = false;
+		}
+
+		*maxY = ceilPow2(*maxY);
+
+		return fits;
+	}
+
+
+    inline void LayoutPack3( Layout* pResult,
+                             const Layout* pSourceLayout,
+                             SCRATCH_LAYOUT_PACK* scratch  )
+    {
+		MUTABLE_CPUPROFILER_SCOPE(MeshLayoutPack3);
+
+        check( pResult->GetBlockCount() == pSourceLayout->GetBlockCount() );
+
+        int blockCount = pSourceLayout->GetBlockCount();
+        check( (int)scratch->blocks.size()==blockCount );
+
+		//Getting maximum layout grid size:
+		int layoutSizeX, layoutSizeY;
+		pSourceLayout->GetMaxGridSize(&layoutSizeX, &layoutSizeY);
+
+		bool usePriority = false;
+
+        // Look for the maximum block sizes on the layout and the total area
+        int maxX = 0;
+        int maxY = 0;
+        int area = 0;
+        for ( int index=0; index<blockCount; ++index )
+        {
+            box< vec2<int> > b;
+            pSourceLayout->GetBlock( index, &b.min[0], &b.min[1], &b.size[0], &b.size[1] );
+
+			int p;
+			pSourceLayout->GetBlockPriority(index, &p);
+
+			vec2<int> reductions;
+			reductions[0] = 0;
+			reductions[1] = 0;
+
+			if (p > 0)
+			{
+				usePriority = true;
+			}
+
+            maxX = std::max( maxX, b.size[0] );
+            maxY = std::max( maxY, b.size[1] );
+
+            area += b.size[0] * b.size[1];
+
+            scratch->blocks[index] = b.size;
+			scratch->priorities[index] = p;
+			scratch->reductions[index] = reductions;
+        }
+
+
+        // Grow until the area is big enough to fit all blocks. We always grow X first, because
+        // in case we cannot pack everything, we will grow Y with the current horizon algorithm.
+		if (pSourceLayout->GetLayoutPackingStrategy() == EPackStrategy::RESIZABLE_LAYOUT)
+		{
+			maxX = ceilPow2(maxX);
+			maxY = ceilPow2(maxY);
+
+			while (maxX*maxY < area)
+			{
+				if (maxX > maxY)
+				{
+					maxY *= 2;
+				}
+				else
+				{
+					maxX *= 2;
+				}
+			}
+		}
+		else
+		{
+			//Increase the maximum layout size if the grid area is smaller than the number of blocks
+			while (blockCount > layoutSizeX*layoutSizeY)
+			{
+				if (layoutSizeX > layoutSizeY)
+				{
+					layoutSizeY *= 2;
+
+					if (layoutSizeY == 0)
+					{
+						layoutSizeY = 1;
+					}
+				}
+				else
+				{
+					layoutSizeX *= 2;
+
+					if (layoutSizeX == 0)
+					{
+						layoutSizeX = 1;
+					}
+				}
+			}
+
+			//The highest block is bigger than the layout max size
+			while (maxX > layoutSizeX)
+			{
+				area = 0;
+
+				for (int index = 0; index < blockCount; ++index)
+				{
+					if (scratch->blocks[index][0] == maxX)
+					{
+						scratch->blocks[index][0] /= 2;
+						scratch->reductions[index][0]++;
+					}
+				}
+				maxX = 0;
+				
+				//recalculating area and maximum block sizes
+				for (int index = 0; index < blockCount; ++index)
+				{
+					maxX = std::max(maxX, scratch->blocks[index][0]);
+					area += scratch->blocks[index][0] * scratch->blocks[index][1];
+				}
+			}
+
+			while (maxY > layoutSizeY)
+			{
+				area = 0;
+
+				for (int index = 0; index < blockCount; ++index)
+				{
+					if (scratch->blocks[index][1] == maxY)
+					{
+						scratch->blocks[index][1] /= 2;
+						scratch->reductions[index][1]++;
+					}
+				}
+				maxY = 0;
+
+				//recalculating area and maximum block sizes
+				for (int index = 0; index < blockCount; ++index)
+				{
+					maxY = std::max(maxY, scratch->blocks[index][1]);
+					area += scratch->blocks[index][0] * scratch->blocks[index][1];
+				}
+			}
+
+			maxX = ceilPow2(maxX);
+			maxY = ceilPow2(maxY);
+
+			// Grow until the area is big enough to fit all blocks or the size is equal to the max layout size.
+			while (maxX*maxY < area && (maxX < layoutSizeX || maxY < layoutSizeY))
+			{
+				if (maxX > maxY)
+				{
+					maxY *= 2;
+				}
+				else
+				{
+					maxX *= 2;
+				}
+			}
+		}
+        
+        int bestY = maxY;
+		int r_it = 0;
+
+        // Sort by height, area
+        check( (int)scratch->sorted.size()==blockCount );
+        for ( int index=0; index<blockCount; ++index )
+        {
+            scratch->sorted[index] = LAY_BLOCK( index, scratch->blocks[index], scratch->priorities[index] );
+        }
+
+		std::sort(scratch->sorted.begin(), scratch->sorted.end(), CompareBlocks);
+		
+		if(pSourceLayout->GetLayoutPackingStrategy() == EPackStrategy::FIXED_LAYOUT)
+		{
+			if (usePriority)	//Sort by priority
+			{
+				std::sort(scratch->sorted.begin(), scratch->sorted.end(), CompareBlocksPriority);
+			}
+
+			//Shrink blocks in case we do not have enough space to pack everything
+			while (maxX*maxY < area)
+			{
+				ReduceBlock(blockCount, &area, &r_it, scratch);
+			}
+			
+		}
+
+		bool fits = false;
+
+		while (!fits)
+		{
+			//Sort by height&area before packing
+			if (usePriority)
+			{
+				std::sort(scratch->sorted.begin(), scratch->sorted.end(), CompareBlocks);
+			}
+
+			//Try to pack everything
+			fits = SetPositions(bestY, layoutSizeY, &maxX, &maxY, scratch, pSourceLayout->GetLayoutPackingStrategy());
+
+			if (!fits && pSourceLayout->GetLayoutPackingStrategy() == EPackStrategy::FIXED_LAYOUT)
+			{
+				//Sort by priority before shrink
+				if (usePriority)
+				{
+					std::sort(scratch->sorted.begin(), scratch->sorted.end(), CompareBlocksPriority);
+				}
+
+				ReduceBlock(blockCount, &area, &r_it, scratch);
+			}
+		}
+
+        // Set data in the result
+        pResult->SetGridSize( maxX, maxY );
+        pResult->SetMaxGridSize(layoutSizeX, layoutSizeY);
+		pResult->SetLayoutPackingStrategy(pSourceLayout->GetLayoutPackingStrategy());
+
+        for ( int index=0; index<blockCount; ++index )
+        {
+            pResult->SetBlock
+                (
+                    index,
+                    scratch->positions[index][0], scratch->positions[index][1],
+                    scratch->blocks[index][0], scratch->blocks[index][1]
+                );
+
+			pResult->SetBlockPriority(index, scratch->priorities[index]);
+        }
+    }
+
+}
diff --git a/Engine/Plugins/Experimental/Mutable/Source/MutableRuntime/Private/OpLayoutRemoveBlocks.h b/Engine/Plugins/Experimental/Mutable/Source/MutableRuntime/Private/OpLayoutRemoveBlocks.h
new file mode 100644
--- /dev/null
+++ b/Engine/Plugins/Experimental/Mutable/Source/MutableRuntime/Private/OpLayoutRemoveBlocks.h
@@ -0,0 +1,91 @@
+// Copyright Epic Games, Inc. All Rights Reserved.
+
+#pragma once
+
+#include "MutableRuntime/Public/Layout.h"
+#include "MutableRuntime/Private/MeshPrivate.h"
+#include "MutableRuntime/Private/Platform.h"
+
+
+namespace mu
+{
+
+	//---------------------------------------------------------------------------------------------
+    inline LayoutPtr LayoutRemoveBlocks( const Layout* pSource, const Mesh* pMesh, int layoutIndex )
+    {
+        // Create the list of blocks in the mesh
+        vector<bool> blocksFound(1024, false);
+
+        UntypedMeshBufferIteratorConst itBlocks( pMesh->GetVertexBuffers(), MBS_LAYOUTBLOCK,
+                                                 layoutIndex );
+        if ( itBlocks.GetFormat()==MBF_UINT16 )
+        {
+            const uint16_t* pBlocks = reinterpret_cast<const uint16_t*>( itBlocks.ptr() );
+            for ( int i=0; i<pMesh->GetVertexCount(); ++i )
+            {
+                if (pBlocks[i]>=blocksFound.size())
+                {
+                    blocksFound.resize(pBlocks[i]+1024,false);
+                }
+
+                blocksFound[ pBlocks[i] ] = true;
+            }
+        }
+        else if ( itBlocks.GetFormat()==MBF_NONE )
+        {
+            // This seems to happen.
+            // May this happen when entire meshes are removed?
+            return pSource->Clone();
+        }
+        else
+        {
+            // Format not supported yet
+            check( false );
+        }
+
+        // Remove blocks that are not in the mesh
+		LayoutPtr pResult = pSource->Clone();
+		int dest = 0;
+		for ( int32 b=0; b<pResult->m_blocks.Num(); ++b )
+		{
+            int blockIndex = pResult->m_blocks[b].m_id;
+            if ( blockIndex<(int)blocksFound.size() && blocksFound[blockIndex] )
+			{
+				// keep
+				pResult->m_blocks[dest] = pResult->m_blocks[b];
+				++dest;
+			}
+		}
+		pResult->SetBlockCount( dest );
+
+        return pResult;
+	}
+
+
+	//---------------------------------------------------------------------------------------------
+	inline LayoutPtr LayoutMerge(const Layout* pA, const Layout* pB )
+	{
+		LayoutPtr pResult = pA->Clone();
+
+		// This is faster but fails in the rare case of a block being in both layouts, which may 
+		// happen if we merge a mesh with itself.
+		//pResult->GetPrivate()->m_blocks.insert
+		//(
+		//	pResult->GetPrivate()->m_blocks.end(),
+		//	pB->GetPrivate()->m_blocks.begin(),
+		//	pB->GetPrivate()->m_blocks.end()
+		//);
+
+		for ( const Layout::BLOCK& block: pB->m_blocks )
+		{
+			if ( pResult->FindBlock(block.m_id)<0 )
+			{
+				pResult->m_blocks.Add(block);
+			}
+		}
+
+		return pResult;
+	}
+
+
+}
diff --git a/Engine/Plugins/Experimental/Mutable/Source/MutableRuntime/Private/OpMeshApplyLayout.h b/Engine/Plugins/Experimental/Mutable/Source/MutableRuntime/Private/OpMeshApplyLayout.h
new file mode 100644
--- /dev/null
+++ b/Engine/Plugins/Experimental/Mutable/Source/MutableRuntime/Private/OpMeshApplyLayout.h
@@ -0,0 +1,197 @@
+// Copyright Epic Games, Inc. All Rights Reserved.
+
+#pragma once
+
+#include "MeshPrivate.h"
+#include "Platform.h"
+
+namespace mu
+{
+
+
+	//---------------------------------------------------------------------------------------------
+    inline void MeshApplyLayout( Mesh* pApplied, const Layout* pLayout, int texCoordsSet )
+	{
+		MUTABLE_CPUPROFILER_SCOPE(MeshApplyLayout);
+
+        int buffer = -1;
+		int channel = -1;
+		pApplied->GetVertexBuffers().FindChannel( MBS_TEXCOORDS, texCoordsSet, &buffer, &channel );
+
+        int layoutBuffer = -1;
+        int layoutChannel = -1;
+        pApplied->GetVertexBuffers().FindChannel( MBS_LAYOUTBLOCK, texCoordsSet, &layoutBuffer, &layoutChannel );
+
+		if (buffer < 0 || layoutBuffer < 0)
+		{
+			return;
+		}
+
+		//{
+		//	FString debugLog;
+		//	pApplied->GetPrivate()->Log(debugLog, 4);
+		//	UE_LOG(LogMutableCore, Warning, TEXT("before: %s"), *debugLog);
+		//}
+
+		// Get the information about the texture coordinates channel
+		MESH_BUFFER_SEMANTIC semantic;
+		int semanticIndex;
+		MESH_BUFFER_FORMAT format;
+		int components;
+		int offset;
+		pApplied->GetVertexBuffers().GetChannel
+				( buffer, channel, &semantic, &semanticIndex, &format, &components, &offset );
+		check( semantic == MBS_TEXCOORDS );
+		//check( format == MBF_FLOAT32 );
+		//check( components == 2 );
+
+        uint8_t* pData = pApplied->GetVertexBuffers().GetBufferData( buffer );
+		int elemSize = pApplied->GetVertexBuffers().GetElementSize( buffer );
+		int channelOffset = pApplied->GetVertexBuffers().GetChannelOffset( buffer, channel );
+		pData += channelOffset;
+
+		struct Box
+		{
+			FVector2f min, size;
+		};
+		vector< Box > transforms( pLayout->GetBlockCount() );
+		for ( int b=0; b<pLayout->GetBlockCount(); ++b )
+		{
+			FIntPoint grid = pLayout->GetGridSize();
+
+			box< vec2<int> > block;
+			pLayout->GetBlock( b, &block.min[0], &block.min[1], &block.size[0], &block.size[1] );
+
+			Box rect;
+			rect.min[0] = ( (float)block.min[0] ) / (float) grid[0];
+			rect.min[1] = ( (float)block.min[1] ) / (float) grid[1];
+			rect.size[0] = ( (float)block.size[0] ) / (float) grid[0];
+			rect.size[1] = ( (float)block.size[1] ) / (float) grid[1];
+			transforms[b] = rect;
+		}
+
+        check( layoutBuffer>=0 && layoutChannel>=0 );
+        check( pApplied->GetVertexBuffers().m_buffers[layoutBuffer].m_channels.Num()==1 );
+        check( pApplied->GetVertexBuffers().m_buffers[layoutBuffer].m_channels[layoutChannel].m_format==MBF_UINT16 );
+        check( pApplied->GetVertexBuffers().m_buffers[layoutBuffer].m_channels[layoutChannel].m_componentCount==1 );
+
+        const uint16* pLayoutData = reinterpret_cast<const uint16*>
+                ( pApplied->GetVertexBuffers().GetBufferData( layoutBuffer ) );
+
+        uint8* pVertices = pData;
+		for ( int v=0; v<pApplied->GetVertexBuffers().GetElementCount(); ++v )
+		{
+			int blockId = pLayoutData[v];
+
+			// TODO: Optimise
+			int relBlock = pLayout->FindBlock( blockId );
+
+			// once we do remove the vertices, this assert can be activated.
+			//check( relBlock>=0 && relBlock<(int)transforms.size() );
+
+			// once we do remove the vertices, this condition is not necessary
+			if (relBlock < 0)
+			{
+				continue;
+			}
+
+			if ( format == MBF_FLOAT32 )
+			{
+				// Not valid in arm because of the potential non-alignement of floats
+				//vec2<float>* pUV = reinterpret_cast<vec2<float>*>( pVertices );
+				//*pUV = (*pUV) * transforms[ relBlock ].size + transforms[ relBlock ].min;
+
+				FVector2f uv;
+				FMemory::Memcpy( &uv, pVertices, sizeof(float)*2 );
+
+				uv = uv * transforms[ relBlock ].size + transforms[ relBlock ].min;
+				FMemory::Memcpy( pVertices, &uv, sizeof(float)*2 );
+			}
+			else if ( format == MBF_FLOAT16 )
+			{
+				// TODO: Optimise
+				float16* pUV = reinterpret_cast<float16*>( pVertices );
+
+				FVector2f uv;
+				uv[0] = halfToFloat( pUV[0] );
+				uv[1] = halfToFloat( pUV[1] );
+				uv = uv * transforms[ relBlock ].size + transforms[ relBlock ].min;
+
+				pUV[0] = floatToHalf( uv[0] );
+				pUV[1] = floatToHalf( uv[1] );
+			}
+
+			else if ( ( format == MBF_NUINT32 )
+						|| ( format == MBF_NINT32 )
+						|| ( format == MBF_UINT32 )
+						|| ( format == MBF_INT32 )
+						)
+			{
+				// TODO: Optimise
+                uint32_t* pUV = reinterpret_cast<uint32_t*>( pVertices );
+
+				for ( int c=0; c<2; ++c )
+				{
+                    uint64_t u_32 = pUV[c];
+                    uint64_t u_48 = u_32 * ((uint64_t)(((float)0xffff) * transforms[ relBlock ].size[c]));
+                    u_48 += (uint64_t)(((float)0xffffffffffULL) * transforms[ relBlock ].min[c]);
+                    pUV[c] = (uint16_t)( u_48 >> 16 );
+				}
+			}
+
+			else if ( ( format == MBF_NUINT16 )
+						|| ( format == MBF_NINT16 )
+						|| ( format == MBF_UINT16 )
+						|| ( format == MBF_INT16 )
+						)
+			{
+				// TODO: Optimise
+                uint16_t* pUV = reinterpret_cast<uint16_t*>( pVertices );
+
+				for ( int c=0; c<2; ++c )
+				{
+                    uint32_t u_16 = pUV[c];
+                    uint32_t u_32 = u_16 * ((uint8_t)(((float)0xffff) * transforms[ relBlock ].size[c]));
+                    u_32 += (uint32_t)(((float)0xffffffff) * transforms[ relBlock ].min[c]);
+                    pUV[c] = (uint16_t)( u_32 >> 16 );
+				}
+			}
+
+			else if ( ( format == MBF_NUINT8 )
+						|| ( format == MBF_NINT8 )
+						|| ( format == MBF_UINT8 )
+						|| ( format == MBF_INT8 )
+						)
+			{
+				// TODO: Optimise
+                uint8_t* pUV = reinterpret_cast<uint8_t*>( pVertices );
+
+				for ( int c=0; c<2; ++c )
+				{
+                    uint32_t u_8 = pUV[c];
+                    uint32_t u_24 = u_8 * ((uint32_t)(((float)0xffff) * transforms[ relBlock ].size[c]));
+                    u_24 += (uint32_t)(((float)0xffffff) * transforms[ relBlock ].min[c]);
+                    pUV[c] = (uint8_t)( u_24 >> 16 );
+				}
+			}
+
+			else
+			{
+				UE_LOG(LogMutableCore, Warning, TEXT("This case is not implemented.."));
+				check( false );
+			}
+
+			pVertices += elemSize;
+		}
+
+		//
+		pApplied->SetLayout( texCoordsSet, pLayout->Clone() );
+
+		//{
+		//	FString debugLog;
+		//	pApplied->GetPrivate()->Log(debugLog, 4);
+		//	UE_LOG(LogMutableCore, Warning, TEXT("after: %s"), *debugLog);
+		//}
+	}
+
+}
diff --git a/Engine/Plugins/Experimental/Mutable/Source/MutableRuntime/Private/OpMeshApplyPose.h b/Engine/Plugins/Experimental/Mutable/Source/MutableRuntime/Private/OpMeshApplyPose.h
new file mode 100644
--- /dev/null
+++ b/Engine/Plugins/Experimental/Mutable/Source/MutableRuntime/Private/OpMeshApplyPose.h
@@ -0,0 +1,128 @@
+// Copyright Epic Games, Inc. All Rights Reserved.
+
+#pragma once
+
+#include "MeshPrivate.h"
+#include "ConvertData.h"
+#include "Platform.h"
+#include "OpMeshChartDifference.h"
+
+#include <memory>
+
+
+namespace mu
+{
+
+    //---------------------------------------------------------------------------------------------
+    //! Reference version
+    //---------------------------------------------------------------------------------------------
+    inline MeshPtr MeshApplyPose(const Mesh* pBase, const Mesh* pPose)
+    {
+        MUTABLE_CPUPROFILER_SCOPE(MeshApplyPose);
+
+        MeshPtr pDest = pBase->Clone();
+
+		mu::SkeletonPtrConst pSkeleton = pBase->GetSkeleton();
+		if (!pSkeleton)
+		{
+			return pDest;
+		}
+
+        // We assume the matrices are transforms from the binding pose bone to the new pose
+        // For now we only convert the vertex positions.
+        // \TODO: normals and tangents
+
+        // Prepare the skin matrices. They may be in different order, and we only need the ones
+        // relevant for the base mesh deformation.
+        vector<FTransform3f> skinTransforms(pSkeleton->GetBoneCount());
+        for ( int32 b=0; b< pSkeleton->GetBoneCount(); ++b )
+        {
+            int32 poseBoneIndex = pPose->FindBonePose(pSkeleton->GetBoneName(b));
+
+            if( poseBoneIndex != INDEX_NONE )
+            {
+				skinTransforms[b] = pPose->m_bonePoses[poseBoneIndex].m_boneTransform;
+            }
+            else
+            {
+                // This could happen if the model is compiled with maximum optimization and the skeleton
+				// has bones used by other meshes.
+				skinTransforms[b] = FTransform3f::Identity;
+            }
+
+        }
+
+        // Get pointers to vertex position data
+		MeshBufferIteratorConst<MBF_FLOAT32, float, 3> itSource( pBase->m_VertexBuffers, MBS_POSITION, 0 );
+		MeshBufferIterator<MBF_FLOAT32, float, 3> itTarget( pDest->m_VertexBuffers, MBS_POSITION, 0 );
+        if (!itSource.ptr() || !itTarget.ptr())
+        {
+            // Formats not implemented
+            check(false);
+            return pDest;
+        }
+
+        // Get pointers to skinning data
+		UntypedMeshBufferIteratorConst itBoneIndices( pBase->m_VertexBuffers, MBS_BONEINDICES, 0 );
+		UntypedMeshBufferIteratorConst itBoneWeights( pBase->m_VertexBuffers, MBS_BONEWEIGHTS, 0 );
+        if (!itBoneIndices.ptr() || !itBoneWeights.ptr())
+        {
+            // No skinning data
+            check(false);
+            return pDest;
+        }
+
+        // Proceed
+        int vertexCount = pBase->GetVertexCount();
+        int weightCount = itBoneIndices.GetComponents();
+        check( weightCount == itBoneWeights.GetComponents() );
+
+        constexpr int MAX_BONES_PER_VERTEX = 16;
+        check( weightCount < MAX_BONES_PER_VERTEX );
+
+        for ( int v=0; v<vertexCount; ++v )
+        {
+            FVector3f sourcePosition = ToUnreal(*itSource);
+			FVector3f position = FVector3f(0,0,0);
+
+            float totalWeight = 0.0f;
+
+            for ( int w=0; w<weightCount; ++w )
+            {
+                float weight[MAX_BONES_PER_VERTEX] = {};
+                ConvertData( w, &weight, MBF_FLOAT32, itBoneWeights.ptr(), itBoneWeights.GetFormat() );
+
+                uint32_t boneIndex[MAX_BONES_PER_VERTEX] = {};
+                ConvertData( w, &boneIndex, MBF_UINT32, itBoneIndices.ptr(), itBoneIndices.GetFormat() );
+
+                totalWeight += weight[w];
+
+                //vec4f p = skinMatrices[ boneIndex[w] ]
+                //        * vec4f( sourcePosition, 1.0f )
+                //        * weight[w];
+				FVector3f p = skinTransforms[boneIndex[w]].TransformPosition(sourcePosition);
+				position += p * weight[w];
+            }
+
+            if (totalWeight<1e-5f)
+            {
+                position = sourcePosition;
+            }
+            else
+            {
+                position /= totalWeight;
+            }
+
+            *itTarget = FromUnreal(position);
+
+            ++itTarget;
+            ++itSource;
+            ++itBoneIndices;
+            ++itBoneWeights;
+        }
+
+
+        return pDest;
+    }
+
+}
