[mutable] Move Mutable plugin from the restricted folder to the experimental plugins folder.

diff --git a/Engine/Plugins/Experimental/Mutable/Source/MutableRuntime/Private/OpMeshBind.h b/Engine/Plugins/Experimental/Mutable/Source/MutableRuntime/Private/OpMeshBind.h
new file mode 100644
--- /dev/null
+++ b/Engine/Plugins/Experimental/Mutable/Source/MutableRuntime/Private/OpMeshBind.h
@@ -0,0 +1,1299 @@
+// Copyright Epic Games, Inc. All Rights Reserved.
+
+#pragma once
+
+#include "MutableRuntime/Private/ConvertData.h"
+#include "MutableRuntime/Private/MeshPrivate.h"
+
+#include "Math/Plane.h"
+#include "Math/Ray.h"
+#include "Math/NumericLimits.h"
+
+#include "IndexTypes.h"
+#include "Spatial/MeshAABBTree3.h"
+
+
+
+// TODO: Make the handling of rotations an option. It is more expensive on CPU and memory, and for some
+// cases it is not required at all.
+
+// TODO: Face stretch to scale the deformation per-vertex? 
+
+// TODO: Support multiple binding influences per vertex, to have smoother deformations.
+
+// TODO: Support multiple binding sets, to have different shapes deformations at once.
+
+// TODO: Deformation mask, to select the intensisty of the deformation per-vertex.
+
+// TODO: This is a reference implementation with ample roof for optimization.
+
+namespace mu
+{
+
+	template<class PositionCompType, class NormalCompType>
+	struct FShapeMeshDescriptor
+	{
+		TArray<UE::Math::TVector<PositionCompType>> Positions;
+		TArray<UE::Math::TVector<NormalCompType>> Normals;
+		TArray<UE::Geometry::FIndex3i> Triangles;
+	};
+
+	using FShapeMeshDescriptorBind = FShapeMeshDescriptor<double, float>;
+		
+	struct FShapeMeshAdapter
+	{
+		const FShapeMeshDescriptorBind& Mesh;
+
+		FShapeMeshAdapter(const FShapeMeshDescriptorBind& InMesh ) 
+			: Mesh(InMesh)
+		{}
+
+		int32 MaxTriangleID() const
+		{
+			return Mesh.Triangles.Num();
+		}
+
+		int32 MaxVertexID() const
+		{
+			return Mesh.Positions.Num();
+		}
+
+		bool IsTriangle(int32 Index) const
+		{
+			return Mesh.Triangles.IsValidIndex(Index);
+		}
+
+		bool IsVertex(int32 Index) const
+		{
+			return Mesh.Positions.IsValidIndex(Index);
+		}
+
+		int32 TriangleCount() const
+		{
+			return Mesh.Triangles.Num();
+		}
+
+		FORCEINLINE int32 VertexCount() const
+		{
+			return Mesh.Positions.Num();
+		}
+
+		FORCEINLINE uint64 GetChangeStamp() const
+		{
+			return 1;
+		}
+
+		FORCEINLINE UE::Geometry::FIndex3i GetTriangle(int32 Index) const
+		{
+			return Mesh.Triangles[Index];
+		}
+
+		FORCEINLINE FVector3d GetVertex(int32 Index) const
+		{
+			return Mesh.Positions[Index];
+		}
+
+		FORCEINLINE void GetTriVertices(int32 TriIndex, FVector3d& V0, FVector3d& V1, FVector3d& V2) const
+		{
+			const UE::Geometry::FIndex3i& Indices = Mesh.Triangles[TriIndex];
+
+			V0 = Mesh.Positions[Indices.A];
+			V1 = Mesh.Positions[Indices.B];
+			V2 = Mesh.Positions[Indices.C];
+		}
+	};
+
+	using FShapeMeshTree = UE::Geometry::TMeshAABBTree3<const FShapeMeshAdapter>;
+
+	// Structure used for vertex bing data in vertexbuffers for reshape operations.
+	struct FReshapeVertexBindingData
+	{
+		// Barycentric coordinates on the shape triangle
+		float S, T;
+
+		// Distance along the interpolated normal of the shape triangle.
+		float D;
+
+		// Index of the shape triangle.
+		int32 Triangle;
+
+
+		// Used to calculate the rotation to apply to the reshaped vertex tangent space.
+		FVector3f ShapeNormal;
+		
+		// Bind point, if the point belongs to a rigid cluster, the attachment point, otherwise
+		// the original point.
+		FVector3f AttachmentPoint;
+
+		// Weight of the effect for this binding. Ranged between 0 and 1 where 0 denotes no effect at all and 1
+		// full effect. This weight should be proportional to the confidence we have that the binding data is valid. 
+		float Weight;
+	};
+	static_assert(sizeof(FReshapeVertexBindingData) == 44);
+
+	struct FReshapeVertexBindingDataBufferDescriptor
+	{
+		constexpr static int Channels = 4;
+		constexpr static MESH_BUFFER_SEMANTIC Semantics[Channels] = { MBS_BARYCENTRICCOORDS, MBS_DISTANCE, MBS_TRIANGLEINDEX, MBS_OTHER };
+		constexpr static MESH_BUFFER_FORMAT Formats[Channels] = { MBF_FLOAT32, MBF_FLOAT32, MBF_INT32, MBF_FLOAT32 };
+		constexpr static int Components[Channels] = { 2, 1, 1, 3+3+1 };
+		constexpr static int Offsets[Channels] = { 0, 8, 12, 16 };
+
+		const int SemanticIndices[Channels] = { 0, 0, 0, 0 };
+
+		FReshapeVertexBindingDataBufferDescriptor(int32 DataSetIndex)
+			: SemanticIndices{ DataSetIndex, DataSetIndex, DataSetIndex, DataSetIndex }
+		{
+		}
+	};
+
+	struct FUsedIndicesBufferDescriptor
+	{
+		constexpr static int Channels = 1;
+		constexpr static MESH_BUFFER_SEMANTIC Semantics[Channels] = { MBS_OTHER };
+		constexpr static MESH_BUFFER_FORMAT Formats[Channels] = { MBF_INT32 };
+		constexpr static int Components[Channels] = { 1 };
+		constexpr static int Offsets[Channels] = { 0 };
+		constexpr static int SemanticIndices[Channels] = { 0 };
+
+		FUsedIndicesBufferDescriptor()
+		{
+		}
+	};
+
+	// Structure used for vertex bing data in vertexbuffers for Clip deform operations.
+	struct FClipDeformVertexBindingData
+	{
+		// Barycentric coordinates on the shape triangle
+		float S, T;
+
+		// Index of the shape triangle.
+		int32 Triangle;
+
+		float Weight;
+	};
+	static_assert(sizeof(FClipDeformVertexBindingData) == 16);
+
+	struct FClipDeformVertexBindingDataBufferDescriptor
+	{
+		constexpr static int Channels = 3;
+		constexpr static MESH_BUFFER_SEMANTIC Semantics[Channels] = { MBS_BARYCENTRICCOORDS, MBS_TRIANGLEINDEX, MBS_OTHER };
+		constexpr static MESH_BUFFER_FORMAT Formats[Channels] = { MBF_FLOAT32, MBF_INT32, MBF_FLOAT32 };
+		constexpr static int Components[Channels] = { 2, 1, 1 };
+		constexpr static int Offsets[Channels] = { 0, 8, 12 };
+
+		const int SemanticIndices[Channels] = { 0, 0, 0 };
+
+		FClipDeformVertexBindingDataBufferDescriptor(int32 DataSetIndex)
+			: SemanticIndices{ DataSetIndex, DataSetIndex, DataSetIndex }
+		{
+		}
+	};
+
+
+	//---------------------------------------------------------------------------------------------
+    //! Generate the mesh-shape binding data for Reshape operations
+    //---------------------------------------------------------------------------------------------
+
+	inline void BindReshapePoint(
+			FShapeMeshTree& ShapeMeshTree,
+			const FVector3f& Point, FReshapeVertexBindingData& OutBindData, 
+			const float ValidityTolerance = UE_KINDA_SMALL_NUMBER)
+	{
+
+		const FShapeMeshDescriptorBind& ShapeMesh = ShapeMeshTree.GetMesh()->Mesh;
+		
+		OutBindData.S = 0.0f;
+		OutBindData.T = 0.0f;
+		OutBindData.D = 0.0f;
+		OutBindData.Triangle = -1;
+		OutBindData.ShapeNormal = FVector3f::Zero();
+		OutBindData.AttachmentPoint = Point;
+		float Weight = 1.0f;
+
+		double DistSqr = 0.0;
+
+		int32 FoundIndex = ShapeMeshTree.FindNearestTriangle( FVector3d(Point), DistSqr );
+			
+		if (FoundIndex < 0)
+		{
+			return;
+		}
+
+		// Calculate the binding data of the base mesh vertex to its bound shape triangle
+		UE::Geometry::FIndex3i Triangle = ShapeMesh.Triangles[FoundIndex];
+
+		// Project on the triangle, but using the vertex normals.
+		// See reference implementation for details. 
+		FVector3f TriangleA = (FVector3f)ShapeMesh.Positions[Triangle.A];
+		FVector3f TriangleB = (FVector3f)ShapeMesh.Positions[Triangle.B];
+		FVector3f TriangleC = (FVector3f)ShapeMesh.Positions[Triangle.C];
+
+		FPlane4f TrianglePlane(TriangleA, TriangleB, TriangleC);
+		FPlane4f VertexPlane(Point, TrianglePlane.GetNormal());
+	
+		// T1 = Triangle projected on the vertex plane along the triangle vertex normals
+		FVector3f TriangleA_VertexPlane = FMath::RayPlaneIntersection(TriangleA, ShapeMesh.Normals[Triangle.A], VertexPlane);
+		FVector3f TriangleB_VertexPlane = FMath::RayPlaneIntersection(TriangleB, ShapeMesh.Normals[Triangle.B], VertexPlane);
+		FVector3f TriangleC_VertexPlane = FMath::RayPlaneIntersection(TriangleC, ShapeMesh.Normals[Triangle.C], VertexPlane);
+	
+		// Barycentric coordinates of the vertex on in T1
+		FVector3f Barycentric = (FVector3f)FMath::ComputeBaryCentric2D((FVector)Point, (FVector)TriangleA_VertexPlane, (FVector)TriangleB_VertexPlane, (FVector)TriangleC_VertexPlane);
+
+		FVector3f ProjectedPoint = TriangleA * Barycentric.X + TriangleB * Barycentric.Y + TriangleC * Barycentric.Z;
+
+		FVector3f InterpolatedVertexNormal = ShapeMesh.Normals[Triangle.A] * Barycentric.X + ShapeMesh.Normals[Triangle.B] * Barycentric.Y + ShapeMesh.Normals[Triangle.C] * Barycentric.Z;
+		FVector3f ProjectedToVertex = (Point - ProjectedPoint);
+
+		// If the interpolated normal is not normalized the dot product gives the distance to the point times the length of the normal.
+		// We don't want that since when deforming a point, using the interpolated normal, we'd multiply by the length twice. See GetDeform().
+		// One option would be to normalize here (computing the distance) and when deforming the mesh.
+		// We can also compensate the signed distance so that the same interpolated normal (not normalized) gives us the correct point
+		// ( InterpolatedVertexNormal * d = ProjectedToVertex). That way modifications in the interpolated bound normal will affect the
+		// resulting point. We are doing the later.
+		const float InterpolatedNormalSizeSquared = InterpolatedVertexNormal.SizeSquared();
+		const float InvInterpolatedNormalSizeSquared = InterpolatedNormalSizeSquared > UE_SMALL_NUMBER ? 1.0f / InterpolatedNormalSizeSquared : 0.0f;	
+		float d = FVector3f::DotProduct(ProjectedToVertex, InterpolatedVertexNormal) * InvInterpolatedNormalSizeSquared;
+	
+		OutBindData.S = Barycentric.Y;
+		OutBindData.T = Barycentric.Z;
+		OutBindData.D = d;
+		
+		OutBindData.Triangle = FoundIndex;
+		
+		OutBindData.ShapeNormal = ((TriangleB - TriangleA) ^ (TriangleC - TriangleA)).GetSafeNormal();
+		
+		const FVector3f ReprojectedPoint = ProjectedPoint + InterpolatedVertexNormal * d;
+		
+		const FVector3f ReprojectedVector = ReprojectedPoint - Point;
+		const float ErrorEstimate = (ReprojectedPoint - Point).GetAbsMax();
+		
+		// If within the tolerance, 1.0, otherwise linear falloff based on the tolerance
+
+		// Arbitrary factor, a binding will be considered valid (with its corresponding weight) to ErrorFalloffFactor times the validity tolerance.
+		constexpr float ErrorFalloffFactor = 4.0f;
+		OutBindData.Weight = 1.0f - FMath::Clamp( (ErrorEstimate - ValidityTolerance) / (ValidityTolerance * ErrorFalloffFactor), 0.0f, 1.0f );
+		
+		OutBindData.Triangle = FMath::IsNearlyZero(OutBindData.Weight) ? -1 : OutBindData.Triangle;
+	}
+
+	inline void BindPointReference(
+			const TArray<UE::Geometry::FIndex3i>& ShapeTriangles, const TArray<FVector3f>& ShapePositions, const TArray<FVector3f>& ShapeNormals,
+			const FVector3f& Point, FReshapeVertexBindingData& BindData,
+			const float ValidityTolerance = UE_KINDA_SMALL_NUMBER)
+	{
+		BindData.S = 0.0f;
+		BindData.T = 0.0f;
+		BindData.D = 0.0f;
+		BindData.Triangle = -1;
+		BindData.ShapeNormal = FVector3f::Zero();
+		BindData.AttachmentPoint = Point;
+		float Weight = 1.0f;
+
+		// Find the best shape triangle for each base mesh vertex
+
+		// Select the triangle closest to the vertex.
+		// \TODO: review this strategy.
+		int32 TriangleIndex = -1;
+		float BestDistSquared = TNumericLimits<float>::Max();
+		for (int ShapeTriangleIndex = 0; ShapeTriangleIndex < ShapeTriangles.Num(); ++ShapeTriangleIndex)
+		{
+			UE::Geometry::FIndex3i Triangle = ShapeTriangles[ShapeTriangleIndex];
+			FVector3f TriangleA = ShapePositions[Triangle.A];
+			FVector3f TriangleB = ShapePositions[Triangle.B];
+			FVector3f TriangleC = ShapePositions[Triangle.C];
+			FVector3f ClosestPoint = (FVector3f)FMath::ClosestPointOnTriangleToPoint((FVector)Point, (FVector)TriangleA, (FVector)TriangleB, (FVector)TriangleC);
+			float DistSquared = (ClosestPoint - Point).SquaredLength();
+			if (DistSquared < BestDistSquared)
+			{
+				TriangleIndex = ShapeTriangleIndex;
+				BestDistSquared = DistSquared;
+			}
+		}
+
+		if (TriangleIndex < 0)
+		{
+			//check(false);
+			return;
+		}
+
+		// Calculate the binding data of the base mesh vertex to its bound shape triangle
+		UE::Geometry::FIndex3i Triangle = ShapeTriangles[TriangleIndex];
+
+		// Project on the triangle, but using the vertex normals.
+		// See figure 3 in https://www.sciencedirect.com/science/article/pii/S1319157820303931
+		FVector3f TriangleA = ShapePositions[Triangle.A];
+		FVector3f TriangleB = ShapePositions[Triangle.B];
+		FVector3f TriangleC = ShapePositions[Triangle.C];
+
+		FPlane4f TrianglePlane(TriangleA, TriangleB, TriangleC);
+		FPlane4f VertexPlane(Point, TrianglePlane.GetNormal());
+	
+		// T1 = Triangle projected on the vertex plane along the triangle vertex normals
+		FVector3f TriangleA_VertexPlane = FMath::RayPlaneIntersection(TriangleA, ShapeNormals[Triangle.A], VertexPlane);
+		FVector3f TriangleB_VertexPlane = FMath::RayPlaneIntersection(TriangleB, ShapeNormals[Triangle.B], VertexPlane);
+		FVector3f TriangleC_VertexPlane = FMath::RayPlaneIntersection(TriangleC, ShapeNormals[Triangle.C], VertexPlane);
+	
+		// Barycentric coordinates of the vertex on in T1
+		FVector3f Barycentric = (FVector3f)FMath::ComputeBaryCentric2D((FVector)Point, (FVector)TriangleA_VertexPlane, (FVector)TriangleB_VertexPlane, (FVector)TriangleC_VertexPlane);
+
+		FVector3f ProjectedPoint = TriangleA * Barycentric.X + TriangleB * Barycentric.Y + TriangleC * Barycentric.Z;
+
+		FVector3f InterpolatedVertexNormal = ShapeNormals[Triangle.A] * Barycentric.X + ShapeNormals[Triangle.B] * Barycentric.Y + ShapeNormals[Triangle.C] * Barycentric.Z;
+		FVector3f ProjectedToVertex = (Point - ProjectedPoint);
+
+		// If the interpolated normal is not normalized the dot product gives the distance to the point times the length of the normal.
+		// We don't want that since when deforming a point, using the interpolated normal, we'd multiply by the length twice. See GetDeform().
+		// One option would be to normalize here (computing the distance) and when deforming the mesh.
+		// We can also compensate the signed distance so that the same interpolated normal (not normalized) gives us the correct point
+		// ( InterpolatedVertexNormal * d = ProjectedToVertex). That way modifications in the interpolated bound normal will affect the
+		// resulting point. We are doing the later.
+		const float InterpolatedNormalSizeSquared = InterpolatedVertexNormal.SizeSquared();
+		const float InvInterpolatedNormalSizeSquared = InterpolatedNormalSizeSquared > SMALL_NUMBER ? 1.0f / InterpolatedNormalSizeSquared : 0.0f;	
+		float d = FVector3f::DotProduct(ProjectedToVertex, InterpolatedVertexNormal) * InvInterpolatedNormalSizeSquared;
+	
+		BindData.S = Barycentric.Y;
+		BindData.T = Barycentric.Z;
+		BindData.D = d;
+		
+		BindData.Triangle = TriangleIndex;
+		
+		BindData.ShapeNormal = ((TriangleB - TriangleA) ^ (TriangleC - TriangleA)).GetSafeNormal();
+		
+		const FVector3f ReprojectedPoint = ProjectedPoint + InterpolatedVertexNormal * d;
+		
+		const FVector3f ReprojectedVector = ReprojectedPoint - Point;
+		const float ErrorEstimate = (ReprojectedPoint - Point).GetAbsMax();
+		
+		// If within the tolerance, 1.0, otherwise linear falloff based on the tolerance
+
+		// Arbitrary factor, a binding will be considered valid (with its corresponding weight) to ErrorFalloffFactor times the validity tolerance.
+		constexpr float ErrorFalloffFactor = 4.0f;
+		BindData.Weight = 1.0f - FMath::Clamp( (ErrorEstimate - ValidityTolerance) / (ValidityTolerance * ErrorFalloffFactor), 0.0f, 1.0f );
+		
+		BindData.Triangle = FMath::IsNearlyZero(BindData.Weight) ? -1 : BindData.Triangle;
+	}
+	
+    //---------------------------------------------------------------------------------------------
+    //! Find mesh clusters.
+	//! Colour {1,1,1} is reserved for the non rigid cluster
+    //---------------------------------------------------------------------------------------------
+	
+	inline void FindRigidClusters( const Mesh* Mesh, TArray<TArray<int32>>& OutClusters, int32& OutNonRigidClusterIdx)
+	{
+		MUTABLE_CPUPROFILER_SCOPE(MeshFindRigidClusters);
+		
+		const FVector3f NonRigidColour = FVector3f(1.0f, 1.0f, 1.0f);
+
+		OutNonRigidClusterIdx = -1;
+		
+		const UntypedMeshBufferIteratorConst ItColoursBase(Mesh->GetVertexBuffers(), MBS_COLOUR);
+		if (ItColoursBase.ptr())
+		{
+			const int32 VertexCount = Mesh->GetVertexCount();
+
+			// Check if all colors are the same. (maybe should only search for white as it indicates no rigidity)
+			{
+				int32 VertexIndex = 0;
+				for (; VertexIndex < VertexCount; ++VertexIndex)
+				{
+					 // Colours are treated as ids, no need to use tolerance with the comparison.
+					 if (NonRigidColour != (ItColoursBase + VertexIndex).GetAsVec3f())
+					 {
+						  break;
+					 }
+				}
+
+				// If all equal to the non rigid we are done. 
+				if (VertexIndex == VertexCount)
+				{
+					 return;
+				}
+			}
+			
+			OutClusters.Empty(16);
+			TArray<TArray<int32>>& ClusterData = OutClusters;
+			// Consider using NUINT8 colors for id and pack them in a uint32 as key.
+			TMap<FVector3f, int32> ClusterSet;
+
+			for (int32 VertexIndex = 0; VertexIndex < VertexCount; ++VertexIndex)
+			{
+				FVector3f Colour = (ItColoursBase + VertexIndex).GetAsVec3f();
+				int32& ClusterIdx = ClusterSet.FindOrAdd( Colour, INDEX_NONE );
+				TArray<int32>& Cluster = ClusterIdx < 0 ? ClusterData.Emplace_GetRef() : ClusterData[ClusterIdx];
+				if (ClusterIdx < 0)
+				{
+					ClusterIdx = ClusterData.Num() - 1;
+					Cluster.Reserve(32);
+				}
+
+				Cluster.Push(VertexIndex);	
+			}
+
+			// Colour {1, 1, 1} is reserved for the nonrigid cluster. The choice if this value is not arbitrary,
+			// meshes without colour will get white as default.
+			const int32* NonRigidClusterIdxFound = ClusterSet.Find( NonRigidColour );
+
+			// If not found, add an empty cluster for the non rigid id.
+			OutNonRigidClusterIdx = NonRigidClusterIdxFound ? *NonRigidClusterIdxFound : ClusterData.Emplace();
+		}
+	}
+	
+	inline void FindBindingForCluster( 
+			const Mesh* BaseMesh, FShapeMeshTree& ShapeMeshTree,
+			const TArray<int32>& Cluster, FReshapeVertexBindingData& OutBindingData, float BindTolerance )
+	{
+		const UntypedMeshBufferIteratorConst ItPositionBase(BaseMesh->GetVertexBuffers(), MBS_POSITION);
+
+		FBox3f ClusterBoundingBox(ForceInit);
+		for ( const int32& V : Cluster)
+		{
+			ClusterBoundingBox += (ItPositionBase + V).GetAsVec3f();
+		}
+
+		FVector3f BoundPoint = ClusterBoundingBox.GetCenter();
+		
+		BindReshapePoint(ShapeMeshTree, BoundPoint, OutBindingData, BindTolerance);
+	}
+
+	inline TTuple<TArray<FReshapeVertexBindingData>, TArray<int32>> BindPhysicsBodies( 
+			const PhysicsBody* PBody, FShapeMeshTree& ShapeMeshTree, const Mesh* pMesh, 
+			bool bDeformAllPhysics, const TArray<string>& PhysicsToDeform )
+	{
+		TTuple<TArray<FReshapeVertexBindingData>, TArray<int32>> ReturnValue;
+
+		TArray<int32>& BodiesToDeformIndices = ReturnValue.Get<1>();
+	
+		if (!bDeformAllPhysics)
+		{ 
+			BodiesToDeformIndices.Reserve(PhysicsToDeform.Num());
+			const int32 NumBodies = PBody->GetBodyCount();
+			for (int32 I = 0; I < NumBodies; ++I)
+			{
+				if ( PhysicsToDeform.Contains(string(PBody->GetBodyBoneName(I))) )
+				{
+					BodiesToDeformIndices.Add(I);
+				}
+			}
+		}
+		else
+		{
+			const int32 NumBodies = PBody->GetBodyCount();
+			BodiesToDeformIndices.SetNumUninitialized(NumBodies);
+			for (int32 I = 0; I < NumBodies; ++I)
+			{
+				BodiesToDeformIndices[I] = I;
+			}
+		}
+		
+		// Count how many poinst will be needed.
+		const int32 NumBodiesToDeformIndices = BodiesToDeformIndices.Num();
+		int32 NumPoints = 0;
+		for ( const int32 I : BodiesToDeformIndices )
+		{
+			NumPoints += PBody->GetSphereCount(I) * 8;
+			NumPoints += PBody->GetBoxCount(I) * 8;
+			NumPoints += PBody->GetSphylCount(I) * 8;
+			NumPoints += PBody->GetTaperedCapsuleCount(I) * 8;
+
+			const int32 ConvexCount = PBody->GetConvexCount(I);
+			for ( int C = 0; C < ConvexCount; ++C )
+			{
+				const FVector3f* VertData;
+				int32 VertCount;
+				const int32* IndicesData;
+				int32 IndicesCount;
+				FTransform3f Transform;
+				PBody->GetConvex(I, C, VertData, VertCount, IndicesData, IndicesCount, Transform);
+
+				NumPoints += VertCount;
+			}
+		}
+		
+		TArray<FVector3f> Points;
+		Points.SetNumUninitialized(NumPoints);
+	
+		// Bone transform needs to be applied to the body's sample points so they are in mesh space.
+
+		// Create a point soup to be deformed based on the shapes in the aggregate.
+		// Currently using oriented bounding boxes corners for all shapes except convex which uses the deformed mesh directly.
+
+		int32 AddedPoints = 0;
+		for ( const int32 B : BodiesToDeformIndices )
+		{
+			const char* BoneName = PBody->GetBodyBoneName( B );
+			int32 BoneIdx = pMesh->FindBonePose( BoneName );
+
+			FTransform3f T = FTransform3f::Identity;
+			if (BoneIdx > 0)
+			{
+				pMesh->GetBoneTransform(BoneIdx, T);
+			}
+			
+			const int32 SphereCount = PBody->GetSphereCount(B);
+			for ( int32 I = 0; I < SphereCount; ++I, AddedPoints += 8)
+			{
+				FVector3f P;
+				float R;
+
+				PBody->GetSphere(B, I, P, R);	
+
+				Points[AddedPoints + 0] = T.TransformPosition( P + FVector3f(R, R, R)  );
+				Points[AddedPoints + 1] = T.TransformPosition( P + FVector3f(-R, R, R)  );
+				Points[AddedPoints + 2] = T.TransformPosition( P + FVector3f(-R, -R, R)  );
+				Points[AddedPoints + 3] = T.TransformPosition( P + FVector3f(R, -R, R)  );
+
+														  
+				Points[AddedPoints + 4] = T.TransformPosition( P + FVector3f(R, R, -R) );
+				Points[AddedPoints + 5] = T.TransformPosition( P + FVector3f(-R, R, -R) );
+				Points[AddedPoints + 6] = T.TransformPosition( P + FVector3f(-R, -R, -R) );
+				Points[AddedPoints + 7] = T.TransformPosition( P + FVector3f(R, -R, -R) );
+
+			}
+
+			const int32 BoxCount = PBody->GetBoxCount(B);
+			for ( int32 I = 0; I < BoxCount; ++I, AddedPoints += 8)
+			{
+				FVector3f P;
+				FQuat4f Q;
+				FVector3f S;
+			
+				PBody->GetBox(B, I, P, Q, S);
+
+				Points[AddedPoints + 0] = T.TransformPosition( P + Q.RotateVector(FVector3f(S.X, S.Y, S.Z) ) );
+				Points[AddedPoints + 1] = T.TransformPosition( P + Q.RotateVector(FVector3f(-S.X, S.Y, S.Z) ) );
+				Points[AddedPoints + 2] = T.TransformPosition( P + Q.RotateVector(FVector3f(-S.X, -S.Y, S.Z) ) );
+				Points[AddedPoints + 3] = T.TransformPosition( P + Q.RotateVector(FVector3f(S.X, -S.Y, S.Z) ) );
+
+														  
+				Points[AddedPoints + 4] = T.TransformPosition( P + Q.RotateVector(FVector3f(S.X, S.Y, -S.Z) ) );
+				Points[AddedPoints + 5] = T.TransformPosition( P + Q.RotateVector(FVector3f(-S.X, S.Y, -S.Z) ) );
+				Points[AddedPoints + 6] = T.TransformPosition( P + Q.RotateVector(FVector3f(-S.X, -S.Y, -S.Z) ) );
+				Points[AddedPoints + 7] = T.TransformPosition( P + Q.RotateVector(FVector3f(S.X, -S.Y, -S.Z) ) );
+			}
+
+			const int32 SphylCount = PBody->GetSphylCount(B);
+			for ( int32 I = 0; I < SphylCount; ++I, AddedPoints += 8)
+			{
+				FVector3f P;
+				FQuat4f Q;
+				float R;
+				float L;
+			
+				PBody->GetSphyl(B, I, P, Q, R, L);
+				
+				const FVector3f S = FVector3f(R, R, (L + R*2.0f) * 0.5f);
+
+				Points[AddedPoints + 0] = T.TransformPosition( P + Q.RotateVector(FVector3f(S.X, S.Y, S.Z) ) );
+				Points[AddedPoints + 1] = T.TransformPosition( P + Q.RotateVector(FVector3f(-S.X, S.Y, S.Z) ) );
+				Points[AddedPoints + 2] = T.TransformPosition( P + Q.RotateVector(FVector3f(-S.X, -S.Y, S.Z) ) );
+				Points[AddedPoints + 3] = T.TransformPosition( P + Q.RotateVector(FVector3f(S.X, -S.Y, S.Z) ) );
+
+														  
+				Points[AddedPoints + 4] = T.TransformPosition( P + Q.RotateVector(FVector3f(S.X, S.Y, -S.Z) ) );
+				Points[AddedPoints + 5] = T.TransformPosition( P + Q.RotateVector(FVector3f(-S.X, S.Y, -S.Z) ) );
+				Points[AddedPoints + 6] = T.TransformPosition( P + Q.RotateVector(FVector3f(-S.X, -S.Y, -S.Z) ) );
+				Points[AddedPoints + 7] = T.TransformPosition( P + Q.RotateVector(FVector3f(S.X, -S.Y, -S.Z) ) );
+			}
+	
+			const int32 TaperedCapsuleCount = PBody->GetTaperedCapsuleCount(B);
+			for ( int32 I = 0; I < TaperedCapsuleCount; ++I, AddedPoints += 8)
+			{
+				FVector3f P;
+				FQuat4f Q;
+				float R0;
+				float R1;
+				float L;
+			
+				PBody->GetTaperedCapsule(B, I, P, Q, R0, R1, L);
+				
+				const FVector3f ST = FVector3f(R0, R0, L*0.5f);
+				const FVector3f SB = FVector3f(R1, R1, L*0.5f);
+				
+				Points[AddedPoints + 0] = T.TransformPosition( P + Q.RotateVector(FVector3f(ST.X, ST.Y, ST.Z + R0) ) );
+				Points[AddedPoints + 1] = T.TransformPosition( P + Q.RotateVector(FVector3f(-ST.X, ST.Y, ST.Z + R0) ) );
+				Points[AddedPoints + 2] = T.TransformPosition( P + Q.RotateVector(FVector3f(-ST.X, -ST.Y, ST.Z + R0) ) );
+				Points[AddedPoints + 3] = T.TransformPosition( P + Q.RotateVector(FVector3f(ST.X, -ST.Y, ST.Z + R0) ) );
+
+				Points[AddedPoints + 4] = T.TransformPosition( P + Q.RotateVector(FVector3f(SB.X, SB.Y, -SB.Z - R1) ) );
+				Points[AddedPoints + 5] = T.TransformPosition( P + Q.RotateVector(FVector3f(-SB.X, SB.Y, -SB.Z - R1) ) );
+				Points[AddedPoints + 6] = T.TransformPosition( P + Q.RotateVector(FVector3f(-SB.X, -SB.Y, -SB.Z - R1) ) );
+				Points[AddedPoints + 7] = T.TransformPosition( P + Q.RotateVector(FVector3f(SB.X, -SB.Y, -SB.Z - R1) ) );
+			}
+
+			const int32 ConvexCount = PBody->GetConvexCount(B);
+			for ( int32 I = 0; I < ConvexCount; ++I )
+			{
+				const FVector3f* VertexData;
+				int32 VertexCount;
+				const int32* IndicesData;
+				int32 IndicesCount;
+				FTransform3f ConvexT;
+				
+				PBody->GetConvex(B, I, VertexData, VertexCount, IndicesData, IndicesCount, ConvexT);
+				
+				TArrayView<const FVector3f> VertexDataView( VertexData, VertexCount );
+
+				ConvexT = T * ConvexT;
+				for ( const FVector3f& P : VertexDataView )
+				{
+					Points[AddedPoints++] = ConvexT.TransformPosition(P);
+				}	
+			}
+		}
+		
+		TArray<FReshapeVertexBindingData>& PhysicsBodyBindData = ReturnValue.Get<0>();
+		PhysicsBodyBindData.SetNumUninitialized( NumPoints );
+
+		for (int32 I = 0; I < NumPoints; ++I)
+		{
+			BindReshapePoint(ShapeMeshTree, Points[I], PhysicsBodyBindData[I], 0.01f);
+		}
+
+		return ReturnValue;
+	}
+
+	
+	inline TArray<FReshapeVertexBindingData> BindVerticesReshape(
+			const Mesh* BaseMesh, FShapeMeshTree& ShapeMeshTree, bool bEnableRigidParts )	
+	{
+	
+		UE::Geometry::FAxisAlignedBox3d ShapeAABBox = ShapeMeshTree.GetBoundingBox();
+
+		const float BindValidityTolerance = ShapeAABBox.MaxDim() < 1.0 
+										  ? UE_KINDA_SMALL_NUMBER
+										  : static_cast<float>(ShapeAABBox.MaxDim()) * 1e-3f;
+
+		TArray<TArray<int32>> VertexClusters;
+		int32 NonRigidClusterIdx = -1;
+
+		if (bEnableRigidParts)
+		{
+			FindRigidClusters( BaseMesh, VertexClusters, NonRigidClusterIdx );
+		}
+		
+		// Find nearest shape triangle for each base mesh vertex		
+		const int32 MeshVertexCount = BaseMesh->GetVertexCount();
+		TArray<FReshapeVertexBindingData> BindData;
+		{
+			MUTABLE_CPUPROFILER_SCOPE(Project);
+
+			BindData.SetNum( MeshVertexCount );
+
+			UntypedMeshBufferIteratorConst ItPositionBase(BaseMesh->GetVertexBuffers(), MBS_POSITION);
+
+			// Special case for non rigid parts
+
+			// This indicates that we don't care about rigid parts,
+			// only one clusters is found or there is no cluster data or rigid parts is disabled.
+			if (NonRigidClusterIdx < 0)
+			{
+				const int32 VertexCount = BaseMesh->GetVertexCount();
+				for (int32 VertexIndex = 0; VertexIndex < VertexCount; ++VertexIndex)
+				{
+					FVector3f VertexPosition = (ItPositionBase + VertexIndex).GetAsVec3f();
+					BindReshapePoint(ShapeMeshTree, VertexPosition, BindData[VertexIndex], BindValidityTolerance);
+				}
+			}
+			else
+			{
+				TArray<int32>& NonRigidCluster = VertexClusters[NonRigidClusterIdx];
+				const int32 NonRigidVertexCount = NonRigidCluster.Num();
+
+				for (int32 I = 0; I < NonRigidVertexCount; ++I)
+				{
+					const int32 VertexIndex = NonRigidCluster[I];
+					FVector3f VertexPosition = (ItPositionBase + VertexIndex).GetAsVec3f();
+					BindReshapePoint(ShapeMeshTree, VertexPosition, BindData[VertexIndex], BindValidityTolerance);
+				}
+
+				// Remove data form the non rigid cluster so it ios not processed in the rigid parts binding step.
+				NonRigidCluster.Reset();
+			
+				for (const TArray<int32>& RigidCluster : VertexClusters)
+				{
+					FReshapeVertexBindingData ClusterBinding;
+
+					// Copy cluster binding to every vertex of the cluster.
+					const int32 ClusterVertexCount = RigidCluster.Num();
+					if (ClusterVertexCount)
+					{
+						FindBindingForCluster(BaseMesh, ShapeMeshTree, RigidCluster, ClusterBinding, BindValidityTolerance );
+						for (int32 I = 0; I < ClusterVertexCount; ++I)
+						{
+							const int32 VertexIndex = RigidCluster[I];
+							BindData[VertexIndex] = ClusterBinding;
+						}
+					}
+				}
+			}
+		}
+
+		return BindData;
+	}
+
+	inline TTuple<TArray<FReshapeVertexBindingData>, TArray<int32>> BindPose(
+			const Mesh* Mesh, FShapeMeshTree& ShapeMeshTree, 
+			bool bDeformAllBones, const TArray<string>& BonesToDeform )
+	{
+		UE::Geometry::FAxisAlignedBox3d ShapeAABBox = ShapeMeshTree.GetBoundingBox();
+
+		const float BindValidityTolerance = ShapeAABBox.MaxDim() < 1.0 
+										  ? UE_KINDA_SMALL_NUMBER
+										  : static_cast<float>(ShapeAABBox.MaxDim()) * 1e-3f;
+		
+		TTuple<TArray<FReshapeVertexBindingData>, TArray<int32>> ReturnValue;
+
+		TArray<FReshapeVertexBindingData>& SkeletonBindDataArray = ReturnValue.Get<0>();
+		TArray<int32>& BoneIndices = ReturnValue.Get<1>();
+		
+		for ( int32 b = 0; b < Mesh->GetBonePoseCount(); ++b)
+		{
+			if (!bDeformAllBones)
+			{
+				bool bFound = false;
+				for (int32 bd = 0; bd < BonesToDeform.Num(); ++bd)
+				{
+					if (Mesh->GetBonePoseName(b) == BonesToDeform[bd])
+					{
+						bFound = true;
+						break;
+					}
+				}
+
+				if (!bFound)
+				{
+					continue;
+				}
+			}
+
+			const FTransform3f& t = Mesh->m_bonePoses[b].m_boneTransform;
+
+			FVector3f Origin = t.GetLocation();
+
+			FReshapeVertexBindingData SkeletonBindData;
+			BindReshapePoint(ShapeMeshTree, Origin, SkeletonBindData, BindValidityTolerance);
+
+			SkeletonBindDataArray.Add(SkeletonBindData);
+			BoneIndices.Add(b);
+		}
+
+		return ReturnValue;
+	}
+
+
+	//---------------------------------------------------------------------------------------------
+    //! Generate the mesh-shape binding data
+    //---------------------------------------------------------------------------------------------
+    inline MeshPtr MeshBindShapeReshape(
+			const Mesh* BaseMesh, const Mesh* ShapeMesh, 
+			const TArray<string>& BonesToDeform, const TArray<string>& PhysicsToDeform,
+			bool bDeformAllBones, bool bDeformAllPhyiscs,
+			bool reshapeVertices, bool reshapeSkeleton, bool reshapePhysicsVolumes, bool bEnableRigidParts )
+    {
+		MUTABLE_CPUPROFILER_SCOPE(MeshBindShape);
+		
+		if (!BaseMesh)
+		{
+			return nullptr;
+		}
+
+		// Early out if nothing will be modified and the vertices discarted. return null in this
+		// case indicating nothing has modified so the Base Mesh can be reused.
+		const bool bSkeletonModification = BaseMesh->GetSkeleton() && reshapeSkeleton;
+		const bool bPhysicsModification = BaseMesh->GetPhysicsBody() && reshapePhysicsVolumes;
+
+		if (!reshapeVertices && !bSkeletonModification && !bPhysicsModification)
+		{
+			return nullptr;
+		}
+
+		if (!ShapeMesh)
+		{
+			return BaseMesh->Clone();
+		}
+	
+		int32 ShapeVertexCount = ShapeMesh->GetVertexCount();
+		int ShapeTriangleCount = ShapeMesh->GetFaceCount();
+		if (!ShapeVertexCount || !ShapeTriangleCount)
+		{
+			return BaseMesh->Clone();
+		}	
+
+		FShapeMeshDescriptorBind ShapeMeshDescriptor;
+		{
+			MUTABLE_CPUPROFILER_SCOPE(GenerateVertexQueryData);	
+
+			ShapeMeshDescriptor.Positions.SetNum(ShapeVertexCount);
+			ShapeMeshDescriptor.Normals.SetNum(ShapeVertexCount);
+
+			// \TODO: Simple but inefficient
+			UntypedMeshBufferIteratorConst ItPosition(ShapeMesh->GetVertexBuffers(), MBS_POSITION);
+			UntypedMeshBufferIteratorConst ItNormal(ShapeMesh->GetVertexBuffers(), MBS_NORMAL);
+			for (int32 ShapeVertexIndex = 0; ShapeVertexIndex < ShapeVertexCount; ++ShapeVertexIndex)
+			{
+				FVector3f Position = ItPosition.GetAsVec3f();
+				ShapeMeshDescriptor.Positions[ShapeVertexIndex] = static_cast<FVector3d>(Position);
+				++ItPosition;
+				
+				FVector3f Normal = ItNormal.GetAsVec3f();
+				ShapeMeshDescriptor.Normals[ShapeVertexIndex] = Normal;
+				++ItNormal;
+			}
+		}
+		// Generate the temp face query data for the shape
+		// TODO: Index data copy may be saved in most cases.
+		{
+			MUTABLE_CPUPROFILER_SCOPE(GenerateTrianglesQueryData);	
+			ShapeMeshDescriptor.Triangles.SetNum(ShapeTriangleCount);
+			// \TODO: Simple but inefficient
+			UntypedMeshBufferIteratorConst ItIndices(ShapeMesh->GetIndexBuffers(), MBS_VERTEXINDEX);
+			for (int32 TriangleIndex = 0; TriangleIndex < ShapeTriangleCount; ++TriangleIndex)
+			{
+				UE::Geometry::FIndex3i Triangle;
+				Triangle.A = int(ItIndices.GetAsUINT32());
+				++ItIndices;
+				Triangle.B = int(ItIndices.GetAsUINT32());
+				++ItIndices;
+				Triangle.C = int(ItIndices.GetAsUINT32());
+				++ItIndices;
+
+				ShapeMeshDescriptor.Triangles[TriangleIndex] = Triangle;
+			}	
+		}
+
+		FShapeMeshAdapter ShapeMeshAdapter(ShapeMeshDescriptor);
+
+		constexpr bool bAutoBuildTree = false;
+		FShapeMeshTree ShapeMeshTree(&ShapeMeshAdapter, bAutoBuildTree);	
+		{	
+			MUTABLE_CPUPROFILER_SCOPE(BuildShapeTree);
+			
+			ShapeMeshTree.Build();
+		}
+	
+		Ptr<Mesh> Result;
+		
+		// If no vertices are needed, it is assumed we only want to reshape physics or skeleton
+		// In that case, remove everything except physics bodies, the skeleton and pose.
+		if (!reshapeVertices)
+		{
+			constexpr EMeshCloneFlags CloneFlags = 
+					EMeshCloneFlags::WithSkeleton    | 
+					EMeshCloneFlags::WithPhysicsBody | 
+					EMeshCloneFlags::WithPoses;
+
+			Result = BaseMesh->Clone(CloneFlags);
+		}
+		else
+		{	
+			Result = BaseMesh->Clone();
+		}
+
+		int32 BindingDataIndex = 0;
+		if (reshapeVertices)
+		{
+			TArray<FReshapeVertexBindingData> VerticesBindData = BindVerticesReshape( BaseMesh, ShapeMeshTree, bEnableRigidParts );
+			
+			// Add the binding information to the mesh
+			// \TODO: Check that there is no other binding data.
+			// \TODO: Support specifying the binding data channel for multiple binding support.
+			FMeshBufferSet& VB = Result->GetVertexBuffers();
+			int NewBufferIndex = VB.GetBufferCount();
+			VB.SetBufferCount(NewBufferIndex+1);
+
+			FReshapeVertexBindingDataBufferDescriptor BufDesc(BindingDataIndex);
+			VB.SetBuffer(NewBufferIndex, sizeof(FReshapeVertexBindingData), BufDesc.Channels, BufDesc.Semantics, BufDesc.SemanticIndices, BufDesc.Formats, BufDesc.Components);
+			FMemory::Memcpy(VB.GetBufferData(NewBufferIndex), VerticesBindData.GetData(), VerticesBindData.Num() * sizeof(FReshapeVertexBindingData));
+		}
+
+		// Bind the skeleton bones
+		// \TODO: Build bind data only for actually modified bones?
+		if (reshapeSkeleton && (bDeformAllBones || BonesToDeform.Num()))
+		{
+			MUTABLE_CPUPROFILER_SCOPE(BindSkeleton);
+
+			TTuple<TArray<FReshapeVertexBindingData>, TArray<int32>> SkeletonBindingData = 
+					BindPose(Result.get(), ShapeMeshTree, bDeformAllBones, BonesToDeform);
+			const TArray<FReshapeVertexBindingData>& SkeletonBindDataArray = SkeletonBindingData.Get<0>();
+			const TArray<int32>& BoneIndices = SkeletonBindingData.Get<1>();
+
+			check(BoneIndices.Num() == SkeletonBindDataArray.Num());
+
+			const int32 NumBonesToDeform = SkeletonBindDataArray.Num();
+
+			FMeshBufferSet SkeletonBuffer;
+			SkeletonBuffer.SetBufferCount(2);
+			SkeletonBuffer.SetElementCount(SkeletonBindDataArray.Num());
+
+			FReshapeVertexBindingDataBufferDescriptor BufDesc(BindingDataIndex);
+			SkeletonBuffer.SetBuffer(0, sizeof(FReshapeVertexBindingData), BufDesc.Channels, BufDesc.Semantics, BufDesc.SemanticIndices, BufDesc.Formats, BufDesc.Components);
+
+			// Bone indices buffer
+			MESH_BUFFER_SEMANTIC BoneSemantics[1] = { MBS_OTHER };
+			MESH_BUFFER_FORMAT BoneFormats[1] = { MBF_INT32 };
+			int BoneSemanticIndices[1] = { 0 };
+			int BoneComponents[1] = { 1 };
+			int BoneOffsets[1] = { 0 };
+
+			SkeletonBuffer.SetBuffer(1, sizeof(int32), 1, BoneSemantics, BoneSemanticIndices, BoneFormats, BoneComponents, BoneOffsets);
+
+			FMemory::Memcpy(SkeletonBuffer.GetBufferData(0), SkeletonBindDataArray.GetData(), NumBonesToDeform * sizeof(FReshapeVertexBindingData));
+			FMemory::Memcpy(SkeletonBuffer.GetBufferData(1), BoneIndices.GetData(), NumBonesToDeform * sizeof(int32));
+
+			Result->m_AdditionalBuffers.Emplace(EMeshBufferType::SkeletonDeformBinding, MoveTemp(SkeletonBuffer));
+		}
+
+		const PhysicsBody* PhysicsBody = Result->m_pPhysicsBody.get();
+		if (reshapePhysicsVolumes && PhysicsBody && (bDeformAllPhyiscs || PhysicsToDeform.Num()))
+		{
+			MUTABLE_CPUPROFILER_SCOPE(BindPhysicsBody);
+
+			TTuple<TArray<FReshapeVertexBindingData>, TArray<int32>> PhysicsBindingData = 
+					BindPhysicsBodies(PhysicsBody, ShapeMeshTree, Result.get(), bDeformAllPhyiscs, PhysicsToDeform);
+			const TArray<FReshapeVertexBindingData>& PhysicsBindDataArray = PhysicsBindingData.Get<0>();
+			const TArray<int32>& DeformedBodyIndices = PhysicsBindingData.Get<1>();
+
+			FMeshBufferSet PhysicsBodyBuffer;
+			PhysicsBodyBuffer.SetBufferCount(1);
+			PhysicsBodyBuffer.SetElementCount(PhysicsBindDataArray.Num());
+
+			constexpr int32 BindDataIndex = 0; 
+			FReshapeVertexBindingDataBufferDescriptor BufDesc(BindingDataIndex);
+			PhysicsBodyBuffer.SetBuffer(0, sizeof(FReshapeVertexBindingData), BufDesc.Channels, BufDesc.Semantics, BufDesc.SemanticIndices, BufDesc.Formats, BufDesc.Components, BufDesc.Offsets);	
+			FMemory::Memcpy(PhysicsBodyBuffer.GetBufferData(0), PhysicsBindDataArray.GetData(), PhysicsBindDataArray.Num() * sizeof(FReshapeVertexBindingData));
+
+			// Bone indices buffer
+			MESH_BUFFER_SEMANTIC BoneSemantics[1] = { MBS_OTHER };
+			MESH_BUFFER_FORMAT BoneFormats[1] = { MBF_INT32 };
+			int BoneSemanticIndices[1] = { 0 };
+			int BoneComponents[1] = { 1 };
+			int BoneOffsets[1] = { 0 };
+
+			FMeshBufferSet PhysicsBodySelectionBuffer;
+			PhysicsBodySelectionBuffer.SetBufferCount(1);
+			PhysicsBodySelectionBuffer.SetElementCount(DeformedBodyIndices.Num());
+			PhysicsBodySelectionBuffer.SetBuffer(0, sizeof(int32), 1, BoneSemantics, BoneSemanticIndices, BoneFormats, BoneComponents, BoneOffsets);	
+			FMemory::Memcpy(PhysicsBodySelectionBuffer.GetBufferData(0), DeformedBodyIndices.GetData(), DeformedBodyIndices.Num() * sizeof(int32));
+
+			Result->m_AdditionalBuffers.Emplace(EMeshBufferType::PhysicsBodyDeformBinding, MoveTemp(PhysicsBodyBuffer));
+			Result->m_AdditionalBuffers.Emplace(EMeshBufferType::PhysicsBodyDeformSelection, MoveTemp(PhysicsBodySelectionBuffer));
+		}
+		
+        return Result;
+    }
+
+	//---------------------------------------------------------------------------------------------
+    //! Generate the mesh-shape binding data for ClipDeform operations
+    //---------------------------------------------------------------------------------------------
+	inline void BindClipDeformPointClosestProject(
+			FShapeMeshTree& ShapeMeshTree,
+			const FVector3f& Point, FClipDeformVertexBindingData& OutBindData, float ValidityTolerance )
+	{
+		const FShapeMeshDescriptorBind& ShapeMesh = ShapeMeshTree.GetMesh()->Mesh;
+		
+		OutBindData.S = 0.0f;
+		OutBindData.T = 0.0f;
+		OutBindData.Weight = 0.0f;
+		OutBindData.Triangle = -1;
+
+		double DistSqr = 0.0;
+
+		int32 FoundIndex = ShapeMeshTree.FindNearestTriangle( FVector3d(Point), DistSqr );
+			
+		if (FoundIndex < 0)
+		{
+			return;
+		}
+
+		// Calculate the binding data of the base mesh vertex to its bound shape triangle
+		UE::Geometry::FIndex3i TriangleIndices = ShapeMesh.Triangles[FoundIndex];
+
+		// Project on the triangle, but using the vertex normals.
+		// See reference implementation for details.
+		UE::Geometry::FTriangle3f ShapeTriangle {
+				(FVector3f)ShapeMesh.Positions[TriangleIndices.A],
+				(FVector3f)ShapeMesh.Positions[TriangleIndices.B],
+				(FVector3f)ShapeMesh.Positions[TriangleIndices.C] };
+
+		const FVector3f ShapeTriangleNormal = ShapeTriangle.Normal();
+		const FPlane4f VertexPlane(Point, ShapeTriangleNormal);
+	
+		const UE::Geometry::FTriangle3f TriangleVertexPlane {
+				FMath::RayPlaneIntersection(ShapeTriangle.V[0], ShapeMesh.Normals[TriangleIndices.A], VertexPlane),
+				FMath::RayPlaneIntersection(ShapeTriangle.V[1], ShapeMesh.Normals[TriangleIndices.B], VertexPlane),
+				FMath::RayPlaneIntersection(ShapeTriangle.V[2], ShapeMesh.Normals[TriangleIndices.C], VertexPlane) };
+
+		FVector3f Barycentric = TriangleVertexPlane.GetBarycentricCoords(Point);
+
+		FVector3f InterpolatedShapeNormal = 
+				ShapeMesh.Normals[TriangleIndices.A] * Barycentric.X + 
+				ShapeMesh.Normals[TriangleIndices.B] * Barycentric.Y + 
+				ShapeMesh.Normals[TriangleIndices.C] * Barycentric.Z;
+
+		FVector3f BindPoint = ShapeTriangle.BarycentricPoint(Barycentric);
+		FVector3f ProjectedToVertex = (Point - BindPoint);
+
+		// Compute reprojection value to see if the binding is valid.
+		const float InterpolatedNormalSizeSquared = InterpolatedShapeNormal.SizeSquared();
+		const float InvInterpolatedNormalSizeSquared = InterpolatedNormalSizeSquared > SMALL_NUMBER ? 1.0f / InterpolatedNormalSizeSquared : 0.0f;	
+		float D = FVector3f::DotProduct(ProjectedToVertex, InterpolatedShapeNormal) * InvInterpolatedNormalSizeSquared;
+	
+		const FVector3f ReprojectedPoint = BindPoint + InterpolatedShapeNormal * D;
+		
+		const FVector3f ReprojectedVector = ReprojectedPoint - Point;
+		const float ErrorEstimate = (ReprojectedPoint - Point).GetAbsMax();
+		
+		// If within the tolerance, 1.0, otherwise linear falloff based on the tolerance
+
+		// Arbitrary factor, a binding will be considered valid (with its corresponding weight) to ErrorFalloffFactor times the validity tolerance.
+		constexpr float ErrorFalloffFactor = 4.0f;
+		OutBindData.Weight = 1.0f - FMath::Clamp( (ErrorEstimate - ValidityTolerance) / (ValidityTolerance * ErrorFalloffFactor), 0.0f, 1.0f);
+	
+		OutBindData.S = Barycentric.Y;
+		OutBindData.T = Barycentric.Z;
+
+		// Only move points that bind outside the shape.
+		OutBindData.Triangle = FVector3f::DotProduct(ShapeTriangleNormal, Point - BindPoint) < 0 ? -1 : FoundIndex;		
+	}
+
+	inline void BindClipDeformPointClosestToSurface( 
+       		const FShapeMeshTree& ShapeMeshTree, 
+			const FVector3f& Point, FClipDeformVertexBindingData& OutBindData)
+	{
+		const FShapeMeshDescriptorBind& ShapeMesh = ShapeMeshTree.GetMesh()->Mesh;
+		
+		OutBindData.S = 0.0f;
+		OutBindData.T = 0.0f;
+		OutBindData.Weight = 1.0f;
+		OutBindData.Triangle = -1;
+
+		const FVector3d P = FVector3d(Point);
+
+		double DistSqr = 0;
+		int32 FoundTriIndex = ShapeMeshTree.FindNearestTriangle(P, DistSqr);
+	
+		if (FoundTriIndex < 0)
+		{
+			return;
+		}
+
+		check( FoundTriIndex >= 0); 
+
+		UE::Geometry::FIndex3i Triangle = ShapeMesh.Triangles[FoundTriIndex];
+		UE::Geometry::FTriangle3d NearestShapeTriangle { 
+				ShapeMesh.Positions[Triangle.A], 
+				ShapeMesh.Positions[Triangle.B], 
+				ShapeMesh.Positions[Triangle.C] };
+	
+		UE::Geometry::FDistPoint3Triangle3d Dist(P, NearestShapeTriangle);
+		Dist.ComputeResult();
+
+		const FVector3d BindPoint = NearestShapeTriangle.BarycentricPoint(Dist.TriangleBaryCoords);
+
+		OutBindData.S = float(Dist.TriangleBaryCoords.Y);
+		OutBindData.T = float(Dist.TriangleBaryCoords.Z); 
+
+		// Only move points that bind outside the shape.
+		OutBindData.Triangle = FVector3d::DotProduct(NearestShapeTriangle.Normal(), P - BindPoint) < 0 ? -1 : FoundTriIndex;
+	}
+
+	inline void BindClipDeformPointNormalProject( 
+       		const FShapeMeshTree& ShapeMeshTree,
+			const FVector3f& Point, const FVector3f& Normal, FClipDeformVertexBindingData& OutBindData)
+	{
+		const FShapeMeshDescriptorBind& ShapeMesh = ShapeMeshTree.GetMesh()->Mesh;
+
+		OutBindData.S = 0.0f;
+		OutBindData.T = 0.0f;
+		OutBindData.Weight = 1.0f;
+		OutBindData.Triangle = -1;
+	
+		const FRay3d NormalRay = FRay3d(FVector3d(Point), FVector3d(-Normal));
+		double RayHitDist = 0.0;
+		FVector3d RayHitBarycenticCoords = FVector3d::ZeroVector;
+		int32 TriangleIndex = -1;
+
+		const bool bHitFound = ShapeMeshTree.FindNearestHitTriangle(NormalRay, RayHitDist, TriangleIndex, RayHitBarycenticCoords);
+
+		if (!bHitFound || TriangleIndex < 0)
+		{
+			return;
+		}
+		
+		UE::Geometry::FIndex3i ShapeTriangleIndices = ShapeMesh.Triangles[TriangleIndex];
+		UE::Geometry::FTriangle3d HitShapeTriangle {
+				ShapeMesh.Positions[ShapeTriangleIndices.A],
+				ShapeMesh.Positions[ShapeTriangleIndices.B],
+				ShapeMesh.Positions[ShapeTriangleIndices.C] };
+	
+		OutBindData.S = float(RayHitBarycenticCoords.Y);
+		OutBindData.T = float(RayHitBarycenticCoords.Z);
+
+		OutBindData.Triangle = FVector3d::DotProduct(HitShapeTriangle.Normal(), -NormalRay.Direction) < 0 ? -1 : TriangleIndex;
+	}
+	
+	inline TArray<FClipDeformVertexBindingData> BindVerticesClipDeform(
+			const Mesh* BaseMesh, 
+			FShapeMeshTree& ShapeMeshTree, 
+			EShapeBindingMethod BindingMethod )
+	{	
+		UE::Geometry::FAxisAlignedBox3d ShapeAABBox = ShapeMeshTree.GetBoundingBox();
+
+		const float BindValidityTolerance = ShapeAABBox.MaxDim() < 1.0 
+										  ? UE_KINDA_SMALL_NUMBER
+										  : static_cast<float>(ShapeAABBox.MaxDim()) * 1e-3f;
+
+		// Find nearest shape triangle for each base mesh vertex		
+		const int32 MeshVertexCount = BaseMesh->GetVertexCount();
+		TArray<FClipDeformVertexBindingData> BindData;
+		{
+			MUTABLE_CPUPROFILER_SCOPE(ClipDeformBind);
+
+			BindData.SetNum( MeshVertexCount );
+
+			UntypedMeshBufferIteratorConst ItPositionBase(BaseMesh->GetVertexBuffers(), MBS_POSITION);
+			UntypedMeshBufferIteratorConst ItNormalBase(BaseMesh->GetVertexBuffers(), MBS_NORMAL);
+
+			EShapeBindingMethod ActualBindingMethod = BindingMethod == EShapeBindingMethod::ClipDeformNormalProject && !ItNormalBase.ptr() 
+					? EShapeBindingMethod::ClipDeformClosestToSurface 
+					: BindingMethod;
+
+			switch ( ActualBindingMethod )
+			{
+			case EShapeBindingMethod::ClipDeformNormalProject:
+			{
+				check(ItNormalBase.ptr());
+				for (int32 VertexIndex = 0; VertexIndex < MeshVertexCount; ++VertexIndex)
+				{
+					FVector3f VertexPosition = (ItPositionBase + VertexIndex).GetAsVec3f();
+					FVector3f VertexNormal = (ItNormalBase + VertexIndex).GetAsVec3f();
+					BindClipDeformPointNormalProject(ShapeMeshTree, VertexPosition, VertexNormal, BindData[VertexIndex]);
+				}	
+
+				break;
+			}
+			case EShapeBindingMethod::ClipDeformClosestToSurface:
+			{
+				for (int32 VertexIndex = 0; VertexIndex < MeshVertexCount; ++VertexIndex)
+				{
+					FVector3f VertexPosition = (ItPositionBase + VertexIndex).GetAsVec3f();
+					BindClipDeformPointClosestToSurface(ShapeMeshTree, VertexPosition, BindData[VertexIndex]);
+				}
+
+				break;
+			}
+			case EShapeBindingMethod::ClipDeformClosestProject:
+			{
+				for (int32 VertexIndex = 0; VertexIndex < MeshVertexCount; ++VertexIndex)
+				{
+					FVector3f VertexPosition = (ItPositionBase + VertexIndex).GetAsVec3f();
+					BindClipDeformPointClosestProject(ShapeMeshTree, VertexPosition, BindData[VertexIndex], BindValidityTolerance);
+				}
+
+				break;
+			}
+
+			default:
+			{
+				check(false);
+			}
+			};
+		}
+
+		return BindData;
+	}
+
+    inline MeshPtr MeshBindShapeClipDeform(const Mesh* BaseMesh, const Mesh* ShapeMesh, EShapeBindingMethod BindingMethod)
+    {
+		MUTABLE_CPUPROFILER_SCOPE(MeshBindShapeClipDeform);
+		
+		if (!BaseMesh)
+		{
+			return nullptr;
+		}
+
+		if (!ShapeMesh)
+		{
+			return BaseMesh->Clone();
+		}
+	
+		int32 ShapeVertexCount = ShapeMesh->GetVertexCount();
+		int ShapeTriangleCount = ShapeMesh->GetFaceCount();
+		if (!ShapeVertexCount || !ShapeTriangleCount)
+		{
+			return BaseMesh->Clone();
+		}	
+
+		FShapeMeshDescriptorBind ShapeMeshDescriptor;
+		{
+			MUTABLE_CPUPROFILER_SCOPE(GenerateVertexQueryData);	
+
+			ShapeMeshDescriptor.Positions.SetNum(ShapeVertexCount);
+			ShapeMeshDescriptor.Normals.SetNum(ShapeVertexCount);
+
+			// \TODO: Simple but inefficient
+			UntypedMeshBufferIteratorConst ItPosition(ShapeMesh->GetVertexBuffers(), MBS_POSITION);
+			UntypedMeshBufferIteratorConst ItNormal(ShapeMesh->GetVertexBuffers(), MBS_NORMAL);
+			for (int32 ShapeVertexIndex = 0; ShapeVertexIndex < ShapeVertexCount; ++ShapeVertexIndex)
+			{
+				FVector3f Position = ItPosition.GetAsVec3f();
+				ShapeMeshDescriptor.Positions[ShapeVertexIndex] = static_cast<FVector3d>(Position);
+				++ItPosition;
+				
+				FVector3f Normal = ItNormal.GetAsVec3f();
+				ShapeMeshDescriptor.Normals[ShapeVertexIndex] = Normal;
+				++ItNormal;
+			}
+		}
+		// Generate the temp face query data for the shape
+		// TODO: Index data copy may be saved in most cases.
+		{
+			MUTABLE_CPUPROFILER_SCOPE(GenerateTrianglesQueryData);	
+			ShapeMeshDescriptor.Triangles.SetNum(ShapeTriangleCount);
+			// \TODO: Simple but inefficient
+			UntypedMeshBufferIteratorConst ItIndices(ShapeMesh->GetIndexBuffers(), MBS_VERTEXINDEX);
+			for (int32 TriangleIndex = 0; TriangleIndex < ShapeTriangleCount; ++TriangleIndex)
+			{
+				UE::Geometry::FIndex3i Triangle;
+				Triangle.A = int(ItIndices.GetAsUINT32());
+				++ItIndices;
+				Triangle.B = int(ItIndices.GetAsUINT32());
+				++ItIndices;
+				Triangle.C = int(ItIndices.GetAsUINT32());
+				++ItIndices;
+
+				ShapeMeshDescriptor.Triangles[TriangleIndex] = Triangle;
+			}	
+		}
+
+		FShapeMeshAdapter ShapeMeshAdapter(ShapeMeshDescriptor);
+
+		constexpr bool bAutoBuildTree = false;
+		FShapeMeshTree ShapeMeshTree(&ShapeMeshAdapter, bAutoBuildTree);	
+		{	
+			MUTABLE_CPUPROFILER_SCOPE(BuildShapeTree);
+			
+			ShapeMeshTree.Build();
+		}
+	
+		Ptr<Mesh> Result = BaseMesh->Clone();
+		TArray<FClipDeformVertexBindingData> VerticesBindData = BindVerticesClipDeform(BaseMesh, ShapeMeshTree, BindingMethod);
+			
+		// Add the binding information to the mesh
+		// \TODO: Check that there is no other binding data.
+		// \TODO: Support specifying the binding data channel for multiple binding support.
+		FMeshBufferSet& VB = Result->GetVertexBuffers();
+		int NewBufferIndex = VB.GetBufferCount();
+		VB.SetBufferCount(NewBufferIndex + 1);
+
+		// \TODO: Multiple binding dataset support
+		int32 BindingDataIndex = 0;
+			
+		FClipDeformVertexBindingDataBufferDescriptor BufDesc(BindingDataIndex);
+		VB.SetBuffer(NewBufferIndex, sizeof(FClipDeformVertexBindingData), BufDesc.Channels, BufDesc.Semantics, BufDesc.SemanticIndices, BufDesc.Formats, BufDesc.Components, BufDesc.Offsets);
+		FMemory::Memcpy(VB.GetBufferData(NewBufferIndex), VerticesBindData.GetData(), VerticesBindData.Num()*sizeof(FClipDeformVertexBindingData));
+		
+        return Result;
+    }
+}
diff --git a/Engine/Plugins/Experimental/Mutable/Source/MutableRuntime/Private/OpMeshChartDifference.h b/Engine/Plugins/Experimental/Mutable/Source/MutableRuntime/Private/OpMeshChartDifference.h
new file mode 100644
--- /dev/null
+++ b/Engine/Plugins/Experimental/Mutable/Source/MutableRuntime/Private/OpMeshChartDifference.h
@@ -0,0 +1,222 @@
+// Copyright Epic Games, Inc. All Rights Reserved.
+
+#pragma once
+
+#include "MeshPrivate.h"
+#include "Platform.h"
+
+
+namespace mu
+{
+
+
+    //---------------------------------------------------------------------------------------------
+    inline void MeshChartCreateChartBuffer( Mesh* pSource, unsigned chart )
+    {
+        check( chart<=0xFF );
+        check( chart<32 );
+		(void)chart;
+
+        //UINT32 chartMask = (1 << chart);
+
+        // Make sure that the source mesh has a chart buffer
+        int chartBuffer = -1;
+        int chartChannel = -1;
+        pSource->GetFaceBuffers().FindChannel( MBS_CHART, 0, &chartBuffer, &chartChannel );
+        if ( chartBuffer<0 )
+        {
+            // Add a new empty chart channel in its own new face buffer
+            int bufCount = pSource->GetFaceBuffers().GetBufferCount();
+            pSource->GetFaceBuffers().SetBufferCount( bufCount + 1 );
+
+            MESH_BUFFER_SEMANTIC semantic = MBS_CHART;
+            int semanticIndex = 0;
+            MESH_BUFFER_FORMAT format = MBF_UINT32;
+            int components = 1;
+            int offset = 0;
+            pSource->GetFaceBuffers().SetBuffer
+                    ( bufCount, 4, 1, &semantic, &semanticIndex, &format, &components, &offset );
+            uint8_t* pBuf = pSource->GetFaceBuffers().GetBufferData( bufCount );
+            int size = pSource->GetFaceBuffers().GetElementSize( bufCount )
+                     * pSource->GetFaceBuffers().GetElementCount();
+			FMemory::Memzero( pBuf, size );
+        }
+    }
+
+
+	//---------------------------------------------------------------------------------------------
+    inline void MeshChartDifference( Mesh* pSource, const Mesh* pTarget, unsigned chart )
+	{
+		check( chart<=0xFF );
+		check( chart<32 );
+
+		int sourceFaceCount = pSource->GetFaceCount();
+		int targetFaceCount = pTarget->GetFaceCount();
+		if ( (!sourceFaceCount) || (!targetFaceCount) )
+		{
+			return;
+        }
+
+        // Make a tolerance proportional to the mesh bounding box size
+        // TODO: Use precomputed bounding box
+        box< vec3<float> > aabbox;
+        if ( targetFaceCount > 0 )
+        {
+            MeshBufferIteratorConst<MBF_FLOAT32,float,3> itp( pTarget->GetVertexBuffers(), MBS_POSITION );
+
+            aabbox.min = *itp;
+            ++itp;
+
+            for ( int v=1; v<pTarget->GetVertexBuffers().GetElementCount(); ++v )
+            {
+                aabbox.Bound( *itp );
+                ++itp;
+            }
+        }
+        float tolerance = 1e-5f * length(aabbox.size);
+
+
+        uint32_t chartMask = (1 << chart);
+
+		// Make sure that the source mesh has a chart buffer
+        MeshChartCreateChartBuffer(pSource,chart);
+
+        // Bucketed version
+
+		// For every face in the source, see if it is also in the target.
+		Mesh::VERTEX_MATCH_MAP vertexMap;
+        pTarget->GetVertexMap( *pSource, vertexMap, tolerance );
+
+        MeshBufferIterator<MBF_UINT32,uint32_t,1> itChart( pSource->GetFaceBuffers(), MBS_CHART );
+
+
+		// Classify the target faces in buckets along the Y axis
+#define NUM_BUCKETS	128
+#define AXIS		1
+		vector<int> buckets[ NUM_BUCKETS ];
+		float bucketStart = aabbox.min[AXIS];
+		float bucketSize = aabbox.size[AXIS] / NUM_BUCKETS;
+
+		float bucketThreshold = ( 4 * tolerance ) / bucketSize;
+        UntypedMeshBufferIteratorConst iti( pTarget->GetIndexBuffers(), MBS_VERTEXINDEX );
+		MeshBufferIteratorConst<MBF_FLOAT32,float,3> itp( pTarget->GetVertexBuffers(), MBS_POSITION );
+		for ( int tf=0; tf<targetFaceCount; tf++ )
+		{
+            uint32_t index0 = iti.GetAsUINT32(); ++iti;
+            uint32_t index1 = iti.GetAsUINT32(); ++iti;
+            uint32_t index2 = iti.GetAsUINT32(); ++iti;
+			float y = ( (*(itp+index0))[AXIS] + (*(itp+index1))[AXIS] + (*(itp+index2))[AXIS] ) / 3;
+			float fbucket = (y-bucketStart) / bucketSize;
+			int bucket = std::min( NUM_BUCKETS-1, std::max( 0, (int)fbucket ) );
+			buckets[bucket].push_back(tf);
+			int hibucket = std::min( NUM_BUCKETS-1, std::max( 0, (int)(fbucket+bucketThreshold) ) );
+			if (hibucket!=bucket)
+			{
+				buckets[hibucket].push_back(tf);
+			}
+			int lobucket = std::min( NUM_BUCKETS-1, std::max( 0, (int)(fbucket-bucketThreshold) ) );
+			if (lobucket!=bucket)
+			{
+				buckets[lobucket].push_back(tf);
+			}
+		}
+
+//		LogDebug("Box : min %.3f, %.3f, %.3f    size %.3f,%.3f,%.3f\n",
+//				aabbox.min[0], aabbox.min[1], aabbox.min[2],
+//				aabbox.size[0], aabbox.size[1], aabbox.size[2] );
+//		for ( int b=0; b<NUM_BUCKETS; ++b )
+//		{
+//			LogDebug("bucket : %d\n", buckets[b].size() );
+//		}
+
+        UntypedMeshBufferIteratorConst ito( pSource->GetIndexBuffers(), MBS_VERTEXINDEX );
+		MeshBufferIteratorConst<MBF_FLOAT32,float,3> itop( pSource->GetVertexBuffers(), MBS_POSITION );
+        UntypedMeshBufferIteratorConst itti( pTarget->GetIndexBuffers(), MBS_VERTEXINDEX );
+		for ( int f=0; f<sourceFaceCount; ++f )
+		{
+			bool hasFace = false;
+            vec3<uint32_t> ov;
+            ov[0] = ito.GetAsUINT32(); ++ito;
+            ov[1] = ito.GetAsUINT32(); ++ito;
+            ov[2] = ito.GetAsUINT32(); ++ito;
+
+			// find the bucket for this face
+			float y = ( (*(itop+ov[0]))[AXIS] + (*(itop+ov[1]))[AXIS] + (*(itop+ov[2]))[AXIS] ) / 3;
+			float fbucket = (y-bucketStart) / bucketSize;
+			int bucket = std::min( NUM_BUCKETS-1, std::max( 0, (int)fbucket ) );
+
+			for ( std::size_t btf=0; !hasFace && btf<buckets[bucket].size(); btf++ )
+			{
+				int tf =  buckets[bucket][btf];
+
+                vec3<uint32_t> v;
+                v[0] = (itti+3*tf+0).GetAsUINT32();
+                v[1] = (itti+3*tf+1).GetAsUINT32();
+                v[2] = (itti+3*tf+2).GetAsUINT32();
+
+				hasFace = true;
+				for ( int vi=0; hasFace && vi<3; ++vi )
+				{
+					hasFace = vertexMap.Matches(v[vi],ov[0])
+						 || vertexMap.Matches(v[vi],ov[1])
+						 || vertexMap.Matches(v[vi],ov[2]);
+				}
+			}
+
+			if ( hasFace )
+			{
+				(*itChart)[0] |= chartMask;
+			}
+
+			++itChart;
+		}
+
+        /*
+        // Non bucket version
+		// For every face in the source, see if it is also in the target.
+		Mesh::Private::VERTEX_MATCH_MAP vertexMap;
+        pTarget->GetPrivate()->GetVertexMap( *pSource->GetPrivate(), vertexMap, tolerance );
+
+		MeshBufferIterator<MBF_UINT32,UINT32,1> itChart( *pSource->GetFaceBuffers(), MBS_CHART );
+
+		MeshBufferIteratorConst<MBF_UINT32,UINT32,1> ito( *pSource->GetIndexBuffers(), MBS_VERTEXINDEX );
+		for ( int f=0; f<sourceFaceCount; ++f )
+		{
+			bool hasFace = false;
+			vec3<UINT32> ov;
+			ov[0] = (*ito)[0]; ++ito;
+			ov[1] = (*ito)[0]; ++ito;
+			ov[2] = (*ito)[0]; ++ito;
+
+			if ( targetFaceCount>0 )
+			{
+				MeshBufferIteratorConst<MBF_UINT32,UINT32,1> it( *pTarget->GetIndexBuffers(), MBS_VERTEXINDEX );
+				for ( int tf=0; !hasFace && tf<targetFaceCount; tf++ )
+				{
+					vec3<UINT32> v;
+					v[0] = (*it)[0]; ++it;
+					v[1] = (*it)[0]; ++it;
+					v[2] = (*it)[0]; ++it;
+
+					hasFace = true;
+					for ( int vi=0; hasFace && vi<3; ++vi )
+					{
+						hasFace = vertexMap.Matches(v[vi],ov[0])
+							 || vertexMap.Matches(v[vi],ov[1])
+							 || vertexMap.Matches(v[vi],ov[2]);
+					}
+				}
+			}
+
+			if ( hasFace )
+			{
+				(*itChart)[0] |= chartMask;
+			}
+
+			++itChart;
+		}
+        */
+
+	}
+
+}
diff --git a/Engine/Plugins/Experimental/Mutable/Source/MutableRuntime/Private/OpMeshClipDeform.h b/Engine/Plugins/Experimental/Mutable/Source/MutableRuntime/Private/OpMeshClipDeform.h
new file mode 100644
--- /dev/null
+++ b/Engine/Plugins/Experimental/Mutable/Source/MutableRuntime/Private/OpMeshClipDeform.h
@@ -0,0 +1,457 @@
+// Copyright Epic Games, Inc. All Rights Reserved.
+
+#pragma once
+
+
+#include "OpMeshBind.h"
+
+#include "MutableTrace.h"
+
+#include "MeshPrivate.h"
+
+#include "Algo/Copy.h"
+
+
+// TODO: Make the handling of rotations an option. It is more expensive on CPU and memory, and for some
+// cases it is not required at all.
+
+// TODO: Face stretch to scale the deformation per-vertex? 
+
+// TODO: Support multiple binding influences per vertex, to have smoother deformations.
+
+// TODO: Support multiple binding sets, to have different shapes deformations at once.
+
+// TODO: Deformation mask, to select the intensisty of the deformation per-vertex.
+
+// TODO: This is a reference implementation with ample roof for optimization.
+
+namespace mu
+{
+
+	struct FClipDeformShapeMeshDescriptorApply
+	{
+		TArrayView<const FVector3f> Positions;
+		TArrayView<const FVector3f> Normals;
+		TArrayView<const float> Weights;
+		TArrayView<const UE::Geometry::FIndex3i> Triangles;
+	};
+
+	// Method to actually deform a point
+	inline bool GetDeform(
+			const FVector3f& Position,
+			const FVector3f& Normal,
+			const FClipDeformShapeMeshDescriptorApply& ShapeMesh,
+			const FClipDeformVertexBindingData& Binding,
+			FVector3f& OutPosition, 
+			FQuat4f& OutRotation,
+			float& OutWeight)
+	{
+		if (Binding.Triangle < 0 || Binding.Triangle >= ShapeMesh.Triangles.Num())
+		{
+			return false;
+		}
+
+		FVector2f BaryCoords = FVector2f(Binding.S, Binding.T);
+
+		// Clamp Barycoords so we are always inside the bound triangle if the binding is not good.
+		// This is only needed for Closest Project, which is not very robust and sometimes the binding
+		// are not valid.
+		if ( FMath::IsNearlyZero(Binding.Weight) )
+		{
+			BaryCoords.X = FMath::Max( 0.0f, BaryCoords.X );
+			BaryCoords.Y = FMath::Max( 0.0f, BaryCoords.Y );
+			if (BaryCoords.X + BaryCoords.Y > 1.0f)
+			{
+				BaryCoords.X = BaryCoords.X / (BaryCoords.X + BaryCoords.Y);
+				BaryCoords.Y = 1.0f - BaryCoords.X;
+			}
+		}
+		
+		const UE::Geometry::FIndex3i& Triangle = ShapeMesh.Triangles[Binding.Triangle];
+	
+		const FVector3f ProjectedVertexPosition = 
+				ShapeMesh.Positions[Triangle.A] * (1.0f - BaryCoords.X - BaryCoords.Y) +
+				ShapeMesh.Positions[Triangle.B] * BaryCoords.X +
+				ShapeMesh.Positions[Triangle.C] * BaryCoords.Y;
+
+		const float ProjectedVertexWeight = 
+				ShapeMesh.Weights[Triangle.A] * (1.0f - BaryCoords.X - BaryCoords.Y) + 
+				ShapeMesh.Weights[Triangle.B] * BaryCoords.X + 
+				ShapeMesh.Weights[Triangle.C] * BaryCoords.Y;
+		
+		// Morph from the projected position in shape to the original position based on the weight defined in shape.
+		OutWeight = FMath::Clamp(ProjectedVertexWeight, 0.0f, 1.0f);
+		OutPosition = FMath::Lerp(Position, ProjectedVertexPosition, OutWeight);
+
+		// This method approximates the shape face rotation
+		FVector3f InterpolatedNormal = 
+				ShapeMesh.Normals[Triangle.A] * (1.0f - BaryCoords.X - BaryCoords.Y) + 
+				ShapeMesh.Normals[Triangle.B] * BaryCoords.X + 
+				ShapeMesh.Normals[Triangle.C] * BaryCoords.Y;
+
+		FVector3f NewNormal = FMath::Lerp( Normal, InterpolatedNormal, ProjectedVertexWeight ).GetSafeNormal();
+		
+		// Use shape normal to interpolate.
+		//FVector3f NewNormal = FMath::Lerp( Normal, -Binding->ShapeNormal, ProjectedVertexWeight ).GetSafeNormal();
+		
+		OutRotation = FQuat4f::FindBetween(Normal, NewNormal);
+		
+		return true;
+	}
+
+	inline MeshPtr MeshClipDeform(const Mesh* BaseMesh, const Mesh* ShapeMesh, const float ClipWeightThreshold )
+	{
+		MUTABLE_CPUPROFILER_SCOPE(ClipDeform);
+
+		if (!BaseMesh)
+		{
+			return nullptr;
+		}
+
+		constexpr EMeshCloneFlags MeshCloneFlags = ~EMeshCloneFlags::WithVertexBuffers;
+		MeshPtr Result = BaseMesh->Clone(MeshCloneFlags);
+	
+		int BarycentricDataBuffer = 0;
+		int BarycentricDataChannel = 0;
+		const FMeshBufferSet& VB = BaseMesh->GetVertexBuffers();
+		VB.FindChannel(MBS_BARYCENTRICCOORDS, 0, &BarycentricDataBuffer, &BarycentricDataChannel);
+		
+		// Copy buffers skipping binding data.
+		FMeshBufferSet& ResultBuffers = Result->GetVertexBuffers();
+		ResultBuffers.m_elementCount = VB.m_elementCount;
+		// Remove one element to the number of buffers if BarycentricDataBuffer found. 
+		ResultBuffers.m_buffers.SetNum( FMath::Max( 0, VB.m_buffers.Num() - int32(BarycentricDataBuffer >= 0) ) );
+		
+		for (int32 B = 0, R = 0; B < VB.m_buffers.Num(); ++B)
+		{
+			if (B != BarycentricDataBuffer)
+			{
+				ResultBuffers.m_buffers[R++] = VB.m_buffers[B];
+			}
+		}
+
+		if (BarycentricDataBuffer < 0)
+		{
+			return Result;
+		}
+
+		if (!ShapeMesh)
+		{
+			return Result;
+		}
+
+		// \TODO: More checks
+		check(BarycentricDataChannel==0);
+		check(VB.GetElementSize(BarycentricDataBuffer) == (int)sizeof(FClipDeformVertexBindingData));
+		TArrayView<const FClipDeformVertexBindingData> BindingDataView = TArrayView<const FClipDeformVertexBindingData>(
+				(const FClipDeformVertexBindingData*)VB.GetBufferData(BarycentricDataBuffer), VB.GetElementCount() );
+
+		//
+		int32 ShapeVertexCount = ShapeMesh->GetVertexCount();
+		int ShapeTriangleCount = ShapeMesh->GetFaceCount();
+		if (!ShapeVertexCount || !ShapeTriangleCount)
+		{
+			return Result;
+		}
+
+		TArray<FVector3f> ShapePositions;
+		TArray<FVector3f> ShapeNormals;
+		TArray<float> ShapeWeights;
+
+		FClipDeformShapeMeshDescriptorApply ShapeMeshDescriptor;
+		{
+			MUTABLE_CPUPROFILER_SCOPE(GenerateVertexQueryData);
+
+			ShapeNormals.SetNum(ShapeVertexCount);
+			ShapeWeights.SetNum(ShapeVertexCount);
+
+			UntypedMeshBufferIteratorConst ItPositionBegin(ShapeMesh->GetVertexBuffers(), MBS_POSITION);
+			UntypedMeshBufferIteratorConst ItNormalBegin(ShapeMesh->GetVertexBuffers(), MBS_NORMAL);
+			UntypedMeshBufferIteratorConst ItUvsBegin(ShapeMesh->GetVertexBuffers(), MBS_TEXCOORDS);
+
+			const bool bIsPositionBufferCompatible = 
+					ItPositionBegin.GetFormat() == MBF_FLOAT32 && 
+					ItPositionBegin.GetElementSize() == sizeof(FVector3f);
+			if (bIsPositionBufferCompatible)
+			{
+				ShapeMeshDescriptor.Positions = TArrayView<FVector3f>((FVector3f*)ItPositionBegin.ptr(), ShapeVertexCount);
+			}
+			else
+			{
+				ShapePositions.SetNum(ShapeVertexCount);
+				UntypedMeshBufferIteratorConst ItPosition =  ItPositionBegin;
+				for (int32 ShapeVertexIndex = 0; ShapeVertexIndex < ShapeVertexCount; ++ShapeVertexIndex)
+				{
+					FVector3f Position = ItPosition.GetAsVec3f();
+					ShapePositions[ShapeVertexIndex] = Position;
+					++ItPosition;
+				}
+
+				ShapeMeshDescriptor.Positions = TArrayView<FVector3f>(ShapePositions.GetData(), ShapePositions.Num());
+			}
+
+			const bool bIsNormalBufferCompatible = 
+					ItNormalBegin.GetFormat() == MBF_FLOAT32 && 
+					ItNormalBegin.GetElementSize() == sizeof(FVector3f);
+			if (bIsNormalBufferCompatible)
+			{
+				ShapeMeshDescriptor.Normals = TArrayView<const FVector3f>((const FVector3f*)ItNormalBegin.ptr(), ShapeVertexCount);
+			}
+			else
+			{
+				ShapeNormals.SetNum(ShapeVertexCount);
+				UntypedMeshBufferIteratorConst ItNormal = ItNormalBegin;
+				for (int32 ShapeVertexIndex = 0; ShapeVertexIndex < ShapeVertexCount; ++ShapeVertexIndex)
+				{
+					FVector3f Normal = ItNormal.GetAsVec3f();
+					ShapeNormals[ShapeVertexIndex] = Normal;
+					++ItNormal;
+				}
+
+				ShapeMeshDescriptor.Normals = TArrayView<const FVector3f>(ShapeNormals.GetData(), ShapeNormals.Num());
+			}
+	
+			// Don't try to use the buffer directly for weights since we only need a component.
+			UntypedMeshBufferIteratorConst ItUvs = ItUvsBegin;
+			for (int32 ShapeVertexIndex = 0; ShapeVertexIndex < ShapeVertexCount; ++ShapeVertexIndex)
+			{
+				FVector2f Uvs = ItUvs.GetAsVec2f();
+				ShapeWeights[ShapeVertexIndex] = 1.0f - Uvs.Y;
+				++ItUvs;
+			}
+
+			ShapeMeshDescriptor.Weights = TArrayView<const float>(ShapeWeights.GetData(), ShapeWeights.Num());
+		}
+	
+		TArray<UE::Geometry::FIndex3i> ShapeTriangles;
+		{
+			MUTABLE_CPUPROFILER_SCOPE(GenerateTriangleQueryData);
+
+			ShapeTriangles.SetNum(ShapeTriangleCount);
+
+			const UntypedMeshBufferIteratorConst ItIndicesBegin(ShapeMesh->GetIndexBuffers(), MBS_VERTEXINDEX);
+
+			check(ShapeMesh->GetIndexCount() % 3 == 0);
+	
+			const bool bIsIndexBufferCompatible = 
+					(ItIndicesBegin.GetFormat() == MBF_INT32 || ItIndicesBegin.GetFormat() == MBF_UINT32) &&
+					ItIndicesBegin.GetElementSize() == int32(sizeof(int32));
+
+			if (bIsIndexBufferCompatible)
+			{
+				ShapeMeshDescriptor.Triangles = TArrayView<UE::Geometry::FIndex3i>( 
+						(UE::Geometry::FIndex3i*)ItIndicesBegin.ptr(), ShapeMesh->GetIndexCount() / 3 );
+			}
+			else
+			{
+				UntypedMeshBufferIteratorConst ItIndices = ItIndicesBegin;
+				for (int32 TriangleIndex = 0; TriangleIndex < ShapeTriangleCount; ++TriangleIndex)
+				{
+					UE::Geometry::FIndex3i Triangle;
+					Triangle.A = int(ItIndices.GetAsUINT32());
+					++ItIndices;
+					Triangle.B = int(ItIndices.GetAsUINT32());
+					++ItIndices;
+					Triangle.C = int(ItIndices.GetAsUINT32());
+					++ItIndices;
+
+					ShapeTriangles[TriangleIndex] = Triangle;
+				}
+
+				ShapeMeshDescriptor.Triangles = TArrayView<const UE::Geometry::FIndex3i>(ShapeTriangles.GetData(), ShapeTriangles.Num());
+			}
+		}
+
+		ShapeMeshDescriptor.Positions = TArrayView<const FVector3f>(ShapePositions.GetData(), ShapePositions.Num());
+		ShapeMeshDescriptor.Normals = TArrayView<const FVector3f>(ShapeNormals.GetData(), ShapeNormals.Num());
+		ShapeMeshDescriptor.Weights = TArrayView<const float>(ShapeWeights.GetData(), ShapeWeights.Num());
+
+		// Update the result mesh positions
+		int32 MeshVertexCount = BaseMesh->GetVertexCount();
+
+		TArray<uint8> ExcludedVertices;
+		ExcludedVertices.SetNumZeroed(MeshVertexCount);
+		{
+			MUTABLE_CPUPROFILER_SCOPE(UpdateClipDeformVertices);
+
+			UntypedMeshBufferIterator ItPosition(Result->GetVertexBuffers(), MBS_POSITION);
+			UntypedMeshBufferIterator ItNormal(Result->GetVertexBuffers(), MBS_NORMAL);
+			UntypedMeshBufferIterator ItTangent(Result->GetVertexBuffers(), MBS_TANGENT);
+
+			for (int32 MeshVertexIndex = 0; MeshVertexIndex < MeshVertexCount; ++MeshVertexIndex)
+			{
+				FVector3f NewPosition;
+				FQuat4f TangentSpaceCorrection;
+				float ClipWeight = 0.0f;
+				const bool bModified = GetDeform( 
+						ItPosition.GetAsVec3f(), ItNormal.GetAsVec3f(), 
+						ShapeMeshDescriptor, 
+						BindingDataView[MeshVertexIndex],
+						NewPosition, TangentSpaceCorrection, ClipWeight);
+
+				if (bModified)
+				{
+					ItPosition.SetFromVec3f(NewPosition);
+
+					if (ItNormal.ptr())
+					{
+						FVector3f OldNormal = ItNormal.GetAsVec3f();
+						FVector3f NewNormal = TangentSpaceCorrection.RotateVector(OldNormal);
+						ItNormal.SetFromVec3f(NewNormal);
+					}
+
+					if (ItTangent.ptr())
+					{
+						FVector3f OldTangent = ItTangent.GetAsVec3f();
+						FVector3f NewTangent = TangentSpaceCorrection.RotateVector(OldTangent);
+						ItTangent.SetFromVec3f(NewTangent);
+					}
+					
+					ExcludedVertices[MeshVertexIndex] = ClipWeight >= (1.0f - SMALL_NUMBER);	
+				}
+
+				++ItPosition;
+
+				if (ItNormal.ptr())
+				{
+					++ItNormal;
+				}
+
+				if (ItTangent.ptr())
+				{
+					++ItTangent;
+				}
+			}
+		}
+
+		// Remove excluded vertices
+		{
+			MUTABLE_CPUPROFILER_SCOPE(RemoveExcludedVertices);
+			
+			const int32 NumOriginalVerts = ExcludedVertices.Num();
+		
+			// Remove faces only if all vertices are gone.
+			UntypedMeshBufferIteratorConst ItSrc(Result->GetIndexBuffers(), MBS_VERTEXINDEX);
+			UntypedMeshBufferIterator ItDest(Result->GetIndexBuffers(), MBS_VERTEXINDEX);
+
+			UntypedMeshBufferIterator It(Result->GetIndexBuffers(), MBS_VERTEXINDEX);
+
+			// We cannot reintroduce vertices to the same set ExcludedVertices since some other triangle exclusion might
+			// depend on those.
+			TArray<uint8> FinalExcludedVertices;
+			FinalExcludedVertices.Init(1, ExcludedVertices.Num());
+
+			const SIZE_T TriangleSizeBytes =  ItSrc.GetElementSize() * 3;
+			const int32 SrcTriangleCount = Result->GetIndexCount() / 3;
+			for (int32 TriangleIndex = 0; TriangleIndex < SrcTriangleCount; ++TriangleIndex)
+			{
+				int32 A = It.GetAsUINT32();
+				++It;
+				int32 B = It.GetAsUINT32();
+				++It;
+				int32 C = It.GetAsUINT32();
+				++It;
+
+				const bool bAllVerticesExcluded = static_cast<bool>( ExcludedVertices[A] & ExcludedVertices[B] & ExcludedVertices[C] );
+
+				// Some sort of latching mechanisms is needed, if a vertex needs to be included for a vertex triangle
+				// but not for another, the vertex must stay.
+				FinalExcludedVertices[A] = FMath::Min((uint8)bAllVerticesExcluded, FinalExcludedVertices[A] );
+				FinalExcludedVertices[B] = FMath::Min((uint8)bAllVerticesExcluded, FinalExcludedVertices[B] );
+				FinalExcludedVertices[C] = FMath::Min((uint8)bAllVerticesExcluded, FinalExcludedVertices[C] );
+			 
+				if (!bAllVerticesExcluded)
+				{
+					if (ItSrc.ptr() != ItDest.ptr())
+					{
+						FMemory::Memcpy( ItDest.ptr(), ItSrc.ptr(),  TriangleSizeBytes);
+					}
+
+					ItDest += 3;
+				}
+			 
+				ItSrc += 3;
+			}
+
+			const int32 NumRemainingTris = SrcTriangleCount - ( ( ItSrc.ptr() - ItDest.ptr() ) / TriangleSizeBytes ); 
+			Result->GetIndexBuffers().SetElementCount( NumRemainingTris*3 );
+		
+			TArray<int32> IndexMap;
+			IndexMap.Init(-1, NumOriginalVerts);
+		
+			int32 NumRemainingVerts = 0;
+			for ( int32 I = 0; I < NumOriginalVerts; ++I )
+			{
+				if ( FinalExcludedVertices[I] == 0 )
+				{
+					IndexMap[I] = NumRemainingVerts++;
+				}
+			}
+		
+			UntypedMeshBufferIterator ResIndicesIt(Result->GetIndexBuffers(), MBS_VERTEXINDEX);
+		
+			// Remap Indices
+			for ( int32 TIdx = 0; TIdx < NumRemainingTris; ++TIdx )
+			{
+				const int32 A = ResIndicesIt.GetAsUINT32();
+				check(IndexMap[A] >= 0);
+				ResIndicesIt.SetFromUINT32(IndexMap[A]);
+				++ResIndicesIt;
+			
+				const int32 B = ResIndicesIt.GetAsUINT32();
+				check(IndexMap[B] >= 0);
+				ResIndicesIt.SetFromUINT32(IndexMap[B]);
+				++ResIndicesIt;
+
+				const int32 C = ResIndicesIt.GetAsUINT32();
+				check(IndexMap[C] >= 0);
+				ResIndicesIt.SetFromUINT32(IndexMap[C]);
+				++ResIndicesIt;
+			}
+
+
+			if ( NumRemainingVerts == NumOriginalVerts )
+			{
+				return Result;
+			}
+		
+			// We are guaranteed to find a removed vertex.
+			int32 SpanStart = 0; // points at the span first element
+			while ( FinalExcludedVertices[SpanStart++] == 1 );
+			--SpanStart;
+		
+			int32 DestEnd = 0;
+		
+			const int32 VertexBufferCount = Result->GetVertexBuffers().GetBufferCount();
+			while ( SpanStart < NumOriginalVerts )
+			{
+				// Find span to move.
+				int32 SpanEnd = SpanStart; // Points at the last span element + 1
+				for (; SpanEnd < NumOriginalVerts && FinalExcludedVertices[SpanEnd] == 0; SpanEnd++ );
+
+				const int32 SpanSize = SpanEnd - SpanStart;
+
+				// Assume we always have large spans, so it is more work to find the spans than iterate over the different
+				// buffers. This probably is not the best strategy and recomputing the spans for every buffer might be
+				// more performant.
+
+				if (SpanStart > 0)
+				{
+					for (int32 B = 0; B < VertexBufferCount; ++B)
+					{
+						const int32 ElementSize = Result->GetVertexBuffers().GetElementSize(B);
+						uint8* Ptr = Result->GetVertexBuffers().GetBufferData(B);
+						FMemory::Memmove( Ptr + DestEnd*ElementSize, Ptr + SpanStart*ElementSize , SpanSize*ElementSize );
+					}
+				}
+			
+				DestEnd += SpanSize;
+				
+				for (SpanStart = SpanEnd; SpanStart < NumOriginalVerts && FinalExcludedVertices[SpanStart] == 1; ++SpanStart);
+			}
+
+			Result->GetVertexBuffers().SetElementCount( DestEnd );
+		}
+		return Result;
+	}
+}
diff --git a/Engine/Plugins/Experimental/Mutable/Source/MutableRuntime/Private/OpMeshClipMorphPlane.h b/Engine/Plugins/Experimental/Mutable/Source/MutableRuntime/Private/OpMeshClipMorphPlane.h
new file mode 100644
--- /dev/null
+++ b/Engine/Plugins/Experimental/Mutable/Source/MutableRuntime/Private/OpMeshClipMorphPlane.h
@@ -0,0 +1,414 @@
+// Copyright Epic Games, Inc. All Rights Reserved.
+
+#pragma once
+
+#include "MeshPrivate.h"
+#include "ConvertData.h"
+#include "Platform.h"
+
+
+namespace mu
+{
+	inline bool PointInBoundingBox(const vec3f& point, const SHAPE& selectionShape)
+	{
+		vec3f v = point - selectionShape.position;
+
+		for (int i = 0; i < 3; ++i)
+		{
+			if (fabs(v[i]) > selectionShape.size[i])
+			{
+				return false;
+			}
+		}
+
+		return true;
+	}
+
+	inline bool VertexIsInMaxRadius(const vec3f& vertex, const vec3f& origin, float vertexSelectionBoneMaxRadius)
+	{
+		if (vertexSelectionBoneMaxRadius < 0.f)
+		{
+			return true;
+		}
+
+		vec3f radiusVec = vertex - origin;
+		float radius = dot(radiusVec, radiusVec);
+
+		return radius < vertexSelectionBoneMaxRadius* vertexSelectionBoneMaxRadius;
+	}
+
+	struct vertex_bone_info
+	{
+		mu::vector<int32_t> bone_indices;
+		mu::vector<int32_t> bone_weights;
+	};
+
+	inline bool VertexIsAffectedByBone(int32_t vertex_idx, const mu::vector<bool>& bone_is_affected, const mu::vector<vertex_bone_info>& vertex_info)
+	{
+		if (vertex_idx >= (int)vertex_info.size())
+		{
+			return false;
+		}
+
+		check(vertex_info[vertex_idx].bone_indices.size() == vertex_info[vertex_idx].bone_weights.size());		
+
+        for (size_t i = 0; i < vertex_info[vertex_idx].bone_indices.size(); ++i)
+		{			
+            check(vertex_info[vertex_idx].bone_indices[i] <= (int)bone_is_affected.size());
+
+			if (bone_is_affected[vertex_info[vertex_idx].bone_indices[i]] && vertex_info[vertex_idx].bone_weights[i] > 0)
+			{
+				return true;
+			}
+		}
+
+		return false;
+	}
+
+	//---------------------------------------------------------------------------------------------
+	//! Reference version
+	//---------------------------------------------------------------------------------------------
+	inline MeshPtr MeshClipMorphPlane(const Mesh* pBase, const vec3f& origin, const vec3f& normal, float dist, float factor, float radius,
+		float radius2, float angle, const SHAPE& selectionShape, const mu::string& boneName = mu::string(), float vertexSelectionBoneMaxRadius = -1.f)
+	{
+		//float radius = 8.f;
+		//float radius2 = 4.f;
+		//float factor = 1.f;
+
+		// Generate vector perpendicular to normal for ellipse rotation reference base
+		vec3f aux_base(0.f, 1.f, 0.f);
+
+		if (fabs(dot(normal, aux_base)) > 0.95f)		// fabs = absolute value
+		{
+			aux_base = vec3f(0.f, 0.f, 1.f);
+		}
+		
+		vec3f origin_radius_vector = cross(normal, aux_base);		// PERPENDICULAR VECTOR TO THE PLANE normal and aux base
+		check(fabs(dot(normal, origin_radius_vector)) < 0.05f);
+
+		MeshPtr pDest = pBase->Clone();
+
+        uint32_t vcount = pBase->GetVertexBuffers().GetElementCount();
+
+		if (!vcount)
+		{
+			return pDest;
+		}
+
+		mu::vector<bool> bone_is_affected;
+		mu::vector<vertex_bone_info> vertex_info;
+
+        auto pBaseSkeleton = pBase->GetSkeleton();
+        if (!boneName.empty() && pBaseSkeleton )
+		{
+            bone_is_affected.resize(pBaseSkeleton->GetBoneCount());
+
+            for (int32_t bone_idx = 0; bone_idx < pBaseSkeleton->GetBoneCount(); ++bone_idx)
+			{
+				int32_t current_idx = bone_idx;
+				bool found_bone = false;
+
+				while (current_idx >= 0)
+				{
+                    if (pBaseSkeleton->GetBoneName(current_idx) == boneName)
+					{
+						found_bone = true;
+						break;
+					}
+					else
+					{
+                        current_idx = pBaseSkeleton->GetBoneParent(current_idx);
+					}
+				}
+
+				bone_is_affected[bone_idx] = found_bone;
+			}
+
+			// Look for affected vertex indices
+			vertex_info.resize(vcount);
+			//int firstCount = pBase->GetVertexBuffers().GetElementCount();
+
+			for (int32 vb = 0; vb < pBase->GetVertexBuffers().m_buffers.Num(); ++vb)
+			{
+				const MESH_BUFFER& result = pBase->GetVertexBuffers().m_buffers[vb];
+
+				int elemSize = pBase->GetVertexBuffers().GetElementSize((int)vb);
+				//int firstSize = firstCount * elemSize;
+
+				for (int c = 0; c < pBase->GetVertexBuffers().GetBufferChannelCount((int)vb); ++c)
+				{
+					// Get info about the destination channel
+					MESH_BUFFER_SEMANTIC semantic = MBS_NONE;
+					int semanticIndex = 0;
+					MESH_BUFFER_FORMAT format = MBF_NONE;
+					int components = 0;
+					int offset = 0;
+					pBase->GetVertexBuffers().GetChannel((int)vb, c, &semantic, &semanticIndex, &format, &components, &offset);
+
+					int secondOffset = offset;
+					//int resultOffset = firstSize + offset;
+
+					if (semantic == MBS_BONEINDICES)
+					{
+						for (int i = 0; i < pBase->GetVertexCount(); ++i)
+						{
+							switch (format)
+							{
+							case MBF_INT8:
+							case MBF_UINT8:
+							{
+                                const uint8_t* pD = &result.m_data[secondOffset];
+								
+								for (int j = 0; j < components; ++j)
+								{
+									vertex_info[i].bone_indices.push_back(pD[j]);
+								}
+
+								secondOffset += elemSize;
+								break;
+							}
+
+							case MBF_INT16:
+							case MBF_UINT16:
+							{
+								const uint16_t* pD = reinterpret_cast<const uint16_t*>(&result.m_data[secondOffset]);
+
+								for (int j = 0; j < components; ++j)
+								{
+									vertex_info[i].bone_indices.push_back(pD[j]);
+								}
+
+								secondOffset += elemSize;
+								break;
+							}
+
+							case MBF_INT32:
+							case MBF_UINT32:
+							{
+								const uint32_t* pD = reinterpret_cast<const uint32_t*>(&result.m_data[secondOffset]);
+
+								for (int j = 0; j < components; ++j)
+								{
+									vertex_info[i].bone_indices.push_back(pD[j]);
+								}
+
+								secondOffset += elemSize;
+								break;
+							}
+
+							default:
+								checkf(false, TEXT("Bone index format not supported.") );
+							}
+						}
+					}
+					else if (semantic == MBS_BONEWEIGHTS)
+					{
+						for (int i = 0; i < pBase->GetVertexCount(); ++i)
+						{
+							switch (format)
+							{
+							case MBF_INT8:
+							case MBF_UINT8:
+							case MBF_NUINT8:
+							{
+								const uint8_t* pD = &result.m_data[secondOffset];
+
+								for (int j = 0; j < components; ++j)
+								{
+                                    vertex_info[i].bone_weights.push_back(pD[j]);
+								}
+
+								secondOffset += elemSize;
+								break;
+							}
+
+							case MBF_INT16:
+							case MBF_UINT16:
+							case MBF_NUINT16:
+							{
+								const uint16_t* pD = reinterpret_cast<const uint16_t*>(&result.m_data[secondOffset]);
+
+								for (int j = 0; j < components; ++j)
+								{
+									vertex_info[i].bone_weights.push_back(pD[j]);
+								}
+
+								secondOffset += elemSize;
+								break;
+							}
+
+							case MBF_INT32:
+							case MBF_UINT32:
+							case MBF_NUINT32:
+							{
+								const uint32_t* pD = reinterpret_cast<const uint32_t*>(&result.m_data[secondOffset]);
+
+								for (int j = 0; j < components; ++j)
+								{
+									vertex_info[i].bone_weights.push_back(pD[j]);
+								}
+
+								secondOffset += elemSize;
+								break;
+							}
+
+                            case MBF_FLOAT32:
+                            {
+								const float* pD = reinterpret_cast<const float*>(&result.m_data[secondOffset]);
+
+                                for (int j = 0; j < components; ++j)
+                                {
+                                    vertex_info[i].bone_weights.push_back(pD[j]>0.0f);
+                                }
+
+                                secondOffset += elemSize;
+                                break;
+                            }
+
+							default:
+								checkf(false, TEXT("Bone weight format not supported.") );
+							}
+						}
+					}
+				}
+			}
+		}
+
+        mu::map<uint32_t, bool> removed_vertexs;
+
+		const FMeshBufferSet& MBSPriv = pDest->GetVertexBuffers();
+		for (int32 b = 0; b < MBSPriv.m_buffers.Num(); ++b)
+		{
+			for (int32 c = 0; c < MBSPriv.m_buffers[b].m_channels.Num(); ++c)
+			{
+				MESH_BUFFER_SEMANTIC sem = MBSPriv.m_buffers[b].m_channels[c].m_semantic;
+				int semIndex = MBSPriv.m_buffers[b].m_channels[c].m_semanticIndex;
+
+				UntypedMeshBufferIterator it(pDest->GetVertexBuffers(), sem, semIndex);
+
+				switch (sem)
+				{
+				case MBS_POSITION:
+                    for (uint32_t v = 0; v < vcount; ++v)
+					{
+						vec3f vertex(0.0f, 0.0f, 0.0f);
+						for (int i = 0; i < 3; ++i)
+						{
+							ConvertData(i, &vertex[0], MBF_FLOAT32, it.ptr(), it.GetFormat());
+						}
+
+						if (
+							(  !boneName.empty() && VertexIsAffectedByBone(v, bone_is_affected, vertex_info) && VertexIsInMaxRadius(vertex, origin, vertexSelectionBoneMaxRadius))
+                            || (boneName.empty() && selectionShape.type == (uint8_t)SHAPE::Type::None)
+                            || (selectionShape.type == (uint8_t)SHAPE::Type::AABox && PointInBoundingBox(vertex, selectionShape))
+							)
+						{
+							vec3f morph_plane_center = origin;					// MORPH PLANE POS relative to root of the selected bone
+							vec3f clip_plane_center = origin + normal * dist;	// CPLIPPING PLANE POS
+							vec3f aux_morph = vertex - morph_plane_center;		// MORPH PLANE --> CURRENT VERTEX
+							vec3f aux_clip = vertex - clip_plane_center;		// CLIPPING PLANE --> CURRENT VERTEX
+
+							float dot_morph = dot(aux_morph, normal);			// ANGLE (MORPH PLANE TO VERTEX AND NORMAL)
+							float dot_cut = dot(aux_clip, normal);				// ANGLE (CLIPPING PLANE TO VERTEX AND NORMAL )
+
+							if (dot_morph >= 0.f || dot_cut >= 0.f)				// CHECK IF CLIPPING OR MORPH SHOULD BE COMPUTED FOR V VERTEX
+							{
+								vec3f current_center = morph_plane_center + normal * dot_morph;	// PROJECTED POINT FROM THE MROPH PLANE (THE CLOSER THE DOT VALUE OF NORMAL AND VERTEX THE FURTHER IT GOES )
+								vec3f radius_vector = vertex - current_center;					// 	
+								float radius_vector_len = length(radius_vector);
+								vec3f radius_vector_unit = radius_vector_len != 0.f ? radius_vector / radius_vector_len : vec3f(0.f, 0.f, 0.f); // UNITARY VECTOR THAT GOES FROM THE POINT TO THE VERTEX
+
+								float angle_from_origin = acosf(dot(radius_vector_unit, origin_radius_vector));
+
+								// Cross product between the perpendicular vector from radius vector and origin radius and the normal vector
+								if (dot(cross(radius_vector_unit, origin_radius_vector), normal) < 0)
+								{
+									angle_from_origin = -angle_from_origin;
+								}
+
+								angle_from_origin += angle * PI / 180.f;
+
+								float term1 = radius2 * cosf(angle_from_origin);
+								float term2 = radius * sinf(angle_from_origin);
+								float ellipse_radius_at_angle = radius * radius2 / sqrtf(term1 * term1 + term2 * term2);
+
+								vec3f vertex_proj_ellipse = current_center + radius_vector_unit * ellipse_radius_at_angle;
+								//vec3f vertex_proj_ellipse = current_center + radius_vector_unit * radius;
+
+								float morph_alpha = dist != 0.f && dot_morph <= dist ? clamp(0.f, 1.f, powf(dot_morph / dist, factor)) : 1.f;
+
+								vertex = vertex * (1.f - morph_alpha) + vertex_proj_ellipse * morph_alpha;
+
+								if (dot_cut >= 0.f)		// CHECK IF THE VERTEX SHOULD BE CLIPPED
+								{
+									vec3f vert_displ = normal * -dot_cut;
+									vertex = vertex + vert_displ;
+									removed_vertexs[v] = true;
+								}
+							}
+						}
+
+						for (int i = 0; i < 3; ++i)
+						{
+							ConvertData(i, it.ptr(), it.GetFormat(), &vertex[0], MBF_FLOAT32);
+						}
+
+						++it;
+					}
+					break;
+
+				default:
+					break;
+				}
+			}
+		}
+
+		// Now remove all the faces from the result mesh that have all vertices removed
+		UntypedMeshBufferIteratorConst itBase(pDest->GetIndexBuffers(), MBS_VERTEXINDEX);
+		UntypedMeshBufferIterator itDest(pDest->GetIndexBuffers(), MBS_VERTEXINDEX);
+		int aFaceCount = pDest->GetFaceCount();
+
+		UntypedMeshBufferIteratorConst ito(pDest->GetIndexBuffers(), MBS_VERTEXINDEX);
+		for (int f = 0; f < aFaceCount; ++f)
+		{
+            vec3<uint32_t> ov;
+			ov[0] = ito.GetAsUINT32(); ++ito;
+			ov[1] = ito.GetAsUINT32(); ++ito;
+			ov[2] = ito.GetAsUINT32(); ++ito;
+
+			bool all_vertexs_removed = removed_vertexs[ov[0]] && removed_vertexs[ov[1]] && removed_vertexs[ov[2]];
+
+			if (!all_vertexs_removed)
+			{
+				if (itDest.ptr() != itBase.ptr())
+				{
+					FMemory::Memcpy(itDest.ptr(), itBase.ptr(), itBase.GetElementSize() * 3);
+				}
+
+				itDest += 3;
+			}
+
+			itBase += 3;
+		}
+
+		std::size_t removedIndices = itBase - itDest;
+		check(removedIndices % 3 == 0);
+
+		pDest->GetFaceBuffers().SetElementCount(aFaceCount - (int)removedIndices / 3);
+		pDest->GetIndexBuffers().SetElementCount(aFaceCount * 3 - (int)removedIndices);
+
+		// TODO: Should redo/reorder the face buffer before SetElementCount since some deleted faces could be left and some remaining faces deleted.
+
+        // Fix the surface data if present.
+        if (pDest->m_surfaces.Num())
+        {
+            // We assume there will be only one.
+            check(pDest->m_surfaces.Num()==1);
+
+            pDest->m_surfaces[0].m_indexCount -= (int32_t)removedIndices;
+        }
+
+		return pDest;
+	}
+
+}
