[mutable] Move Mutable plugin from the restricted folder to the experimental plugins folder.

diff --git a/Engine/Plugins/Experimental/Mutable/Source/MutableRuntime/Private/OpMeshClipWithMesh.h b/Engine/Plugins/Experimental/Mutable/Source/MutableRuntime/Private/OpMeshClipWithMesh.h
new file mode 100644
--- /dev/null
+++ b/Engine/Plugins/Experimental/Mutable/Source/MutableRuntime/Private/OpMeshClipWithMesh.h
@@ -0,0 +1,1353 @@
+// Copyright Epic Games, Inc. All Rights Reserved.
+
+#pragma once
+
+#include "MeshPrivate.h"
+#include "ConvertData.h"
+#include "Platform.h"
+#include "OpMeshChartDifference.h"
+#include "MutableTrace.h"
+#include "Math/Ray.h"
+#include "GeometryCore/Public/TriangleTypes.h"
+#include "GeometryCore/Public/BoxTypes.h"
+#include "GeometryCore/Public/Intersection/IntrRay3Triangle3.h"
+#include "GeometryCore/Public/MathUtil.h"
+#include "GeometryCore/Public/IntVectorTypes.h"
+
+#include "Math/UnrealMathUtility.h"
+
+#include <memory>
+
+
+namespace mu
+{
+    const static float vert_collapse_eps = 0.0001f;
+
+    //---------------------------------------------------------------------------------------------
+    //! Create a map from vertices into vertices, collapsing vertices that have the same position
+    //---------------------------------------------------------------------------------------------
+    inline void MeshCreateCollapsedVertexMap( const Mesh* pMesh,
+                                       mu::vector<int>& collapsedVertexMap,
+                                       mu::vector<vec3f>& vertices
+                                       )
+    {
+        MUTABLE_CPUPROFILER_SCOPE(MeshCreateCollapsedVertexMap);
+    	
+        int vcount = pMesh->GetVertexCount();
+        collapsedVertexMap.resize(vcount);
+        vertices.resize(vcount);
+
+        const FMeshBufferSet& MBSPriv = pMesh->GetVertexBuffers();
+        for (int32 b = 0; b < MBSPriv.m_buffers.Num(); ++b)
+        {
+            for (int32 c = 0; c < MBSPriv.m_buffers[b].m_channels.Num(); ++c)
+            {
+                MESH_BUFFER_SEMANTIC sem = MBSPriv.m_buffers[b].m_channels[c].m_semantic;
+                int semIndex = MBSPriv.m_buffers[b].m_channels[c].m_semanticIndex;
+
+                UntypedMeshBufferIteratorConst it(pMesh->GetVertexBuffers(), sem, semIndex);
+
+                switch (sem)
+                {
+                case MBS_POSITION:
+                    // First create a cache of the vertices in the vertices array
+                    for (int v = 0; v < vcount; ++v)
+                    {
+                        vec3f vertex(0.0f, 0.0f, 0.0f);
+                        for (int i = 0; i < 3; ++i)
+                        {
+                            ConvertData(i, &vertex[0], MBF_FLOAT32, it.ptr(), it.GetFormat());
+                        }
+
+                        vertices[v] = vertex;
+
+                        ++it;
+                    }
+
+                    // Create map to store which vertices are the same (collapse nearby vertices)
+                    for (std::size_t v = 0; v < vertices.size(); ++v)
+                    {
+                        collapsedVertexMap[v] = (int32_t)v;
+
+                        for (std::size_t candidate_v = 0; candidate_v < v; ++candidate_v)
+                        {
+                            int collapsed_candidate_v_idx = collapsedVertexMap[candidate_v];
+
+                            // Test whether v and the collapsed candidate are close enough to be collapsed
+                            vec3f r = vertices[collapsed_candidate_v_idx] - vertices[v];
+
+                            if (dot(r, r) <= vert_collapse_eps * vert_collapse_eps)
+                            {
+                                collapsedVertexMap[v] = collapsed_candidate_v_idx;
+                                break;
+                            }
+                        }
+                    }
+
+                    break;
+
+                default:
+                    break;
+                }
+            }
+        }
+    }
+
+    //---------------------------------------------------------------------------------------------
+    //! Create a map from vertices into vertices, collapsing vertices that have the same position, 
+	//! This version uses UE Containers to return.	
+    //---------------------------------------------------------------------------------------------
+    inline void MeshCreateCollapsedVertexMap( const Mesh* pMesh, TArray<int32>& OutCollapsedVertexMap, TArray<FVector3f>& Vertices )
+    {
+        MUTABLE_CPUPROFILER_SCOPE(MeshCreateCollapsedVertexMap);
+    	
+        int32 VCount = pMesh->GetVertexCount();
+
+        OutCollapsedVertexMap.SetNum(VCount);
+        Vertices.SetNum(VCount);
+
+        const FMeshBufferSet& MBSPriv = pMesh->GetVertexBuffers();
+        for (int32 b = 0; b < MBSPriv.m_buffers.Num(); ++b)
+        {
+            for (int32 c = 0; c < MBSPriv.m_buffers[b].m_channels.Num(); ++c)
+            {
+                MESH_BUFFER_SEMANTIC Sem = MBSPriv.m_buffers[b].m_channels[c].m_semantic;
+                int32 SemIndex = MBSPriv.m_buffers[b].m_channels[c].m_semanticIndex;
+
+                UntypedMeshBufferIteratorConst It(pMesh->GetVertexBuffers(), Sem, SemIndex);
+
+                switch (Sem)
+                {
+                case MBS_POSITION:
+				{
+                    // First create a cache of the vertices in the vertices array
+                    for (int32 V = 0; V < VCount; ++V)
+                    {
+                        FVector3f Vertex(0.0f, 0.0f, 0.0f);
+                        for (int32 I = 0; I < 3; ++I)
+                        {
+                            ConvertData(I, &Vertex[0], MBF_FLOAT32, It.ptr(), It.GetFormat());
+                        }
+
+                        Vertices[V] = Vertex;
+
+                        ++It;
+                    }
+					
+                    // Create map to store which vertices are the same (collapse nearby vertices)
+					const int32 VerticesNum = Vertices.Num();
+                    for (int32 V = 0; V < VerticesNum; ++V)
+                    {
+                        OutCollapsedVertexMap[V] = V;
+
+                        for (int32 CandidateV = 0; CandidateV < V; ++CandidateV)
+                        {
+                            const int32 CollapsedCandidateVIdx = OutCollapsedVertexMap[CandidateV];
+
+                            // Test whether v and the collapsed candidate are close enough to be collapsed
+                            const FVector3f R = Vertices[CollapsedCandidateVIdx] - Vertices[V];
+
+                            if (R.Dot(R) <= TMathUtilConstants<float>::ZeroTolerance)
+                            {
+                                OutCollapsedVertexMap[V] = CollapsedCandidateVIdx;
+                                break;
+                            }
+                        }
+                    }
+
+                    break;
+				}
+                default:
+                    break;
+                }
+            }
+        }
+    }
+
+    //---------------------------------------------------------------------------------------------
+    //! Return true if a mesh is closed
+    //---------------------------------------------------------------------------------------------
+    inline bool MeshIsClosed( const Mesh* pMesh )
+    {
+        MUTABLE_CPUPROFILER_SCOPE(MeshIsClosed);
+
+        int vcount = pMesh->GetVertexCount();
+        int fcount = pMesh->GetFaceCount();
+
+        mu::vector<vec3f> vertices(vcount);
+
+        // Map in ClipMesh from vertices to the one they are collapsed to because they are very similar, if they aren't collapsed then they are mapped to themselves
+        mu::vector<int> collapsedVertexMap(vcount);
+
+        MeshCreateCollapsedVertexMap( pMesh, collapsedVertexMap, vertices );
+
+
+        // Acumulate edges
+        mu::map< mu::pair<int,int>, int > faceCountPerEdge;
+
+        UntypedMeshBufferIteratorConst itClipMesh(pMesh->GetIndexBuffers(), MBS_VERTEXINDEX);
+        for (int f = 0; f < fcount; ++f)
+        {
+            int face[3];
+            face[0] = collapsedVertexMap[itClipMesh.GetAsUINT32()]; ++itClipMesh;
+            face[1] = collapsedVertexMap[itClipMesh.GetAsUINT32()]; ++itClipMesh;
+            face[2] = collapsedVertexMap[itClipMesh.GetAsUINT32()]; ++itClipMesh;
+
+            for (int e=0; e<3; ++e)
+            {
+                int v0 = face[e];
+                int v1 = face[(e+1)%3];
+
+                if(v0==v1)
+                {
+                    // Degenerated mesh
+                    return false;
+                }
+
+                mu::pair<int,int> edge;
+                edge.first = std::min( v0, v1 );
+                edge.second = std::max( v0, v1 );
+
+                faceCountPerEdge[edge]++;
+            }
+        }
+
+
+        // See if every edge has 2 faces
+        for( const auto& e: faceCountPerEdge)
+        {
+            if (e.second!=2)
+            {
+                return false;
+            }
+        }
+
+        return true;
+    }
+
+
+    //---------------------------------------------------------------------------------------------
+    //! Remove all unused vertices from a mesh, and fix its index buffers.
+    //---------------------------------------------------------------------------------------------
+    inline void MeshRemoveUnusedVertices( Mesh* pMesh )
+    {
+        MUTABLE_CPUPROFILER_SCOPE(MeshRemoveUnusedVertices);
+
+        // Mark used vertices
+        mu::vector<uint8_t> used( pMesh->GetVertexCount(), false );
+        UntypedMeshBufferIteratorConst iti(pMesh->GetIndexBuffers(), MBS_VERTEXINDEX);
+        int IndexCount = pMesh->GetIndexCount();
+        for (int i = 0; i < IndexCount; ++i)
+        {
+            uint32_t Index = iti.GetAsUINT32();
+            ++iti;
+            used[Index] = true;
+        }
+
+        // Build vertex map
+        mu::vector<int32_t> oldToNewVertex(pMesh->GetVertexCount());
+        int totalNewVertices = 0;
+        for (int v = 0; v<pMesh->GetVertexCount(); ++v)
+        {
+            if (used[v])
+            {
+                oldToNewVertex[v] = totalNewVertices;
+                ++totalNewVertices;
+            }
+            else
+            {
+                oldToNewVertex[v] = -1;
+            }
+        }
+
+        // Remove unused vertices and build index map
+        for (int b = 0; b<pMesh->GetVertexBuffers().GetBufferCount(); ++b)
+        {
+            int elemSize = pMesh->GetVertexBuffers().GetElementSize(b);
+            const uint8_t* pSourceData = pMesh->GetVertexBuffers().GetBufferData(b);
+            uint8_t* pData = pMesh->GetVertexBuffers().GetBufferData(b);
+            for (int v = 0; v<pMesh->GetVertexCount(); ++v)
+            {
+                if (oldToNewVertex[v]!=-1)
+                {
+                    uint8_t* pElemData = pData + elemSize*oldToNewVertex[v];
+                    const uint8_t* pElemSourceData = pSourceData + elemSize*v;
+                    // Avoid warning for overlapping memcpy in valgrind
+                    if (pElemData != pElemSourceData)
+                    {
+                        memcpy(pElemData, pElemSourceData, elemSize);
+                    }
+                }
+            }
+        }
+        pMesh->GetVertexBuffers().SetElementCount(totalNewVertices);
+
+        // Update indices
+        UntypedMeshBufferIteratorConst ito(pMesh->GetIndexBuffers(), MBS_VERTEXINDEX);
+        if (ito.GetFormat() == MBF_UINT32)
+        {
+            for (int i = 0; i < IndexCount; ++i)
+            {
+                uint32_t Index = *(uint32_t*)ito.ptr();
+                int32_t NewIndex = oldToNewVertex[Index];
+                check(NewIndex >= 0);
+                *(uint32_t*)ito.ptr() = (uint32_t)NewIndex;
+                ++ito;
+            }
+        }
+        else if (ito.GetFormat() == MBF_UINT16)
+        {
+            for (int i = 0; i < IndexCount; ++i)
+            {
+                uint16_t Index = *(uint16_t*)ito.ptr();
+                int32_t NewIndex = oldToNewVertex[Index];
+                check(NewIndex >= 0);
+                *(uint16_t*)ito.ptr() = (uint16_t)NewIndex;
+                ++ito;
+            }
+        }
+        else
+        {
+            checkf(false, TEXT("Format not implemented.") );
+        }
+
+
+        // \todo: face buffers?
+    }
+
+
+    //---------------------------------------------------------------------------------------------
+    //! Reference version
+    //---------------------------------------------------------------------------------------------
+    inline int get_num_intersections( const vec3f& vertex,
+                                      const vec3f& ray,
+                                      const mu::vector<vec3f>& vertices,
+                                      const mu::vector<uint32_t>& faces,
+                                      mu::vector<int>& collapsedVertexMap,
+                                      mu::vector<uint8_t>& vertex_already_intersected,
+                                      mu::map<mu::pair<int, int>, bool>& edge_already_intersected,
+									  float dynamic_epsilon
+                                      )
+    {
+    	
+        MUTABLE_CPUPROFILER_SCOPE(get_num_intersections);
+    	
+        int num_intersections = 0;
+        vec3f intersection;
+
+		FMemory::Memzero( vertex_already_intersected.data(), vertex_already_intersected.size() );
+
+        edge_already_intersected.clear();
+
+    #define get_collapsed_vertex(vertex_idx) (vertices[collapsedVertexMap[vertex_idx]])
+
+        size_t fcount = faces.size()/3;
+
+        // Check vertex against all ClipMorph faces
+        for (uint32_t face = 0; face < fcount; ++face)
+        {
+            uint32_t vertex_indexs[3] = { faces[3 * face], faces[3 * face + 1], faces[3 * face + 2] };
+
+            vec3f v0 = get_collapsed_vertex(vertex_indexs[0]);
+            vec3f v1 = get_collapsed_vertex(vertex_indexs[1]);
+            vec3f v2 = get_collapsed_vertex(vertex_indexs[2]);
+
+            int out_intersected_vert, out_intersected_edge_v0, out_intersected_edge_v1;
+
+            if (rayIntersectsFace(vertex, ray, v0, v1, v2, intersection, out_intersected_vert, out_intersected_edge_v0, out_intersected_edge_v1, dynamic_epsilon))
+            {
+                bool vertex_not_intersected_before = true;
+                bool edge_not_intersected_before = true;
+
+                if (out_intersected_vert >= 0)
+                {
+                    int collapsed_vert_index = collapsedVertexMap[vertex_indexs[out_intersected_vert]];
+                    vertex_not_intersected_before = !vertex_already_intersected[collapsed_vert_index];
+
+                    vertex_already_intersected[collapsed_vert_index] = true;
+                }
+
+                if (out_intersected_edge_v0 >= 0)
+                {
+                    int collapsed_edge_vert_index0 = collapsedVertexMap[vertex_indexs[out_intersected_edge_v0]];
+                    int collapsed_edge_vert_index1 = collapsedVertexMap[vertex_indexs[out_intersected_edge_v1]];
+
+                    mu::pair<int, int> key;
+                    key.first = collapsed_edge_vert_index0 <= collapsed_edge_vert_index1 ? collapsed_edge_vert_index0 : collapsed_edge_vert_index1;
+                    key.second = collapsed_edge_vert_index1 > collapsed_edge_vert_index0 ? collapsed_edge_vert_index1 : collapsed_edge_vert_index0;
+                    edge_not_intersected_before = edge_already_intersected.find(key) == edge_already_intersected.end();
+
+                    if (edge_not_intersected_before)
+                    {
+                        edge_already_intersected[key] = true;
+                    }
+                }
+
+                if (vertex_not_intersected_before && edge_not_intersected_before)
+                {
+                    num_intersections++;
+                }
+            }
+        }
+
+        return num_intersections;
+    }
+
+
+    //---------------------------------------------------------------------------------------------
+    //! Reference version
+    //---------------------------------------------------------------------------------------------
+    inline void MeshClipMeshClassifyVertices(mu::vector<uint8_t>& vertex_in_clip_mesh,
+                                             const Mesh* pBase, const Mesh* pClipMesh)
+    {
+        MUTABLE_CPUPROFILER_SCOPE(MeshClipMeshClassifyVertices);
+
+        // TODO: this fails in bro for some object. do it at graph construction time to report it 
+        // nicely.
+        //check( MeshIsClosed(pClipMesh) );
+
+        int vcount = pClipMesh->GetVertexBuffers().GetElementCount();
+        int fcount = pClipMesh->GetFaceCount();
+
+        int orig_vcount = pBase->GetVertexBuffers().GetElementCount();
+
+        // Stores whether each vertex in the original mesh in the clip mesh volume
+        vertex_in_clip_mesh.resize(orig_vcount);
+        std::fill(vertex_in_clip_mesh.begin(), vertex_in_clip_mesh.end(), false);
+
+        if (!vcount)
+        {
+            return;
+        }
+
+        mu::vector<vec3f> vertices(vcount); // ClipMesh vertex cache
+        mu::vector<uint32_t> faces(fcount * 3); // ClipMesh face cache
+
+        // Map in ClipMesh from vertices to the one they are collapsed to because they are very 
+        // similar, if they aren't collapsed then they are mapped to themselves
+        mu::vector<int> collapsedVertexMap(vcount);
+
+        MeshCreateCollapsedVertexMap( pClipMesh, collapsedVertexMap, vertices );
+
+
+        // Create cache of the faces
+        UntypedMeshBufferIteratorConst itClipMesh(pClipMesh->GetIndexBuffers(), MBS_VERTEXINDEX);
+
+        for (int f = 0; f < fcount; ++f)
+        {
+            faces[3 * f]     = itClipMesh.GetAsUINT32(); ++itClipMesh;
+            faces[3 * f + 1] = itClipMesh.GetAsUINT32(); ++itClipMesh;
+            faces[3 * f + 2] = itClipMesh.GetAsUINT32(); ++itClipMesh;
+        }
+
+
+        // Create a bounding box of the clip mesh
+        box<vec3f> clipMeshBoundingBox;
+        clipMeshBoundingBox.min = vertices[0];
+        clipMeshBoundingBox.size = vec3f(0.0f,0.0f,0.0f);
+        for (size_t i=1; i<vertices.size(); ++i)
+        {
+            clipMeshBoundingBox.Bound(vertices[i]);
+        }
+
+		// Dynamic distance epsilon to support different engines
+		float maxDimensionBoundingBox = length(clipMeshBoundingBox.size);
+		// 0.000001 is the value that helps to achieve the dynamic epsilon, do not change it
+		float dynamicEpsilon = 0.000001f * maxDimensionBoundingBox * (maxDimensionBoundingBox < 1.0f ? maxDimensionBoundingBox : 1.0f);
+
+        // Z-direction. Don't change this without reviewing the acceleration structures below.
+        const vec3f ray = vec3f(0.f, 0.f, 203897203.f);
+
+        // Create an acceleration grid to avoid testing all clip-mesh triangles.
+        // This assumes that the testing ray direction is Z
+        const int GRID_SIZE = 8;
+        std::unique_ptr<mu::vector<uint32_t>[]> gridFaces( new mu::vector<uint32_t>[GRID_SIZE*GRID_SIZE] );
+        vec2f gridCellSize = clipMeshBoundingBox.size.xy() / (float)GRID_SIZE;
+        for ( int i=0; i<GRID_SIZE; ++i)
+        {
+            for ( int j=0; j<GRID_SIZE; ++j)
+            {
+                box<vec2f> cellBox;
+                cellBox.size = gridCellSize;
+                cellBox.min = clipMeshBoundingBox.min.xy() + gridCellSize * vec2f((float)i,(float)j);
+
+                mu::vector<uint32_t>& cellFaces = gridFaces[i+j*GRID_SIZE];
+                cellFaces.reserve(fcount/GRID_SIZE);
+                for (int f = 0; f < fcount; ++f)
+                {
+                    // Imprecise, conservative classification of faces.
+                    box<vec2f> faceBox;
+                    faceBox.size = vec2f(0.0f,0.0f);
+                    faceBox.min = vertices[ faces[3 * f + 0] ].xy();
+                    faceBox.Bound(vertices[ faces[3 * f + 1] ].xy());
+                    faceBox.Bound(vertices[ faces[3 * f + 2] ].xy());
+
+                    if (cellBox.Intersects(faceBox))
+                    {
+                        cellFaces.push_back(faces[3 * f + 0]);
+                        cellFaces.push_back(faces[3 * f + 1]);
+                        cellFaces.push_back(faces[3 * f + 2]);
+                    }
+                }
+            }
+        }
+
+
+        // Now go through all vertices in the mesh and record whether they are inside or outside of the ClipMesh
+        uint32_t dest_vertex_count = pBase->GetVertexCount();
+
+        const FMeshBufferSet& MBSPriv2 = pBase->GetVertexBuffers();
+        for (int32 b = 0; b < MBSPriv2.m_buffers.Num(); ++b)
+        {
+            for (int32 c = 0; c < MBSPriv2.m_buffers[b].m_channels.Num(); ++c)
+            {
+                MESH_BUFFER_SEMANTIC sem = MBSPriv2.m_buffers[b].m_channels[c].m_semantic;
+                int semIndex = MBSPriv2.m_buffers[b].m_channels[c].m_semanticIndex;
+
+                UntypedMeshBufferIteratorConst it(pBase->GetVertexBuffers(), sem, semIndex);
+
+                mu::vector<uint8_t> vertex_already_intersected(vcount);
+                mu::map<mu::pair<int, int>, bool> edge_already_intersected;
+
+                switch (sem)
+                {
+                case MBS_POSITION:
+                    for (uint32_t v = 0; v < dest_vertex_count; ++v)
+                    {
+                        vec3f vertex(0.0f, 0.0f, 0.0f);
+                        for (int i = 0; i < 3; ++i)
+                        {
+                            ConvertData(i, &vertex[0], MBF_FLOAT32, it.ptr(), it.GetFormat());
+                        }
+
+                        // Find out grid cell
+                        vec2f hPos = (vertex.xy()-clipMeshBoundingBox.min.xy()) / clipMeshBoundingBox.size.xy() * (float)GRID_SIZE;
+                        int i = std::min( std::max( (int)hPos.x(), 0 ), GRID_SIZE-1 );
+                        int j = std::min( std::max( (int)hPos.y(), 0 ), GRID_SIZE-1 );
+
+                        // Early discard test: if the vertex is not inside the bounding box of the clip mesh, it won't be clipped.
+                        bool contains = (clipMeshBoundingBox.ContainsInclusive(vertex));
+
+                        if (contains)
+                        {
+                            // Optimised test
+                            int cellIndex = i+j*GRID_SIZE;
+                            mu::vector<uint32_t>& thisCell = gridFaces[cellIndex];
+                            int num_intersections = get_num_intersections(
+                                        vertex,
+                                        ray,
+                                        vertices,
+                                        thisCell,
+                                        collapsedVertexMap,
+                                        vertex_already_intersected,
+                                        edge_already_intersected,
+										dynamicEpsilon
+                                        );
+
+                            // Full test BLEH, debug
+//                            int full_num_intersections = get_num_intersections(
+//                                        vertex,
+//                                        ray,
+//                                        vertices,
+//                                        faces,
+//                                        collapsedVertexMap,
+//                                        vertex_already_intersected,
+//                                        edge_already_intersected
+//                                        );
+
+
+                            vertex_in_clip_mesh[v] = num_intersections % 2 == 1;
+
+                            // This may be used to debug degenerated cases if the conditional above is also removed.
+                            // \todo: make sure it works well
+//                            if (!contains && vertex_in_clip_mesh[v])
+//                            {
+//                                assert(false);
+//                            }
+                        }
+
+                        ++it;
+                    }
+                    break;
+
+                default:
+                    break;
+                }
+            }
+        }
+    }
+
+
+    //---------------------------------------------------------------------------------------------
+    //! Core Geometry version
+    //---------------------------------------------------------------------------------------------
+    inline int GetNumIntersections( const FRay3f& Ray,
+                                    const TArray<FVector3f>& Vertices,
+                                    const TArray<uint32>& Faces,
+                                    const TArray<int32>& CollapsedVertexMap,
+                                    TArray<uint8>& InOutVertexAlreadyIntersected,
+                                    TSet<uint64>& InOutEdgeAlreadyIntersected,
+									const float DynamicEpsilon )
+    {
+		using namespace UE::Geometry;    
+
+        MUTABLE_CPUPROFILER_SCOPE(GetNumIntersections);
+    	
+        int32 NumIntersections = 0;
+
+		FMemory::Memzero( InOutVertexAlreadyIntersected.GetData(), InOutVertexAlreadyIntersected.Num() );
+		InOutEdgeAlreadyIntersected.Empty();
+
+		auto GetCollapsedVertex = [&CollapsedVertexMap, &Vertices]( const uint32 V ) -> const FVector3f&
+		{
+			return Vertices[ CollapsedVertexMap[V] ];
+		};
+
+        int32 FaceCount = Faces.Num() / 3;
+
+		UE::Geometry::FIntrRay3Triangle3f Intersector(Ray, UE::Geometry::FTriangle3f());
+
+        // Check vertex against all ClipMorph faces
+        for (int32 Face = 0; Face < FaceCount; ++Face)
+        {
+            const uint32 VertexIndexs[3] = { Faces[3 * Face], Faces[3 * Face + 1], Faces[3 * Face + 2] };
+
+            const FVector3f& V0 = GetCollapsedVertex(VertexIndexs[0]);
+            const FVector3f& V1 = GetCollapsedVertex(VertexIndexs[1]);
+            const FVector3f& V2 = GetCollapsedVertex(VertexIndexs[2]);
+
+			Intersector.Triangle = UE::Geometry::FTriangle3f( V0, V1, V2 );
+
+			if ( Intersector.Find() )
+            {
+				// Find if close to edge using barycentric coordinates form intersector.
+	
+				// Is the Dynamic Epsilon needed?. Intersector bary coords are in the range [0.0, 1.0], 
+				// furthermore, IntrRay3Triangle3f::TriangleBaryCoords are double pressison, not sure 
+				// is intentional or is a bug ( is double even when the FReal type is float ) 
+				const bool IntersectsEdge01 = FMath::IsNearlyZero( Intersector.TriangleBaryCoords.Z, DynamicEpsilon );
+				const bool IntersectsEdge02 = FMath::IsNearlyZero( Intersector.TriangleBaryCoords.Y, DynamicEpsilon );
+				const bool IntersectsEdge12 = FMath::IsNearlyZero( Intersector.TriangleBaryCoords.X, DynamicEpsilon );
+
+				int32 IntersectedTriangleVertexId = -1;
+				IntersectedTriangleVertexId =  (IntersectsEdge01 & IntersectsEdge02) ? 0 : IntersectedTriangleVertexId;
+				IntersectedTriangleVertexId =  (IntersectsEdge01 & IntersectsEdge12) ? 1 : IntersectedTriangleVertexId;
+				IntersectedTriangleVertexId =  (IntersectsEdge02 & IntersectsEdge12) ? 2 : IntersectedTriangleVertexId;
+
+				bool bIsAlreadyIntersected = false;
+
+                if ( IntersectedTriangleVertexId >= 0 )
+                {
+                    const int32 CollapsedVertIndex = CollapsedVertexMap[VertexIndexs[IntersectedTriangleVertexId]];
+                    bIsAlreadyIntersected = InOutVertexAlreadyIntersected[CollapsedVertIndex] == 0;
+
+                    InOutVertexAlreadyIntersected[CollapsedVertIndex] = true;
+                }
+				else if ( IntersectsEdge01 | IntersectsEdge02 | IntersectsEdge12 )
+                {
+					int32 EdgeV0 = (IntersectsEdge01 | IntersectsEdge02) ? 0 : 1;
+					int32 EdgeV1 = IntersectsEdge01 ? 1 : 2;				
+
+                    const int32 CollapsedEdgeVertIndex0 = CollapsedVertexMap[VertexIndexs[EdgeV0]];
+                    const int32 CollapsedEdgeVertIndex1 = CollapsedVertexMap[VertexIndexs[EdgeV1]];
+
+                    const uint64 EdgeKey = ( ( (uint64)FMath::Max( CollapsedEdgeVertIndex0, CollapsedEdgeVertIndex1 ) ) << 32 ) | 
+							                   (uint64)FMath::Min( CollapsedEdgeVertIndex0, CollapsedEdgeVertIndex1 );
+	
+					InOutEdgeAlreadyIntersected.FindOrAdd( EdgeKey, &bIsAlreadyIntersected );
+                }
+
+                NumIntersections += (int32)(!bIsAlreadyIntersected);
+            }
+        }
+
+        return NumIntersections;
+    }
+
+
+    //---------------------------------------------------------------------------------------------
+    //! Core Geometry version
+    //---------------------------------------------------------------------------------------------
+    inline void MeshClipMeshClassifyVertices(TArray<uint8>& VertexInClipMesh,
+                                             const Mesh* pBase, const Mesh* pClipMesh)
+    {
+		using namespace UE::Geometry;
+
+        MUTABLE_CPUPROFILER_SCOPE(MeshClipMeshClassifyVertices);
+    	
+        // TODO: this fails in bro for some object. do it at graph construction time to report it 
+        // nicely.
+        //check( MeshIsClosed(pClipMesh) );
+
+        const int32 VCount = pClipMesh->GetVertexBuffers().GetElementCount();
+        const int32 FCount = pClipMesh->GetFaceCount();
+
+        int32 OrigVertCount = pBase->GetVertexBuffers().GetElementCount();
+
+        // Stores whether each vertex in the original mesh in the clip mesh volume
+        VertexInClipMesh.AddZeroed(OrigVertCount);
+
+        if (VCount == 0)
+        {
+            return;
+        }
+
+        TArray<FVector3f> Vertices; // ClipMesh vertex cache
+		Vertices.SetNumUninitialized(VCount);
+
+        TArray<uint32> Faces; // ClipMesh face cache
+		Faces.SetNumUninitialized(FCount * 3); 
+
+        // Map in ClipMesh from vertices to the one they are collapsed to because they are very 
+        // similar, if they aren't collapsed then they are mapped to themselves
+        TArray<int32> CollapsedVertexMap;
+		CollapsedVertexMap.AddUninitialized(VCount);
+
+        MeshCreateCollapsedVertexMap( pClipMesh, CollapsedVertexMap, Vertices );
+
+        // Create cache of the faces
+        UntypedMeshBufferIteratorConst ItClipMesh(pClipMesh->GetIndexBuffers(), MBS_VERTEXINDEX);
+
+        for ( int32 F = 0; F < FCount; ++F )
+        {
+            Faces[3 * F]     = ItClipMesh.GetAsUINT32(); ++ItClipMesh;
+            Faces[3 * F + 1] = ItClipMesh.GetAsUINT32(); ++ItClipMesh;
+            Faces[3 * F + 2] = ItClipMesh.GetAsUINT32(); ++ItClipMesh;
+        }
+
+
+        // Create a bounding box of the clip mesh
+		UE::Geometry::FAxisAlignedBox3f ClipMeshBoundingBox = UE::Geometry::FAxisAlignedBox3f::Empty();
+
+        for ( const FVector3f& Vert : Vertices )
+        {
+            ClipMeshBoundingBox.Contain( Vert );
+        }
+
+		// Dynamic distance epsilon to support different engines
+		const float MaxDimensionBoundingBox = ClipMeshBoundingBox.DiagonalLength();
+		// 0.000001 is the value that helps to achieve the dynamic epsilon, do not change it
+		const float DynamicEpsilon = 0.000001f * MaxDimensionBoundingBox * (MaxDimensionBoundingBox < 1.0f ? MaxDimensionBoundingBox : 1.0f);
+
+        // Create an acceleration grid to avoid testing all clip-mesh triangles.
+        // This assumes that the testing ray direction is Z
+        constexpr int32 GRID_SIZE = 8;
+        TUniquePtr<TArray<uint32>[]> GridFaces( new TArray<uint32>[GRID_SIZE*GRID_SIZE] );
+        const FVector2f GridCellSize = FVector2f( ClipMeshBoundingBox.Width(), ClipMeshBoundingBox.Height() ) / (float)GRID_SIZE;
+
+        for ( int32 I = 0; I < GRID_SIZE; ++I )
+        {
+            for ( int32 J = 0; J < GRID_SIZE; ++J )
+            {
+				const FVector2f BBoxMin = FVector2f(ClipMeshBoundingBox.Min.X, ClipMeshBoundingBox.Min.Y) + GridCellSize * FVector2f(I, J);
+
+                FAxisAlignedBox2f CellBox( BBoxMin, BBoxMin + GridCellSize );
+
+                TArray<uint32>& CellFaces = GridFaces[I+J*GRID_SIZE];
+                CellFaces.Empty(FCount/GRID_SIZE);
+                for (int32 F = 0; F < FCount; ++F)
+                {
+                    // Imprecise, conservative classification of faces.
+					const FVector3f& V0 = Vertices[ Faces[3*F + 0] ];
+					const FVector3f& V1 = Vertices[ Faces[3*F + 1] ];
+					const FVector3f& V2 = Vertices[ Faces[3*F + 2] ];
+
+					FAxisAlignedBox2f FaceBox;
+					FaceBox.Contain( FVector2f( V0.X, V0.Y ) );
+                    FaceBox.Contain( FVector2f( V1.X, V1.Y ) );
+                    FaceBox.Contain( FVector2f( V2.X, V2.Y ) );
+
+                    if (CellBox.Intersects(FaceBox))
+                    {
+                        CellFaces.Add( Faces[3*F + 0] );
+                        CellFaces.Add( Faces[3*F + 1] );
+                        CellFaces.Add( Faces[3*F + 2] );
+                    }
+                }
+            }
+        }	
+
+        // Now go through all vertices in the mesh and record whether they are inside or outside of the ClipMesh
+        uint32 DestVertexCount = pBase->GetVertexCount();
+
+        const FMeshBufferSet& MBSPriv2 = pBase->GetVertexBuffers();
+        for (int32 b = 0; b < MBSPriv2.m_buffers.Num(); ++b)
+        {
+            for (int32 c = 0; c < MBSPriv2.m_buffers[b].m_channels.Num(); ++c)
+            {
+                MESH_BUFFER_SEMANTIC Sem = MBSPriv2.m_buffers[b].m_channels[c].m_semantic;
+                int32 SemIndex = MBSPriv2.m_buffers[b].m_channels[c].m_semanticIndex;
+
+                UntypedMeshBufferIteratorConst It(pBase->GetVertexBuffers(), Sem, SemIndex);
+
+                TArray<uint8> VertexAlreadyIntersected;
+				VertexAlreadyIntersected.AddZeroed(VCount);
+
+                TSet<uint64> EdgeAlreadyIntersected;
+
+                switch ( Sem )
+                {
+                case MBS_POSITION:
+                    for (uint32 V = 0; V < DestVertexCount; ++V)
+                    {
+                        FVector3f Vertex(0.0f, 0.0f, 0.0f);
+                        for (int32 Offset = 0; Offset < 3; ++Offset)
+                        {
+                            ConvertData(Offset, &Vertex[0], MBF_FLOAT32, It.ptr(), It.GetFormat());
+                        }
+						
+						const FVector2f ClipBBoxMin = FVector2f( ClipMeshBoundingBox.Min.X, ClipMeshBoundingBox.Min.Y);
+						const FVector2f ClipBBoxSize = FVector2f( ClipMeshBoundingBox.Width(), ClipMeshBoundingBox.Height() );
+
+						const FVector2i HPos = FVector2i( ( ( FVector2f( Vertex.X, Vertex.Y ) - ClipBBoxMin ) / ClipBBoxSize ) * (float)GRID_SIZE );   
+						const FVector2i CellCoord = FVector2i( FMath::Clamp( HPos.X, 0, GRID_SIZE - 1 ), 
+														       FMath::Clamp( HPos.Y, 0, GRID_SIZE - 1 ) );
+
+                        // Early discard test: if the vertex is not inside the bounding box of the clip mesh, it won't be clipped.
+                        const bool bContainsVertex = ClipMeshBoundingBox.Contains( Vertex );
+
+                        if ( bContainsVertex )
+                        {
+                            // Optimised test	
+							
+							// Z-direction. Don't change this without reviewing the acceleration structure.
+							FVector3f RayDir = FVector3f(0.0f, 0.0f, 1.0f);	
+                            const int32 CellIndex = CellCoord.X + CellCoord.Y*GRID_SIZE;
+
+                            int32 NumIntersections = GetNumIntersections( 
+									FRay3f( Vertex, RayDir ), 
+									Vertices,
+									GridFaces[CellIndex],
+									CollapsedVertexMap, 
+									VertexAlreadyIntersected, 
+									EdgeAlreadyIntersected, 
+								    DynamicEpsilon);
+
+                            // Full test BLEH, debug
+//                            int32 FullNumIntersections = GetNumIntersections(
+//									FRay3f( Vertex, RayDir ), 
+//									Vertices,
+//									Faces,
+//									CollapsedVertexMap, 
+//									VertexAlreadyIntersected, 
+//									EdgeAlreadyIntersected, 
+//								    DynamicEpsilon);
+
+							VertexInClipMesh[V] = NumIntersections % 2 == 1;
+
+                            // This may be used to debug degenerated cases if the conditional above is also removed.
+                            // \todo: make sure it works well
+//                          if (!bContainsVertex && VertexInClipMesh[V])
+//                          {
+//                              assert(false);
+//                          }
+                        }
+
+                        ++It;
+                    }
+                    break;
+
+                default:
+                    break;
+                }
+            }
+        }
+    }	
+
+
+    //---------------------------------------------------------------------------------------------
+    //! Reference version
+    //---------------------------------------------------------------------------------------------
+    inline MeshPtr MeshClipWithMesh_Reference(const Mesh* pBase, const Mesh* pClipMesh)
+    {
+        MUTABLE_CPUPROFILER_SCOPE(MeshClipWithMesh_Reference);
+        
+		MeshPtr pDest = pBase->Clone();
+
+        uint32_t vcount = pClipMesh->GetVertexBuffers().GetElementCount();
+        if (!vcount)
+        {
+            return pDest; // Since there's nothing to clip against return a copy of the original base mesh
+        }
+
+        mu::vector<uint8_t> vertex_in_clip_mesh;  // Stores whether each vertex in the original mesh in in the clip mesh volume
+        MeshClipMeshClassifyVertices( vertex_in_clip_mesh, pBase, pClipMesh );
+
+        // Now remove all the faces from the result mesh that have all the vertices outside the clip volume
+        UntypedMeshBufferIteratorConst itBase(pDest->GetIndexBuffers(), MBS_VERTEXINDEX);
+        UntypedMeshBufferIterator itDest(pDest->GetIndexBuffers(), MBS_VERTEXINDEX);
+        int aFaceCount = pDest->GetFaceCount();
+
+        UntypedMeshBufferIteratorConst ito(pDest->GetIndexBuffers(), MBS_VERTEXINDEX);
+        for (int f = 0; f < aFaceCount; ++f)
+        {
+            vec3<uint32_t> ov;
+            ov[0] = ito.GetAsUINT32(); ++ito;
+            ov[1] = ito.GetAsUINT32(); ++ito;
+            ov[2] = ito.GetAsUINT32(); ++ito;
+
+            bool all_verts_in = vertex_in_clip_mesh[ov[0]] && vertex_in_clip_mesh[ov[1]] && vertex_in_clip_mesh[ov[2]];
+
+            if (!all_verts_in)
+            {
+                if (itDest.ptr() != itBase.ptr())
+                {
+					FMemory::Memcpy(itDest.ptr(), itBase.ptr(), itBase.GetElementSize() * 3);
+                }
+
+                itDest += 3;
+            }
+
+            itBase += 3;
+        }
+
+        std::size_t removedIndices = itBase - itDest;
+        check(removedIndices % 3 == 0);
+
+        pDest->GetFaceBuffers().SetElementCount(aFaceCount - (int)removedIndices / 3);
+        pDest->GetIndexBuffers().SetElementCount(aFaceCount * 3 - (int)removedIndices);
+
+        // TODO: Should redo/reorder the face buffer before SetElementCount since some deleted faces could be left and some remaining faces deleted.
+
+        //if (pDest->GetFaceBuffers().GetElementCount() == 0)
+        //{
+        //	pDest = pBase->Clone(); // If all faces have been discarded, return a copy of the unmodified mesh because unreal doesn't like empty meshes
+        //}
+
+        // [jordi] Remove unused vertices. This is necessary to avoid returning a mesh with vertices and no faces, which screws some
+        // engines like Unreal Engine 4.
+        MeshRemoveUnusedVertices( pDest.get() );
+
+        return pDest;
+    }
+
+
+    //---------------------------------------------------------------------------------------------
+    //! CoreGeometry version
+    //---------------------------------------------------------------------------------------------
+    inline MeshPtr MeshClipWithMesh(const Mesh* pBase, const Mesh* pClipMesh)
+    {
+        MUTABLE_CPUPROFILER_SCOPE(MeshClipWithMesh);
+
+        MeshPtr pDest = pBase->Clone();
+
+        uint32_t VCount = pClipMesh->GetVertexBuffers().GetElementCount();
+        if (!VCount)
+        {
+            return pDest; // Since there's nothing to clip against return a copy of the original base mesh
+        }
+
+        TArray<uint8> VertexInClipMesh;  // Stores whether each vertex in the original mesh in in the clip mesh volume
+        MeshClipMeshClassifyVertices( VertexInClipMesh, pBase, pClipMesh );
+
+        // Now remove all the faces from the result mesh that have all the vertices outside the clip volume
+        UntypedMeshBufferIteratorConst ItBase(pDest->GetIndexBuffers(), MBS_VERTEXINDEX);
+        UntypedMeshBufferIterator ItDest(pDest->GetIndexBuffers(), MBS_VERTEXINDEX);
+        int32 AFaceCount = pDest->GetFaceCount();
+
+        UntypedMeshBufferIteratorConst Ito(pDest->GetIndexBuffers(), MBS_VERTEXINDEX);
+        for (int32 F = 0; F < AFaceCount; ++F)
+        {
+			uint32 OV[3] = {0, 0, 0};
+
+            OV[0] = Ito.GetAsUINT32(); ++Ito;
+            OV[1] = Ito.GetAsUINT32(); ++Ito;
+            OV[2] = Ito.GetAsUINT32(); ++Ito;
+
+            bool AllVertsIn = VertexInClipMesh[OV[0]] && VertexInClipMesh[OV[1]] && VertexInClipMesh[OV[2]];
+
+            if (!AllVertsIn)
+            {
+                if (ItDest.ptr() != ItBase.ptr())
+                {
+					FMemory::Memcpy(ItDest.ptr(), ItBase.ptr(), ItBase.GetElementSize() * 3);
+                }
+
+                ItDest += 3;
+            }
+
+            ItBase += 3;
+        }
+
+        SIZE_T RemovedIndices = ItBase - ItDest;
+        check(RemovedIndices % 3 == 0);
+
+        pDest->GetFaceBuffers().SetElementCount(AFaceCount - (int32)RemovedIndices / 3);
+        pDest->GetIndexBuffers().SetElementCount(AFaceCount * 3 - (int32)RemovedIndices);
+
+        // TODO: Should redo/reorder the face buffer before SetElementCount since some deleted faces could be left and some remaining faces deleted.
+
+        //if (pDest->GetFaceBuffers().GetElementCount() == 0)
+        //{
+        //	pDest = pBase->Clone(); // If all faces have been discarded, return a copy of the unmodified mesh because unreal doesn't like empty meshes
+        //}
+
+        // [jordi] Remove unused vertices. This is necessary to avoid returning a mesh with vertices and no faces, which screws some
+        // engines like Unreal Engine 4.
+        MeshRemoveUnusedVertices( pDest.get() );
+
+        return pDest;
+    }
+
+
+    //---------------------------------------------------------------------------------------------
+    //!
+    //---------------------------------------------------------------------------------------------
+    inline MeshPtr CreateMask( MeshPtrConst pBase, const mu::vector<uint8_t>& excludedVertices )
+    {
+
+        int maskVertexCount = 0;
+        for( uint8_t b: excludedVertices )
+        {
+            if (!b) ++maskVertexCount;
+        }
+
+        MeshPtr pMask = new Mesh();
+
+        // Create the vertex buffer
+        {
+            pMask->GetVertexBuffers().SetElementCount( maskVertexCount );
+            pMask->GetVertexBuffers().SetBufferCount( 1 );
+
+            vector<MESH_BUFFER_SEMANTIC> semantic;
+            vector<int> semanticIndex;
+            vector<MESH_BUFFER_FORMAT> format;
+            vector<int> components;
+            vector<int> offsets;
+
+            // Vertex index channel
+            semantic.push_back( MBS_VERTEXINDEX );
+            semanticIndex.push_back( 0 );
+            format.push_back( MBF_UINT32 );
+            components.push_back( 1 );
+            offsets.push_back( 0 );
+
+            pMask->GetVertexBuffers().SetBuffer
+                (
+                    0,
+                    4,
+                    1,
+                    &semantic[0],
+                    &semanticIndex[0],
+                    &format[0],
+                    &components[0],
+                    &offsets[0]
+                );
+        }
+
+        MeshBufferIterator<MBF_UINT32,uint32_t,1> itMask( pMask->GetVertexBuffers(), MBS_VERTEXINDEX, 0 );
+        UntypedMeshBufferIteratorConst itBase( pBase->GetVertexBuffers(), MBS_VERTEXINDEX, 0 );
+
+        for (size_t v = 0; v<excludedVertices.size(); ++v)
+        {
+            if (!excludedVertices[v])
+            {
+                (*itMask)[0] = itBase.GetAsUINT32();
+                ++itMask;
+            }
+
+            ++itBase;
+        }
+
+        return pMask;
+    }
+
+    //---------------------------------------------------------------------------------------------
+    //!
+    //---------------------------------------------------------------------------------------------
+    inline MeshPtr CreateMask( MeshPtrConst pBase, const TArray<uint8>& ExcludedVertices )
+    {
+
+        int32 MaskVertexCount = 0;
+        for( uint8 B: ExcludedVertices )
+        {
+            MaskVertexCount += (int32)(!B);
+        }
+
+        MeshPtr pMask = new Mesh();
+
+        // Create the vertex buffer
+        {
+            pMask->GetVertexBuffers().SetElementCount( MaskVertexCount );
+            pMask->GetVertexBuffers().SetBufferCount( 1 );
+
+            vector<MESH_BUFFER_SEMANTIC> semantic;
+            vector<int> semanticIndex;
+            vector<MESH_BUFFER_FORMAT> format;
+            vector<int> components;
+            vector<int> offsets;
+
+            // Vertex index channel
+            semantic.push_back( MBS_VERTEXINDEX );
+            semanticIndex.push_back( 0 );
+            format.push_back( MBF_UINT32 );
+            components.push_back( 1 );
+            offsets.push_back( 0 );
+
+            pMask->GetVertexBuffers().SetBuffer
+                (
+                    0,
+                    4,
+                    1,
+                    &semantic[0],
+                    &semanticIndex[0],
+                    &format[0],
+                    &components[0],
+                    &offsets[0]
+                );
+        }
+
+        MeshBufferIterator<MBF_UINT32,uint32_t,1> itMask( pMask->GetVertexBuffers(), MBS_VERTEXINDEX, 0 );
+        UntypedMeshBufferIteratorConst itBase( pBase->GetVertexBuffers(), MBS_VERTEXINDEX, 0 );
+
+		const int32 ExcludedVerticesNum = ExcludedVertices.Num();
+        for ( int32 V = 0; V < ExcludedVerticesNum; ++V )
+        {
+            if (!ExcludedVertices[V])
+            {
+                (*itMask)[0] = itBase.GetAsUINT32();
+                ++itMask;
+            }
+
+            ++itBase;
+        }
+
+        return pMask;
+    }
+
+
+    //---------------------------------------------------------------------------------------------
+    //! Generate a mask mesh with the faces of the base mesh inside the clip mesh.
+    //---------------------------------------------------------------------------------------------
+    inline MeshPtr MeshMaskClipMesh(const Mesh* pBase, const Mesh* pClipMesh )
+    {
+        MUTABLE_CPUPROFILER_SCOPE(MeshMaskClipMesh);
+
+        uint32 VCount = pClipMesh->GetVertexBuffers().GetElementCount();
+        if (!VCount)
+        {
+            return nullptr; // Since there's nothing to clip against return a copy of the original base mesh
+        }
+
+        TArray<uint8> VertexInClipMesh;  // Stores whether each vertex in the original mesh in in the clip mesh volume
+        MeshClipMeshClassifyVertices( VertexInClipMesh, pBase, pClipMesh );
+
+        // We only remove vertices if all their faces are clipped
+        TArray<uint8> VertexWithFaceNotClipped;
+		VertexWithFaceNotClipped.AddZeroed( VertexInClipMesh.Num() );
+
+        UntypedMeshBufferIteratorConst Ito(pBase->GetIndexBuffers(), MBS_VERTEXINDEX);
+        int32 AFaceCount = pBase->GetFaceCount();
+        for (int32 F = 0; F < AFaceCount; ++F)
+        {
+			uint32 OV[3] = { 0, 0, 0 };
+
+            OV[0] = Ito.GetAsUINT32(); ++Ito;
+            OV[1] = Ito.GetAsUINT32(); ++Ito;
+            OV[2] = Ito.GetAsUINT32(); ++Ito;
+
+            bool bFaceClipped =
+                    VertexInClipMesh[OV[0]] &&
+                    VertexInClipMesh[OV[1]] &&
+                    VertexInClipMesh[OV[2]];
+
+            if (!bFaceClipped)
+            {
+                VertexWithFaceNotClipped[OV[0]] = true;
+                VertexWithFaceNotClipped[OV[1]] = true;
+                VertexWithFaceNotClipped[OV[2]] = true;
+            }
+        }
+
+        return CreateMask( pBase, VertexWithFaceNotClipped );
+    }
+
+    //---------------------------------------------------------------------------------------------
+    //! Generate a mask mesh with the faces of the base mesh inside the clip mesh.
+    //---------------------------------------------------------------------------------------------
+    inline MeshPtr MeshMaskClipMesh_Reference(const Mesh* pBase, const Mesh* pClipMesh )
+    {
+        MUTABLE_CPUPROFILER_SCOPE(MeshMaskClipMesh_Reference);
+
+        uint32_t vcount = pClipMesh->GetVertexBuffers().GetElementCount();
+        if (!vcount)
+        {
+            return nullptr; // Since there's nothing to clip against return a copy of the original base mesh
+        }
+
+        mu::vector<uint8_t> vertex_in_clip_mesh;  // Stores whether each vertex in the original mesh in in the clip mesh volume
+        MeshClipMeshClassifyVertices( vertex_in_clip_mesh, pBase, pClipMesh );
+
+        // We only remove vertices if all their faces are clipped
+        mu::vector<uint8_t> vertex_with_face_not_clipped(vertex_in_clip_mesh.size(),0);
+
+        UntypedMeshBufferIteratorConst ito(pBase->GetIndexBuffers(), MBS_VERTEXINDEX);
+        int aFaceCount = pBase->GetFaceCount();
+        for (int f = 0; f < aFaceCount; ++f)
+        {
+            vec3<uint32_t> ov;
+            ov[0] = ito.GetAsUINT32(); ++ito;
+            ov[1] = ito.GetAsUINT32(); ++ito;
+            ov[2] = ito.GetAsUINT32(); ++ito;
+
+            bool faceClipped =
+                    vertex_in_clip_mesh[ov[0]] &&
+                    vertex_in_clip_mesh[ov[1]] &&
+                    vertex_in_clip_mesh[ov[2]];
+
+            if (!faceClipped)
+            {
+                vertex_with_face_not_clipped[ov[0]] = true;
+                vertex_with_face_not_clipped[ov[1]] = true;
+                vertex_with_face_not_clipped[ov[2]] = true;
+            }
+        }
+
+
+        return CreateMask( pBase, vertex_with_face_not_clipped );
+    }
+
+    //---------------------------------------------------------------------------------------------
+    //! Generate a mask mesh with the faces of the base mesh matching the fragment.
+    //---------------------------------------------------------------------------------------------
+    inline MeshPtr MeshMaskDiff(const Mesh* pBase, const Mesh* pFragment )
+    {
+        MUTABLE_CPUPROFILER_SCOPE(MeshMaskDiff);
+
+        uint32_t vcount = pFragment->GetVertexBuffers().GetElementCount();
+        if (!vcount)
+        {
+            return nullptr;
+        }
+
+        int sourceFaceCount = pBase->GetFaceCount();
+        int sourceVertexCount = pBase->GetVertexCount();
+        int fragmentFaceCount = pFragment->GetFaceCount();
+
+
+        // Make a tolerance proportional to the mesh bounding box size
+        // TODO: Use precomputed bounding box
+        box< vec3<float> > aabbox;
+        if ( fragmentFaceCount > 0 )
+        {
+            MeshBufferIteratorConst<MBF_FLOAT32,float,3> itp( pFragment->GetVertexBuffers(), MBS_POSITION );
+
+            aabbox.min = *itp;
+            ++itp;
+
+            for ( int v=1; v<pFragment->GetVertexBuffers().GetElementCount(); ++v )
+            {
+                aabbox.Bound( *itp );
+                ++itp;
+            }
+        }
+        float tolerance = 1e-5f * length(aabbox.size);
+        Mesh::VERTEX_MATCH_MAP vertexMap;
+        pFragment->GetVertexMap( *pBase, vertexMap, tolerance );
+
+
+        // Classify the target faces in buckets along the Y axis
+#define NUM_BUCKETS	128
+#define AXIS		1
+        vector<int> buckets[ NUM_BUCKETS ];
+        float bucketStart = aabbox.min[AXIS];
+        float bucketSize = aabbox.size[AXIS] / NUM_BUCKETS;
+
+        float bucketThreshold = ( 4 * tolerance ) / bucketSize;
+        UntypedMeshBufferIteratorConst iti( pFragment->GetIndexBuffers(), MBS_VERTEXINDEX );
+        MeshBufferIteratorConst<MBF_FLOAT32,float,3> itp( pFragment->GetVertexBuffers(), MBS_POSITION );
+        for ( int tf=0; tf<fragmentFaceCount; tf++ )
+        {
+            uint32_t index0 = iti.GetAsUINT32(); ++iti;
+            uint32_t index1 = iti.GetAsUINT32(); ++iti;
+            uint32_t index2 = iti.GetAsUINT32(); ++iti;
+            float y = ( (*(itp+index0))[AXIS] + (*(itp+index1))[AXIS] + (*(itp+index2))[AXIS] ) / 3;
+            float fbucket = (y-bucketStart) / bucketSize;
+            int bucket = std::min( NUM_BUCKETS-1, std::max( 0, (int)fbucket ) );
+            buckets[bucket].push_back(tf);
+            int hibucket = std::min( NUM_BUCKETS-1, std::max( 0, (int)(fbucket+bucketThreshold) ) );
+            if (hibucket!=bucket)
+            {
+                buckets[hibucket].push_back(tf);
+            }
+            int lobucket = std::min( NUM_BUCKETS-1, std::max( 0, (int)(fbucket-bucketThreshold) ) );
+            if (lobucket!=bucket)
+            {
+                buckets[lobucket].push_back(tf);
+            }
+        }
+
+//		LogDebug("Box : min %.3f, %.3f, %.3f    size %.3f,%.3f,%.3f\n",
+//				aabbox.min[0], aabbox.min[1], aabbox.min[2],
+//				aabbox.size[0], aabbox.size[1], aabbox.size[2] );
+//		for ( int b=0; b<NUM_BUCKETS; ++b )
+//		{
+//			LogDebug("bucket : %d\n", buckets[b].size() );
+//		}
+
+        vector<uint8_t> faceClipped(sourceFaceCount,false);
+
+        UntypedMeshBufferIteratorConst ito( pBase->GetIndexBuffers(), MBS_VERTEXINDEX );
+        MeshBufferIteratorConst<MBF_FLOAT32,float,3> itop( pBase->GetVertexBuffers(), MBS_POSITION );
+        UntypedMeshBufferIteratorConst itti( pFragment->GetIndexBuffers(), MBS_VERTEXINDEX );
+        for ( int f=0; f<sourceFaceCount; ++f )
+        {
+            bool hasFace = false;
+            vec3<uint32_t> ov;
+            ov[0] = ito.GetAsUINT32(); ++ito;
+            ov[1] = ito.GetAsUINT32(); ++ito;
+            ov[2] = ito.GetAsUINT32(); ++ito;
+
+            // find the bucket for this face
+            float y = ( (*(itop+ov[0]))[AXIS] + (*(itop+ov[1]))[AXIS] + (*(itop+ov[2]))[AXIS] ) / 3;
+            float fbucket = (y-bucketStart) / bucketSize;
+            int bucket = std::min( NUM_BUCKETS-1, std::max( 0, (int)fbucket ) );
+
+            for ( std::size_t btf=0; !hasFace && btf<buckets[bucket].size(); btf++ )
+            {
+                int tf =  buckets[bucket][btf];
+
+                vec3<uint32_t> v;
+                v[0] = (itti+3*tf+0).GetAsUINT32();
+                v[1] = (itti+3*tf+1).GetAsUINT32();
+                v[2] = (itti+3*tf+2).GetAsUINT32();
+
+                hasFace = true;
+                for ( int vi=0; hasFace && vi<3; ++vi )
+                {
+                    hasFace = vertexMap.Matches(v[vi],ov[0])
+                         || vertexMap.Matches(v[vi],ov[1])
+                         || vertexMap.Matches(v[vi],ov[2]);
+                }
+            }
+
+            if ( hasFace )
+            {
+                faceClipped[f] = true;
+            }
+        }
+
+        // We only remove vertices if all their faces are clipped
+        mu::vector<uint8_t> vertex_with_face_not_clipped(sourceVertexCount,0);
+
+        UntypedMeshBufferIteratorConst itoi(pBase->GetIndexBuffers(), MBS_VERTEXINDEX);
+        int aFaceCount = pBase->GetFaceCount();
+        for (int f = 0; f < aFaceCount; ++f)
+        {
+            vec3<uint32_t> ov;
+            ov[0] = itoi.GetAsUINT32(); ++itoi;
+            ov[1] = itoi.GetAsUINT32(); ++itoi;
+            ov[2] = itoi.GetAsUINT32(); ++itoi;
+
+            if (!faceClipped[f])
+            {
+                vertex_with_face_not_clipped[ov[0]] = true;
+                vertex_with_face_not_clipped[ov[1]] = true;
+                vertex_with_face_not_clipped[ov[2]] = true;
+            }
+        }
+
+        return CreateMask( pBase, vertex_with_face_not_clipped );
+    }
+
+}
diff --git a/Engine/Plugins/Experimental/Mutable/Source/MutableRuntime/Private/OpMeshDifference.h b/Engine/Plugins/Experimental/Mutable/Source/MutableRuntime/Private/OpMeshDifference.h
new file mode 100644
--- /dev/null
+++ b/Engine/Plugins/Experimental/Mutable/Source/MutableRuntime/Private/OpMeshDifference.h
@@ -0,0 +1,326 @@
+// Copyright Epic Games, Inc. All Rights Reserved.
+
+#pragma once
+
+#include "MeshPrivate.h"
+#include "Platform.h"
+
+namespace mu
+{
+
+#pragma pack(push,1)
+struct MORPH_VERTEX
+{
+    uint32_t index;
+	vec3<float> pos;
+	vec3<float> nor;
+};
+#pragma pack(pop)
+
+
+	//---------------------------------------------------------------------------------------------
+	//! Create a diff from the mesh vertices. The meshes must have the same amount of vertices.
+	//! If the channel list is empty all the channels will be compared.
+	//---------------------------------------------------------------------------------------------
+    inline MeshPtr MeshDifference( const Mesh* pBase, const Mesh* pTarget,
+                                   int numChannels,
+                                   const MESH_BUFFER_SEMANTIC* semantics,
+                                   const int* semanticIndices,
+                                   bool ignoreTexCoords )
+
+	{
+        if (!pBase || !pTarget)
+        {
+            return nullptr;
+        }
+
+		bool correct =
+                pBase
+                &&
+				( pBase->GetVertexBuffers().GetElementCount()
+						== pTarget->GetVertexBuffers().GetElementCount() )
+				&&
+                ( pBase->GetIndexCount() == pTarget->GetIndexCount() )
+                &&
+                ( pBase->GetVertexBuffers().GetElementCount()>0 )
+                ;
+
+		MeshPtr pDest = new Mesh();				
+
+		if (!correct)
+            return pDest;
+
+		uint32_t vcount = pBase->GetVertexBuffers().GetElementCount();		
+
+		// If no channels were specified, get them all
+		vector<MESH_BUFFER_SEMANTIC> allSemantics;		
+		vector<int> allSemanticIndices;
+		if ( !numChannels )
+		{
+			for ( int vb=0; vb<pBase->GetVertexBuffers().GetBufferCount(); ++vb )
+			{
+				for ( int c=0; c<pBase->GetVertexBuffers().GetBufferChannelCount(vb); ++c )
+				{
+					MESH_BUFFER_SEMANTIC sem = MBS_NONE;
+					int semIndex = 0;
+					pBase->GetVertexBuffers().GetChannel( vb, c,
+															&sem, &semIndex,
+															nullptr, nullptr, nullptr );			
+					if ( sem != MBS_VERTEXINDEX &&
+					     sem != MBS_BONEINDICES &&
+					     sem != MBS_BONEWEIGHTS &&
+					     sem != MBS_CHART &&
+					     sem != MBS_LAYOUTBLOCK &&
+					     sem != MBS_OTHER && 
+					     ( !ignoreTexCoords || sem!=MBS_TEXCOORDS ) )
+					{
+						allSemantics.push_back( sem );
+						allSemanticIndices.push_back( semIndex );
+						++numChannels;
+					}
+				}
+			}
+
+			semantics = &allSemantics[0];
+			semanticIndices = &allSemanticIndices[0];
+		}
+
+
+		// Make a delta of every vertex
+		// TODO: Not always 4 components
+		// TODO: Not always floats
+		int differentVertexCount = 0;
+		vector<bool> isVertexDifferent( vcount, false );
+		vector< vec4<float> > deltas( vcount*numChannels );
+
+		for ( int c=0; c<numChannels; ++c )
+		{
+			UntypedMeshBufferIteratorConst baseIt		
+					( pBase->GetVertexBuffers(), semantics[c], semanticIndices[c] );
+
+			UntypedMeshBufferIteratorConst targetIt		
+					( pTarget->GetVertexBuffers(), semantics[c], semanticIndices[c] );
+
+			for ( size_t v=0; v<vcount; ++v )
+			{
+				vec4<float> base = baseIt.GetAsVec4f();			
+				vec4<float> target = targetIt.GetAsVec4f();		
+				vec4<float> delta = target-base;
+				deltas[ v*numChannels+c ] = delta;		
+
+				if (!delta.AlmostNull())
+				{
+					if ( !isVertexDifferent[v] )
+					{
+						++differentVertexCount;
+					}
+
+					isVertexDifferent[v] = true;
+				}
+
+				++baseIt;		
+				++targetIt;		
+			}
+		}
+
+
+		// Create the morph mesh
+		{
+			pDest->GetVertexBuffers().SetElementCount( differentVertexCount );			
+			pDest->GetVertexBuffers().SetBufferCount( 1 );								
+
+			vector<MESH_BUFFER_SEMANTIC> semantic;				
+			vector<int> semanticIndex;							
+			vector<MESH_BUFFER_FORMAT> format;					
+			vector<int> components;								
+			vector<int> offsets;								
+			int offset = 0;										
+
+			// Vertex index channel
+			semantic.push_back( MBS_VERTEXINDEX );					
+			semanticIndex.push_back( 0 );
+			format.push_back( MBF_UINT32 );		
+			components.push_back( 1 );
+			offsets.push_back( offset );
+			offset += 4;						
+
+			for ( int c=0; c<numChannels; ++c )
+			{
+				semantic.push_back( semantics[c] );
+				semanticIndex.push_back( semanticIndices[c] );
+				format.push_back( MBF_FLOAT32 );
+				components.push_back( 4 );
+				offsets.push_back( offset );
+				offset += 16;				
+			}
+
+			pDest->GetVertexBuffers().SetBuffer
+				(
+					0,
+					offset,
+					numChannels+1,
+					&semantic[0],
+					&semanticIndex[0],
+					&format[0],
+					&components[0],
+					&offsets[0]
+				);
+
+			// Source vertex index channel
+			UntypedMeshBufferIteratorConst
+					itBaseId( pBase->GetVertexBuffers(), MBS_VERTEXINDEX, 0 );
+
+			// Set the data
+			uint8_t* pDestData = pDest->GetVertexBuffers().GetBufferData( 0 );
+			for (size_t v=0; v<vcount; ++v)
+			{
+				if ( isVertexDifferent[v] )
+				{
+					// index
+					*((uint32_t*)pDestData) = itBaseId.GetAsUINT32();			
+					pDestData += 4;
+
+					// all channels
+					for ( int c=0; c<numChannels; ++c )
+					{
+						*((vec4<float>*)pDestData) = deltas[ v*numChannels+c ];		
+						pDestData += 4*4;											
+					}
+				}
+
+				++itBaseId;
+			}
+		}
+
+
+        // Rebuild surface data.
+        // \todo: For now make single surfaced
+        pDest->m_surfaces.SetNum(0);
+        pDest->EnsureSurfaceData();
+
+		return pDest;
+	}
+
+	/*/
+	{
+		bool correct =
+				( pBase->GetVertexBuffers().GetElementCount()
+						== pTarget->GetVertexBuffers().GetElementCount() )
+				&&
+				( pBase->GetIndexCount() == pTarget->GetIndexCount()
+				&&
+				 pBase->GetVertexBuffers().GetElementCount()>0
+				);
+
+		MeshPtr pDest = new Mesh();
+
+		if (correct)
+		{
+			// For now, just make a difference from the position and normal buffers.
+			UINT32 vcount = pBase->GetVertexBuffers().GetElementCount();
+
+			// Base mesh position
+			int basePosElementSize = 0;
+			const UINT8* pBasePosBuf = 0;
+			GetMeshBuf( pBase, MBS_POSITION, MBF_FLOAT32, 3, pBasePosBuf, basePosElementSize );
+
+			// Base mesh normal
+			int baseNorElementSize = 0;
+			const UINT8* pBaseNorBuf = 0;
+			GetMeshBuf( pBase, MBS_NORMAL, MBF_FLOAT32, 3, pBaseNorBuf, baseNorElementSize );
+
+			// Target mesh position
+			int targetPosElementSize = 0;
+			const UINT8* pTargetPosBuf = 0;
+			GetMeshBuf( pTarget, MBS_POSITION, MBF_FLOAT32, 3, pTargetPosBuf, targetPosElementSize );
+
+			// Target mesh normal
+			int targetNorElementSize = 0;
+			const UINT8* pTargetNorBuf = 0;
+			GetMeshBuf( pTarget, MBS_NORMAL, MBF_FLOAT32, 3, pTargetNorBuf, targetNorElementSize );
+
+			// Actual diff
+			vector< MORPH_VERTEX > vertices;
+
+			for ( UINT32 vi=0; vi<vcount; ++vi )
+			{
+				MORPH_VERTEX v;
+
+				vec3<float> basePos = *reinterpret_cast< const vec3<float>* >( pBasePosBuf );
+				vec3<float> targetPos = *reinterpret_cast< const vec3<float>* >( pTargetPosBuf );
+				v.pos = targetPos - basePos;
+
+				vec3<float> baseNor = *reinterpret_cast< const vec3<float>* >( pBaseNorBuf );
+				vec3<float> targetNor = *reinterpret_cast< const vec3<float>* >( pTargetNorBuf );
+				v.nor = targetNor - baseNor;
+
+				if ( !( v.pos.AlmostNull() &&
+						v.nor.AlmostNull() ) )
+				{
+					v.index = vi;
+					vertices.push_back( v );
+				}
+
+				pBasePosBuf += basePosElementSize;
+				pTargetPosBuf += targetPosElementSize;
+				pBaseNorBuf += baseNorElementSize;
+				pTargetNorBuf += targetNorElementSize;
+			}
+
+			// Create the morph mesh
+			{
+				int morphCount = (int)vertices.size();
+				pDest->GetVertexBuffers().SetElementCount( morphCount );
+				pDest->GetVertexBuffers().SetBufferCount( 1 );
+
+				// Vertex indices
+				{
+					vector<MESH_BUFFER_SEMANTIC> semantic;
+					semantic.push_back( MBS_VERTEXINDEX );
+					semantic.push_back( MBS_POSITION );
+					semantic.push_back( MBS_NORMAL );
+
+					vector<int> semanticIndex;
+					semanticIndex.push_back( 0 );
+					semanticIndex.push_back( 0 );
+					semanticIndex.push_back( 0 );
+
+					vector<MESH_BUFFER_FORMAT> format;
+					format.push_back( MBF_UINT32 );
+					format.push_back( MBF_FLOAT32 );
+					format.push_back( MBF_FLOAT32 );
+
+					vector<int> components;
+					components.push_back( 1 );
+					components.push_back( 3 );
+					components.push_back( 3 );
+
+					vector<int> offset;
+					offset.push_back( 0 );
+					offset.push_back( 4 );
+					offset.push_back( 16 );
+
+					pDest->GetVertexBuffers().SetBuffer
+						(
+							0,
+							28,
+							3,
+							&semantic[0],
+							&semanticIndex[0],
+							&format[0],
+							&components[0],
+							&offset[0]
+						);
+
+					UINT8* pDestData = pDest->GetVertexBuffers().GetBufferData( 0 );
+					int size = morphCount*sizeof(MORPH_VERTEX);
+					memcpy( pDestData, &vertices[0], size );
+				}
+
+			}
+		}
+
+		return pDest;
+	}
+*/
+}
diff --git a/Engine/Plugins/Experimental/Mutable/Source/MutableRuntime/Private/OpMeshExtractLayoutBlock.h b/Engine/Plugins/Experimental/Mutable/Source/MutableRuntime/Private/OpMeshExtractLayoutBlock.h
new file mode 100644
--- /dev/null
+++ b/Engine/Plugins/Experimental/Mutable/Source/MutableRuntime/Private/OpMeshExtractLayoutBlock.h
@@ -0,0 +1,215 @@
+// Copyright Epic Games, Inc. All Rights Reserved.
+
+#pragma once
+
+#include "MeshPrivate.h"
+#include "Platform.h"
+
+
+namespace mu
+{
+
+
+    //---------------------------------------------------------------------------------------------
+    inline void MeshExtractFromVertices( const Mesh* pSource,
+                                         Mesh* pResult,
+                                         const vector<int>& oldToNew,
+                                         const vector<int>& newToOld )
+    {
+        int resultVertices = (int)newToOld.size();
+
+        // Assemble the new vertex buffer
+        pResult->GetVertexBuffers().SetElementCount( resultVertices );
+        for ( int b=0; b<pResult->GetVertexBuffers().GetBufferCount(); ++b )
+        {
+            const uint8_t* pSourceData = pSource->GetVertexBuffers().GetBufferData(b);
+            uint8_t* pDest = pResult->GetVertexBuffers().GetBufferData(b);
+            int size = pResult->GetVertexBuffers().GetElementSize(b);
+            for ( int v=0; v<resultVertices; ++v )
+            {
+                int i = newToOld[v];
+                memcpy( pDest, pSourceData + size*i, size );
+                pDest+=size;
+            }
+        }
+
+        // Assemble the new index buffers
+        vector<bool> usedSourceFaces( pSource->GetFaceCount(), false );
+        UntypedMeshBufferIteratorConst itIndex( pSource->GetIndexBuffers(), MBS_VERTEXINDEX );
+        UntypedMeshBufferIterator itResultIndex( pResult->GetIndexBuffers(), MBS_VERTEXINDEX );
+        int indexCount = 0;
+        if ( itIndex.GetFormat()==MBF_UINT32 )
+        {
+            const uint32_t* pIndices = reinterpret_cast<const uint32_t*>( itIndex.ptr() );
+            uint32_t* pDestIndices = reinterpret_cast<uint32_t*>( itResultIndex.ptr() );
+            for ( int i=0; i<pSource->GetIndexCount()/3; ++i )
+            {
+                if ( oldToNew[ pIndices[i*3+0] ]>=0
+                     &&
+                     oldToNew[ pIndices[i*3+1] ]>=0
+                     &&
+                     oldToNew[ pIndices[i*3+2] ]>=0 )
+                {
+                    usedSourceFaces[i] = true;
+
+                    // Clamp in case triangles go across blocks
+                    pDestIndices[ indexCount++ ] = std::max( 0, oldToNew[ pIndices[i*3+0] ] );
+                    pDestIndices[ indexCount++ ] = std::max( 0, oldToNew[ pIndices[i*3+1] ] );
+                    pDestIndices[ indexCount++ ] = std::max( 0, oldToNew[ pIndices[i*3+2] ] );
+                }
+            }
+        }
+        else if ( itIndex.GetFormat()==MBF_UINT16 )
+        {
+            const uint16_t* pIndices = reinterpret_cast<const uint16_t*>( itIndex.ptr() );
+            uint16_t* pDestIndices = reinterpret_cast<uint16_t*>( itResultIndex.ptr() );
+            for ( int i=0; i<pSource->GetIndexCount()/3; ++i )
+            {
+                if ( oldToNew[ pIndices[i*3+0] ]>=0
+                     &&
+                     oldToNew[ pIndices[i*3+1] ]>=0
+                     &&
+                     oldToNew[ pIndices[i*3+2] ]>=0 )
+                {
+                    usedSourceFaces[i] = true;
+
+                    // Clamp in case triangles go across blocks
+                    pDestIndices[ indexCount++ ] = (uint16_t)std::max( 0, oldToNew[ pIndices[i*3+0] ] );
+                    pDestIndices[ indexCount++ ] = (uint16_t)std::max( 0, oldToNew[ pIndices[i*3+1] ] );
+                    pDestIndices[ indexCount++ ] = (uint16_t)std::max( 0, oldToNew[ pIndices[i*3+2] ] );
+                }
+            }
+        }
+        else
+        {
+            check( false );
+        }
+        pResult->GetIndexBuffers().SetElementCount( indexCount );
+
+
+        // Face buffers
+        {
+            int faceCount = 0;
+            for ( int i=0; i<pSource->GetFaceCount(); ++i )
+            {
+                if ( usedSourceFaces[i] )
+                {
+                    for ( int b=0; b<pResult->GetFaceBuffers().GetBufferCount(); ++b )
+                    {
+                        const uint8_t* pSourceData = pSource->GetFaceBuffers().GetBufferData(b);
+                        uint8_t* pDest = pResult->GetFaceBuffers().GetBufferData(b);
+                        int size = pResult->GetFaceBuffers().GetElementSize(b);
+                        memcpy( pDest + size*faceCount, pSourceData + size*i, size );
+                    }
+
+                    faceCount++;
+                }
+            }
+            check( faceCount == indexCount/3 );
+            pResult->GetFaceBuffers().SetElementCount( faceCount );
+        }
+    }
+
+
+	//---------------------------------------------------------------------------------------------
+    inline MeshPtr MeshExtractLayoutBlock( const Mesh* pSource,
+                                           uint32_t layout,
+                                           uint16_t blockCount,
+                                           const uint32_t* pExtractBlocks )
+	{
+		// TODO: Optimise
+		MeshPtr pResult = pSource->Clone();
+
+        // TODO
+        pResult->SetFaceGroupCount( 0 );
+
+		UntypedMeshBufferIteratorConst itBlocks( pSource->GetVertexBuffers(),
+												 MBS_LAYOUTBLOCK, layout );
+
+        if (itBlocks.GetFormat()!=MBF_NONE)
+        {
+            int resultVertices = 0;
+            vector<int> oldToNew( pSource->GetVertexCount(), -1 );
+            vector<int> newToOld;
+            newToOld.reserve( pSource->GetVertexCount() );
+
+            if ( itBlocks.GetFormat()==MBF_UINT16 )
+            {
+                const uint16_t* pBlocks = reinterpret_cast<const uint16_t*>( itBlocks.ptr() );
+                for ( int i=0; i<pSource->GetVertexCount(); ++i )
+                {
+                    uint32_t vertexBlock = pBlocks[i];
+
+                    bool found = false;
+                    for ( int j=0; j<blockCount; ++j)
+                    {
+                        if (vertexBlock==pExtractBlocks[j])
+                        {
+                            found = true;
+                            break;
+                        }
+                    }
+
+                    if ( found )
+                    {
+                        oldToNew[i] = resultVertices++;
+                        newToOld.push_back( i );
+                    }
+                }
+            }
+            else
+            {
+                check( false );
+            }
+
+            MeshExtractFromVertices( pSource, pResult.get(), oldToNew, newToOld );
+        }
+
+        pResult->m_surfaces.Empty();
+        pResult->EnsureSurfaceData();
+
+        return pResult;
+	}
+
+
+    //---------------------------------------------------------------------------------------------
+    inline MeshPtr MeshExtractFaceGroup( const Mesh* pSource, int group )
+    {
+        if ( group<0 || group>=pSource->GetFaceGroupCount() )
+        {
+            return new Mesh;
+        }
+
+        // TODO: Optimise
+        MeshPtr pResult = pSource->Clone();
+
+        int resultVertices = 0;
+        vector<int> oldToNew( pSource->GetVertexCount(), -1 );
+        vector<int> newToOld;
+        newToOld.reserve( pSource->GetVertexCount() );
+
+        UntypedMeshBufferIteratorConst itIndex( pSource->GetIndexBuffers(), MBS_VERTEXINDEX );
+        for ( int32 f=0; f<pSource->m_faceGroups[group].m_faces.Num(); ++f )
+        {
+            int32 face = pSource->m_faceGroups[group].m_faces[f];
+
+            for (int i=0; i<3; i++)
+            {
+                uint32_t vertexIndex = (itIndex+face*3+i).GetAsUINT32();
+                if ( oldToNew[vertexIndex]<0 )
+                {
+                    oldToNew[vertexIndex] = resultVertices++;
+                    newToOld.push_back( vertexIndex );
+                }
+            }
+        }
+
+        MeshExtractFromVertices( pSource, pResult.get(), oldToNew, newToOld );
+
+        pResult->m_surfaces.Empty();
+        pResult->EnsureSurfaceData();
+
+        return pResult;
+    }
+
+}
