[mutable] Move Mutable plugin from the restricted folder to the experimental plugins folder.

diff --git a/Engine/Plugins/Experimental/Mutable/Source/MutableRuntime/Private/OpMeshFormat.cpp b/Engine/Plugins/Experimental/Mutable/Source/MutableRuntime/Private/OpMeshFormat.cpp
new file mode 100644
--- /dev/null
+++ b/Engine/Plugins/Experimental/Mutable/Source/MutableRuntime/Private/OpMeshFormat.cpp
@@ -0,0 +1,947 @@
+// Copyright Epic Games, Inc. All Rights Reserved.
+
+
+#include "OpMeshFormat.h"
+#include "Platform.h"
+#include "MutableMath.h"
+#include "MeshPrivate.h"
+#include "ConvertData.h"
+#include "MutableTrace.h"
+
+
+using namespace mu;
+
+
+//-------------------------------------------------------------------------------------------------
+// Eric Lengyel method's
+// http://www.terathon.com/code/tangent.html
+//-------------------------------------------------------------------------------------------------
+namespace
+{
+
+#define MUTABLE_VERTEX_MERGE_TEX_RANGE          1024
+#define MUTABLE_TANGENT_GENERATION_EPSILON_1    0.000001f
+#define MUTABLE_TANGENT_GENERATION_EPSILON_2    0.001f
+#define MUTABLE_TANGENT_MIN_AXIS_DIFFERENCE     0.999f
+
+    struct TVertex
+    {
+
+        TVertex() = default;
+
+        TVertex( const vec3<float>& p, const vec3<float>&  n, const vec2<float>& t )
+        {
+            pos = p;
+            nor = n;
+            tex = t;
+        }
+
+        vec3<float> pos;
+        vec3<float> nor;
+        vec2<float> tex;
+
+        inline bool operator< ( const TVertex& other ) const
+        {
+            if( pos < other.pos )
+                return true;
+            if( other.pos < pos )
+                return false;
+
+            if( nor < other.nor )
+                return true;
+            if( other.nor < nor )
+                return false;
+
+            // Compare the texture coordinates with a particular precission.
+            vec2<int> uv0 = vec2<int>( int( tex[0] * MUTABLE_VERTEX_MERGE_TEX_RANGE ),
+                                       int( tex[1] * MUTABLE_VERTEX_MERGE_TEX_RANGE ) );
+            vec2<int> uv1 = vec2<int>( int( other.tex[0] * MUTABLE_VERTEX_MERGE_TEX_RANGE ),
+                                       int( other.tex[1] * MUTABLE_VERTEX_MERGE_TEX_RANGE ) );
+
+            if( uv0 < uv1 )
+                return true;
+            if( uv1 < uv0 )
+                return false;
+
+            return false;
+        }
+
+        inline bool operator== ( const TVertex& other ) const
+        {
+            return  ! ( (*this) < other )
+                    &&
+                    ! ( other < (*this) );
+        }
+    };
+
+
+    struct TFace
+    {
+        vec3<float> T;
+        vec3<float> B;
+        vec3<float> N;
+
+        TFace()
+            : T(0,0,0)
+            , B(0,0,0)
+            , N(0,0,0)
+        {}
+
+        TFace
+            (
+                const vec3<float>& v1,
+                const vec3<float>& v2,
+                const vec3<float>& v3,
+                const vec2<float>& w1,
+                const vec2<float>& w2,
+                const vec2<float>& w3
+            )
+        {
+            vec3<float> E1 = v2 - v1;
+            vec3<float> E2 = v3 - v1;
+
+            vec2<float> UV1 = w2 - w1;
+            vec2<float> UV2 = w3 - w1;
+
+            float  UVdet = UV1[0] * UV2[1] - UV2[0] * UV1[1];
+
+            N = normalise( cross( E1, E2 ) );
+
+            if( !(fabs(UVdet)<=MUTABLE_TANGENT_GENERATION_EPSILON_1 ) )
+            {
+                double r = 1.0 / UVdet;
+
+                T = vec3<float>
+                    (
+                        (UV2[1] * E1[0] - UV1[1] * E2[0]),
+                        (UV2[1] * E1[1] - UV1[1] * E2[1]),
+                        (UV2[1] * E1[2] - UV1[1] * E2[2])
+                    );
+
+                B = vec3<float>
+                    (
+                        (UV1[0] * E2[0] - UV2[0] * E1[0]),
+                        (UV1[0] * E2[1] - UV2[0] * E1[1]),
+                        (UV1[0] * E2[2] - UV2[0] * E1[2])
+                    );
+
+                T = T * float(r);
+                B = B * float(r);
+
+                T = normalise(T);
+                B = normalise(B);
+            }
+            else
+            {
+                T = vec3<float>(0,0,0);
+                B = vec3<float>(0,0,0);
+                N = vec3<float>(0,0,0);
+            }
+        }
+
+    };
+
+
+	void RebuildTangents
+		(
+			const FMeshBufferSet& IndexBuffers,
+			FMeshBufferSet& VertexBuffers
+		)
+	{
+		// Because of these, this method only works for a specific mesh format
+        UntypedMeshBufferIteratorConst vertex( VertexBuffers, MBS_POSITION, 0 );
+        UntypedMeshBufferIteratorConst normal( VertexBuffers, MBS_NORMAL, 0 );
+        UntypedMeshBufferIteratorConst texcoord( VertexBuffers, MBS_TEXCOORDS, 0 );
+
+		int vertexCount = VertexBuffers.GetElementCount();
+
+		// Group similar vertices
+		vector< TVertex > vertexData;
+		vertexData.reserve( vertexCount );
+        vector< uint32_t > lVertexGroups( vertexCount, 0 );
+
+        //Agrupamos los vertices por posicion, normal, texcoord
+        for( int32_t i = 0; i < vertexCount; ++i )
+        {
+            TVertex v( (vertex+i).GetAsVec4f().xyz(),
+                       (normal+i).GetAsVec4f().xyz(),
+                       (texcoord+i).GetAsVec4f().xy() );
+            vector<TVertex>::const_iterator it = std::find( vertexData.begin(),
+                                                            vertexData.end(),
+                                                            v);
+
+            size_t pos = 0;
+            if ( it != vertexData.end() )
+            {
+                pos = it - vertexData.begin();
+            }
+            else
+            {
+                pos = vertexData.size();
+                vertexData.push_back( v );
+            }
+
+            lVertexGroups[ i ] = (int)pos;
+        }
+
+        // Calculate the tangent space
+        MeshBufferIterator<MBF_FLOAT32, float, 3> tangent( VertexBuffers, MBS_TANGENT, 0 );
+        MeshBufferIterator<MBF_FLOAT32, float, 3> bitangent( VertexBuffers, MBS_BINORMAL, 0 );
+        UntypedMeshBufferIteratorConst indices( IndexBuffers, MBS_VERTEXINDEX, 0 );
+        int indexCount = IndexBuffers.GetElementCount();
+
+        vector<TFace> vertexFaces( indexCount );
+
+        // Worst case: a group for every vertex
+        vector< vector< TFace > > grp_faces( indexCount );
+
+        for ( long faceidx = 0, count = indexCount/3; faceidx < count; faceidx++ )
+        {
+            long i1 = 0;
+            long i2 = 0;
+            long i3 = 0;
+
+            if ( indices.GetFormat()==MBF_UINT16)
+            {
+                auto pIndices = reinterpret_cast<const uint16_t*>( indices.ptr() );
+                i1 = pIndices [faceidx * 3    ];
+                i2 = pIndices [faceidx * 3 + 1];
+                i3 = pIndices [faceidx * 3 + 2];
+            }
+            else if ( indices.GetFormat()==MBF_UINT32)
+            {
+                auto pIndices = reinterpret_cast<const uint32_t*>( indices.ptr() );
+                i1 = pIndices [faceidx * 3    ];
+                i2 = pIndices [faceidx * 3 + 1];
+                i3 = pIndices [faceidx * 3 + 2];
+            }
+            else
+            {
+                check( false );
+            }
+
+            TFace face;
+
+            const vec3<float> v1 = (vertex+(int)i1).GetAsVec4f().xyz();
+            const vec3<float> v2 = (vertex+(int)i2).GetAsVec4f().xyz();
+            const vec3<float> v3 = (vertex+(int)i3).GetAsVec4f().xyz();
+
+            // Compute face t & b
+            const vec2<float> w1 = (texcoord+(int)i1).GetAsVec4f().xy();
+            const vec2<float> w2 = (texcoord+(int)i2).GetAsVec4f().xy();
+            const vec2<float> w3 = (texcoord+(int)i3).GetAsVec4f().xy();
+            face = TFace( v1, v2, v3, w1, w2, w3 );
+
+            if( length( face.N ) != 0 )
+            {
+                vertexFaces[ i1 ] = vertexFaces[ i2 ] = vertexFaces[ i3 ] = face;
+
+                grp_faces[lVertexGroups[i1]].push_back( face );
+                grp_faces[lVertexGroups[i2]].push_back( face );
+                grp_faces[lVertexGroups[i3]].push_back( face );
+            }
+        }
+
+        for ( long vtxidx = 0; vtxidx < vertexCount; vtxidx++ )
+        {
+            const vec3<float> n = (normal+(int)vtxidx).GetAsVec4f().xyz();
+            const uint32_t group    = lVertexGroups[vtxidx];
+
+            vec3<float> lTangent  (0,0,0);
+            vec3<float> lBiTangent(0,0,0);
+
+            const vector< TFace >& faces   = grp_faces[group];
+
+            const TFace& vFace = vertexFaces[ vtxidx ];
+
+            // Ignore mirrors
+            for( unsigned i = 0; i < faces.size(); ++i )
+            {
+                const float DN = dot( vFace.N, faces[i].N );
+                const float DT = dot( vFace.T, faces[i].T );
+                const float DB = dot( vFace.B, faces[i].B );
+
+                if( DN > 0 && DT > 0 && DB > 0 )
+                {
+                    lTangent   += faces[i].T;
+                    lBiTangent += faces[i].B;
+                }
+            }
+
+            vec3<float> ortogonalizedTangent
+                    = normalise( lTangent
+                                 - ( n * dot( n, lTangent  ) ) );
+
+            vec3<float> ortogonalizedBiTangent
+                    = normalise( lBiTangent
+                                 - ( n * dot( n, lBiTangent) )
+                                 - ( ortogonalizedTangent * dot(ortogonalizedTangent,lBiTangent) ));
+
+            // this is where we can do a final check for zero length vectors
+            // and set them to something appropriate
+            float lenTan = length( ortogonalizedTangent );
+            float lenBin = length( ortogonalizedBiTangent );
+
+            if( (lenTan <= MUTABLE_TANGENT_GENERATION_EPSILON_2)
+                    || (lenBin <= MUTABLE_TANGENT_GENERATION_EPSILON_2)  ) //should be approx 1.0f
+            {
+                // the tangent space is ill defined at this vertex
+                // so we can generate a valid one based on the normal vector,
+                // which I'm assuming is valid!
+                if(lenTan > 0.5f)
+                {
+                    // the tangent is valid, so we can just use that to calculate the binormal
+                    ortogonalizedBiTangent = cross( n, ortogonalizedTangent );
+                }
+                else if(lenBin > 0.5)
+                {
+                    // the binormal is good and we can use it to calculate the tangent
+                    ortogonalizedTangent = cross( ortogonalizedBiTangent, n );
+                }
+                else
+                {
+                    // both vectors are invalid, so we should create something
+                    // that is at least valid if not correct
+                    // I'm checking two possible axis, because the normal could be one of them,
+                    // and we want to chose a different one to start making our valid basis.
+                    // I can find out which is further away from it by checking the dot product
+                    vec3<float> startAxis;
+
+                    if( dot(vec3<float>(1,0,0), n) <  dot(vec3<float>(0,1,0), n) )
+                    {
+                        //the xAxis is more different than the yAxis when compared to the normal
+                        startAxis = vec3<float>(1,0,0);
+                    }
+                    else
+                    {
+                        // the yAxis is more different than the xAxis when compared to the normal
+                        startAxis = vec3<float>(0,1,0);
+                    }
+
+                    ortogonalizedTangent   = cross( n, startAxis);
+                    ortogonalizedBiTangent = cross( n, ortogonalizedTangent);
+                }
+            }
+            else
+            {
+                // one final sanity check, make sure that they tangent and binormal are different
+                // enough
+                if ( dot(ortogonalizedTangent, ortogonalizedBiTangent )
+                     >
+                     MUTABLE_TANGENT_MIN_AXIS_DIFFERENCE )
+                {
+                    // then they are too similar lets make them more different
+                    ortogonalizedBiTangent = cross( n, ortogonalizedTangent );
+                }
+            }
+
+            tangent  [ vtxidx ] = ortogonalizedTangent;
+            bitangent[ vtxidx ] = ortogonalizedBiTangent;
+        }
+
+    }
+
+}
+
+
+//-------------------------------------------------------------------------------------------------
+void mu::MeshFormatBuffer
+	(
+		const FMeshBufferSet& Source,
+        FMeshBufferSet& Result,
+        int bufferIndex
+	)
+{	
+	int vCount = Source.GetElementCount();
+
+    int b = bufferIndex;
+    {
+		// For every channel in this buffer
+		for ( int c=0; c<Result.GetBufferChannelCount(b); ++c )
+		{
+			// Find this channel in the source mesh
+			MESH_BUFFER_SEMANTIC resultSemantic;
+			int resultSemanticIndex;
+			MESH_BUFFER_FORMAT resultFormat;
+			int resultComponents;
+			int resultOffset;
+			Result.GetChannel
+					(
+						b, c,
+						&resultSemantic, &resultSemanticIndex,
+						&resultFormat, &resultComponents,
+						&resultOffset
+					);
+
+			int sourceBuffer;
+			int sourceChannel;
+			Source.FindChannel
+					( resultSemantic, resultSemanticIndex, &sourceBuffer, &sourceChannel );
+
+			int resultElemSize = Result.GetElementSize( b );
+			int resultChannelSize = GetMeshFormatData( resultFormat ).m_size * resultComponents;
+            uint8_t* pResultBuf = Result.GetBufferData( b );
+			pResultBuf += resultOffset;
+
+			if ( sourceBuffer<0 )
+			{
+				// Not found: fill with zeros.
+
+                // Special case for derived channel data
+                bool generated = false;
+                if ( resultSemantic == MBS_TANGENTSIGN )
+                {
+                    // Look for the full tangent space
+                    int tanXBuf, tanXChan, tanYBuf, tanYChan, tanZBuf, tanZChan;
+                    Source.FindChannel( MBS_TANGENT, resultSemanticIndex, &tanXBuf, &tanXChan );
+                    Source.FindChannel( MBS_BINORMAL, resultSemanticIndex, &tanYBuf, &tanYChan );
+                    Source.FindChannel( MBS_NORMAL, resultSemanticIndex, &tanZBuf, &tanZChan );
+
+                    if ( tanXBuf>=0 && tanYBuf>=0 && tanZBuf>=0 )
+                    {
+                        generated = true;
+                        UntypedMeshBufferIteratorConst xIt( Source, MBS_TANGENT, resultSemanticIndex );
+                        UntypedMeshBufferIteratorConst yIt( Source, MBS_BINORMAL, resultSemanticIndex );
+                        UntypedMeshBufferIteratorConst zIt( Source, MBS_NORMAL, resultSemanticIndex );
+                        for ( int v=0; v<vCount; ++v )
+                        {
+                            mat3f mat;
+                            mat[0] = xIt.GetAsVec4f().xyz();
+                            mat[1] = yIt.GetAsVec4f().xyz();
+                            mat[2] = zIt.GetAsVec4f().xyz();
+                            float sign = mat.GetDeterminant()<0?-1.0f:1.0f;
+                            ConvertData( 0, pResultBuf, resultFormat, &sign, MBF_FLOAT32 );
+
+                            for ( int i=1; i<resultComponents; ++i )
+                            {
+                                // Add zeros
+								FMemory::Memzero
+                                        (
+                                            pResultBuf+GetMeshFormatData( resultFormat ).m_size*i,
+                                            GetMeshFormatData( resultFormat ).m_size
+                                        );
+                            }
+                            pResultBuf += resultElemSize;
+                            xIt++;
+                            yIt++;
+                            zIt++;
+                        }
+                    }
+                }
+
+                // If we have to add colour channels, we will add them as white, to be neutral.
+                // \todo: normal channels also should have special values.
+                else if ( resultSemantic == MBS_COLOUR )
+                {
+                    generated = true;
+
+                    switch(resultFormat)
+                    {
+                    case MBF_FLOAT32:
+                    {
+						for (int v = 0; v < vCount; ++v)
+						{
+							auto pTypedResultBuf = (float*)pResultBuf;
+							for (int i = 0; i < resultComponents; ++i)
+							{
+								pTypedResultBuf[i] = 1.0f;
+							}
+							pResultBuf += resultElemSize;
+						}
+                        break;
+                    }
+
+                    case MBF_NUINT8:
+                    {
+						for (int v = 0; v < vCount; ++v)
+						{
+							auto pTypedResultBuf = (uint8_t*)pResultBuf;
+							for (int i = 0; i < resultComponents; ++i)
+							{
+								pTypedResultBuf[i] = 255;
+							}
+							pResultBuf += resultElemSize;
+						}
+                        break;
+                    }
+
+                    case MBF_NUINT16:
+                    {
+						for (int v = 0; v < vCount; ++v)
+						{
+							auto pTypedResultBuf = (uint16_t*)pResultBuf;
+							for (int i = 0; i < resultComponents; ++i)
+							{
+								pTypedResultBuf[i] = 65535;
+							}
+							pResultBuf += resultElemSize;
+						}
+                        break;
+                    }
+
+                    default:
+                        // Format not implemented
+                        check( false );
+                        break;
+                    }
+                }
+
+                if (!generated)
+                {
+                    // TODO: and maybe raise a warning?
+                    for ( int v=0; v<vCount; ++v )
+                    {
+						FMemory::Memzero( pResultBuf, resultChannelSize );
+                        pResultBuf += resultElemSize;
+                    }
+                }
+			}
+			else
+			{
+				// Get the data about the source format
+				MESH_BUFFER_SEMANTIC sourceSemantic;
+				int sourceSemanticIndex;
+				MESH_BUFFER_FORMAT sourceFormat;
+				int sourceComponents;
+				int sourceOffset;
+				Source.GetChannel
+						(
+							sourceBuffer, sourceChannel,
+							&sourceSemantic, &sourceSemanticIndex,
+							&sourceFormat, &sourceComponents,
+							&sourceOffset
+						);
+				check( sourceSemantic==resultSemantic
+								&&
+								sourceSemanticIndex == resultSemanticIndex );
+
+				int sourceElemSize = Source.GetElementSize( sourceBuffer );
+                const uint8_t* pSourceBuf = Source.GetBufferData( sourceBuffer );
+				pSourceBuf += sourceOffset;
+
+				// Copy element by element
+				for ( int v=0; v<vCount; ++v )
+                {
+                    if ( resultFormat==sourceFormat && resultComponents==sourceComponents)
+					{
+						memcpy( pResultBuf, pSourceBuf, resultChannelSize );
+                    }
+                    else if (resultFormat == MBF_PACKEDDIR8_W_TANGENTSIGN
+                             ||
+                             resultFormat == MBF_PACKEDDIRS8_W_TANGENTSIGN)
+                    {
+                        check(sourceComponents>=3);
+                        check(resultComponents==4);
+
+                        // convert the 3 first components
+                        //memcpy(pResultBuf, pSourceBuf, resultChannelSize);
+                        for ( int i=0; i<3; ++i )
+                        {
+                            if ( i<sourceComponents)
+                            {
+                                ConvertData
+                                        (
+                                            i,
+                                            pResultBuf, resultFormat,
+                                            pSourceBuf, sourceFormat
+                                        );
+                            }
+                        }
+
+
+                        // Add the tangent sign
+                        auto pData = (uint8_t*)pResultBuf;
+
+                        // Look for the full tangent space
+                        int tanXBuf, tanXChan, tanYBuf, tanYChan, tanZBuf, tanZChan;
+                        Source.FindChannel(MBS_TANGENT, resultSemanticIndex, &tanXBuf, &tanXChan);
+                        Source.FindChannel(MBS_BINORMAL, resultSemanticIndex, &tanYBuf, &tanYChan);
+                        Source.FindChannel(MBS_NORMAL, resultSemanticIndex, &tanZBuf, &tanZChan);
+
+                        if (tanXBuf >= 0 && tanYBuf >= 0 && tanZBuf >= 0)
+                        {
+                            UntypedMeshBufferIteratorConst xIt(Source, MBS_TANGENT, resultSemanticIndex);
+                            UntypedMeshBufferIteratorConst yIt(Source, MBS_BINORMAL, resultSemanticIndex);
+                            UntypedMeshBufferIteratorConst zIt(Source, MBS_NORMAL, resultSemanticIndex);
+
+                            mat3f mat;
+                            xIt += v;
+                            yIt += v;
+                            zIt += v;
+                            mat[0] = xIt.GetAsVec4f().xyz();
+                            mat[1] = yIt.GetAsVec4f().xyz();
+                            mat[2] = zIt.GetAsVec4f().xyz();
+
+                            uint8_t sign = 0;
+                            if (resultFormat == MBF_PACKEDDIR8_W_TANGENTSIGN)
+                            {
+                                sign = mat.GetDeterminant() < 0 ? 0 : 255;
+                            }
+                            else if (resultFormat == MBF_PACKEDDIRS8_W_TANGENTSIGN)
+                            {
+                                sign = mat.GetDeterminant() < 0 ? -128 : 127;
+                            }
+                            pData[3] = sign;
+                        }
+                    }
+                    else
+                    {
+						// Convert formats
+						for ( int i=0; i<resultComponents; ++i )
+						{
+							if ( i<sourceComponents)
+							{
+								ConvertData
+										(
+											i,
+											pResultBuf, resultFormat,
+											pSourceBuf, sourceFormat
+										);
+							}
+							else
+							{
+								// Add zeros. TODO: Warning?
+								FMemory::Memzero
+										(
+											pResultBuf+GetMeshFormatData( resultFormat ).m_size*i,
+											GetMeshFormatData( resultFormat ).m_size
+										);
+							}
+						}
+
+
+                        // Extra step to normalise some semantics in some formats
+                        // TODO: Make it optional, and add different normalisation types n, n^2
+                        // TODO: Optimise
+                        if (sourceSemantic==MBS_BONEWEIGHTS)
+                        {
+                            if (resultFormat==MBF_NUINT8)
+                            {
+                                auto pData = (uint8_t*)pResultBuf;
+                                uint8_t accum = 0;
+                                for ( int i=0; i<resultComponents; ++i )
+                                {
+                                    accum+= pData[i];
+                                }
+                                pData[0] += 255-accum;
+                            }
+                        }
+					}
+
+					pResultBuf += resultElemSize;
+					pSourceBuf += sourceElemSize;
+				}
+			}
+		}
+	}
+}
+
+
+//-------------------------------------------------------------------------------------------------
+void FormatBufferSet
+    (
+        const FMeshBufferSet& Source,
+        FMeshBufferSet& Result,
+        bool keepSystemBuffers,
+        bool ignoreMissingChannels,
+        bool isVertexBuffer
+    )
+{
+    if ( ignoreMissingChannels )
+    {
+        // Remove from the result the channels that are not present in the source, and re-pack the
+        // offsets.
+        for ( int b=0; b<Result.GetBufferCount(); ++b )
+        {
+            vector<MESH_BUFFER_SEMANTIC> resultSemantics;
+            vector<int> resultSemanticIndexs;
+            vector<MESH_BUFFER_FORMAT> resultFormats;
+            vector<int> resultComponentss;
+            vector<int> resultOffsets;
+            int offset = 0;
+
+            // For every channel in this buffer
+            for ( int c=0; c<Result.GetBufferChannelCount(b); ++c )
+            {
+                MESH_BUFFER_SEMANTIC resultSemantic;
+                int resultSemanticIndex;
+                MESH_BUFFER_FORMAT resultFormat;
+                int resultComponents;
+
+                // Find this channel in the source mesh
+                Result.GetChannel
+                        (
+                            b, c,
+                            &resultSemantic, &resultSemanticIndex,
+                            &resultFormat, &resultComponents,
+                            nullptr
+                        );
+
+                int sourceBuffer;
+                int sourceChannel;
+                Source.FindChannel
+                        ( resultSemantic, resultSemanticIndex, &sourceBuffer, &sourceChannel );
+
+                if ( sourceBuffer>=0 )
+                {
+                    resultSemantics.push_back( resultSemantic );
+                    resultSemanticIndexs.push_back( resultSemanticIndex );
+                    resultFormats.push_back( resultFormat );
+                    resultComponentss.push_back( resultComponents );
+                    resultOffsets.push_back( offset );
+
+                    offset += GetMeshFormatData( resultFormat ).m_size * resultComponents;
+                }
+            }
+
+            if (resultSemantics.empty())
+            {
+                Result.SetBuffer( b, 0, 0, nullptr, nullptr, nullptr, nullptr, nullptr );
+            }
+            else
+            {
+                Result.SetBuffer( b, offset, (int)resultSemantics.size(),
+                                    &resultSemantics[0],
+                                    &resultSemanticIndexs[0],
+                                    &resultFormats[0],
+                                    &resultComponentss[0],
+                                    &resultOffsets[0] );
+            }
+        }
+    }
+
+
+    // For every vertex buffer in result
+    int vCount = Source.GetElementCount();
+    Result.SetElementCount( vCount );
+    for ( int b=0; b<Result.GetBufferCount(); ++b )
+    {
+        MeshFormatBuffer( Source, Result, b );
+    }
+
+
+    // Detect internal system buffers and clone them unmodified.
+    if ( keepSystemBuffers )
+    {
+        for ( int b=0; b<Source.GetBufferCount(); ++b )
+        {
+            // Detect system buffers and clone them unmodified.
+            if ( Source.GetBufferChannelCount(b)>0 )
+            {
+                MESH_BUFFER_SEMANTIC sourceSemantic;
+                int sourceSemanticIndex;
+                MESH_BUFFER_FORMAT sourceFormat;
+                int sourceComponents;
+                int sourceOffset;
+                Source.GetChannel
+                        (
+                            b, 0,
+                            &sourceSemantic, &sourceSemanticIndex,
+                            &sourceFormat, &sourceComponents,
+                            &sourceOffset
+                        );
+
+                if ( sourceSemantic==MBS_LAYOUTBLOCK
+                     || sourceSemantic==MBS_CHART
+                     || ( isVertexBuffer && sourceSemantic==MBS_VERTEXINDEX )
+                     )
+                {
+                    Result.AddBuffer( Source, b );
+                }
+            }
+        }
+    }
+
+}
+
+
+
+//-------------------------------------------------------------------------------------------------
+MeshPtr mu::MeshFormat
+	(
+		const Mesh* pPureSource,
+		const Mesh* pFormat,
+		bool keepSystemBuffers,
+		bool formatVertices,
+		bool formatIndices,
+		bool formatFaces,
+		bool rebuildTangents,
+		bool ignoreMissingChannels
+	)
+{
+    MUTABLE_CPUPROFILER_SCOPE(MeshFormat);
+
+	if (!pPureSource) { return nullptr; }
+	if (!pFormat) { return pPureSource->Clone(); }
+
+    MeshPtrConst pSource = pPureSource;
+
+    // Rebuild the tangent space if necessary
+    if ( rebuildTangents )
+    {
+		MeshPtr pClonedSource = pPureSource->Clone();
+		pSource = pClonedSource;
+
+		// Make sure there are the tangents and binormal channels.
+        int tanBuf, tanChan, binBuf, binChan;
+        pSource->GetVertexBuffers().FindChannel( MBS_TANGENT, 0, &tanBuf, &tanChan );
+        pSource->GetVertexBuffers().FindChannel( MBS_BINORMAL, 0, &binBuf, &binChan );
+
+		if (tanBuf<=0 ||binBuf<=0)
+        {
+			MESH_BUFFER buffer;
+
+			if (tanBuf<=0)
+			{
+                MESH_BUFFER_CHANNEL channel;
+				channel.m_semantic = MBS_TANGENT;
+				channel.m_format = MBF_FLOAT32;
+				channel.m_componentCount = 3;
+				channel.m_semanticIndex = 0;
+				channel.m_offset = 0;
+				buffer.m_channels.Add( channel );
+				buffer.m_elementSize += 3*sizeof(float);
+			}
+
+			if (binBuf<=0)
+			{
+                MESH_BUFFER_CHANNEL channel;
+				channel.m_semantic = MBS_BINORMAL;
+				channel.m_format = MBF_FLOAT32;
+				channel.m_componentCount = 3;
+				channel.m_semanticIndex = 0;
+				channel.m_offset = (uint8)buffer.m_elementSize;
+				buffer.m_channels.Add( channel );
+				buffer.m_elementSize += 3*sizeof(float);
+			}
+
+			buffer.m_data.SetNum( buffer.m_elementSize*pClonedSource->GetVertexCount() );
+			pClonedSource->GetVertexBuffers().m_buffers.Add( buffer );
+        }
+
+        RebuildTangents( pClonedSource->GetIndexBuffers(), pClonedSource->GetVertexBuffers() );
+    }
+
+    MeshPtr pResult = pFormat->Clone();
+
+    // Make sure that the bone indices will fit in this format, or extend it.
+    if ( formatVertices )
+    {
+        auto& buffs = pSource->GetVertexBuffers();
+
+        int semanticIndex = 0;
+        while( true )
+        {
+            int buf = 0;
+            int chan = 0;
+            buffs.FindChannel( MBS_BONEINDICES, semanticIndex, &buf, &chan );
+            if (buf<0)
+                break;
+
+            int resultBuf = 0;
+            int resultChan = 0;
+            auto& formBuffs = pResult->GetVertexBuffers();
+            formBuffs.FindChannel( MBS_BONEINDICES, semanticIndex, &resultBuf, &resultChan );
+            if ( resultBuf >=0 )
+            {
+                UntypedMeshBufferIteratorConst it( buffs, MBS_BONEINDICES, semanticIndex );
+                int32_t maxBoneIndex = 0;
+                for ( int v = 0; v < buffs.GetElementCount(); ++v )
+                {
+                    auto va = it.GetAsVec8i();
+                    for (int c=0; c<it.GetComponents(); ++c)
+                    {
+                        maxBoneIndex = std::max(maxBoneIndex,va[c]);
+                    }                    
+                    ++it;
+                }
+
+                auto& format = formBuffs.m_buffers[resultBuf].m_channels[resultChan].m_format;
+                if ( maxBoneIndex > 0xffff && (format == MBF_UINT8 || format == MBF_UINT16) )
+                {
+                    format = MBF_UINT32;
+                    formBuffs.UpdateOffsets( resultBuf );
+                }
+                else if ( maxBoneIndex > 0x7fff && (format == MBF_INT8 || format == MBF_INT16) )
+                {
+                    format = MBF_UINT32;
+                    formBuffs.UpdateOffsets( resultBuf );
+                }
+                else if ( maxBoneIndex > 0xff && format == MBF_UINT8 )
+                {
+                    format = MBF_UINT16;
+                    formBuffs.UpdateOffsets( resultBuf );
+                }
+                else if ( maxBoneIndex > 0x7f && format == MBF_INT8 )
+                {
+                    format = MBF_INT16;
+                    formBuffs.UpdateOffsets( resultBuf );
+                }
+            }
+
+            semanticIndex++;
+        }
+    }
+
+    // \todo Make sure that the vertex indices will fit in this format, or extend it.
+
+
+
+    if ( formatVertices )
+	{
+
+        FormatBufferSet( pSource->GetVertexBuffers(), pResult->GetVertexBuffers(),
+                         keepSystemBuffers, ignoreMissingChannels, true );
+    }
+	else
+	{
+		pResult->m_VertexBuffers = pSource->GetVertexBuffers();
+	}
+
+	if ( formatIndices )
+	{
+        FormatBufferSet( pSource->GetIndexBuffers(), pResult->GetIndexBuffers(), keepSystemBuffers,
+                         ignoreMissingChannels, false );
+    }
+	else
+	{
+		pResult->m_IndexBuffers = pSource->GetIndexBuffers();
+	}
+
+	if ( formatFaces )
+	{
+        FormatBufferSet( pSource->GetFaceBuffers(),
+                         pResult->GetFaceBuffers(),
+						 keepSystemBuffers,
+                         ignoreMissingChannels,
+                         false );
+	}
+	else
+	{
+		pResult->m_FaceBuffers = pSource->GetFaceBuffers();
+	}
+
+	// Copy the rest of the data
+    pResult->SetSkeleton( pSource->GetSkeleton() );
+	pResult->SetPhysicsBody( pSource->GetPhysicsBody() );
+
+	pResult->m_layouts.Empty();
+	for ( const auto& Layout: pSource->m_layouts )
+	{
+		pResult->m_layouts.Add(Layout->Clone());
+	}
+
+	pResult->ResetStaticFormatFlags();
+	pResult->EnsureSurfaceData();
+
+	pResult->m_tags = pSource->m_tags;
+
+	pResult->m_AdditionalBuffers = pSource->m_AdditionalBuffers;
+
+	pResult->m_bonePoses = pSource->m_bonePoses;
+
+	return pResult;
+
+}
+
+
diff --git a/Engine/Plugins/Experimental/Mutable/Source/MutableRuntime/Private/OpMeshInterpolate.h b/Engine/Plugins/Experimental/Mutable/Source/MutableRuntime/Private/OpMeshInterpolate.h
new file mode 100644
--- /dev/null
+++ b/Engine/Plugins/Experimental/Mutable/Source/MutableRuntime/Private/OpMeshInterpolate.h
@@ -0,0 +1,72 @@
+// Copyright Epic Games, Inc. All Rights Reserved.
+
+#pragma once
+
+#include "MeshPrivate.h"
+#include "Platform.h"
+
+
+namespace mu
+{
+
+
+	//---------------------------------------------------------------------------------------------
+    inline MeshPtr MeshInterpolate( const Mesh* pA, const Mesh* pB, float factor )
+	{
+		bool correct = ( pA->GetVertexBuffers().GetElementCount() == pB->GetVertexBuffers().GetElementCount() )
+					&& ( pA->GetIndexCount() == pB->GetIndexCount() )
+					&& ( pA->HasCompatibleFormat(pB) );
+
+		MeshPtr pDest = pA->Clone();
+
+		// Deal with degenerated data
+		if ( correct && pDest->GetVertexBuffers().GetElementCount() )
+		{
+			// For now, just interpolate the position and normal buffer.
+
+			// A mesh position
+			int aPosElementSize = 0;
+            uint8_t* pAPosBuf = 0;
+			GetMeshBuf( pDest.get(), MBS_POSITION, MBF_FLOAT32, 3, pAPosBuf, aPosElementSize );
+
+			// A mesh normal
+			int aNorElementSize = 0;
+            uint8_t* pANorBuf = 0;
+			GetMeshBuf( pDest.get(), MBS_NORMAL, MBF_FLOAT32, 3, pANorBuf, aNorElementSize );
+
+			// B mesh position
+			int bPosElementSize = 0;
+            const uint8_t* pBPosBuf = 0;
+			GetMeshBuf( pB, MBS_POSITION, MBF_FLOAT32, 3, pBPosBuf, bPosElementSize );
+
+			// B mesh normal
+			int bNorElementSize = 0;
+            const uint8_t* pBNorBuf = 0;
+			GetMeshBuf( pB, MBS_NORMAL, MBF_FLOAT32, 3, pBNorBuf, bNorElementSize );
+
+            uint32_t vcount = pA->GetVertexBuffers().GetElementCount();
+            for ( uint32_t vi=0; vi<vcount; ++vi )
+			{
+				vec3<float>* aPos = reinterpret_cast< vec3<float>* >( pAPosBuf );
+				const vec3<float>* bPos = reinterpret_cast< const vec3<float>* >( pBPosBuf );
+				*aPos += (*bPos - *aPos) * factor;
+
+				vec3<float>* aNor = reinterpret_cast< vec3<float>* >( pANorBuf );
+				const vec3<float>* bNor = reinterpret_cast< const vec3<float>* >( pBNorBuf );
+				*aNor += (*bNor - *aNor) * factor;
+
+				pAPosBuf += aPosElementSize;
+				pBPosBuf += bPosElementSize;
+				pANorBuf += aNorElementSize;
+				pBNorBuf += bNorElementSize;
+			}
+		}
+		else
+		{
+			// TODO report somehow
+		}
+
+		return pDest;
+	}
+
+}
diff --git a/Engine/Plugins/Experimental/Mutable/Source/MutableRuntime/Private/OpMeshMerge.h b/Engine/Plugins/Experimental/Mutable/Source/MutableRuntime/Private/OpMeshMerge.h
new file mode 100644
--- /dev/null
+++ b/Engine/Plugins/Experimental/Mutable/Source/MutableRuntime/Private/OpMeshMerge.h
@@ -0,0 +1,1174 @@
+// Copyright Epic Games, Inc. All Rights Reserved.
+
+#pragma once
+
+#include "MeshPrivate.h"
+#include "Platform.h"
+#include "MutableMath.h"
+#include "OpMeshFormat.h"
+#include "MutableTrace.h"
+
+
+namespace mu
+{
+
+
+	//---------------------------------------------------------------------------------------------
+	//! Merge two meshes into one new mesh
+	//---------------------------------------------------------------------------------------------
+	inline MeshPtr MeshMerge( const Mesh* pFirst, const Mesh* pSecond )
+	{
+        MUTABLE_CPUPROFILER_SCOPE(MeshMerge);
+
+        MeshPtr pResult = new Mesh;
+
+		// Should never happen, but fixes static analysis warnings.
+		if (!(pFirst && pSecond))
+		{
+			return pResult;
+		}
+
+		// Indices
+		//-----------------
+		if ( pFirst->GetIndexBuffers().GetBufferCount()>0 )
+		{
+            MUTABLE_CPUPROFILER_SCOPE(Indices);
+
+            int firstCount = pFirst->GetIndexBuffers().GetElementCount();
+			int secondCount = pSecond->GetIndexBuffers().GetElementCount();
+			pResult->GetIndexBuffers().SetElementCount( firstCount+secondCount );
+
+			check( pFirst->GetIndexBuffers().GetBufferCount() <= 1 );
+			check( pSecond->GetIndexBuffers().GetBufferCount() <= 1 );
+			pResult->GetIndexBuffers().SetBufferCount( 1 );
+
+            // This will be changed below if need to change the format of the index buffers.
+            MESH_BUFFER_FORMAT requiredBufferFormatChange = MBF_NONE;
+
+            if (firstCount && secondCount)
+			{
+				const MESH_BUFFER& first = pFirst->GetIndexBuffers().m_buffers[0];
+				const MESH_BUFFER& second = pSecond->GetIndexBuffers().m_buffers[0];
+				check(first.m_channels == second.m_channels);
+				check(first.m_elementSize == second.m_elementSize);
+                check(first.m_channels.IsEmpty() || first.m_channels[0].m_format == second.m_channels[0].m_format);
+				
+				// Avoid unused variable warnings
+				(void)first;
+				(void)second;
+
+                // We need to know the total number of vertices in case we need to adjust the index buffer format.
+                uint64_t totalVertexCount = pFirst->GetVertexBuffers().GetElementCount()
+                        + pSecond->GetVertexBuffers().GetElementCount();
+                uint64_t maxValueBits = GetMeshFormatData(pFirst->GetIndexBuffers().m_buffers[0].m_channels[0].m_format).m_maxValueBits;
+                uint64_t maxSupportedVertices = uint64_t(1)<<maxValueBits;
+                if (totalVertexCount>maxSupportedVertices)
+                {
+                    if (totalVertexCount > 0xFFFF)
+                    {
+                        requiredBufferFormatChange = MBF_UINT32;
+                    }
+                    else
+                    {
+                        requiredBufferFormatChange = MBF_UINT16;
+                    }
+                }
+            }
+
+			MESH_BUFFER& result = pResult->GetIndexBuffers().m_buffers[0];
+
+            if (requiredBufferFormatChange != MBF_NONE)
+            {
+                // We only support vertex indices in case of having to change the format.
+                check(pFirst->GetIndexBuffers().m_buffers[0].m_channels.Num()==1);
+
+                result.m_channels.SetNum(1);
+                result.m_channels[0].m_semantic = MBS_VERTEXINDEX;
+                result.m_channels[0].m_format = requiredBufferFormatChange;
+                result.m_channels[0].m_componentCount = 1;
+                result.m_channels[0].m_semanticIndex = 0;
+                result.m_channels[0].m_offset = 0;
+                result.m_elementSize = GetMeshFormatData(requiredBufferFormatChange).m_size;
+            }
+            else if (firstCount)
+			{
+				const MESH_BUFFER& first = pFirst->GetIndexBuffers().m_buffers[0];
+				result.m_channels = first.m_channels;
+				result.m_elementSize = first.m_elementSize;
+			}
+            else if (secondCount)
+			{
+				const MESH_BUFFER& second = pSecond->GetIndexBuffers().m_buffers[0];
+				result.m_channels = second.m_channels;
+				result.m_elementSize = second.m_elementSize;
+			}
+
+			result.m_data.SetNum(result.m_elementSize*(firstCount + secondCount));
+
+			check( result.m_channels.Num()==1 );
+			check( result.m_channels[0].m_semantic == MBS_VERTEXINDEX );
+
+			if ( result.m_data.Num() )
+			{
+				if ( firstCount )
+				{
+					const MESH_BUFFER& first = pFirst->GetIndexBuffers().m_buffers[0];
+
+                    if (requiredBufferFormatChange==MBF_NONE
+                            || requiredBufferFormatChange==first.m_channels[0].m_format)
+                    {
+                        FMemory::Memcpy( &result.m_data[0],
+                                        &first.m_data[0],
+                                        first.m_elementSize*firstCount );
+                    }
+                    else
+                    {
+                        // Conversion required
+                        const uint8_t* pSource = &first.m_data[0];
+                        uint8_t* pDest = &result.m_data[0];
+                        switch( requiredBufferFormatChange )
+                        {
+                        case MBF_UINT32:
+                        {
+                            switch( first.m_channels[0].m_format )
+                            {
+                            case MBF_UINT16:
+                            {
+                                for ( int v=0; v<firstCount; ++v )
+                                {
+                                    *(uint32_t*)pDest = *(const uint16_t*)pSource;
+                                    pSource+=first.m_elementSize;
+                                    pDest+=result.m_elementSize;
+                                }
+                                break;
+                            }
+
+                            case MBF_UINT8:
+                            {
+                                for ( int v=0; v<firstCount; ++v )
+                                {
+                                    *(uint32_t*)pDest = *(const uint8_t*)pSource;
+                                    pSource+=first.m_elementSize;
+                                    pDest+=result.m_elementSize;
+                                }
+                                break;
+                            }
+
+                            default:
+                                checkf( false, TEXT("Format not supported.") );
+                                break;
+                            }
+                            break;
+                        }
+
+                        case MBF_UINT16:
+                        {
+                            switch( first.m_channels[0].m_format )
+                            {
+
+                            case MBF_UINT8:
+                            {
+                                for ( int v=0; v<firstCount; ++v )
+                                {
+                                    *(uint16_t*)pDest = *(const uint8_t*)pSource;
+                                    pSource+=first.m_elementSize;
+                                    pDest+=result.m_elementSize;
+                                }
+                                break;
+                            }
+
+                            default:
+								checkf(false, TEXT("Format not supported."));
+								break;
+                            }
+                            break;
+                        }
+
+                        default:
+							checkf(false, TEXT("Format not supported."));
+							break;
+                        }
+                    }
+                }
+
+				if ( secondCount )
+				{
+					const MESH_BUFFER& second = pSecond->GetIndexBuffers().m_buffers[0];
+                    const uint8_t* pSource = &second.m_data[0];
+                    uint8_t* pDest = &result.m_data[result.m_elementSize*firstCount ];
+
+                    uint32_t firstVertexCount = pFirst->GetVertexBuffers().GetElementCount();
+
+                    if (requiredBufferFormatChange==MBF_NONE
+                            || requiredBufferFormatChange==second.m_channels[0].m_format)
+                    {
+                        switch( second.m_channels[0].m_format )
+                        {
+                        case MBF_INT32:
+                        case MBF_UINT32:
+                        case MBF_NINT32:
+                        case MBF_NUINT32:
+                        {
+                            for ( int v=0; v<secondCount; ++v )
+                            {
+                                *(uint32_t*)pDest = firstVertexCount + *(const uint32_t*)pSource;
+                                pSource+=second.m_elementSize;
+                                pDest+=result.m_elementSize;
+                            }
+                            break;
+                        }
+
+                        case MBF_INT16:
+                        case MBF_UINT16:
+                        case MBF_NINT16:
+                        case MBF_NUINT16:
+                        {
+                            for ( int v=0; v<secondCount; ++v )
+                            {
+                                *(uint16_t*)pDest = uint16_t(firstVertexCount) + *(const uint16_t*)pSource;
+                                pSource+=second.m_elementSize;
+                                pDest+=result.m_elementSize;
+                            }
+                            break;
+                        }
+
+                        case MBF_INT8:
+                        case MBF_UINT8:
+                        case MBF_NINT8:
+                        case MBF_NUINT8:
+                        {
+                            for ( int v=0; v<secondCount; ++v )
+                            {
+                                *(uint8_t*)pDest = uint8_t(firstVertexCount) + *(const uint8_t*)pSource;
+                                pSource+=second.m_elementSize;
+                                pDest+=result.m_elementSize;
+                            }
+                            break;
+                        }
+
+                        default:
+							checkf(false, TEXT("Format not supported."));
+							break;
+                        }
+                    }
+                    else
+                    {
+                        // Format conversion required
+                        switch (requiredBufferFormatChange)
+                        {
+
+                        case MBF_UINT32:
+                        {
+                            switch( second.m_channels[0].m_format )
+                            {
+                            case MBF_INT16:
+                            case MBF_UINT16:
+                            case MBF_NINT16:
+                            case MBF_NUINT16:
+                            {
+                                for ( int v=0; v<secondCount; ++v )
+                                {
+                                    *(uint32_t*)pDest = uint32_t(firstVertexCount) + *(const uint16_t*)pSource;
+                                    pSource+=second.m_elementSize;
+                                    pDest+=result.m_elementSize;
+                                }
+                                break;
+                            }
+
+                            case MBF_INT8:
+                            case MBF_UINT8:
+                            case MBF_NINT8:
+                            case MBF_NUINT8:
+                            {
+                                for ( int v=0; v<secondCount; ++v )
+                                {
+                                    *(uint32_t*)pDest = uint32_t(firstVertexCount) + *(const uint8_t*)pSource;
+                                    pSource+=second.m_elementSize;
+                                    pDest+=result.m_elementSize;
+                                }
+                                break;
+                            }
+
+                            default:
+								checkf(false, TEXT("Format not supported."));
+								break;
+                            }
+
+                            break;
+                        }
+
+                        case MBF_UINT16:
+                        {
+                            switch( second.m_channels[0].m_format )
+                            {
+                            case MBF_INT8:
+                            case MBF_UINT8:
+                            case MBF_NINT8:
+                            case MBF_NUINT8:
+                            {
+                                for ( int v=0; v<secondCount; ++v )
+                                {
+                                    *(uint16_t*)pDest = uint16_t(firstVertexCount) + *(const uint8_t*)pSource;
+                                    pSource+=second.m_elementSize;
+                                    pDest+=result.m_elementSize;
+                                }
+                                break;
+                            }
+
+                            default:
+								checkf(false, TEXT("Format not supported."));
+								break;
+                            }
+
+                            break;
+                        }
+
+                        default:
+							checkf(false, TEXT("Format not supported."));
+							break;
+
+                        }
+                    }
+                }
+			}
+		}
+
+		// Faces
+		//-----------------
+		{
+            MUTABLE_CPUPROFILER_SCOPE(Faces);
+
+            int firstCount = pFirst->GetFaceBuffers().GetElementCount();
+			int secondCount = pSecond->GetFaceBuffers().GetElementCount();
+			pResult->GetFaceBuffers().SetElementCount( firstCount+secondCount );
+			pResult->GetFaceBuffers().SetBufferCount( pFirst->GetFaceBuffers().GetBufferCount() );
+
+			// Merge only the buffers present in the first mesh
+			for ( int b = 0; b<pResult->GetFaceBuffers().GetBufferCount(); ++b )
+			{
+				const MESH_BUFFER& first = pFirst->GetFaceBuffers().m_buffers[b];
+
+				MESH_BUFFER& result = pResult->GetFaceBuffers().m_buffers[b];
+				result.m_channels = first.m_channels;
+				result.m_elementSize = first.m_elementSize;
+				result.m_data.SetNum( result.m_elementSize * (firstCount+secondCount) );
+
+				MESH_BUFFER_SEMANTIC semantic = MBS_NONE;
+				int semanticIndex = 0;
+				MESH_BUFFER_FORMAT type = MBF_NONE;
+				int components = 0;
+				int offset = 0;
+				pFirst->GetFaceBuffers().GetChannel
+						( b, 0, &semantic, &semanticIndex, &type, &components, &offset );
+
+				if ( firstCount )
+				{
+					FMemory::Memcpy( &result.m_data[0],
+									&first.m_data[0],
+									result.m_elementSize * firstCount
+									);
+				}
+
+				if ( secondCount )
+				{
+					// Find in the second mesh
+					int otherBuffer = -1;
+					int otherChannel = -1;
+					pSecond->GetFaceBuffers().FindChannel
+							( semantic, 0, &otherBuffer, &otherChannel );
+
+					if (otherBuffer>=0)
+					{
+						const MESH_BUFFER& second =
+								pSecond->GetFaceBuffers().m_buffers[otherBuffer];
+						check( first.m_channels == second.m_channels );
+
+						// Raw copy
+						FMemory::Memcpy( &result.m_data[ result.m_elementSize * firstCount ],
+										&second.m_data[0],
+										result.m_elementSize * secondCount );
+					}
+					else
+					{
+						// Fill with zeroes
+						FMemory::Memzero( &result.m_data[ result.m_elementSize * firstCount ],
+										result.m_elementSize * secondCount );
+					}
+				}
+			}
+		}
+
+
+        // Layouts
+        //-----------------
+        {
+            MUTABLE_CPUPROFILER_SCOPE(Layouts);
+
+            pResult->m_layouts.SetNum( pFirst->m_layouts.Num() );
+            for ( int i=0; i<pFirst->m_layouts.Num(); ++i )
+            {
+                const Layout* pF = pFirst->m_layouts[i].get();
+                LayoutPtr pR = pF->Clone();
+
+                if ( i<pSecond->m_layouts.Num() )
+                {
+                    const Layout* pS = pSecond->m_layouts[i].get();
+
+                    pR->m_blocks.Append(pS->m_blocks);
+                }
+
+                pResult->m_layouts[i] = pR;
+            }
+        }
+
+
+        // Surfaces
+        //-----------------
+        pResult->m_surfaces = pFirst->m_surfaces;
+
+        for ( const auto& s : pSecond->m_surfaces )
+        {
+            pResult->m_surfaces.Add(s);
+            pResult->m_surfaces.Last().m_firstVertex += pFirst->GetVertexCount();
+            pResult->m_surfaces.Last().m_firstIndex += pFirst->GetIndexCount();
+        }
+
+
+		// Skeleton
+		//---------------------------
+
+        // Do they have the same skeleton?
+        bool remapNeeded = pFirst->GetSkeleton() != pSecond->GetSkeleton();
+
+        // Are they different skeletons but with the same data?
+        if ( remapNeeded && pFirst->GetSkeleton() && pSecond->GetSkeleton() )
+        {
+            remapNeeded = ! ( *pFirst->GetSkeleton()
+                              ==
+                              *pSecond->GetSkeleton() );
+        }
+
+
+		TArray<int32> secondToFirstBones;
+
+        if ( remapNeeded )
+        {
+            MUTABLE_CPUPROFILER_SCOPE(Skeleton);
+            Ptr<Skeleton> pResultSkeleton;
+
+			mu::SkeletonPtrConst pFirstSkeleton = pFirst->GetSkeleton();
+			mu::SkeletonPtrConst pSecondSkeleton = pSecond->GetSkeleton();
+
+			const int32 pFirstBoneCount = pFirstSkeleton ? pFirstSkeleton->GetBoneCount() : 0;
+			const int32 pSecondBoneCount = pSecondSkeleton ? pSecondSkeleton->GetBoneCount() : 0;
+			
+			pResultSkeleton = pFirstSkeleton ? pFirstSkeleton->Clone() : new Skeleton; 
+			pResult->SetSkeleton(pResultSkeleton);
+
+			secondToFirstBones.SetNumUninitialized(pSecondBoneCount);
+
+			// Merge pSecond and build the remap array 
+            for ( int32 ob=0; ob< pSecondBoneCount; ++ob )
+            {
+                int32 tb = INDEX_NONE;
+                for ( int32 c=0; tb<0 && c<pResultSkeleton->m_bones.Num(); ++c )
+                {
+                    if ( pResultSkeleton->m_bones[c] == pSecondSkeleton->GetBoneName( ob ) )
+                    {
+                        tb = c;
+						break;
+                    }
+                }
+
+                // Add a new bone
+                if( tb == INDEX_NONE)
+                {
+                    tb = pResultSkeleton->m_bones.Num();
+                    pResultSkeleton->m_bones.Add( pSecondSkeleton->m_bones[ob] );
+					pResultSkeleton->m_boneIds.Add(pSecondSkeleton->m_boneIds[ob]);
+					
+					// Add an incorrect index, to be fixed below in case the parent index is later in the bone array.
+					pResultSkeleton->m_boneParents.Add(pSecondSkeleton->m_boneParents[ob]);
+                }
+
+                secondToFirstBones[ob] = tb;
+			}
+
+            // Fix second mesh bone parents
+            for (int32 ob = pFirstBoneCount; ob < pResultSkeleton->m_boneParents.Num(); ++ob)
+            {
+                int16_t secondMeshIndex = pResultSkeleton->m_boneParents[ob];
+                if (secondMeshIndex != INDEX_NONE)
+                {
+                    pResultSkeleton->m_boneParents[ob] = (uint16_t)secondToFirstBones[secondMeshIndex];
+                }
+            }
+        }
+        else
+        {
+            pResult->SetSkeleton( pFirst->GetSkeleton() );
+        }
+
+
+		// Pose
+		//---------------------------
+		{
+			MUTABLE_CPUPROFILER_SCOPE(Pose);
+
+			pResult->m_bonePoses.Reserve(pResult->GetSkeleton()->GetBoneCount());
+			
+			// Copy poses from the first mesh
+			pResult->m_bonePoses = pFirst->m_bonePoses;
+
+			// Add or override bone poses
+			for (const Mesh::FBonePose& pSecondBonePose : pSecond->m_bonePoses)
+			{
+				const int32 resultBoneIndex = pResult->FindBonePose(pSecondBonePose.m_boneName.c_str());
+
+				if (resultBoneIndex != INDEX_NONE)
+				{
+					Mesh::FBonePose& pResultBonePose = pResult->m_bonePoses[resultBoneIndex];
+
+					// override pose if the bone is not used by the skinning of the first mesh
+					if (!pResultBonePose.m_boneSkinned)
+					{
+						pResultBonePose.m_boneName = pSecondBonePose.m_boneName;
+						pResultBonePose.m_boneSkinned = pSecondBonePose.m_boneSkinned;
+						pResultBonePose.m_boneTransform = pSecondBonePose.m_boneTransform;
+					}
+				}
+				else
+				{
+					pResult->m_bonePoses.Add(pSecondBonePose);
+				}
+			}
+
+			pResult->m_bonePoses.Shrink();
+		}
+
+		
+		// PhysicsBodies
+		//---------------------------
+
+		// Appends InPhysicsBody to OutPhysicsBody removing Bodies that are equal, have same bone and customId and its properies are identical.	
+		auto AppendPhysicsBodiesUnique = [](PhysicsBody& OutPhysicsBody, const PhysicsBody& InPhysicsBody)
+		{	
+			TArray<string>& OutBones = OutPhysicsBody.Bones;
+			TArray<int32>& OutCustomIds = OutPhysicsBody.CustomIds;
+			TArray<FPhysicsBodyAggregate>& OutBodies = OutPhysicsBody.Bodies;
+			
+			const TArray<string>& InBones = InPhysicsBody.Bones;
+			const TArray<int32>& InCustomIds = InPhysicsBody.CustomIds;
+			const TArray<FPhysicsBodyAggregate>& InBodies = InPhysicsBody.Bodies;
+			
+			const int32 InBodyCount = InPhysicsBody.GetBodyCount();
+			const int32 OutBodyCount = OutPhysicsBody.GetBodyCount();
+
+			for (int32 I = 0; I < InBodyCount; ++I)
+			{
+				int32 FoundIndex = INDEX_NONE;
+				for (int32 O = 0; O < OutBodyCount; ++O)
+				{
+					if (InCustomIds[I] == OutCustomIds[O] && InBones[I] == OutBones[O])
+					{
+						FoundIndex = O;
+						break;
+					}
+				}
+
+				if (FoundIndex == INDEX_NONE)
+				{
+					OutBones.Add(InBones[I]);
+					OutCustomIds.Add(InCustomIds[I]);
+					OutBodies.Add(InBodies[I]);
+				}
+				else
+				{
+					for (const FSphereBody& Body : InBodies[I].Spheres)
+					{
+						OutBodies[FoundIndex].Spheres.AddUnique(Body);
+					}
+
+					for (const FBoxBody& Body : InBodies[I].Boxes)
+					{
+						OutBodies[FoundIndex].Boxes.AddUnique(Body);
+					}
+
+					for (const FSphylBody& Body : InBodies[I].Sphyls)
+					{
+						OutBodies[FoundIndex].Sphyls.AddUnique(Body);
+					}
+
+					for (const FTaperedCapsuleBody& Body : InBodies[I].TaperedCapsules)
+					{
+						OutBodies[FoundIndex].TaperedCapsules.AddUnique(Body);
+					}
+
+					for (const FConvexBody& Body : InBodies[I].Convex)
+					{
+						OutBodies[FoundIndex].Convex.AddUnique(Body);
+					}
+				}
+			}
+		};
+
+
+		if (pFirst->GetPhysicsBody() == pSecond->GetPhysicsBody())
+		{
+			// Both meshes are sharing the physics body, let the result share it as well.
+			pResult->SetPhysicsBody( pFirst->GetPhysicsBody() );
+		}
+		else
+		{
+			Ptr<PhysicsBody> ResultPhysicsBody = new PhysicsBody;
+
+			if (pFirst->GetPhysicsBody())
+			{
+				AppendPhysicsBodiesUnique(*ResultPhysicsBody, *pFirst->GetPhysicsBody());
+			}
+			
+			if (pSecond->GetPhysicsBody())
+			{
+				AppendPhysicsBodiesUnique(*ResultPhysicsBody, *pSecond->GetPhysicsBody());
+			}
+			bool FirstBodyModified = pFirst->GetPhysicsBody() ? pFirst->GetPhysicsBody()->bBodiesModified : false;
+			bool SecondBodyModified = pSecond->GetPhysicsBody() ? pSecond->GetPhysicsBody()->bBodiesModified : false;
+			ResultPhysicsBody->bBodiesModified = FirstBodyModified || SecondBodyModified;
+			pResult->SetPhysicsBody( ResultPhysicsBody );
+		}
+
+		// Vertices
+		//-----------------
+		{
+            MUTABLE_CPUPROFILER_SCOPE(Vertices);
+
+            int firstCount = pFirst->GetVertexBuffers().GetElementCount();
+			int secondCount = pSecond->GetVertexBuffers().GetElementCount();
+
+			// TODO: when formats match, which at runtime should be always.
+			bool fastPath = pFirst->GetVertexBuffers().HasSameFormat( pSecond->GetVertexBuffers() );
+
+			fastPath = fastPath && !remapNeeded;
+
+			MeshPtrConst pVFirst;
+			MeshPtrConst pVSecond;
+
+			if ( !fastPath )
+			{
+                MUTABLE_CPUPROFILER_SCOPE(SlowPath);
+
+                // Very slow! Should we warn somehow?
+
+				// Expand component counts in vertex channels of the format mesh
+				int32 vbcount = pFirst->GetVertexBuffers().m_buffers.Num();
+				pResult->GetVertexBuffers().SetBufferCount( (int)vbcount );
+
+				for ( int32 vb = 0; vb<vbcount; ++vb )
+				{
+					MESH_BUFFER& result = pResult->GetVertexBuffers().m_buffers[vb];
+					const MESH_BUFFER& first = pFirst->GetVertexBuffers().m_buffers[vb];
+
+					result.m_channels = first.m_channels;
+					result.m_elementSize = first.m_elementSize;
+
+					// See if we need to enlarge the components of any of the result channels
+					bool resetOffsets = false;
+					for ( int32 c=0; c<result.m_channels.Num(); ++c )
+					{
+						int sb = -1;
+						int sc = -1;
+						pSecond->GetVertexBuffers().FindChannel
+								(
+									result.m_channels[c].m_semantic,
+									result.m_channels[c].m_semanticIndex,
+									&sb, &sc
+								);
+						if ( sb>=0 )
+						{
+							const MESH_BUFFER& second = pSecond->GetVertexBuffers().m_buffers[sb];
+
+							if ( second.m_channels[sc].m_componentCount
+								 >
+								 result.m_channels[c].m_componentCount )
+							{
+								result.m_channels[c].m_componentCount =
+										second.m_channels[sc].m_componentCount;
+								resetOffsets = true;
+							}
+						}
+					}
+
+					// Reset the channel offsets if necessary
+					if (resetOffsets)
+					{
+						int offset = 0;
+						for ( int32 c=0; c<result.m_channels.Num(); ++c )
+						{
+                            result.m_channels[c].m_offset = (uint8_t)offset;
+							offset += result.m_channels[c].m_componentCount
+									*
+									GetMeshFormatData(result.m_channels[c].m_format).m_size;
+						}
+						result.m_elementSize = offset;
+					}
+				}
+
+                // See if we need to add additional buffers from the second mesh (like vertex colours)
+                // This is a bit ad-hoc: we only add buffers containing all new channels
+                for ( const auto& buf : pSecond->GetVertexBuffers().m_buffers )
+                {
+                    bool someChannel = false;
+                    bool allNewChannels = true;
+                    for ( const auto& chan : buf.m_channels )
+                    {
+                        // Not a system buffer
+                        if ( chan.m_semantic!=MBS_VERTEXINDEX
+                             &&
+                             chan.m_semantic!=MBS_LAYOUTBLOCK
+                             &&
+                             chan.m_semantic!=MBS_CHART
+                             )
+                        {
+                            someChannel = true;
+
+                            int foundBuffer=-1;
+                            int foundChannel=-1;
+                            pFirst->GetVertexBuffers().FindChannel(chan.m_semantic, chan.m_semanticIndex,&foundBuffer,&foundChannel);
+                            if (foundBuffer>=0)
+                            {
+                                // There is a meaningful channel that we don't have in the first
+                                // mesh, we'll need to add the buffer.
+                                allNewChannels = false;
+                            }
+                        }
+                    }
+
+                    if (someChannel && allNewChannels)
+                    {
+                        pResult->GetVertexBuffers().m_buffers.Add(buf);
+                    }
+                }
+
+				// Change the format of the bone indices buffer if there are more than 256 bones.
+                int pResultBoneCount =
+                    pResult->GetSkeleton()  ? pResult->GetSkeleton()->GetBoneCount() : 0;
+                if (pResultBoneCount > 256)
+				{
+					int foundBuffer = -1;
+					int foundChannel = -1;
+					pResult->GetVertexBuffers().FindChannel(MBS_BONEINDICES, 0, &foundBuffer, &foundChannel);
+
+					if (foundChannel >= 0)
+					{
+						MESH_BUFFER& boneBuffer = pResult->GetVertexBuffers().m_buffers[foundBuffer];
+
+						MESH_BUFFER_FORMAT format = pResultBoneCount > 65535 ? MBF_UINT32 : MBF_UINT16;
+						boneBuffer.m_channels[foundChannel].m_format = format;
+
+						// Reset offsets
+						int offset = 0;
+						for (int32 c = 0; c < boneBuffer.m_channels.Num(); ++c)
+						{
+							boneBuffer.m_channels[c].m_offset = (uint8_t)offset;
+							offset += boneBuffer.m_channels[c].m_componentCount
+								*
+								GetMeshFormatData(boneBuffer.m_channels[c].m_format).m_size;
+						}
+						boneBuffer.m_elementSize = offset;
+					}
+				}
+
+                // Allocate vertices
+                pResult->GetVertexBuffers().SetElementCount( firstCount + secondCount );
+
+				// Convert the source meshes to the new format
+                if (pFirst->GetVertexBuffers().HasSameFormat(pResult->GetVertexBuffers()))
+                {
+                    pVFirst = pFirst;
+                }
+                else
+                {
+                    pVFirst = MeshFormat( pFirst, pResult.get(), false, true, false, false, false, false );
+                }
+
+
+                if (pSecond->GetVertexBuffers().HasSameFormat(pResult->GetVertexBuffers()))
+                {
+                    pVSecond = pSecond;
+                }
+                else
+                {
+                    pVSecond = MeshFormat( pSecond, pResult.get(), false, true, false, false, false, false );
+                }
+
+				check( pVFirst->GetVertexBuffers().HasSameFormat
+                                ( pVSecond->GetVertexBuffers() ) );
+			}
+			else
+			{
+                MUTABLE_CPUPROFILER_SCOPE(FastPath);
+
+                pVFirst = pFirst;
+				pVSecond = pSecond;
+			}
+
+
+			pResult->m_VertexBuffers = pVFirst->m_VertexBuffers;
+			pResult->GetVertexBuffers().SetElementCount( firstCount + secondCount );
+
+            // first copy all the vertex data
+			{
+				for ( int32 vb = 0; vb<pResult->GetVertexBuffers().m_buffers.Num(); ++vb )
+				{
+					MESH_BUFFER& result = pResult->GetVertexBuffers().m_buffers[vb];
+					const MESH_BUFFER& second = pVSecond->GetVertexBuffers().m_buffers[vb];
+					if ( secondCount )
+					{
+						int elemSize = pResult->GetVertexBuffers().GetElementSize((int)vb);
+						int firstSize = firstCount * elemSize;
+						int secondSize = secondCount * elemSize;
+
+						FMemory::Memcpy( &result.m_data[firstSize], &second.m_data[0], secondSize );
+					}
+				}
+			}
+
+            // TODO This could eventually be part of the mesh format: force the same skeleton.
+            if (remapNeeded)
+            {
+                MUTABLE_CPUPROFILER_SCOPE(Remap);
+
+                // We need to remap the bone indices of the second mesh vertices that we already copied
+                // to result
+				for ( int32 vb=0; vb<pResult->GetVertexBuffers().m_buffers.Num(); ++vb )
+				{
+					MESH_BUFFER& result = pResult->GetVertexBuffers().m_buffers[vb];
+
+					int elemSize = pResult->GetVertexBuffers().GetElementSize((int)vb);
+                    int firstSize = firstCount * elemSize;
+
+					for ( int c=0; c<pResult->GetVertexBuffers().GetBufferChannelCount( (int)vb ); ++c )
+					{
+						// Get info about the destination channel
+						MESH_BUFFER_SEMANTIC semantic = MBS_NONE;
+						int semanticIndex = 0;
+						MESH_BUFFER_FORMAT format = MBF_NONE;
+						int components = 0;
+						int offset = 0;
+						pResult->GetVertexBuffers().GetChannel
+								( (int)vb, c, &semantic, &semanticIndex, &format, &components, &offset );
+
+                        int resultOffset = firstSize + offset;
+
+						// Check special channels that need extra work
+						if ( semantic==MBS_BONEINDICES )
+						{
+							// Bone indices may need remapping
+                            for ( int si=0; si<pVSecond->GetVertexCount(); ++si )
+							{
+								switch ( format )
+								{
+								case MBF_INT8:
+								case MBF_UINT8:
+								{
+                                    uint8_t* pD = reinterpret_cast<uint8_t*>
+											( &result.m_data[resultOffset] );
+
+                                    int comp=0;
+                                    for ( ; comp<components; ++comp )
+									{
+                                        uint8_t bone = pD[comp];
+
+                                        // be defensive
+                                        if (bone<secondToFirstBones.Num())
+                                        {
+                                            pD[comp] = (uint8_t)secondToFirstBones[ bone ];
+                                        }
+                                        else
+                                        {
+                                            pD[comp] = 0;
+                                        }
+									}
+
+									resultOffset += elemSize;
+									break;
+								}
+
+								case MBF_INT16:
+								case MBF_UINT16:
+								{
+                                    uint16_t* pD = reinterpret_cast<uint16_t*>
+											( &result.m_data[resultOffset] );
+
+                                    int comp=0;
+                                    for ( ; comp<components; ++comp )
+									{
+                                        uint16_t bone = pD[comp];
+
+                                        // be defensive
+                                        if (bone<secondToFirstBones.Num())
+                                        {
+                                            pD[comp] = (uint16_t)secondToFirstBones[ bone ];
+                                        }
+                                        else
+                                        {
+                                            pD[comp] = 0;
+                                        }
+                                    }
+
+									resultOffset += elemSize;
+									break;
+								}
+
+								case MBF_INT32:
+								case MBF_UINT32:
+								{
+                                    uint32_t* pD = reinterpret_cast<uint32_t*>
+											( &result.m_data[resultOffset] );
+
+                                    for ( int comp=0; comp<components; ++comp )
+									{
+                                        uint32_t bone = pD[comp];
+
+                                        // be defensive
+                                        if (bone< (uint32_t)secondToFirstBones.Num())
+                                        {
+                                            pD[comp] = (uint32_t)secondToFirstBones[ bone ];
+                                        }
+                                        else
+                                        {
+                                            pD[comp] = 0;
+                                        }
+                                    }
+
+									resultOffset += elemSize;
+									break;
+								}
+
+								default:
+									checkf(false, TEXT("Format not supported."));
+								}
+							}
+						}
+					}
+				}
+			}
+		}
+
+		// Tags
+		pResult->m_tags = pFirst->m_tags;
+
+		for (auto& secondTag : pSecond->m_tags)
+		{
+			bool repeated = false;
+
+			for (auto& firstTag : pFirst->m_tags)
+			{
+				if (firstTag == secondTag)
+				{
+					repeated = true;
+					break;
+				}				
+			}
+
+			if (!repeated)
+			{
+				pResult->m_tags.Add(secondTag);
+			}
+		}
+
+		return pResult;
+	}
+
+
+    //---------------------------------------------------------------------------------------------
+    //!
+    //---------------------------------------------------------------------------------------------
+    inline void ExtendSkeleton( Skeleton* pBase, const Skeleton* pOther )
+    {
+        map<int,int> otherToResult;
+
+        int initialBones = pBase->GetBoneCount();
+        for ( int b=0; pOther && b<pOther->GetBoneCount(); ++b)
+        {
+            int resultBoneIndex = pBase->FindBone( pOther->GetBoneName(b) );
+            if ( resultBoneIndex<0 )
+            {
+                int32 newIndex = pBase->m_bones.Num();
+                otherToResult[b] = int(newIndex);
+                pBase->m_bones.Add( pOther->m_bones[b] );
+
+                // Will be remapped below
+                pBase->m_boneParents.Add(pOther->m_boneParents[b] );
+
+                // Copy the external bone id
+				pBase->m_boneIds.Add(pOther->m_boneIds[b]);
+            }
+            else
+            {
+                otherToResult[b] = resultBoneIndex;
+            }
+        }
+
+        // Fix bone parent indices of the bones added from pOther
+        for ( int b=initialBones;b<pBase->GetBoneCount(); ++b)
+        {
+            int16_t sourceIndex = pBase->m_boneParents[b];
+            if (sourceIndex>=0)
+            {
+                pBase->m_boneParents[b] = (int16_t)otherToResult[ sourceIndex ];
+            }
+        }
+    }
+
+
+    //---------------------------------------------------------------------------------------------
+    //! Return null if there is no need to remap the mesh
+    //---------------------------------------------------------------------------------------------
+    inline MeshPtr MeshRemapSkeleton( const Mesh* pSource, const Skeleton* pSkeleton )
+    {
+        if ( pSource->GetVertexCount()==0
+             ||
+             !pSource->GetSkeleton()
+             ||
+             pSource->GetSkeleton()->GetBoneCount()==0 )
+        {
+            return nullptr;
+        }
+
+        auto pSourceSkeleton = pSource->GetSkeleton();
+
+        MeshPtr pResult = pSource->Clone();
+        pResult->SetSkeleton( pSkeleton );
+
+        // Build a skeleton map
+        vector<int> sourceToSkeleton;
+        for ( const auto& boneName: pSourceSkeleton->m_bones )
+        {
+            sourceToSkeleton.push_back( pResult->GetSkeleton()->FindBone(boneName.c_str()) );
+        }
+
+        bool somethingRemapped = false;
+
+        // Convert all bone index channels
+        FMeshBufferSet& BufSet = pResult->m_VertexBuffers;
+        for (auto& buf: BufSet.m_buffers)
+        {
+            for (auto& chan: buf.m_channels)
+            {
+                if (chan.m_semantic==MBS_BONEINDICES)
+                {
+                    uint8_t* pData = &buf.m_data[chan.m_offset];
+
+                    switch ( chan.m_format )
+                    {
+                    case MBF_INT8:
+                    case MBF_UINT8:
+                    {
+                        for (uint32_t v=0;v<BufSet.m_elementCount; ++v)
+                        {
+                            uint8_t* pTypedData = (uint8_t*)pData;
+
+                            for (int c=0; c<chan.m_componentCount; ++c)
+                            {
+                                size_t boneIndex = size_t( pTypedData[c] );
+                                //check( boneIndex < sourceToSkeleton.Num() );
+                                if ( boneIndex < sourceToSkeleton.size() )
+                                {
+                                    int bone = sourceToSkeleton[pTypedData[c]];
+                                    if ( pTypedData[c] != uint8_t( bone ) )
+                                    {
+                                        pTypedData[c] = uint8_t( bone );
+                                        somethingRemapped = true;
+                                    }
+                                }
+                                else
+                                {
+                                    UE_LOG(LogMutableCore, Error, TEXT("Bone index out of range.") );
+                                }
+                            }
+
+                            pData += buf.m_elementSize;
+                        }
+                        break;
+                    }
+
+                    case MBF_INT16:
+                    case MBF_UINT16:
+                    {
+                        for (uint32_t v=0;v<BufSet.m_elementCount; ++v)
+                        {
+                            uint16_t* pTypedData = (uint16_t*)pData;
+
+                            for (int c=0; c<chan.m_componentCount; ++c)
+                            {
+                                size_t boneIndex = size_t(pTypedData[c]);
+                                //check( boneIndex < sourceToSkeleton.size() );
+                                if ( boneIndex < sourceToSkeleton.size() )
+                                {
+                                    int bone = sourceToSkeleton[boneIndex];
+                                    if ( pTypedData[c] != uint16_t(bone) )
+                                    {
+                                        pTypedData[c] = uint16_t(bone);
+                                        somethingRemapped = true;
+                                    }
+                                }
+                                else
+                                {
+                                    UE_LOG(LogMutableCore, Error, TEXT("Bone index out of range.") );
+                                }
+                            }
+
+                            pData += buf.m_elementSize;
+                        }
+                        break;
+                    }
+
+                    case MBF_INT32:
+                    case MBF_UINT32: 
+                    {
+                        for (uint32_t v=0;v<BufSet.m_elementCount; ++v)
+                        {
+                            uint32_t* pTypedData = (uint32_t*)pData;
+
+                            for (int c=0; c<chan.m_componentCount; ++c)
+                            {
+                                size_t boneIndex = size_t( pTypedData[c] );
+                                //check( boneIndex < sourceToSkeleton.size() );
+                                if ( boneIndex < sourceToSkeleton.size() )
+                                {
+                                    int bone = sourceToSkeleton[pTypedData[c]];
+                                    if ( pTypedData[c] != uint32_t( bone ) )
+                                    {
+                                        pTypedData[c] = uint32_t( bone );
+                                        somethingRemapped = true;
+                                    }
+                                }
+                                else
+                                {
+                                    UE_LOG(LogMutableCore, Error, TEXT("Bone index out of range.") );
+                                }
+                            }
+
+                            pData += buf.m_elementSize;
+                        }
+                        break;
+                    }
+
+                    default:
+						checkf(false, TEXT("Not implemented."));
+						break;
+                    }
+                }
+            }
+        }
+
+        if (!somethingRemapped)
+        {
+            pResult = nullptr;
+        }
+
+        return pResult;
+    }
+
+}
diff --git a/Engine/Plugins/Experimental/Mutable/Source/MutableRuntime/Private/OpMeshMorph.h b/Engine/Plugins/Experimental/Mutable/Source/MutableRuntime/Private/OpMeshMorph.h
new file mode 100644
--- /dev/null
+++ b/Engine/Plugins/Experimental/Mutable/Source/MutableRuntime/Private/OpMeshMorph.h
@@ -0,0 +1,472 @@
+// Copyright Epic Games, Inc. All Rights Reserved.
+
+#pragma once
+
+#include "MeshPrivate.h"
+#include "ConvertData.h"
+#include "Platform.h"
+
+#include "SparseIndexMap.h"
+
+namespace mu
+{
+
+    //---------------------------------------------------------------------------------------------
+    //! Reference Linear factor version
+    //---------------------------------------------------------------------------------------------
+    inline MeshPtr MeshMorphReference( const Mesh* pBase, const Mesh* pMorph, float factor )
+    {
+        MUTABLE_CPUPROFILER_SCOPE(MeshMorph);
+
+        if (!pBase) return nullptr;
+
+        MeshPtr pDest = pBase->Clone();
+
+        if (!pMorph) return pDest;
+
+        // Number of vertices to modify
+        uint32_t vcountMorph = pMorph->GetVertexBuffers().GetElementCount();
+        uint32_t vcountBase = pBase->GetVertexBuffers().GetElementCount();
+        if ( !vcountMorph || !vcountBase )
+        {
+            return pDest;
+        }
+
+        uint32_t ccount = 0;
+        if ( pMorph->GetVertexBuffers().GetBufferCount()>0 )
+        {
+            ccount = pMorph->GetVertexBuffers().GetBufferChannelCount(0);
+        }
+
+        // Iterator of the vertex ids of the base vertices
+        UntypedMeshBufferIteratorConst itBaseId( pBase->GetVertexBuffers(), MBS_VERTEXINDEX );
+
+        vector< UntypedMeshBufferIterator > itBaseChannels( ccount );
+        vector< UntypedMeshBufferIteratorConst > itMorphChannels( ccount );
+        for ( size_t c=1; c<ccount; ++c )
+        {
+            const FMeshBufferSet& MBSPriv = pMorph->GetVertexBuffers();
+            MESH_BUFFER_SEMANTIC sem = MBSPriv.m_buffers[0].m_channels[c].m_semantic;
+            int semIndex = MBSPriv.m_buffers[0].m_channels[c].m_semanticIndex;
+
+            itBaseChannels[c] = UntypedMeshBufferIterator
+                    ( pDest->GetVertexBuffers(), sem, semIndex );
+            itMorphChannels[c] = UntypedMeshBufferIteratorConst
+                    ( pMorph->GetVertexBuffers(), sem, semIndex );
+        }
+
+        // Morph
+
+        // Number of vertices consumed from the morph mesh
+        uint32_t processedMorphVertices = 0;
+
+        // Morph mesh channels are always this friendly types
+        MeshBufferIteratorConst<MBF_UINT32,uint32_t,1>
+                itFirstMorphVertex( pMorph->GetVertexBuffers(), MBS_VERTEXINDEX );
+
+        // Number of vertices at the beginning of the morph already consumed (optimisation)
+        int morphVerticesFromStartAlreadyConsumed = 0;
+
+        // Number of vertices we need to advance the base mesh iterators to get to the current base
+        // mesh vertex
+        int step = 0;
+
+        // Iterate all the base mesh vertices
+        for ( uint32_t vi=0; vi<vcountBase && processedMorphVertices<vcountMorph; ++vi )
+        {
+            uint32_t baseMeshVertexIndex = itBaseId.GetAsUINT32();
+            int32_t morphVertexIndex = -1;
+
+            // Look for a morph vertex to apply to this base mesh vertex
+            auto itMorphVer = itFirstMorphVertex+morphVerticesFromStartAlreadyConsumed;
+            for ( int32_t morphVertexCandidate=morphVerticesFromStartAlreadyConsumed;
+                  morphVertexCandidate<int32_t(vcountMorph);
+                  ++morphVertexCandidate )
+            {
+                if ((*itMorphVer)[0]==baseMeshVertexIndex)
+                {
+                    morphVertexIndex = morphVertexCandidate;
+                    if (morphVertexCandidate==morphVerticesFromStartAlreadyConsumed)
+                    {
+                        ++morphVerticesFromStartAlreadyConsumed;
+                    }
+                    break;
+                }
+                ++itMorphVer;
+            }
+
+            if( morphVertexIndex>=0 )
+            {
+                // Morph one vertex
+                for ( size_t c=1; c<ccount; ++c )
+                {
+                    if ( itBaseChannels[c].ptr() )
+                    {
+                        itBaseChannels[c] += step;
+
+                        vec4<float> value = itBaseChannels[c].GetAsVec4f();
+                        vec4<float> delta = (itMorphChannels[c]+morphVertexIndex).GetAsVec4f();
+
+                        value += delta * factor;
+
+                        for( int i=0; i<itBaseChannels[c].GetComponents(); ++i )
+                        {
+                            ConvertData( i,
+                                         itBaseChannels[c].ptr(),
+                                         itBaseChannels[c].GetFormat(),
+                                         &value, MBF_FLOAT32 );
+                        }
+                    }
+                }
+
+                // Reset the step to the next vertex
+                step = 1;
+            }
+            else
+            {
+                // Vertex was not morphed, the step to the next vertex will be bigger
+                ++step;
+            }
+
+            ++itBaseId;
+        }
+
+        return pDest;
+    }
+
+
+
+    //---------------------------------------------------------------------------------------------
+    //! Reference Factor-less version
+    //---------------------------------------------------------------------------------------------
+    inline MeshPtr MeshMorphReference( const Mesh* pBase, const Mesh* pMorph )
+    {
+        // Trust the compiler to remove the factor
+        return MeshMorphReference( pBase, pMorph, 1.0f );
+    }
+
+
+    //---------------------------------------------------------------------------------------------
+    //! Reference Linear factor version for morphing 2 targets
+    //---------------------------------------------------------------------------------------------
+    inline MeshPtr MeshMorph2Reference( const Mesh* pBase, const Mesh* pMin, const Mesh* pMax, float factor )
+    {
+        MUTABLE_CPUPROFILER_SCOPE(MeshMorph2Reference);
+
+        MeshPtr pDest = pBase->Clone();
+
+        // Number of vertices to modify
+        uint32_t vmincount = pMin->GetVertexBuffers().GetElementCount();
+        uint32_t vmaxcount = pMax->GetVertexBuffers().GetElementCount();
+
+        uint32_t vcountBase = pBase->GetVertexBuffers().GetElementCount();
+
+        if ( (vmincount+vmaxcount)==0 || !vcountBase )
+        {
+            return pDest;
+        }
+
+
+        auto refTarget = vmincount > 0 ? pMin : pMax;
+
+        uint32_t ccount = refTarget->GetVertexBuffers().GetBufferChannelCount(0);
+
+        // Iterator of the vertex ids of the base vertices
+        UntypedMeshBufferIteratorConst itBaseId( pBase->GetVertexBuffers(), MBS_VERTEXINDEX );
+
+        vector< UntypedMeshBufferIterator > itBaseChannels( ccount );
+        vector< UntypedMeshBufferIteratorConst > itMinChannels( ccount );
+        vector< UntypedMeshBufferIteratorConst > itMaxChannels( ccount );
+        for ( size_t c=1; c<ccount; ++c )
+        {
+            const FMeshBufferSet& MBSPriv = refTarget->GetVertexBuffers();
+            MESH_BUFFER_SEMANTIC sem = MBSPriv.m_buffers[0].m_channels[c].m_semantic;
+            int semIndex = MBSPriv.m_buffers[0].m_channels[c].m_semanticIndex;
+
+            itBaseChannels[c] = UntypedMeshBufferIterator
+                    ( pDest->GetVertexBuffers(), sem, semIndex );
+            itMinChannels[c] = UntypedMeshBufferIteratorConst
+                    ( pMin->GetVertexBuffers(), sem, semIndex );
+            itMaxChannels[c] = UntypedMeshBufferIteratorConst
+                    ( pMax->GetVertexBuffers(), sem, semIndex );
+        }
+
+        // Morph
+        uint32_t minProcessed=0;
+        uint32_t maxProcessed=0;
+        float maxFactor = factor;
+        float minFactor = 1.0f-factor;
+
+        // Morph mesh channels are always this friendly types
+        MeshBufferIteratorConst<MBF_UINT32,uint32_t,1>
+                itFirstMinVer( pMin->GetVertexBuffers(), MBS_VERTEXINDEX );
+        MeshBufferIteratorConst<MBF_UINT32,uint32_t,1>
+                itFirstMaxVer( pMax->GetVertexBuffers(), MBS_VERTEXINDEX );
+
+        // Number of vertices at the beginning of the morph already consumed (optimisation)
+        int minMorphVerticesFromStartAlreadyConsumed = 0;
+        int maxMorphVerticesFromStartAlreadyConsumed = 0;
+
+        // Number of vertices we need to advance the base mesh iterators to get to the current base
+        // mesh vertex
+        int step = 0;
+
+        // Iterate all the base mesh vertices
+        for ( uint32_t i=0; i<vcountBase && (minProcessed<vmincount || maxProcessed<vmaxcount); ++i)
+        {
+            uint32_t baseMeshVertexIndex = itBaseId.GetAsUINT32();
+
+            int32_t minMorphVertexIndex = -1;
+            int32_t maxMorphVertexIndex = -1;
+
+            // Look for a min morph vertex to apply to this base mesh vertex
+            auto itMorphVer = itFirstMinVer+minMorphVerticesFromStartAlreadyConsumed;
+            for ( int32_t morphVertexCandidate=minMorphVerticesFromStartAlreadyConsumed;
+                  morphVertexCandidate<int32_t(vmincount);
+                  ++morphVertexCandidate )
+            {
+                if ((*itMorphVer)[0]==baseMeshVertexIndex)
+                {
+                    minMorphVertexIndex = morphVertexCandidate;
+                    if (morphVertexCandidate==minMorphVerticesFromStartAlreadyConsumed)
+                    {
+                        ++minMorphVerticesFromStartAlreadyConsumed;
+                    }
+                    break;
+                }
+                ++itMorphVer;
+            }
+
+            // Look for a max morph vertex to apply to this base mesh vertex
+            itMorphVer = itFirstMaxVer+maxMorphVerticesFromStartAlreadyConsumed;
+            for ( int32_t morphVertexCandidate=maxMorphVerticesFromStartAlreadyConsumed;
+                  morphVertexCandidate<int32_t(vmaxcount);
+                  ++morphVertexCandidate )
+            {
+                if ((*itMorphVer)[0]==baseMeshVertexIndex)
+                {
+                    maxMorphVertexIndex = morphVertexCandidate;
+                    if (morphVertexCandidate==maxMorphVerticesFromStartAlreadyConsumed)
+                    {
+                        ++maxMorphVerticesFromStartAlreadyConsumed;
+                    }
+                    break;
+                }
+                ++itMorphVer;
+            }
+            
+            // Morph one vertex
+            if ( minMorphVertexIndex>=0 || maxMorphVertexIndex>=0 )
+            {
+                for ( size_t c=1; c<ccount; ++c )
+                {
+                    if ( itBaseChannels[c].ptr() )
+                    {
+                        itBaseChannels[c] += step;
+
+                        vec4<float> value = itBaseChannels[c].GetAsVec4f();
+
+                        if ( minMorphVertexIndex>=0 )
+                        {
+                            vec4<float> minvalue = (itMinChannels[c]+minMorphVertexIndex).GetAsVec4f();
+                            value += minvalue * minFactor;
+                        }
+
+                        if ( maxMorphVertexIndex>=0 )
+                        {
+                            vec4<float> maxvalue = (itMaxChannels[c]+maxMorphVertexIndex).GetAsVec4f();
+                            value += maxvalue * maxFactor;
+                        }
+
+						const int32_t componentCount = itBaseChannels[c].GetComponents();
+                        for( int component=0; component<componentCount; ++component )
+                        {
+                            ConvertData( component,
+                                         itBaseChannels[c].ptr(),
+                                         itBaseChannels[c].GetFormat(),
+                                         &value, MBF_FLOAT32 );
+                        }
+                    }
+                }
+
+                step = 1;
+            }
+            else
+            {
+                // Vertex was not morphed, the step to the next vertex will be bigger
+                ++step;
+            }
+
+            ++itBaseId;
+        }
+
+        return pDest;
+    }
+
+	
+	//---------------------------------------------------------------------------------------------
+	//! Optimized linear factor version for morphing 2 targets
+	//---------------------------------------------------------------------------------------------
+    inline MeshPtr MeshMorph2( const Mesh* pBase, const Mesh* pMin, const Mesh* pMax, const float factor )
+    {
+        MUTABLE_CPUPROFILER_SCOPE(MeshMorph2);
+
+		if (!pBase) return nullptr;
+
+        const auto ApplyMorph = []
+            ( auto BaseIdIter, const std::vector< UntypedMeshBufferIterator >& baseChannelsIters, const int32_t baseSize,
+              auto MorphIdIter, const std::vector< UntypedMeshBufferIteratorConst >& morphChannelsIters, const int32_t morphSize,
+              const float factor )
+        -> void
+        {
+            uint32_t minBaseId = std::numeric_limits<uint32_t>::max();
+            uint32_t maxBaseId = 0;
+
+			{
+                auto limitsBaseIdIter = BaseIdIter;
+				for ( int32_t i = 0; i < baseSize; ++i, ++limitsBaseIdIter )
+				{
+					const uint32_t id = limitsBaseIdIter.GetAsUINT32();
+					minBaseId = std::min( id, minBaseId );
+					maxBaseId = std::max( id, maxBaseId );
+				}
+			}
+
+            SparseIndexMap IndexMap( minBaseId, maxBaseId );
+            
+            auto usedMorphIdIter = MorphIdIter;
+            for ( uint32_t i = 0; i < static_cast<uint32_t>(morphSize); ++i, ++usedMorphIdIter)
+            {
+                const uint32_t morphId = (*usedMorphIdIter)[0];
+                
+                IndexMap.Insert( morphId, i );
+            }
+
+            for ( int32_t v = 0; v < baseSize; ++v )
+            {
+                const uint32_t baseId = (BaseIdIter + v).GetAsUINT32();
+                const uint32_t morphIdx = IndexMap.Find( baseId );
+
+                if ( morphIdx == SparseIndexMap::NotFoundValue )
+                {
+                    continue;
+                }
+
+                const int32_t m = static_cast<int32_t>( morphIdx );
+
+                // Find consecutive run.
+                auto RunBaseIter = BaseIdIter + v;
+                auto RunMorphIter = MorphIdIter + m;
+
+                int32_t runSize = 0;
+                for (  ; v + runSize < baseSize && m + runSize < morphSize && RunBaseIter.GetAsUINT32() == (*RunMorphIter)[0];
+                        ++runSize, ++RunBaseIter, ++RunMorphIter );
+
+                const size_t channelCount = morphChannelsIters.size();
+                for ( size_t c = 1; c < channelCount; ++c )
+                {
+                
+                	if ( !(baseChannelsIters[c].ptr() && morphChannelsIters[c].ptr()) )
+                	{
+                		continue;
+                	}
+                
+                    auto channelBaseIter = baseChannelsIters[c] + v;
+                    auto channelMorphIter = morphChannelsIters[c] + m;
+                   
+                    const auto dstChannelFormat = baseChannelsIters[c].GetFormat();
+                    const auto dstChannelComps = baseChannelsIters[c].GetComponents();
+
+                    // Apply Morph to range found above.
+                    for ( int32_t r = 0; r < runSize; ++r, ++channelBaseIter, ++channelMorphIter )
+                    {
+                        const vec4<float> value = channelBaseIter.GetAsVec4f() + channelMorphIter.GetAsVec4f()*factor;
+
+                        // TODO: Optimize this for the specific components.
+                        // Max 4 components
+                        for ( int32_t comp = 0; comp < dstChannelComps && comp < 4; ++comp )
+                        {
+                            ConvertData( comp, channelBaseIter.ptr(), dstChannelFormat, &value, MBF_FLOAT32 );
+                        }
+                    }
+                }
+
+				v += std::max(runSize - 1, 0);
+            }
+        };
+
+		MeshPtr pDest = pBase->Clone();
+
+		// Number of vertices to modify
+		const uint32_t minCount = pMin ? pMin->GetVertexBuffers().GetElementCount() : 0;
+		const uint32_t maxCount = pMax ? pMax->GetVertexBuffers().GetElementCount() : 0;
+		const uint32_t baseCount = pBase ? pBase->GetVertexBuffers().GetElementCount() : 0;
+		const Mesh* refTarget = minCount > 0 ? pMin : pMax;
+
+		if ( baseCount == 0 || (minCount + maxCount) == 0)
+		{
+			return pDest;
+		}
+
+		if (refTarget)
+		{
+			uint32_t ccount = refTarget->GetVertexBuffers().GetBufferChannelCount(0);
+
+			// Iterator of the vertex ids of the base vertices
+			UntypedMeshBufferIteratorConst itBaseId(pBase->GetVertexBuffers(), MBS_VERTEXINDEX);
+
+			std::vector< UntypedMeshBufferIterator > itBaseChannels(ccount);
+			std::vector< UntypedMeshBufferIteratorConst > itMinChannels(ccount);
+			std::vector< UntypedMeshBufferIteratorConst > itMaxChannels(ccount);
+
+			for (size_t c = 1; c < ccount; ++c)
+			{
+				const FMeshBufferSet& MBSPriv = refTarget->GetVertexBuffers();
+				MESH_BUFFER_SEMANTIC sem = MBSPriv.m_buffers[0].m_channels[c].m_semantic;
+				int semIndex = MBSPriv.m_buffers[0].m_channels[c].m_semanticIndex;
+				
+				itBaseChannels[c] = UntypedMeshBufferIterator(pDest->GetVertexBuffers(), sem, semIndex);
+				if (minCount > 0)
+				{
+					itMinChannels[c] = UntypedMeshBufferIteratorConst(pMin->GetVertexBuffers(), sem, semIndex);
+				}
+
+				if (maxCount > 0)
+				{
+					itMaxChannels[c] = UntypedMeshBufferIteratorConst(pMax->GetVertexBuffers(), sem, semIndex);
+				}
+			}
+			
+			if (minCount > 0)
+			{
+				MeshBufferIteratorConst<MBF_UINT32, uint32_t, 1> itMinId(pMin->GetVertexBuffers(), MBS_VERTEXINDEX);
+				ApplyMorph(itBaseId, itBaseChannels, baseCount, itMinId, itMinChannels, minCount, 1.0f - factor);
+			}
+
+			if (maxCount > 0)
+			{
+				MeshBufferIteratorConst<MBF_UINT32, uint32_t, 1> itMaxId(pMax->GetVertexBuffers(), MBS_VERTEXINDEX);
+				ApplyMorph(itBaseId, itBaseChannels, baseCount, itMaxId, itMaxChannels, maxCount, factor);
+			}
+		}
+
+        return pDest;
+    }
+
+	//---------------------------------------------------------------------------------------------
+	//! \TODO Optimized linear factor version
+	//---------------------------------------------------------------------------------------------
+	inline MeshPtr MeshMorph(const Mesh* pBase, const Mesh* pMorph, float factor)
+	{
+		return MeshMorph2( pBase, nullptr, pMorph, factor );
+	}
+
+    //---------------------------------------------------------------------------------------------
+    //! \TODO Optimized Factor-less version
+    //---------------------------------------------------------------------------------------------
+    inline MeshPtr MeshMorph( const Mesh* pBase, const Mesh* pMorph )
+    {
+        // Trust the compiler to remove the factor
+        return MeshMorph( pBase, pMorph, 1.0f );
+    }
+}
