[mutable] Move Mutable plugin from the restricted folder to the experimental plugins folder.

diff --git a/Engine/Plugins/Experimental/Mutable/Source/MutableRuntime/Private/OpMeshRemapIndices.h b/Engine/Plugins/Experimental/Mutable/Source/MutableRuntime/Private/OpMeshRemapIndices.h
new file mode 100644
--- /dev/null
+++ b/Engine/Plugins/Experimental/Mutable/Source/MutableRuntime/Private/OpMeshRemapIndices.h
@@ -0,0 +1,86 @@
+// Copyright Epic Games, Inc. All Rights Reserved.
+
+#pragma once
+
+#include "MeshPrivate.h"
+#include "ConvertData.h"
+#include "Platform.h"
+
+
+namespace mu
+{
+
+    //---------------------------------------------------------------------------------------------
+    //!
+    //---------------------------------------------------------------------------------------------
+    inline MeshPtr MeshRemapIndices( const Mesh* pBase, const Mesh* pReference )
+    {
+        MUTABLE_CPUPROFILER_SCOPE(MeshRemapIndices);
+
+        if (!pBase) return nullptr;
+
+        MeshPtr pDest = pBase->Clone();
+
+        if (!pReference) return pDest;
+
+        // Number of vertices to modify
+        uint32_t vcountReference = pReference->GetVertexBuffers().GetElementCount();
+        uint32_t vcountBase = pBase->GetVertexBuffers().GetElementCount();
+        if ( !vcountReference || !vcountBase )
+        {
+            return pDest;
+        }
+
+        // Iterator of the vertex ids of the base vertices
+        UntypedMeshBufferIteratorConst itReferenceId( pReference->GetVertexBuffers(), MBS_VERTEXINDEX );
+
+        FMeshBufferSet& verts = pDest->m_VertexBuffers;
+        for ( auto& buf: verts.m_buffers )
+        {
+            for ( auto& chan: buf.m_channels )
+            {
+                if (chan.m_semantic!=MBS_VERTEXINDEX)
+                    continue;
+                
+                auto pData = buf.m_data.GetData() + chan.m_offset;
+
+                switch( chan.m_format )
+                {
+                case MBF_INT32:
+                case MBF_UINT32:
+                    for ( uint32_t v = 0; v < vcountBase;  ++v)
+                    {
+                        uint32_t* pIndex = (uint32_t*)pData;
+                        uint32_t oldIndex = *pIndex;
+                        check( oldIndex < vcountReference );
+                        uint32_t newIndex = (itReferenceId+oldIndex).GetAsUINT32();
+                        *pIndex = newIndex;
+                        pData += buf.m_elementSize;
+                    }
+                    break;
+
+                case MBF_INT16:
+                case MBF_UINT16:
+                    for ( uint32_t v = 0; v < vcountBase;  ++v)
+                    {
+                        uint16_t* pIndex = (uint16_t*)pData;
+                        uint16_t oldIndex = *pIndex;
+                        check( *pIndex < vcountReference );
+                        uint32_t newIndex = ( itReferenceId + oldIndex ).GetAsUINT32();
+                        *pIndex = uint16_t(newIndex);
+                        pData += buf.m_elementSize;
+                    }
+                    break;
+
+                default:
+                    // Not implemented
+                    check(false);
+                    break;
+                }
+            }
+        }
+
+        return pDest;
+    }
+
+}
diff --git a/Engine/Plugins/Experimental/Mutable/Source/MutableRuntime/Private/OpMeshRemoveChart.h b/Engine/Plugins/Experimental/Mutable/Source/MutableRuntime/Private/OpMeshRemoveChart.h
new file mode 100644
--- /dev/null
+++ b/Engine/Plugins/Experimental/Mutable/Source/MutableRuntime/Private/OpMeshRemoveChart.h
@@ -0,0 +1,280 @@
+// Copyright Epic Games, Inc. All Rights Reserved.
+
+#pragma once
+
+#include "MeshPrivate.h"
+#include "Platform.h"
+
+
+namespace mu
+{
+
+    //---------------------------------------------------------------------------------------------
+    struct ID_INTERVAL
+    {
+        int idStart;
+        int idPosition;
+        int size;
+    };
+
+
+    inline void ExtractVertexIndexIntervals( vector<ID_INTERVAL>& intervals, const Mesh* pSource )
+    {
+        MeshBufferIteratorConst<MBF_UINT32,uint32_t,1> itVI( pSource->GetVertexBuffers(), MBS_VERTEXINDEX );
+        ID_INTERVAL current;
+		current.idStart = -1;
+		current.idPosition = 0;
+		current.size = 0;
+		for ( int sv=0; sv<pSource->GetVertexBuffers().GetElementCount(); ++sv )
+        {
+            int32_t id = (*itVI)[0];
+            ++itVI;
+
+            if (current.idStart<0)
+            {
+                current.idStart = id;
+                current.idPosition = sv;
+                current.size = 1;
+            }
+            else
+            {
+                if (id==current.idStart+current.size)
+                {
+                    ++current.size;
+                }
+                else
+                {
+                    intervals.push_back(current);
+                    current.idStart = id;
+                    current.idPosition = sv;
+                    current.size = 1;
+                }
+            }
+        }
+
+        if (current.idStart>=0)
+        {
+            intervals.push_back(current);
+        }
+    }
+
+
+    inline int FindPositionInIntervals( const vector<ID_INTERVAL>& intervals, int id )
+    {
+        for( const auto& interval: intervals )
+        {
+            int deltaId = id - interval.idStart;
+            if (deltaId>=0 && deltaId<interval.size)
+            {
+                return interval.idPosition+deltaId;
+            }
+        }
+        return -1;
+    }
+
+
+    //---------------------------------------------------------------------------------------------
+    inline MeshPtr MeshRemoveMask( const Mesh* pSource, const Mesh* pMask )
+    {
+        MUTABLE_CPUPROFILER_SCOPE(MeshRemoveMask);
+
+        MeshPtr pResult = pSource->Clone();
+
+        if (!pMask->GetVertexCount() || !pResult->GetVertexCount() ||
+			!pResult->GetIndexCount() )
+        {
+            return pResult;
+        }
+
+        // TODO
+        pResult->SetFaceGroupCount( 0 );
+
+        MeshBufferIteratorConst<MBF_UINT32,uint32_t,1> itMaskVI( pMask->GetVertexBuffers(), MBS_VERTEXINDEX );
+
+        int firstFreeVertex = 0;
+
+        // For each source vertex, true if it is removed.
+        int resultVertexCount = pResult->GetVertexCount();
+        uint8_t* removedVertices = (uint8_t*)mutable_malloc_aligned( resultVertexCount, 16 );
+		FMemory::Memzero(removedVertices,resultVertexCount);
+        {
+            vector<ID_INTERVAL> intervals;
+            ExtractVertexIndexIntervals( intervals, pResult.get() );
+
+            for ( int mv=0; mv<pMask->GetVertexBuffers().GetElementCount(); ++mv )
+            {
+                uint32_t maskVertexId = (*itMaskVI)[0];
+                ++itMaskVI;
+
+                int indexInSource = FindPositionInIntervals(intervals, maskVertexId);
+                if (indexInSource>=0)
+                {
+                    removedVertices[indexInSource] = 1;
+                }
+            }
+        }
+
+        // Rebuild index buffers
+
+        // Map from source vertex index, to new vertex index for used vertices.
+        // These are indices as in the index buffer, not the absoulte vertex index as in the
+        // vertexbuffer MBS_VERTEXINDEX buffers.
+        vector<int> usedVertices( pResult->GetVertexCount(), -1 );
+        {
+            size_t removedIndices = 0;
+
+            if ( pResult->GetIndexBuffers().GetElementSize(0)==4 )
+            {
+                MeshBufferIteratorConst<MBF_UINT32,uint32_t,1> itSource( pResult->GetIndexBuffers(), MBS_VERTEXINDEX );
+                MeshBufferIterator<MBF_UINT32,uint32_t,1> itDest( pResult->GetIndexBuffers(), MBS_VERTEXINDEX );
+
+                int indexCount = pResult->GetIndexCount();
+                for ( int f=0; f<indexCount/3; ++f )
+                {
+                    uint32_t sourceIndices[3];
+                    sourceIndices[0] = itSource[0][0];
+                    sourceIndices[1] = itSource[1][0];
+                    sourceIndices[2] = itSource[2][0];
+
+                    bool faceRemoved =  ( removedVertices[ sourceIndices[0] ]
+                            + removedVertices[ sourceIndices[1] ]
+                            + removedVertices[ sourceIndices[2] ]
+                            ) > 0;
+
+                    if ( !faceRemoved )
+                    {
+                        for (int i=0;i<3;++i)
+                        {
+                            uint32_t sourceIndex = sourceIndices[i];
+
+                            if ( usedVertices[ sourceIndex ] < 0 )
+                            {
+                                usedVertices[ sourceIndex ] = firstFreeVertex;
+                                firstFreeVertex++;
+                            }
+
+                            uint32_t destIndex = usedVertices[ sourceIndex ];
+                            *(uint32_t*)itDest.ptr() = destIndex;
+
+                            itDest++;
+                        }
+                    }
+
+                    itSource+=3;
+                }
+
+                removedIndices = itSource - itDest;
+            }
+
+            else if ( pResult->GetIndexBuffers().GetElementSize(0)==2 )
+            {
+                MeshBufferIteratorConst<MBF_UINT16,uint16_t,1> itSource( pResult->GetIndexBuffers(), MBS_VERTEXINDEX );
+                MeshBufferIterator<MBF_UINT16,uint16_t,1> itDest( pResult->GetIndexBuffers(), MBS_VERTEXINDEX );
+
+                int indexCount = pResult->GetIndexCount();
+                for ( int f=0; f<indexCount/3; ++f )
+                {
+                    uint16_t sourceIndices[3];
+                    sourceIndices[0] = itSource[0][0];
+                    sourceIndices[1] = itSource[1][0];
+                    sourceIndices[2] = itSource[2][0];
+
+                    bool faceRemoved =  ( removedVertices[ sourceIndices[0] ]
+                            + removedVertices[ sourceIndices[1] ]
+                            + removedVertices[ sourceIndices[2] ]
+                            ) > 0;
+
+                    if ( !faceRemoved )
+                    {
+                        for (int i=0;i<3;++i)
+                        {
+                            uint16_t sourceIndex = sourceIndices[i];
+
+                            if ( usedVertices[ sourceIndex ] < 0 )
+                            {
+                                usedVertices[ sourceIndex ] = firstFreeVertex;
+                                firstFreeVertex++;
+                            }
+
+                            uint16_t destIndex = (uint16_t)usedVertices[ sourceIndex ];
+                            *(uint16_t*)itDest.ptr() = destIndex;
+
+                            itDest++;
+                        }
+                    }
+
+                    itSource+=3;
+                }
+
+                removedIndices = itSource - itDest;
+            }
+
+            else
+            {
+                // Case not implemented
+                check( false );
+            }
+
+            check( removedIndices%3==0 );
+
+            int faceCount = pResult->GetFaceCount();
+            pResult->GetFaceBuffers().SetElementCount( faceCount-(int)removedIndices/3 );
+            pResult->GetIndexBuffers().SetElementCount( faceCount*3-(int)removedIndices );
+        }
+
+        mutable_free_aligned( removedVertices, resultVertexCount );
+
+        // Rebuild the vertex buffers
+        for ( int b=0; b<pResult->GetVertexBuffers().GetBufferCount(); ++b )
+        {
+            int elemSize = pResult->GetVertexBuffers().GetElementSize( b );
+            const uint8_t* pSourceData = pSource->GetVertexBuffers().GetBufferData( b );
+            uint8_t* pData = pResult->GetVertexBuffers().GetBufferData( b );
+            for ( int v=0; v<pResult->GetVertexCount(); ++v )
+            {
+                int span = 0;
+                for ( int s=0; v+s<pResult->GetVertexCount(); ++s )
+                {
+                    if ( usedVertices[v+s]>=0 )
+                    {
+                        if (span==0)
+                        {
+                            ++span;
+                        }
+                        else
+                        {
+                            if ( usedVertices[v+s] == usedVertices[v+s-1]+1 )
+                            {
+                                ++span;
+                            }
+                            else
+                            {
+                                break;
+                            }
+                        }
+                    }
+                    else
+                    {
+                        break;
+                    }
+                }
+
+                if (span>0)
+                {
+                    memcpy( pData+elemSize*usedVertices[v], pSourceData+elemSize*v, elemSize*span );
+
+                    v += span-1;
+                }
+            }
+        }
+        pResult->GetVertexBuffers().SetElementCount( firstFreeVertex );
+
+        // Rebuild surface data.
+        // \todo: For now make single surfaced
+        pResult->m_surfaces.Empty();
+        pResult->EnsureSurfaceData();
+
+        return pResult;
+    }
+
+}
diff --git a/Engine/Plugins/Experimental/Mutable/Source/MutableRuntime/Private/OpMeshReshape.h b/Engine/Plugins/Experimental/Mutable/Source/MutableRuntime/Private/OpMeshReshape.h
new file mode 100644
--- /dev/null
+++ b/Engine/Plugins/Experimental/Mutable/Source/MutableRuntime/Private/OpMeshReshape.h
@@ -0,0 +1,760 @@
+// Copyright Epic Games, Inc. All Rights Reserved.
+
+#pragma once
+
+#include "MutableRuntime/Private/ConvertData.h"
+#include "MutableRuntime/Private/MeshPrivate.h"
+#include "MutableRuntime/Private/OpMeshBind.h"
+
+#include "Math/Plane.h"
+#include "Math/Ray.h"
+#include "Math/NumericLimits.h"
+
+#include "IndexTypes.h"
+#include "Spatial/MeshAABBTree3.h"
+#include "Distance/DistPoint3Triangle3.h"
+#include "LineTypes.h"
+
+
+
+// TODO: Make the handling of rotations an option. It is more expensive on CPU and memory, and for some
+// cases it is not required at all.
+
+// TODO: Face stretch to scale the deformation per-vertex? 
+
+// TODO: Support multiple binding influences per vertex, to have smoother deformations.
+
+// TODO: Support multiple binding sets, to have different shapes deformations at once.
+
+// TODO: Deformation mask, to select the intensisty of the deformation per-vertex.
+
+// TODO: This is a reference implementation with ample roof for optimization.
+
+namespace mu
+{
+
+	struct FShapeMeshDescriptorApply
+	{
+		TArray<FVector3f> Positions;
+		TArray<FVector3f> Normals;
+		TArray<UE::Geometry::FIndex3i> Triangles;
+	};
+
+	// Method to actually deform a point
+	inline bool GetDeform( const FShapeMeshDescriptorApply& Shape, const FReshapeVertexBindingData& Binding, FVector3f& NewPosition, FQuat4f& Rotation)
+	{
+		const TArray<FVector3f>& ShapePositions = Shape.Positions;
+		const TArray<FVector3f>& ShapeNormals = Shape.Normals;
+	    const TArray<UE::Geometry::FIndex3i>& ShapeTriangles = Shape.Triangles;
+		
+		if (Binding.Triangle < 0 || Binding.Triangle >= ShapeTriangles.Num() )
+		{
+			return false;
+		}
+
+		const UE::Geometry::FIndex3i& Triangle = ShapeTriangles[Binding.Triangle];
+
+		FVector3f ProjectedVertexPosition
+			= ShapePositions[Triangle.A] * (1.0f - Binding.S - Binding.T)
+			+ ShapePositions[Triangle.B] * Binding.S
+			+ ShapePositions[Triangle.C] * Binding.T;
+
+		// This method approximates the shape face rotation
+		FVector3f InterpolatedNormal
+			= ShapeNormals[Triangle.A] * (1.0f - Binding.S - Binding.T)
+			+ ShapeNormals[Triangle.B] * Binding.S
+			+ ShapeNormals[Triangle.C] * Binding.T;
+		
+		FVector3f PositionOffset = InterpolatedNormal * Binding.D;
+		NewPosition = ProjectedVertexPosition + PositionOffset;
+
+		FVector3f CurrentShapeNormal = ((ShapePositions[Triangle.B] - ShapePositions[Triangle.A]) ^ (ShapePositions[Triangle.C] - ShapePositions[Triangle.A])).GetSafeNormal();
+		Rotation = FQuat4f::FindBetween(Binding.ShapeNormal, CurrentShapeNormal);
+		
+		return true;
+	}
+
+    //---------------------------------------------------------------------------------------------
+    //! Physics Bodies Reshape 
+    //---------------------------------------------------------------------------------------------
+	inline bool GetDeformedBox( const FShapeMeshDescriptorApply& Shape, const FReshapeVertexBindingData* BindingData,  TStaticArray<FVector3f, 8>& OutBox )
+	{
+		FQuat4f UnusedRotation;
+
+		for ( int32 I = 0; I < OutBox.Num(); ++I )
+		{
+			bool bModified = GetDeform(Shape, *(BindingData + I), OutBox[I], UnusedRotation);
+
+			if (!bModified)
+			{
+				// Invalidate the whole box. 
+				return false;
+			}
+		}
+
+		return true;
+	}
+
+	inline void GetDeformedConvex( const FShapeMeshDescriptorApply& Shape, const FReshapeVertexBindingData* BindingData, TArray<FVector3f>& InOutDeformedVertices )
+	{
+		const int32 ConvexVertCount = InOutDeformedVertices.Num();
+		FQuat4f UnusedTangentCorrection;
+		for ( int32 I = 0; I < ConvexVertCount; ++I )
+		{
+			GetDeform(Shape, *(BindingData + I), InOutDeformedVertices[I], UnusedTangentCorrection);	
+		}
+	}
+
+	inline void ComputeSphereFromDeformedBox( 
+			const TStaticArray<FVector3f, 8>& Box, 
+			FVector3f& InOutP, float& InOutR, 
+			const FTransform3f& InvBoneT )
+	{
+		TStaticArray<FVector3f, 8> TransformedBox = Box;
+		for (int32 I = 0; I < 8; ++I)
+		{
+			TransformedBox[I] = InvBoneT.TransformPosition(Box[I]) - InOutP;
+		}
+		
+		FVector3f Centroid(ForceInitToZero);
+		for ( const FVector3f& V : TransformedBox )
+		{
+			Centroid += V;	
+		}
+
+		Centroid *= (1.0f/8.0f);
+
+		InOutP = Centroid + InOutP;
+		
+		FVector3f Dim(ForceInitToZero);
+		for ( const FVector3f& V : TransformedBox )
+		{
+			Dim += (V - Centroid).GetAbs();
+		}
+
+		Dim *= (1.0f/8.0f);
+		
+		InOutR = (Dim.X + Dim.Y + Dim.Z) * (1.0f/3.0f);
+	}
+
+	inline void ComputeBoxFromDeformedBox( 
+			const TStaticArray<FVector3f, 8>& Box, 
+			FVector3f& InOutP, FQuat4f& InOutQ, FVector3f& InOutS, 
+			const FTransform3f& InvBoneT )
+	{
+		const FQuat4f InvQ = InOutQ.Inverse();
+		
+		TStaticArray<FVector3f, 8> TranformedBox;
+		for (int32 I = 0; I < 8; ++I)
+		{
+			TranformedBox[I] = InvQ.RotateVector(InvBoneT.TransformPosition(Box[I]) - InOutP);
+		}
+	
+		const FVector3f TopC = (Box[0] + Box[1] + Box[2] + Box[3])*0.25f;
+		const FVector3f BottomC = (Box[4+0] + Box[4+1] + Box[4+2] + Box[4+3])*0.25f;
+    
+		const FVector3f FrontC = (Box[0] + Box[1] + Box[4+0] + Box[4+1])*0.25f;
+		const FVector3f BackC = (Box[2] + Box[3] + Box[4+2] + Box[4+3])*0.25f;
+    
+		const FVector3f RightC =  (Box[1] + Box[2] + Box[4+1] + Box[4+2])*0.25f;
+		const FVector3f LeftC =  (Box[3] + Box[0] + Box[4+3] + Box[4+0])*0.25f;
+    
+		FVector3f ZB = (TopC - BottomC).GetSafeNormal();
+		FVector3f XB = (RightC - LeftC).GetSafeNormal();
+		FVector3f YB = (FrontC - BackC).GetSafeNormal();
+
+
+		// Pick the 2 most offaxis vectors and construct a rotation form those.	
+		// TODO: Find a better way of finding and orientation form ZB, XB, TN maybe by averaging somehow different bases created from the vectors, with quaternions?
+		FVector3f OF = FVector3f( XB.X, YB.Y, ZB.Y ).GetAbs();
+
+		float M0 = FMath::Max3( OF.X, OF.Y, OF.X);
+		float M1 = M0 == OF.X ? FMath::Max(OF.Y, OF.Z) : (M0 == OF.Y ? FMath::Max(OF.X,OF.Z) : FMath::Max(OF.Y, OF.X));
+		
+		FMatrix44f RotationMatrix = FMatrix44f::Identity;
+		if (M0 == OF.X)
+		{
+			RotationMatrix = M1 == OF.Y 
+					? FRotationMatrix44f::MakeFromXY(XB, YB)
+					: FRotationMatrix44f::MakeFromXZ(XB, ZB);
+		}
+		else if (M0 == OF.Y)
+		{
+			RotationMatrix = M1 == XB.X 
+					? FRotationMatrix44f::MakeFromYX(YB, XB)
+					: FRotationMatrix44f::MakeFromYZ(YB, ZB);
+		}
+		else
+		{
+			RotationMatrix = M1 == XB.X 
+					? FRotationMatrix44f::MakeFromZX(ZB, XB)
+					: FRotationMatrix44f::MakeFromZY(ZB, YB);	
+		} 
+		
+		InOutQ = RotationMatrix.ToQuat() * InOutQ;
+		InOutP = (TopC + BottomC) * 0.5f + InOutP;
+		InOutS = FVector3f((RightC - LeftC).Size(), (FrontC - BackC).Size(), (TopC - BottomC).Size()) * 0.5f;	
+	}
+
+	inline void ComputeSphylFromDeformedBox( 
+			const TStaticArray<FVector3f, 8>& Box, 
+			FVector3f& InOutP, FQuat4f& InOutQ, float& InOutR, float& InOutL, 
+			const FTransform3f& InvBoneT)
+	{
+		const FQuat4f InvQ = InOutQ.Inverse();
+    	
+		TStaticArray<FVector3f, 8> TranformedBox;
+		for (int32 I = 0; I < 8; ++I)
+		{
+			TranformedBox[I] = InvQ.RotateVector(InvBoneT.TransformPosition(Box[I]) - InOutP);
+		}
+       
+		FVector3f TopCentroid(ForceInitToZero);
+		FVector3f BottomCentroid(ForceInitToZero);
+		for (int32 I = 0; I < 4; ++I)
+		{
+			TopCentroid += TranformedBox[I];
+			BottomCentroid += TranformedBox[4+I];
+		}
+		
+		TopCentroid *= 0.25f;
+		BottomCentroid *= 0.25f;
+
+		FVector3f DirVector = (TopCentroid - BottomCentroid).GetSafeNormal();
+
+		InOutQ = FQuat4f::FindBetweenNormals(FVector3f::UnitZ(), DirVector) * InOutQ; 
+		InOutP = (TopCentroid + BottomCentroid) * 0.5f + InOutP;
+
+		FVector3f TopDim(ForceInitToZero);
+		FVector3f BottomDim(ForceInitToZero);
+		for (int32 I = 0; I < 4; ++I)
+		{
+			TopDim += (TopCentroid - TranformedBox[I]).GetAbs();
+			BottomDim += (BottomCentroid - TranformedBox[4+I]).GetAbs();
+		}
+
+		TopDim *= 0.25f;
+		BottomDim *= 0.25f;
+
+		InOutR = (TopDim.X + TopDim.Y + BottomDim.X + BottomDim.Y) * 0.25f;
+		InOutL = (TopCentroid - BottomCentroid).Size() - InOutR*2.0f;
+	}
+
+	inline void ComputeTaperedCapsuleFromDeformedBox( 
+			const TStaticArray<FVector3f, 8>& Box, 
+			FVector3f& InOutP, FQuat4f& InOutQ, float& InOutR0, float& InOutR1, float& InOutL, 
+			const FTransform3f& InvBoneT)
+	{
+		const FQuat4f InvQ = InOutQ.Inverse();
+    
+		TStaticArray<FVector3f, 8> TranformedBox;
+		for (int32 I = 0; I < 8; ++I)
+		{
+			TranformedBox[I] = InvQ.RotateVector(InvBoneT.TransformPosition(Box[I]) - InOutP);
+		}
+       
+		FVector3f TopCentroid(ForceInitToZero);
+		FVector3f BottomCentroid(ForceInitToZero);
+		for (int32 I = 0; I < 4; ++I)
+		{
+			TopCentroid += TranformedBox[I];
+			BottomCentroid += TranformedBox[4+I];
+		}
+		
+		TopCentroid *= 0.25f;
+		BottomCentroid *= 0.25f;
+
+		FVector3f DirVector = (TopCentroid - BottomCentroid).GetSafeNormal();
+
+		InOutQ = InOutQ * FQuat4f::FindBetweenNormals(FVector3f::UnitZ(), DirVector); 
+		InOutP = (TopCentroid + BottomCentroid) * 0.5f + InOutP;
+
+		FVector3f TopDim(ForceInitToZero);
+		FVector3f BottomDim(ForceInitToZero);
+
+		for (int32 I = 0; I < 4; ++I)
+		{
+			TopDim += (TopCentroid - TranformedBox[I]).GetAbs();
+			BottomDim += (BottomCentroid - TranformedBox[4+I]).GetAbs();
+		}
+
+		TopDim *= 0.25f;
+		BottomDim *= 0.25f;
+
+		InOutR0 = (TopDim.X + TopDim.Y) * 0.5f;
+		InOutR1 = (BottomDim.X + BottomDim.Y) * 0.5f;
+		InOutL = (TopCentroid - BottomCentroid).Size() - (InOutR0 + InOutR1);
+	}
+
+	inline void ApplyToVertices(Mesh* Mesh, TArrayView<const FReshapeVertexBindingData> BindingData, const FShapeMeshDescriptorApply& Shape)
+	{
+		check(Mesh);
+		check(Mesh->GetVertexCount() == BindingData.Num());
+		
+
+		UntypedMeshBufferIterator ItPosition(Mesh->GetVertexBuffers(), MBS_POSITION);
+		UntypedMeshBufferIterator ItNormal(Mesh->GetVertexBuffers(), MBS_NORMAL);
+		UntypedMeshBufferIterator ItTangent(Mesh->GetVertexBuffers(), MBS_TANGENT);
+
+		const int32 MeshVertexCount = BindingData.Num();
+		for (int32 MeshVertexIndex = 0; MeshVertexIndex < MeshVertexCount; ++MeshVertexIndex)
+		{
+			const FReshapeVertexBindingData& Binding = BindingData[MeshVertexIndex];
+
+			FVector3f NewPosition;
+			FQuat4f TangentSpaceCorrection;
+			bool bModified = GetDeform( Shape, Binding, NewPosition, TangentSpaceCorrection );
+
+			if (bModified)
+			{
+				FVector3f OldPosition = ItPosition.GetAsVec3f();
+				
+				if ( !FMath::IsNearlyEqual( Binding.Weight, 1.0f ) )
+				{
+					// Zero weighted vertices are already discarded at the binding phase so there is no gain
+					// checking here.
+					NewPosition = FMath::Lerp( OldPosition, NewPosition, Binding.Weight );
+					TangentSpaceCorrection = FQuat4f::Slerp( FQuat4f::Identity, TangentSpaceCorrection, Binding.Weight );
+				}
+				
+				// Non rigid vertices will not rotate since are attached to themselves 
+				NewPosition = TangentSpaceCorrection.RotateVector(OldPosition - Binding.AttachmentPoint) + NewPosition;
+				
+				ItPosition.SetFromVec3f(NewPosition);
+					
+				if (ItNormal.ptr())
+				{
+					FVector3f OldNormal = ItNormal.GetAsVec3f();
+					FVector3f NewNormal = TangentSpaceCorrection.RotateVector(OldNormal);
+					ItNormal.SetFromVec3f(NewNormal);
+				}
+
+				if (ItTangent.ptr())
+				{
+					FVector3f OldTangent = ItTangent.GetAsVec3f();
+					FVector3f NewTangent = TangentSpaceCorrection.RotateVector(OldTangent);
+					ItTangent.SetFromVec3f(NewTangent);
+				}
+			}
+			
+			++ItPosition;
+
+			if (ItNormal.ptr())
+			{
+				++ItNormal;
+			}
+
+			if (ItTangent.ptr())
+			{
+				++ItTangent;
+			}
+		}
+		
+	}
+
+	inline void ApplyToPose(Mesh* Result, 
+			TArrayView<const FReshapeVertexBindingData> BindingData, 
+			TArrayView<const int> BoneIndices,
+			const FShapeMeshDescriptorApply& Shape)
+	{
+		check(Result);
+
+		const int32 NumBoneIndices = BoneIndices.Num();
+		for (int32 b = 0; b < NumBoneIndices; ++b)
+		{
+			int32 BoneIndex = BoneIndices[b];
+
+			const FReshapeVertexBindingData& Binding = BindingData[b];
+
+			FTransform3f& t = Result->m_bonePoses[BoneIndex].m_boneTransform;
+
+			FVector3f NewPosition;
+			FQuat4f TangentSpaceCorrection;
+			bool bModified = GetDeform(Shape, Binding, NewPosition, TangentSpaceCorrection);
+
+			if (bModified)
+			{
+				// The pose has been modified and should not be replaced
+				Result->m_bonePoses[BoneIndex].m_boneSkinned = true;
+
+				// TODO: Review if the rotation also needs to be applied.
+				t.SetLocation(FMath::Lerp(t.GetLocation(), NewPosition, Binding.Weight));
+			}
+		}
+	}
+
+	inline void ApplyToPhysicsBodies(
+			PhysicsBody* PBody, const Mesh* BaseMesh, 
+			TArrayView<const FReshapeVertexBindingData> BindingData, TArrayView<const int32> UsedIndices, 
+			const FShapeMeshDescriptorApply& Shape)
+	{
+		check(PBody);
+		check(BaseMesh);
+
+		// Not used.	
+		FQuat4f DummyTangentSpaceCorrection;
+		
+		int32 NumProcessedBindPoints = 0;
+
+		bool bAnyModified = false;
+
+		// Retrieve them in the same order the boxes where put in, so they can be linked to the physics body volumes.
+		for ( const int32 B : UsedIndices )
+		{	
+			int32 BoneIdx = BaseMesh->FindBonePose(PBody->GetBodyBoneName(B));
+			FTransform3f BoneTransform = FTransform3f::Identity;
+			if (BoneIdx >= 0)
+			{
+				BaseMesh->GetBoneTransform(BoneIdx, BoneTransform);
+			}
+			
+			FTransform3f InvBoneTransform = BoneTransform.Inverse();
+
+			const int32 NumSpheres = PBody->GetSphereCount(B); 
+			for ( int32 I = 0; I < NumSpheres; ++I )
+			{
+				FVector3f P;
+				float R;
+
+				TStaticArray<FVector3f, 8> Box; 
+				const bool bDeformed = GetDeformedBox( Shape, &BindingData[NumProcessedBindPoints], Box );
+
+				if (bDeformed)
+				{
+					PBody->GetSphere(B, I, P, R);
+					ComputeSphereFromDeformedBox(Box, P, R, InvBoneTransform);
+					PBody->SetSphere(B, I, P, R);
+					bAnyModified = true;
+				}
+
+				NumProcessedBindPoints += Box.Num();
+			}
+
+			const int32 NumBoxes = PBody->GetBoxCount(B);
+			for ( int32 I = 0; I < NumBoxes; ++I )
+			{
+				TStaticArray<FVector3f, 8> Box;
+				const bool bDeformed = GetDeformedBox( Shape, &BindingData[NumProcessedBindPoints], Box );
+				
+				if (bDeformed)
+				{
+					FVector3f P;
+					FQuat4f Q;
+					FVector3f S;
+					
+					PBody->GetBox(B, I, P, Q, S);
+					ComputeBoxFromDeformedBox(Box, P, Q, S, InvBoneTransform);
+					PBody->SetBox(B, I, P, Q, S);
+					bAnyModified = true;
+				}
+
+				NumProcessedBindPoints += Box.Num();
+			}
+			
+			const int32 NumSphyls = PBody->GetSphylCount(B);
+			for ( int32 I = 0; I < NumSphyls; ++I )
+			{
+				TStaticArray<FVector3f, 8> Box;
+				const bool bDeformed = GetDeformedBox( Shape, &BindingData[NumProcessedBindPoints], Box );
+
+				if ( bDeformed )
+				{
+					FVector3f P;
+					FQuat4f Q;
+					float R;
+					float L;						
+
+					PBody->GetSphyl(B, I, P, Q, R, L);
+					ComputeSphylFromDeformedBox(Box, P, Q, R, L, InvBoneTransform);
+					PBody->SetSphyl(B, I, P, Q, R, L);
+					bAnyModified = true;
+				}
+
+				NumProcessedBindPoints += Box.Num();
+			}
+
+			const int32 NumTaperedCapsules = PBody->GetTaperedCapsuleCount(B);
+			for ( int32 I = 0; I < NumTaperedCapsules; ++I )
+			{
+				TStaticArray<FVector3f, 8> Box;
+				const bool bDeformed = GetDeformedBox( Shape, &BindingData[NumProcessedBindPoints], Box );
+
+				if ( bDeformed )
+				{
+					FVector3f P;
+					FQuat4f Q;
+					float R0;
+					float R1;
+					float L;
+				
+					PBody->GetTaperedCapsule(B, I, P, Q, R0, R1, L);
+					ComputeTaperedCapsuleFromDeformedBox(Box, P, Q, R0, R1, L, InvBoneTransform);
+					PBody->SetTaperedCapsule(B, I, P, Q, R0, R1, L);
+					bAnyModified = true;
+				}
+				
+				NumProcessedBindPoints += Box.Num();
+			}
+
+			const int32 NumConvex = PBody->GetConvexCount(B);
+			for ( int32 I = 0; I < NumConvex; ++I )
+			{
+				FVector3f const* VertexData;
+				int32 VertexCount;
+				int32 const* IndexData;
+				int32 IndexCount;
+				FTransform3f ConvexT;
+
+				PBody->GetConvex(B, I, VertexData, VertexCount, IndexData, IndexCount, ConvexT);
+
+				// TODO: Allow access and modification to the convex data directly to avoid a copy here.
+				TArray<FVector3f> DeformedVertexData(VertexData, VertexCount);
+				TArray<int32> Indices(IndexData, IndexCount);
+
+				GetDeformedConvex( Shape, &BindingData[NumProcessedBindPoints], DeformedVertexData);
+
+				FTransform3f InvConvexT = InvBoneTransform * ConvexT.Inverse();
+				for ( FVector3f& V : DeformedVertexData )
+				{
+					V = InvConvexT.TransformPosition( V );				
+				}
+				
+				PBody->SetConvex(B, I, DeformedVertexData.GetData(), DeformedVertexData.Num(), Indices.GetData(), Indices.Num(), ConvexT);
+
+				NumProcessedBindPoints += VertexCount;
+				bAnyModified = true;
+			}
+		}
+
+		PBody->bBodiesModified = bAnyModified;
+
+		check(NumProcessedBindPoints == BindingData.Num());	
+	}
+
+	//---------------------------------------------------------------------------------------------
+	//! Rebuild the (previously bound) mesh data for a new shape.
+	//! Proof-of-concept implementation.
+	//---------------------------------------------------------------------------------------------
+	inline MeshPtr MeshApplyShape(const Mesh* BaseMesh, const Mesh* ShapeMesh, 
+			bool bReshapeVertices, bool bReshapeSkeleton, bool bReshapePhysicsVolumes)
+	{
+		MUTABLE_CPUPROFILER_SCOPE(MeshApplyReshape);
+
+		if (!BaseMesh)
+		{
+			return nullptr;
+		}
+		
+		// Early out if nothing will be modified and the vertices discarted.
+		const bool bSkeletonModification = BaseMesh->GetSkeleton() && bReshapeSkeleton;
+		const bool bPhysicsModification = BaseMesh->GetPhysicsBody() && bReshapePhysicsVolumes;
+
+		if (!bReshapeVertices && !bSkeletonModification && !bPhysicsModification)
+		{
+			return nullptr;
+		}
+	
+		// \TODO: Multiple binding data support
+		int BindingDataIndex = 0;
+
+		// If the base mesh has no binding data, just clone it.
+		int BarycentricDataBuffer = 0;
+		int BarycentricDataChannel = 0;
+		const FMeshBufferSet& VB = BaseMesh->GetVertexBuffers();
+		VB.FindChannel(MBS_BARYCENTRICCOORDS, BindingDataIndex, &BarycentricDataBuffer, &BarycentricDataChannel);
+		
+		// Clone Without VertexBuffers or AdditionalBuffers
+		constexpr EMeshCloneFlags CloneFlags = ~(EMeshCloneFlags::WithVertexBuffers | EMeshCloneFlags::WithAdditionalBuffers);
+		MeshPtr Result = BaseMesh->Clone(CloneFlags);
+	
+		FMeshBufferSet& ResultBuffers = Result->GetVertexBuffers();
+
+		check(ResultBuffers.m_buffers.Num() == 0);
+
+		// Copy buffers skipping binding data. 
+		ResultBuffers.m_elementCount = VB.m_elementCount;
+		// Remove one element to the number of buffers if BarycentricDataBuffer found. 
+		ResultBuffers.m_buffers.SetNum( FMath::Max( 0, VB.m_buffers.Num() - int32(BarycentricDataBuffer >= 0) ) );
+		
+		for (int32 B = 0, R = 0; B < VB.m_buffers.Num(); ++B)
+		{
+			if (B != BarycentricDataBuffer)
+			{
+				ResultBuffers.m_buffers[R++] = VB.m_buffers[B];
+			}
+		}
+
+		// Copy the additional buffers skipping binding data.
+		Result->m_AdditionalBuffers.Reserve( BaseMesh->m_AdditionalBuffers.Num() );
+		for (const TPair<EMeshBufferType, FMeshBufferSet>& A : BaseMesh->m_AdditionalBuffers)
+		{	
+			const bool bIsBindOpBuffer = 
+					A.Key == EMeshBufferType::SkeletonDeformBinding || 
+					A.Key == EMeshBufferType::PhysicsBodyDeformBinding ||
+					A.Key == EMeshBufferType::PhysicsBodyDeformSelection;
+	
+			if (!bIsBindOpBuffer)
+			{
+				Result->m_AdditionalBuffers.Add(A);
+			}
+		}
+		if (!ShapeMesh)
+		{
+			return Result;
+		}
+
+		int32 ShapeVertexCount = ShapeMesh->GetVertexCount();
+		int32 ShapeTriangleCount = ShapeMesh->GetFaceCount();
+		if (!ShapeVertexCount || !ShapeTriangleCount)
+		{
+			return Result;
+		}
+
+		// Generate the temp vertex query data for the shape
+		// TODO: Vertex data copy may be saved in most cases.
+		TArray<FVector3f> ShapePositions;
+		FShapeMeshDescriptorApply ShapeDescriptor;
+		{
+			MUTABLE_CPUPROFILER_SCOPE(GenerateVertexQueryData);
+
+			ShapeDescriptor.Positions.SetNum(ShapeVertexCount);
+			ShapeDescriptor.Normals.SetNum(ShapeVertexCount);
+
+			// \TODO: Simple but inefficient
+			UntypedMeshBufferIteratorConst ItPosition(ShapeMesh->GetVertexBuffers(), MBS_POSITION);
+			UntypedMeshBufferIteratorConst ItNormal(ShapeMesh->GetVertexBuffers(), MBS_NORMAL);
+			for (int32 ShapeVertexIndex = 0; ShapeVertexIndex < ShapeVertexCount; ++ShapeVertexIndex)
+			{
+				FVector3f Position = ItPosition.GetAsVec3f();
+				ShapeDescriptor.Positions[ShapeVertexIndex] = Position;
+				++ItPosition;
+
+				FVector3f Normal = ItNormal.GetAsVec3f();
+				ShapeDescriptor.Normals[ShapeVertexIndex] = Normal;
+				++ItNormal;
+			}
+		}
+
+		// Generate the temp face query data for the shape
+		// TODO: Index data copy may be saved in most cases.
+		{
+			MUTABLE_CPUPROFILER_SCOPE(GenerateTriangleQueryData);
+
+			ShapeDescriptor.Triangles.SetNum(ShapeTriangleCount);
+
+			// \TODO: Simple but inefficient
+			UntypedMeshBufferIteratorConst ItIndices(ShapeMesh->GetIndexBuffers(), MBS_VERTEXINDEX);
+			for (int32 TriangleIndex = 0; TriangleIndex < ShapeTriangleCount; ++TriangleIndex)
+			{
+				UE::Geometry::FIndex3i Triangle;
+				Triangle.A = int(ItIndices.GetAsUINT32());
+				++ItIndices;
+				Triangle.B = int(ItIndices.GetAsUINT32());
+				++ItIndices;
+				Triangle.C = int(ItIndices.GetAsUINT32());
+				++ItIndices;
+
+				ShapeDescriptor.Triangles[TriangleIndex] = Triangle;
+			}
+		}
+
+		if (bReshapeVertices && BarycentricDataBuffer >= 0)
+		{
+			MUTABLE_CPUPROFILER_SCOPE(ReshapeVertices);
+			// \TODO: More checks
+			check(BarycentricDataChannel == 0);
+			check(VB.GetElementSize(BarycentricDataBuffer) == (int)sizeof(FReshapeVertexBindingData));
+
+			TArrayView<const FReshapeVertexBindingData> VerticesBindingData(
+					(const FReshapeVertexBindingData*)VB.GetBufferData(BarycentricDataBuffer),
+					VB.GetElementCount());
+
+			ApplyToVertices(Result.get(), VerticesBindingData, ShapeDescriptor);
+		}
+	
+		if (bReshapeSkeleton)
+		{
+			MUTABLE_CPUPROFILER_SCOPE(ReshapeSkeleton);
+
+			// If the base mesh has no binding data for the skeleton don't do anything.
+			const FMeshBufferSet* SkeletonBindBuffer = nullptr;
+			for ( const TPair<EMeshBufferType, FMeshBufferSet>& A : BaseMesh->m_AdditionalBuffers )
+			{
+				if ( A.Key == EMeshBufferType::SkeletonDeformBinding )
+				{
+					SkeletonBindBuffer = &A.Value;
+					break;
+				}
+			}
+			
+			BarycentricDataBuffer = -1;
+			if (SkeletonBindBuffer)
+			{
+				SkeletonBindBuffer->FindChannel(MBS_BARYCENTRICCOORDS, BindingDataIndex, &BarycentricDataBuffer, &BarycentricDataChannel);
+			}
+
+			if (SkeletonBindBuffer && BarycentricDataBuffer >= 0)
+			{
+				// \TODO: More checks
+				check(BarycentricDataChannel == 0);
+				check(SkeletonBindBuffer && SkeletonBindBuffer->GetElementSize(BarycentricDataBuffer) == (int)sizeof(FReshapeVertexBindingData));
+
+				TArrayView<const FReshapeVertexBindingData> SkeletonBindingData( 
+						(const FReshapeVertexBindingData*)SkeletonBindBuffer->GetBufferData(BarycentricDataBuffer),
+						SkeletonBindBuffer->GetElementCount());
+			
+				check(SkeletonBindBuffer->GetBufferCount() >= 2);
+				TArrayView<const int32> BoneIndices( 
+						(const int32*)SkeletonBindBuffer->GetBufferData(1), SkeletonBindBuffer->GetElementCount());
+
+				ApplyToPose(Result.get(), SkeletonBindingData, BoneIndices, ShapeDescriptor);
+			}
+		}
+
+		// Transform Physics Volumes based on the deformed bounding box points.
+		const PhysicsBody* OldPhysicsBody = Result->m_pPhysicsBody.get();
+
+		if (bReshapePhysicsVolumes && OldPhysicsBody)
+		{	
+			MUTABLE_CPUPROFILER_SCOPE(ReshapePhysicsBodies);
+			
+			using BufferEntryType = TPair<EMeshBufferType, FMeshBufferSet>;
+			const BufferEntryType* FoundPhysicsBindBuffer = BaseMesh->m_AdditionalBuffers.FindByPredicate(
+					[](BufferEntryType& E){ return E.Key == EMeshBufferType::PhysicsBodyDeformBinding; });
+
+			const BufferEntryType* FoundPhysicsBindSelectionBuffer = BaseMesh->m_AdditionalBuffers.FindByPredicate(
+					[](BufferEntryType& E){ return E.Key == EMeshBufferType::PhysicsBodyDeformSelection; });
+	
+			BarycentricDataBuffer = -1;
+			if (FoundPhysicsBindBuffer)
+			{
+				FoundPhysicsBindBuffer->Value.FindChannel(MBS_BARYCENTRICCOORDS, BindingDataIndex, &BarycentricDataBuffer, &BarycentricDataChannel);
+			}
+			
+			if (FoundPhysicsBindBuffer && FoundPhysicsBindSelectionBuffer && BarycentricDataBuffer >= 0)
+			{
+				const FMeshBufferSet& PhysicsBindBuffer = FoundPhysicsBindBuffer->Value;
+				const FMeshBufferSet& PhyiscsBindSelectionBuffer = FoundPhysicsBindSelectionBuffer->Value;
+
+				// \TODO: More checks
+				check(BarycentricDataChannel == 0);
+				check(PhysicsBindBuffer.GetElementSize(BarycentricDataBuffer) == (int)sizeof(FReshapeVertexBindingData));
+					
+				TArrayView<const FReshapeVertexBindingData> PhysicsBindingData(  
+						(const FReshapeVertexBindingData*)PhysicsBindBuffer.GetBufferData(BarycentricDataBuffer),
+						PhysicsBindBuffer.GetElementCount() );
+
+				TArrayView<const int32> UsedIndices( 
+						(const int32*)PhyiscsBindSelectionBuffer.GetBufferData(0), 
+					    PhyiscsBindSelectionBuffer.GetElementCount());
+
+				Ptr<PhysicsBody> NewPhysicsBody = OldPhysicsBody->Clone();
+				Result->SetPhysicsBody(NewPhysicsBody);
+
+				ApplyToPhysicsBodies( NewPhysicsBody.get(), BaseMesh, PhysicsBindingData, UsedIndices, ShapeDescriptor );
+			}
+		}
+		
+		return Result;
+	}
+}
diff --git a/Engine/Plugins/Experimental/Mutable/Source/MutableRuntime/Private/OpMeshSubtract.h b/Engine/Plugins/Experimental/Mutable/Source/MutableRuntime/Private/OpMeshSubtract.h
new file mode 100644
--- /dev/null
+++ b/Engine/Plugins/Experimental/Mutable/Source/MutableRuntime/Private/OpMeshSubtract.h
@@ -0,0 +1,87 @@
+// Copyright Epic Games, Inc. All Rights Reserved.
+
+#pragma once
+
+#include "MeshPrivate.h"
+#include "Platform.h"
+
+namespace mu
+{
+
+
+	//---------------------------------------------------------------------------------------------
+    inline MeshPtr MeshSubtract( const Mesh* pA, const Mesh* pB )
+	{
+		// TODO: Optimise, generate optimised meshes
+
+		MeshPtr pResult = pA->Clone();
+
+		if ( pA->GetVertexCount()==0 || pB->GetVertexCount()==0 )
+		{
+			return pResult;
+		}
+
+		check( pResult->GetIndexBuffers().GetBufferCount()==1 );
+		check( pResult->GetIndexBuffers().GetBufferChannelCount(0)==1 );
+		check( pResult->GetFaceBuffers().GetBufferCount()==0 );
+
+        UntypedMeshBufferIteratorConst itSource( pResult->GetIndexBuffers(), MBS_VERTEXINDEX );
+        UntypedMeshBufferIterator itDest( pResult->GetIndexBuffers(), MBS_VERTEXINDEX );
+
+		// For every face in A, see if it is also in B.
+		Mesh::VERTEX_MATCH_MAP vertexMap;
+		pB->GetVertexMap( *pA, vertexMap );
+		int aFaceCount = pResult->GetFaceCount();
+		int bFaceCount = pB->GetFaceCount();
+
+        UntypedMeshBufferIteratorConst ito( pResult->GetIndexBuffers(), MBS_VERTEXINDEX );
+		for ( int f=0; f<aFaceCount; ++f )
+		{
+			bool hasFace = false;
+            vec3<uint32_t> ov;
+            ov[0] = ito.GetAsUINT32(); ++ito;
+            ov[1] = ito.GetAsUINT32(); ++ito;
+            ov[2] = ito.GetAsUINT32(); ++ito;
+
+            UntypedMeshBufferIteratorConst it( pB->GetIndexBuffers(), MBS_VERTEXINDEX );
+            for ( int fb=0; !hasFace && fb<bFaceCount; fb++ )
+			{
+                vec3<uint32_t> v;
+                v[0] = it.GetAsUINT32(); ++it;
+                v[1] = it.GetAsUINT32(); ++it;
+                v[2] = it.GetAsUINT32(); ++it;
+
+				hasFace = true;
+				for ( int vi=0; hasFace && vi<3; ++vi )
+				{
+					hasFace = vertexMap.Matches(v[vi],ov[0])
+						 || vertexMap.Matches(v[vi],ov[1])
+						 || vertexMap.Matches(v[vi],ov[2]);
+				}
+			}
+
+			if ( !hasFace )
+			{
+				if ( itDest.ptr() != itSource.ptr() )
+				{
+					FMemory::Memcpy( itDest.ptr(), itSource.ptr(), itSource.GetElementSize()*3 );
+				}
+
+				itDest+=3;
+			}
+
+			itSource+=3;
+		}
+
+		std::size_t removedIndices = itSource - itDest;
+		check( removedIndices%3==0 );
+
+		pResult->GetFaceBuffers().SetElementCount( aFaceCount-(int)removedIndices/3 );
+		pResult->GetIndexBuffers().SetElementCount( aFaceCount*3-(int)removedIndices );
+
+		// TODO: Should redo/reorder the face buffer before SetElementCount since some deleted faces could be left and some remaining faces deleted.
+
+		return pResult;
+	}
+
+}
diff --git a/Engine/Plugins/Experimental/Mutable/Source/MutableRuntime/Private/OpMeshTransform.h b/Engine/Plugins/Experimental/Mutable/Source/MutableRuntime/Private/OpMeshTransform.h
new file mode 100644
--- /dev/null
+++ b/Engine/Plugins/Experimental/Mutable/Source/MutableRuntime/Private/OpMeshTransform.h
@@ -0,0 +1,95 @@
+// Copyright Epic Games, Inc. All Rights Reserved.
+
+#pragma once
+
+#include "MeshPrivate.h"
+#include "ConvertData.h"
+#include "Platform.h"
+
+
+namespace mu
+{
+
+    //---------------------------------------------------------------------------------------------
+    //! Reference version
+    //---------------------------------------------------------------------------------------------
+    inline MeshPtr MeshTransform( const Mesh* pBase, const FMatrix44f& transform )
+	{
+		MeshPtr pDest = pBase->Clone();
+
+        uint32_t vcount = pBase->GetVertexBuffers().GetElementCount();
+
+        if ( !vcount )
+		{
+			return pDest;
+		}
+
+		FMatrix44f transformIT = transform.Inverse().GetTransposed();
+
+        const FMeshBufferSet& MBSPriv = pDest->GetVertexBuffers();
+        for ( int32 b=0; b<MBSPriv.m_buffers.Num(); ++b )
+        {
+
+            for ( int32 c=0; c<MBSPriv.m_buffers[b].m_channels.Num(); ++c )
+            {
+                MESH_BUFFER_SEMANTIC sem = MBSPriv.m_buffers[b].m_channels[c].m_semantic;
+                int semIndex = MBSPriv.m_buffers[b].m_channels[c].m_semanticIndex;
+
+                UntypedMeshBufferIterator it( pDest->GetVertexBuffers(), sem, semIndex );
+
+                switch ( sem )
+                {
+                case MBS_POSITION:
+                    for ( uint32_t v=0; v<vcount; ++v )
+                    {
+                        FVector4f value( 0.0f, 0.0f, 0.0f, 1.0f );
+                        for( int i=0; i<it.GetComponents(); ++i )
+                        {
+                            ConvertData( i, &value[0], MBF_FLOAT32, it.ptr(), it.GetFormat() );
+                        }
+
+                        value = transform.TransformFVector4( value );
+
+                        for( int i=0; i<it.GetComponents(); ++i )
+                        {
+                            ConvertData( i, it.ptr(), it.GetFormat(), &value[0], MBF_FLOAT32 );
+                        }
+
+                        ++it;
+                    }
+                    break;
+
+                case MBS_NORMAL:
+                case MBS_TANGENT:
+                case MBS_BINORMAL:
+                    for ( uint32_t v=0; v<vcount; ++v )
+                    {
+						FVector4f value( 0.0f, 0.0f, 0.0f, 1.0f );
+                        for( int i=0; i<it.GetComponents(); ++i )
+                        {
+                            ConvertData( i, &value[0], MBF_FLOAT32, it.ptr(), it.GetFormat() );
+                        }
+
+                        value = transformIT.TransformFVector4(value);
+
+                        for( int i=0; i<it.GetComponents(); ++i )
+                        {
+                            ConvertData( i, it.ptr(), it.GetFormat(), &value[0], MBF_FLOAT32 );
+                        }
+
+                        ++it;
+                    }
+                    break;
+
+                default:
+                    break;
+                }
+            }
+        }
+
+
+		return pDest;
+	}
+
+}
+
diff --git a/Engine/Plugins/Experimental/Mutable/Source/MutableRuntime/Private/Operations.cpp b/Engine/Plugins/Experimental/Mutable/Source/MutableRuntime/Private/Operations.cpp
new file mode 100644
--- /dev/null
+++ b/Engine/Plugins/Experimental/Mutable/Source/MutableRuntime/Private/Operations.cpp
@@ -0,0 +1,1192 @@
+// Copyright Epic Games, Inc. All Rights Reserved.
+
+
+#include "Operations.h"
+#include "Platform.h"
+#include "ModelPrivate.h"
+
+
+namespace mu
+{
+
+	//---------------------------------------------------------------------------------------------
+	//---------------------------------------------------------------------------------------------
+	//---------------------------------------------------------------------------------------------
+    #define OP_DESC_COUNT		size_t(OP_TYPE::COUNT)
+
+    // clang-format off
+    static const OP_DESC s_opDescs[OP_DESC_COUNT] =
+	{ 
+		// type			cached	GPU-izable	supported base image formats
+        { DT_NONE,		false,	false,		{ 0, 0, 0, 0, 0, 0, 0, 0, 0, 0, 0, 0, 0, 0, 0, 0, 0, 0, 0, 0, 0, 0 }		},	// NONE
+
+        { DT_BOOL,		false,	false,		{ 0, 0, 0, 0, 0, 0, 0, 0, 0, 0, 0, 0, 0, 0, 0, 0, 0, 0, 0, 0, 0, 0 }		},	// BO_CONSTANT
+        { DT_INT,		false,	false,		{ 0, 0, 0, 0, 0, 0, 0, 0, 0, 0, 0, 0, 0, 0, 0, 0, 0, 0, 0, 0, 0, 0 }		},	// NU_CONSTANT
+        { DT_SCALAR,	false,	false,		{ 0, 0, 0, 0, 0, 0, 0, 0, 0, 0, 0, 0, 0, 0, 0, 0, 0, 0, 0, 0, 0, 0 }		},	// SC_CONSTANT
+        { DT_COLOUR,	false,	false,		{ 0, 0, 0, 0, 0, 0, 0, 0, 0, 0, 0, 0, 0, 0, 0, 0, 0, 0, 0, 0, 0, 0 }		},	// CO_CONSTANT
+        { DT_IMAGE,		true,   false,		{ 1, 1, 1, 1, 1, 1, 1, 1, 1, 1, 1, 1, 0, 0, 0, 0, 0, 0, 0, 0, 0, 0 }		},	// IM_CONSTANT
+        { DT_MESH,		true,   false,		{ 0, 0, 0, 0, 0, 0, 0, 0, 0, 0, 0, 0, 0, 0, 0, 0, 0, 0, 0, 0, 0, 0 }		},	// ME_CONSTANT
+        { DT_LAYOUT,	false,	false,		{ 0, 0, 0, 0, 0, 0, 0, 0, 0, 0, 0, 0, 0, 0, 0, 0, 0, 0, 0, 0, 0, 0 }		},	// LA_CONSTANT
+        { DT_PROJECTOR,	false,	false,		{ 0, 0, 0, 0, 0, 0, 0, 0, 0, 0, 0, 0, 0, 0, 0, 0, 0, 0, 0, 0, 0, 0 }		},	// PR_CONSTANT
+        { DT_STRING,	false,	false,		{ 0, 0, 0, 0, 0, 0, 0, 0, 0, 0, 0, 0, 0, 0, 0, 0, 0, 0, 0, 0, 0, 0 }		},	// ST_CONSTANT
+
+        { DT_BOOL,		false,	false,		{ 0, 0, 0, 0, 0, 0, 0, 0, 0, 0, 0, 0, 0, 0, 0, 0, 0, 0, 0, 0, 0, 0 }		},	// BO_PARAMETER
+		{ DT_INT,		false,	false,		{ 0, 0, 0, 0, 0, 0, 0, 0, 0, 0, 0, 0, 0, 0, 0, 0, 0, 0, 0, 0, 0, 0 }		},	// NU_PARAMETER
+		{ DT_SCALAR,	false,	false,		{ 0, 0, 0, 0, 0, 0, 0, 0, 0, 0, 0, 0, 0, 0, 0, 0, 0, 0, 0, 0, 0, 0 }		},	// SC_PARAMETER
+		{ DT_COLOUR,	false,	false,		{ 0, 0, 0, 0, 0, 0, 0, 0, 0, 0, 0, 0, 0, 0, 0, 0, 0, 0, 0, 0, 0, 0 }		},	// CO_PARAMETER
+        { DT_PROJECTOR,	false,	false,		{ 0, 0, 0, 0, 0, 0, 0, 0, 0, 0, 0, 0, 0, 0, 0, 0, 0, 0, 0, 0, 0, 0 }		},	// PR_PARAMETER
+        { DT_IMAGE,		false,	false,		{ 0, 0, 0, 0, 0, 0, 0, 0, 0, 0, 0, 0, 0, 0, 0, 0, 0, 0, 0, 0, 0, 0 }		},	// IM_PARAMETER
+        { DT_STRING,	false,	false,		{ 0, 0, 0, 0, 0, 0, 0, 0, 0, 0, 0, 0, 0, 0, 0, 0, 0, 0, 0, 0, 0, 0 }		},	// ST_PARAMETER
+
+        { DT_INT,		false,	false,		{ 0, 0, 0, 0, 0, 0, 0, 0, 0, 0, 0, 0, 0, 0, 0, 0, 0, 0, 0, 0, 0, 0 }		},	// NU_CONDITIONAL
+		{ DT_SCALAR,	false,	false,		{ 0, 0, 0, 0, 0, 0, 0, 0, 0, 0, 0, 0, 0, 0, 0, 0, 0, 0, 0, 0, 0, 0 }		},	// SC_CONDITIONAL
+		{ DT_COLOUR,	false,	false,		{ 0, 0, 0, 0, 0, 0, 0, 0, 0, 0, 0, 0, 0, 0, 0, 0, 0, 0, 0, 0, 0, 0 }		},	// CO_CONDITIONAL
+        { DT_IMAGE,		true,	false,		{ 1, 1, 1, 1, 1, 1, 1, 1, 1, 1, 1, 1, 1, 1, 1, 1, 1, 1, 1, 1, 1, 1 }		},	// IM_CONDITIONAL
+        { DT_MESH,		true,	false,		{ 0, 0, 0, 0, 0, 0, 0, 0, 0, 0, 0, 0, 0, 0, 0, 0, 0, 0, 0, 0, 0, 0 }		},	// ME_CONDITIONAL
+        { DT_LAYOUT,	false,  false,		{ 0, 0, 0, 0, 0, 0, 0, 0, 0, 0, 0, 0, 0, 0, 0, 0, 0, 0, 0, 0, 0, 0 }		},	// LA_CONDITIONAL
+        { DT_INSTANCE,	false,	false,		{ 0, 0, 0, 0, 0, 0, 0, 0, 0, 0, 0, 0, 0, 0, 0, 0, 0, 0, 0, 0, 0, 0 }		},	// IN_CONDITIONAL
+
+        { DT_INT,		false,	false,		{ 0, 0, 0, 0, 0, 0, 0, 0, 0, 0, 0, 0, 0, 0, 0, 0, 0, 0, 0, 0, 0, 0 }		},	// NU_SWITCH
+		{ DT_SCALAR,	false,	false,		{ 0, 0, 0, 0, 0, 0, 0, 0, 0, 0, 0, 0, 0, 0, 0, 0, 0, 0, 0, 0, 0, 0 }		},	// SC_SWITCH
+		{ DT_COLOUR,	false,	false,		{ 0, 0, 0, 0, 0, 0, 0, 0, 0, 0, 0, 0, 0, 0, 0, 0, 0, 0, 0, 0, 0, 0 }		},	// CO_SWITCH
+		{ DT_IMAGE,		true,	false,		{ 1, 1, 1, 1, 1, 1, 1, 1, 1, 1, 1, 1, 1, 1, 1, 1, 1, 1, 1, 1, 1, 1 }		},	// IM_SWITCH
+		{ DT_MESH,		true,	false,		{ 0, 0, 0, 0, 0, 0, 0, 0, 0, 0, 0, 0, 0, 0, 0, 0, 0, 0, 0, 0, 0, 0 }		},	// ME_SWITCH
+		{ DT_LAYOUT,	false,	false,		{ 0, 0, 0, 0, 0, 0, 0, 0, 0, 0, 0, 0, 0, 0, 0, 0, 0, 0, 0, 0, 0, 0 }		},	// LA_SWITCH
+		{ DT_INSTANCE,	false,	false,		{ 0, 0, 0, 0, 0, 0, 0, 0, 0, 0, 0, 0, 0, 0, 0, 0, 0, 0, 0, 0, 0, 0 }		},	// IN_SWITCH
+
+        { DT_BOOL,		false,  false,		{ 0, 0, 0, 0, 0, 0, 0, 0, 0, 0, 0, 0, 0, 0, 0, 0, 0, 0, 0, 0, 0, 0 }		},	// BO_LESS
+        { DT_BOOL,		false,  false,		{ 0, 0, 0, 0, 0, 0, 0, 0, 0, 0, 0, 0, 0, 0, 0, 0, 0, 0, 0, 0, 0, 0 }		},	// BO_EQUAL_SC_CONST
+        { DT_BOOL,		false,  false,		{ 0, 0, 0, 0, 0, 0, 0, 0, 0, 0, 0, 0, 0, 0, 0, 0, 0, 0, 0, 0, 0, 0 }		},	// BO_AND
+        { DT_BOOL,		false,  false,		{ 0, 0, 0, 0, 0, 0, 0, 0, 0, 0, 0, 0, 0, 0, 0, 0, 0, 0, 0, 0, 0, 0 }		},	// BO_OR
+        { DT_BOOL,		false,  false,		{ 0, 0, 0, 0, 0, 0, 0, 0, 0, 0, 0, 0, 0, 0, 0, 0, 0, 0, 0, 0, 0, 0 }		},	// BO_NOT
+
+        { DT_SCALAR,	true,   false,		{ 0, 0, 0, 0, 0, 0, 0, 0, 0, 0, 0, 0, 0, 0, 0, 0, 0, 0, 0, 0, 0, 0 }		},	// SC_MULTIPLYADD
+        { DT_SCALAR,	true,   false,		{ 0, 0, 0, 0, 0, 0, 0, 0, 0, 0, 0, 0, 0, 0, 0, 0, 0, 0, 0, 0, 0, 0 }		},	// SC_ARITHMETIC
+        { DT_SCALAR,	true,   false,		{ 0, 0, 0, 0, 0, 0, 0, 0, 0, 0, 0, 0, 0, 0, 0, 0, 0, 0, 0, 0, 0, 0 }		},	// SC_CURVE
+
+        { DT_COLOUR,	false,	false,		{ 0, 0, 0, 0, 0, 0, 0, 0, 0, 0, 0, 0, 0, 0, 0, 0, 0, 0, 0, 0, 0, 0 }		},	// CO_SAMPLEIMAGE
+        { DT_COLOUR,	false,	false,		{ 0, 0, 0, 0, 0, 0, 0, 0, 0, 0, 0, 0, 0, 0, 0, 0, 0, 0, 0, 0, 0, 0 }		},	// CO_SWIZZLE
+        { DT_COLOUR,	false,	false,		{ 0, 0, 0, 0, 0, 0, 0, 0, 0, 0, 0, 0, 0, 0, 0, 0, 0, 0, 0, 0, 0, 0 }		},	// CO_IMAGESIZE
+        { DT_COLOUR,	false,	false,		{ 0, 0, 0, 0, 0, 0, 0, 0, 0, 0, 0, 0, 0, 0, 0, 0, 0, 0, 0, 0, 0, 0 }		},	// CO_LAYOUTBLOCKTRANSFORM
+        { DT_COLOUR,	false,	false,		{ 0, 0, 0, 0, 0, 0, 0, 0, 0, 0, 0, 0, 0, 0, 0, 0, 0, 0, 0, 0, 0, 0 }		},	// CO_FROMSCALARS
+        { DT_COLOUR,	false,	false,		{ 0, 0, 0, 0, 0, 0, 0, 0, 0, 0, 0, 0, 0, 0, 0, 0, 0, 0, 0, 0, 0, 0 }		},	// CO_ARITHMETIC
+
+        { DT_IMAGE,		true,	true,		{ 0, 1, 1, 1, 0, 0, 0, 0, 0, 0, 0, 0, 0, 0, 0, 0, 0, 0, 0, 0, 0, 0 }		},	// IM_LAYER
+        { DT_IMAGE,		true,	true,		{ 0, 1, 1, 1, 0, 0, 0, 0, 0, 0, 0, 0, 0, 0, 0, 0, 0, 0, 0, 0, 0, 0 }		},	// IM_LAYERCOLOUR
+        { DT_IMAGE,		true,	false,		{ 1, 1, 1, 1, 1, 1, 1, 1, 1, 1, 1, 1, 0, 0, 0, 0, 0, 0, 0, 0, 0, 0 }		},	// IM_PIXELFORMAT
+        { DT_IMAGE,		true,	false,		{ 1, 1, 1, 1, 0, 0, 0, 0, 0, 0, 0, 0, 0, 0, 0, 0, 0, 0, 0, 0, 0, 0 }		},	// IM_MIPMAP
+        { DT_IMAGE,		true,	false,		{ 0, 1, 1, 1, 0, 0, 0, 0, 0, 0, 0, 0, 0, 0, 0, 0, 0, 0, 0, 0, 0, 0 }		},	// IM_RESIZE
+        { DT_IMAGE,		true,	false,		{ 0, 1, 1, 1, 0, 0, 0, 0, 0, 0, 0, 0, 0, 0, 0, 0, 0, 0, 0, 0, 0, 0 }		},	// IM_RESIZELIKE
+		{ DT_IMAGE,		true,	false,		{ 0, 1, 1, 1, 0, 0, 0, 0, 0, 0, 0, 0, 0, 0, 0, 0, 0, 0, 0, 0, 0, 0 }		},	// IM_RESIZEREL
+		{ DT_IMAGE,		true,   false,		{ 1, 1, 1, 1, 1, 1, 1, 1, 1, 1, 1, 1, 0, 0, 0, 0, 0, 0, 0, 0, 0, 0 }		},	// IM_BLANKLAYOUT
+		{ DT_IMAGE,		true,	true,		{ 1, 1, 1, 1, 1, 1, 1, 1, 1, 1, 1, 0, 0, 0, 1, 1, 1, 1, 1, 1, 1, 1 }		},	// IM_COMPOSE
+        { DT_IMAGE,		true,	false,		{ 0, 1, 1, 0, 0, 0, 0, 0, 0, 0, 0, 0, 0, 0, 0, 0, 0, 0, 0, 0, 0, 0 }		},	// IM_DIFFERENCE
+		{ DT_IMAGE,		true,	false,		{ 0, 1, 1, 1, 0, 0, 0, 0, 0, 0, 0, 0, 0, 0, 0, 0, 0, 0, 0, 0, 0, 0 }		},	// IM_INTERPOLATE
+		{ DT_IMAGE,		true,	true,		{ 0, 1, 1, 1, 0, 0, 0, 0, 0, 0, 0, 0, 0, 0, 0, 0, 0, 0, 0, 0, 0, 0 }		},	// IM_INTERPOLATE3
+		{ DT_IMAGE,		true,	false,		{ 0, 1, 1, 0, 0, 0, 0, 0, 0, 0, 0, 0, 0, 0, 0, 0, 0, 0, 0, 0, 0, 0 }		},	// IM_SATURATE
+		{ DT_IMAGE,		true,	false,		{ 0, 1, 1, 0, 0, 0, 0, 0, 0, 0, 0, 0, 0, 0, 0, 0, 0, 0, 0, 0, 0, 0 }		},	// IM_LUMINANCE
+        { DT_IMAGE,		true,	false,		{ 0, 1, 1, 1, 0, 0, 0, 0, 0, 0, 0, 0, 0, 0, 0, 0, 0, 0, 0, 0, 0, 0 }		},	// IM_SWIZZLE
+		{ DT_IMAGE,		true,	false,		{ 0, 1, 1, 1, 0, 0, 0, 0, 0, 0, 0, 0, 0, 0, 0, 0, 0, 0, 0, 0, 0, 0 }		},	// IM_SELECTCOLOUR
+		{ DT_IMAGE,		true,	false,		{ 0, 1, 1, 1, 0, 0, 0, 0, 0, 0, 0, 0, 0, 0, 0, 0, 0, 0, 0, 0, 0, 0 }		},	// IM_COLOURMAP
+		{ DT_IMAGE,		true,	false,		{ 0, 0, 0, 0, 0, 0, 0, 0, 0, 0, 0, 0, 0, 0, 0, 0, 0, 0, 0, 0, 0, 0 }		},	// IM_GRADIENT
+        { DT_IMAGE,		true,	false,		{ 0, 1, 1, 1, 0, 0, 0, 0, 0, 0, 0, 0, 0, 0, 0, 0, 0, 0, 0, 0, 0, 0 }		},	// IM_BINARISE
+		{ DT_IMAGE,		true,	false,		{ 0, 0, 0, 0, 0, 0, 0, 0, 0, 0, 0, 0, 0, 0, 0, 0, 0, 0, 0, 0, 0, 0 }		},	// IM_PLAINCOLOUR
+		{ DT_IMAGE,		true,	false,		{ 0, 0, 0, 0, 0, 0, 0, 0, 0, 0, 0, 0, 0, 0, 0, 0, 0, 0, 0, 0, 0, 0 }		},	// IM_GPU
+		{ DT_IMAGE,		true,	false,		{ 0, 0, 0, 0, 0, 0, 0, 0, 0, 0, 0, 0, 0, 0, 0, 0, 0, 0, 0, 0, 0, 0 }		},	// IM_CROP
+		{ DT_IMAGE,		true,	false,		{ 0, 0, 0, 0, 0, 0, 0, 0, 0, 0, 0, 0, 0, 0, 0, 0, 0, 0, 0, 0, 0, 0 }		},	// IM_PATCH
+        { DT_IMAGE,		true,	false,		{ 0, 0, 0, 0, 0, 0, 0, 0, 0, 0, 0, 0, 0, 0, 0, 0, 0, 0, 0, 0, 0, 0 }		},	// IM_RASTERMESH
+		{ DT_IMAGE,		true,	false,		{ 0, 0, 0, 0, 0, 0, 0, 0, 0, 0, 0, 0, 0, 0, 0, 0, 0, 0, 0, 0, 0, 0 }		},	// IM_MAKEGROWMAP
+		{ DT_IMAGE,		true,	false,		{ 0, 0, 0, 0, 0, 0, 0, 0, 0, 0, 0, 0, 0, 0, 0, 0, 0, 0, 0, 0, 0, 0 }		},	// IM_DISPLACE
+		{ DT_IMAGE,		true,	false,		{ 0, 0, 0, 0, 0, 0, 0, 0, 0, 0, 0, 0, 0, 0, 0, 0, 0, 0, 0, 0, 0, 0 }		},	// IM_MULTILAYER
+		{ DT_IMAGE,		true,	false,		{ 0, 1, 1, 1, 0, 0, 0, 0, 0, 0, 0, 0, 0, 0, 0, 0, 0, 0, 0, 0, 0, 0 }		},	// IM_INVERT
+		{ DT_IMAGE,		true,	false,		{ 0, 1, 1, 1, 0, 0, 0, 0, 0, 0, 0, 0, 0, 0, 0, 0, 0, 0, 0, 0, 0, 0 }		},	// IM_NORMALCOMPOSITE
+		{ DT_IMAGE,		true,	false,		{ 0, 1, 1, 1, 0, 0, 0, 0, 0, 0, 0, 0, 0, 0, 0, 0, 0, 0, 0, 0, 0, 0 }		},	// IM_TRANSFORM
+
+        { DT_MESH,		true,	false,		{ 0, 0, 0, 0, 0, 0, 0, 0, 0, 0, 0, 0, 0, 0, 0, 0, 0, 0, 0, 0, 0, 0 }		},	// ME_APPLYLAYOUT
+        { DT_MESH,		true,	false,		{ 0, 0, 0, 0, 0, 0, 0, 0, 0, 0, 0, 0, 0, 0, 0, 0, 0, 0, 0, 0, 0, 0 }		},	// ME_DIFFERENCE
+        { DT_MESH,		true,	false,		{ 0, 0, 0, 0, 0, 0, 0, 0, 0, 0, 0, 0, 0, 0, 0, 0, 0, 0, 0, 0, 0, 0 }		},	// ME_MORPH2
+        { DT_MESH,		true,	false,		{ 0, 0, 0, 0, 0, 0, 0, 0, 0, 0, 0, 0, 0, 0, 0, 0, 0, 0, 0, 0, 0, 0 }		},	// ME_MERGE
+        { DT_MESH,		true,	false,		{ 0, 0, 0, 0, 0, 0, 0, 0, 0, 0, 0, 0, 0, 0, 0, 0, 0, 0, 0, 0, 0, 0 }		},	// ME_INTERPOLATE
+        { DT_MESH,		true,	false,		{ 0, 0, 0, 0, 0, 0, 0, 0, 0, 0, 0, 0, 0, 0, 0, 0, 0, 0, 0, 0, 0, 0 }		},	// ME_MASKCLIPMESH
+        { DT_MESH,		true,	false,		{ 0, 0, 0, 0, 0, 0, 0, 0, 0, 0, 0, 0, 0, 0, 0, 0, 0, 0, 0, 0, 0, 0 }		},	// ME_MASKDIFF
+        { DT_MESH,		true,	false,		{ 0, 0, 0, 0, 0, 0, 0, 0, 0, 0, 0, 0, 0, 0, 0, 0, 0, 0, 0, 0, 0, 0 }		},	// ME_SUBTRACT
+        { DT_MESH,		true,	false,		{ 0, 0, 0, 0, 0, 0, 0, 0, 0, 0, 0, 0, 0, 0, 0, 0, 0, 0, 0, 0, 0, 0 }		},	// ME_REMOVEMASK
+        { DT_MESH,		true,	false,		{ 0, 0, 0, 0, 0, 0, 0, 0, 0, 0, 0, 0, 0, 0, 0, 0, 0, 0, 0, 0, 0, 0 }		},	// ME_FORMAT
+        { DT_MESH,		true,	false,		{ 0, 0, 0, 0, 0, 0, 0, 0, 0, 0, 0, 0, 0, 0, 0, 0, 0, 0, 0, 0, 0, 0 }		},	// ME_EXTRACTLAYOUTBLOCK
+        { DT_MESH,		true,	false,		{ 0, 0, 0, 0, 0, 0, 0, 0, 0, 0, 0, 0, 0, 0, 0, 0, 0, 0, 0, 0, 0, 0 }		},	// ME_EXTRACTFACEGROUP
+        { DT_MESH,		true,	false,		{ 0, 0, 0, 0, 0, 0, 0, 0, 0, 0, 0, 0, 0, 0, 0, 0, 0, 0, 0, 0, 0, 0 }		},	// ME_TRANSFORM
+        { DT_MESH,		true,	false,		{ 0, 0, 0, 0, 0, 0, 0, 0, 0, 0, 0, 0, 0, 0, 0, 0, 0, 0, 0, 0, 0, 0 }		},	// ME_CLIPMORPHPLANE
+        { DT_MESH,		true,	false,		{ 0, 0, 0, 0, 0, 0, 0, 0, 0, 0, 0, 0, 0, 0, 0, 0, 0, 0, 0, 0, 0, 0 }		},	// ME_CLIPWITHMESH
+        { DT_MESH,		true,	false,		{ 0, 0, 0, 0, 0, 0, 0, 0, 0, 0, 0, 0, 0, 0, 0, 0, 0, 0, 0, 0, 0, 0 }		},	// ME_SETSKELETON
+        { DT_MESH,		true,	false,		{ 0, 0, 0, 0, 0, 0, 0, 0, 0, 0, 0, 0, 0, 0, 0, 0, 0, 0, 0, 0, 0, 0 }		},	// ME_PROJECT
+        { DT_MESH,		true,	false,		{ 0, 0, 0, 0, 0, 0, 0, 0, 0, 0, 0, 0, 0, 0, 0, 0, 0, 0, 0, 0, 0, 0 }		},	// ME_APPLYPOSE
+		{ DT_MESH,		true,	false,		{ 0, 0, 0, 0, 0, 0, 0, 0, 0, 0, 0, 0, 0, 0, 0, 0, 0, 0, 0, 0, 0, 0 }		},	// ME_REMAPINDICES
+		{ DT_MESH,		true,	false,		{ 0, 0, 0, 0, 0, 0, 0, 0, 0, 0, 0, 0, 0, 0, 0, 0, 0, 0, 0, 0, 0, 0 }		},	// ME_GEOMETRYOPERATION
+		{ DT_MESH,		true,	false,		{ 0, 0, 0, 0, 0, 0, 0, 0, 0, 0, 0, 0, 0, 0, 0, 0, 0, 0, 0, 0, 0, 0 }		},	// ME_BINDSHAPE
+		{ DT_MESH,		true,	false,		{ 0, 0, 0, 0, 0, 0, 0, 0, 0, 0, 0, 0, 0, 0, 0, 0, 0, 0, 0, 0, 0, 0 }		},	// ME_APPLYSHAPE
+		{ DT_MESH,		true,	false,		{ 0, 0, 0, 0, 0, 0, 0, 0, 0, 0, 0, 0, 0, 0, 0, 0, 0, 0, 0, 0, 0, 0 }		},	// ME_CLIPDEFORM
+		{ DT_MESH,		true,	false,		{ 0, 0, 0, 0, 0, 0, 0, 0, 0, 0, 0, 0, 0, 0, 0, 0, 0, 0, 0, 0, 0, 0 }		},	// ME_MORPHRESHAPE
+
+        { DT_INSTANCE,	false,	false,		{ 0, 0, 0, 0, 0, 0, 0, 0, 0, 0, 0, 0, 0, 0, 0, 0, 0, 0, 0, 0, 0, 0 }		},	// IN_ADDMESH
+        { DT_INSTANCE,	false,	false,		{ 0, 0, 0, 0, 0, 0, 0, 0, 0, 0, 0, 0, 0, 0, 0, 0, 0, 0, 0, 0, 0, 0 }		},	// IN_ADDIMAGE
+        { DT_INSTANCE,	false,	false,		{ 0, 0, 0, 0, 0, 0, 0, 0, 0, 0, 0, 0, 0, 0, 0, 0, 0, 0, 0, 0, 0, 0 }		},	// IN_ADDVECTOR
+        { DT_INSTANCE,	false,	false,		{ 0, 0, 0, 0, 0, 0, 0, 0, 0, 0, 0, 0, 0, 0, 0, 0, 0, 0, 0, 0, 0, 0 }		},	// IN_ADDSCALAR
+        { DT_INSTANCE,	false,	false,		{ 0, 0, 0, 0, 0, 0, 0, 0, 0, 0, 0, 0, 0, 0, 0, 0, 0, 0, 0, 0, 0, 0 }		},	// IN_ADDSTRING
+        { DT_INSTANCE,	false,	false,		{ 0, 0, 0, 0, 0, 0, 0, 0, 0, 0, 0, 0, 0, 0, 0, 0, 0, 0, 0, 0, 0, 0 }		},	// IN_ADDSURFACE
+        { DT_INSTANCE,	false,	false,		{ 0, 0, 0, 0, 0, 0, 0, 0, 0, 0, 0, 0, 0, 0, 0, 0, 0, 0, 0, 0, 0, 0 }		},	// IN_ADDCOMPONENT
+        { DT_INSTANCE,	false,	false,		{ 0, 0, 0, 0, 0, 0, 0, 0, 0, 0, 0, 0, 0, 0, 0, 0, 0, 0, 0, 0, 0, 0 }		},	// IN_ADDLOD
+
+        { DT_LAYOUT,	true,   false,		{ 0, 0, 0, 0, 0, 0, 0, 0, 0, 0, 0, 0, 0, 0, 0, 0, 0, 0, 0, 0, 0, 0 }		},	// LA_PACK
+        { DT_LAYOUT,	true,   false,		{ 0, 0, 0, 0, 0, 0, 0, 0, 0, 0, 0, 0, 0, 0, 0, 0, 0, 0, 0, 0, 0, 0 }		},	// LA_MERGE
+        { DT_LAYOUT,	true,   false,		{ 0, 0, 0, 0, 0, 0, 0, 0, 0, 0, 0, 0, 0, 0, 0, 0, 0, 0, 0, 0, 0, 0 }		},	// LA_REMOVEBLOCKS
+    };
+
+    // clang-format on
+
+    static_assert( OP_DESC_COUNT == (int)OP_TYPE::COUNT, "OperationDescMismatch" );
+
+	//---------------------------------------------------------------------------------------------
+	const OP_DESC& GetOpDesc( OP_TYPE type )
+	{
+        return s_opDescs[ (int)type ];
+	}
+
+
+    //---------------------------------------------------------------------------------------------
+    void ForEachReference( OP& op, const std::function<void(OP::ADDRESS*)>& f )
+    {
+        switch ( op.type )
+        {
+        case OP_TYPE::NONE:
+        case OP_TYPE::BO_CONSTANT:
+        case OP_TYPE::NU_CONSTANT:
+        case OP_TYPE::SC_CONSTANT:
+        case OP_TYPE::CO_CONSTANT:
+        case OP_TYPE::IM_CONSTANT:
+        case OP_TYPE::ME_CONSTANT:
+        case OP_TYPE::LA_CONSTANT:
+        case OP_TYPE::PR_CONSTANT:
+        case OP_TYPE::BO_PARAMETER:
+        case OP_TYPE::NU_PARAMETER:
+        case OP_TYPE::SC_PARAMETER:
+        case OP_TYPE::CO_PARAMETER:
+        case OP_TYPE::PR_PARAMETER:
+        case OP_TYPE::IM_PARAMETER:
+            break;
+
+        //-------------------------------------------------------------------------------------
+        case OP_TYPE::BO_LESS:
+            f(&op.args.BoolLess.a );
+            f(&op.args.BoolLess.b );
+            break;
+
+        //-------------------------------------------------------------------------------------
+        case OP_TYPE::BO_EQUAL_INT_CONST:
+            f(&op.args.BoolEqualScalarConst.value );
+            break;
+
+        //-------------------------------------------------------------------------------------
+        case OP_TYPE::BO_AND:
+        case OP_TYPE::BO_OR:
+            f(&op.args.BoolBinary.a );
+            f(&op.args.BoolBinary.b );
+            break;
+
+        //-------------------------------------------------------------------------------------
+        case OP_TYPE::BO_NOT:
+            f(&op.args.BoolNot.source );
+            break;
+
+        //-------------------------------------------------------------------------------------
+        case OP_TYPE::SC_MULTIPLYADD:
+            f(&op.args.ScalarMultiplyAdd.factor0 );
+            f(&op.args.ScalarMultiplyAdd.factor1 );
+            f(&op.args.ScalarMultiplyAdd.add );
+            break;
+
+        //-------------------------------------------------------------------------------------
+        case OP_TYPE::SC_ARITHMETIC:
+            f(&op.args.ScalarArithmetic.a );
+            f(&op.args.ScalarArithmetic.b );
+            break;
+
+        //-------------------------------------------------------------------------------------
+        case OP_TYPE::CO_SAMPLEIMAGE:
+            f(&op.args.ColourSampleImage.image );
+            f(&op.args.ColourSampleImage.x );
+            f(&op.args.ColourSampleImage.y );
+            break;
+
+        case OP_TYPE::CO_SWIZZLE:
+            for (int t=0;t<MUTABLE_OP_MAX_SWIZZLE_CHANNELS;++t)
+            {
+                f(&op.args.ColourSwizzle.sources[t] );
+            }
+            break;
+
+        case OP_TYPE::CO_IMAGESIZE:
+            f(&op.args.ColourSampleImage.image );
+            break;
+
+        case OP_TYPE::CO_LAYOUTBLOCKTRANSFORM:
+            f(&op.args.ColourLayoutBlockTransform.layout );
+            break;
+
+        case OP_TYPE::CO_FROMSCALARS:
+            f(&op.args.ColourFromScalars.x);
+            f(&op.args.ColourFromScalars.y);
+            f(&op.args.ColourFromScalars.z);
+            f(&op.args.ColourFromScalars.w);
+            break;
+
+        case OP_TYPE::CO_ARITHMETIC:
+            f(&op.args.ColourArithmetic.a);
+            f(&op.args.ColourArithmetic.b);
+            break;
+
+        //-------------------------------------------------------------------------------------
+        case OP_TYPE::IM_LAYER:
+            f(&op.args.ImageLayer.base );
+            if ( op.args.ImageLayer.mask )
+            {
+                f(&op.args.ImageLayer.mask );
+            }
+            f(&op.args.ImageLayer.blended );
+            break;
+
+        case OP_TYPE::IM_LAYERCOLOUR:
+            f(&op.args.ImageLayerColour.base );
+            if ( op.args.ImageLayerColour.mask )
+            {
+                f(&op.args.ImageLayerColour.mask );
+            }
+            f(&op.args.ImageLayerColour.colour );
+            break;
+
+        case OP_TYPE::IM_RESIZE:
+            f(&op.args.ImageResize.source );
+            break;
+
+        case OP_TYPE::IM_RESIZELIKE:
+            f(&op.args.ImageResizeLike.source );
+            f(&op.args.ImageResizeLike.sizeSource );
+            break;
+
+        case OP_TYPE::IM_RESIZEREL:
+            f(&op.args.ImageResizeRel.source );
+            break;
+
+        case OP_TYPE::IM_BLANKLAYOUT:
+            f(&op.args.ImageBlankLayout.layout );
+            break;
+
+        case OP_TYPE::IM_DIFFERENCE:
+            f(&op.args.ImageDifference.a );
+            f(&op.args.ImageDifference.b );
+            break;
+
+        case OP_TYPE::IM_INTERPOLATE:
+            f(&op.args.ImageInterpolate.factor );
+
+            for (int t=0;t<MUTABLE_OP_MAX_INTERPOLATE_COUNT;++t)
+            {
+                f(&op.args.ImageInterpolate.targets[t] );
+            }
+            break;
+
+        case OP_TYPE::IM_INTERPOLATE3:
+            f(&op.args.ImageInterpolate3.factor1 );
+            f(&op.args.ImageInterpolate3.factor2 );
+            f(&op.args.ImageInterpolate3.target0 );
+            f(&op.args.ImageInterpolate3.target1 );
+            f(&op.args.ImageInterpolate3.target2 );
+            break;
+
+        case OP_TYPE::IM_SWIZZLE:
+            for (int t=0;t<MUTABLE_OP_MAX_SWIZZLE_CHANNELS;++t)
+            {
+                f(&op.args.ImageSwizzle.sources[t] );
+            }
+            break;
+
+        case OP_TYPE::IM_SATURATE:
+            f(&op.args.ImageSaturate.base );
+            f(&op.args.ImageSaturate.factor );
+            break;
+
+        case OP_TYPE::IM_LUMINANCE:
+            f(&op.args.ImageLuminance.base );
+            break;
+
+        case OP_TYPE::IM_SELECTCOLOUR:
+            f(&op.args.ImageSelectColour.base );
+            f(&op.args.ImageSelectColour.colour );
+            break;
+
+        case OP_TYPE::IM_COLOURMAP:
+            f(&op.args.ImageColourMap.base );
+            f(&op.args.ImageColourMap.mask );
+            f(&op.args.ImageColourMap.map );
+            break;
+
+        case OP_TYPE::IM_GRADIENT:
+            f(&op.args.ImageGradient.colour0 );
+            f(&op.args.ImageGradient.colour1 );
+            break;
+
+        case OP_TYPE::IM_BINARISE:
+            f(&op.args.ImageBinarise.base );
+            f(&op.args.ImageBinarise.threshold );
+            break;
+
+        case OP_TYPE::IM_PLAINCOLOUR:
+            f(&op.args.ImagePlainColour.colour );
+            break;
+
+        case OP_TYPE::IM_CROP:
+            f(&op.args.ImageCrop.source );
+            break;
+
+        case OP_TYPE::IM_RASTERMESH:
+            f(&op.args.ImageRasterMesh.mesh );
+            f(&op.args.ImageRasterMesh.image );
+            f(&op.args.ImageRasterMesh.mask );
+            f(&op.args.ImageRasterMesh.angleFadeProperties );
+            f(&op.args.ImageRasterMesh.projector );
+            break;
+
+        case OP_TYPE::IM_MAKEGROWMAP:
+            f(&op.args.ImageMakeGrowMap.mask );
+            break;
+
+        case OP_TYPE::IM_DISPLACE:
+            f(&op.args.ImageDisplace.source );
+            f(&op.args.ImageDisplace.displacementMap );
+            break;
+
+		case OP_TYPE::IM_INVERT:
+			f(&op.args.ImageInvert.base);
+			break;
+
+        //-------------------------------------------------------------------------------------
+        case OP_TYPE::ME_APPLYLAYOUT:
+            f(&op.args.MeshApplyLayout.layout );
+            f(&op.args.MeshApplyLayout.mesh );
+            break;
+
+        case OP_TYPE::ME_DIFFERENCE:
+            f(&op.args.MeshDifference.base );
+            f(&op.args.MeshDifference.target );
+            break;
+
+        case OP_TYPE::ME_MORPH2:
+            f(&op.args.MeshMorph2.factor );
+            f(&op.args.MeshMorph2.base );
+
+            for (int t=0;t<MUTABLE_OP_MAX_MORPH2_TARGETS;++t)
+            {
+                f(&op.args.MeshMorph2.targets[t] );
+            }
+            break;
+
+        case OP_TYPE::ME_MERGE:
+            f(&op.args.MeshMerge.base );
+            f(&op.args.MeshMerge.added );
+            break;
+
+        case OP_TYPE::ME_INTERPOLATE:
+            f(&op.args.MeshInterpolate.factor );
+            f(&op.args.MeshInterpolate.base );
+
+            for (int t=0;t<MUTABLE_OP_MAX_INTERPOLATE_COUNT-1;++t)
+            {
+                f(&op.args.MeshInterpolate.targets[t] );
+            }
+            break;
+
+        case OP_TYPE::ME_MASKDIFF:
+            f(&op.args.MeshMaskDiff.source );
+            f(&op.args.MeshMaskDiff.fragment );
+            break;
+
+        case OP_TYPE::ME_SUBTRACT:
+            f(&op.args.MeshSubtract.a );
+            f(&op.args.MeshSubtract.b );
+            break;
+
+        case OP_TYPE::ME_EXTRACTFACEGROUP:
+            f(&op.args.MeshExtractFaceGroup.source );
+            break;
+
+        case OP_TYPE::ME_CLIPMORPHPLANE:
+            f(&op.args.MeshClipMorphPlane.source);
+            break;
+
+        case OP_TYPE::ME_CLIPWITHMESH :
+            f(&op.args.MeshClipWithMesh.source);
+            f(&op.args.MeshClipWithMesh.clipMesh);
+            break;
+
+        case OP_TYPE::ME_SETSKELETON :
+            f(&op.args.MeshSetSkeleton.source);
+            f(&op.args.MeshSetSkeleton.skeleton);
+            break;
+
+		case OP_TYPE::ME_PROJECT:
+			f(&op.args.MeshProject.mesh);
+			f(&op.args.MeshProject.projector);
+			break;
+
+        //-------------------------------------------------------------------------------------
+        case OP_TYPE::LA_PACK:
+            f(&op.args.LayoutPack.layout );
+            break;
+
+        case OP_TYPE::LA_MERGE:
+            f(&op.args.LayoutMerge.base );
+            f(&op.args.LayoutMerge.added );
+            break;
+
+        case OP_TYPE::LA_REMOVEBLOCKS:
+            f(&op.args.LayoutRemoveBlocks.source );
+            f(&op.args.LayoutRemoveBlocks.mesh );
+            break;
+
+        default:
+			check( false );
+            break;
+        }
+
+    }
+
+
+
+    //---------------------------------------------------------------------------------------------
+    void ForEachReference( const PROGRAM& program, OP::ADDRESS at, const std::function<void(OP::ADDRESS)>& f )
+    {
+        OP_TYPE type = program.GetOpType(at);
+        switch ( type )
+        {
+        case OP_TYPE::NONE:
+        case OP_TYPE::BO_CONSTANT:
+        case OP_TYPE::NU_CONSTANT:
+        case OP_TYPE::SC_CONSTANT:
+        case OP_TYPE::ST_CONSTANT:
+        case OP_TYPE::CO_CONSTANT:
+        case OP_TYPE::IM_CONSTANT:
+        case OP_TYPE::ME_CONSTANT:
+        case OP_TYPE::LA_CONSTANT:
+        case OP_TYPE::PR_CONSTANT:
+        case OP_TYPE::BO_PARAMETER:
+        case OP_TYPE::NU_PARAMETER:
+        case OP_TYPE::SC_PARAMETER:
+        case OP_TYPE::CO_PARAMETER:
+        case OP_TYPE::PR_PARAMETER:
+        case OP_TYPE::IM_PARAMETER:
+            break;
+
+        case OP_TYPE::SC_CURVE:
+        {
+            auto args = program.GetOpArgs<OP::ScalarCurveArgs>(at);
+            f(args.time );
+            break;
+        }
+
+        case OP_TYPE::NU_CONDITIONAL:
+        case OP_TYPE::SC_CONDITIONAL:
+        case OP_TYPE::CO_CONDITIONAL:
+        case OP_TYPE::IM_CONDITIONAL:
+        case OP_TYPE::ME_CONDITIONAL:
+        case OP_TYPE::LA_CONDITIONAL:
+        case OP_TYPE::IN_CONDITIONAL:
+        {
+            auto args = program.GetOpArgs<OP::ConditionalArgs>(at);
+            f(args.condition );
+            f(args.yes );
+            f(args.no );
+            break;
+        }
+
+        case OP_TYPE::NU_SWITCH:
+        case OP_TYPE::SC_SWITCH:
+        case OP_TYPE::CO_SWITCH:
+        case OP_TYPE::IM_SWITCH:
+        case OP_TYPE::ME_SWITCH:
+        case OP_TYPE::LA_SWITCH:
+        case OP_TYPE::IN_SWITCH:
+        {
+			const uint8_t* data = program.GetOpArgsPointer(at);
+			
+			OP::ADDRESS VarAddress;
+			FMemory::Memcpy( &VarAddress, data, sizeof(OP::ADDRESS));
+			data += sizeof(OP::ADDRESS);
+
+			OP::ADDRESS DefAddress;
+			FMemory::Memcpy( &DefAddress, data, sizeof(OP::ADDRESS));
+			data += sizeof(OP::ADDRESS);
+
+			uint32_t CaseCount;
+			FMemory::Memcpy( &CaseCount, data, sizeof(uint32_t));
+			data += sizeof(uint32_t);
+
+			f(VarAddress);
+			f(DefAddress);
+
+			for ( uint32_t C = 0; C < CaseCount; ++C )
+			{	
+				//int32_t Condition;
+				//FMemory::Memcpy( &Condition, data, sizeof(int32_t));
+				data += sizeof(int32_t);
+				
+				OP::ADDRESS At;
+				FMemory::Memcpy( &At, data, sizeof(OP::ADDRESS) );
+				data += sizeof(OP::ADDRESS);
+
+				if (At)
+				{
+					f(At);
+				}
+			}
+
+            break;
+        }
+
+        //-------------------------------------------------------------------------------------
+        case OP_TYPE::BO_LESS:
+        {
+            auto args = program.GetOpArgs<OP::BoolLessArgs>(at);
+            f(args.a );
+            f(args.b );
+            break;
+        }
+
+        //-------------------------------------------------------------------------------------
+        case OP_TYPE::BO_EQUAL_INT_CONST:
+        {
+            auto args = program.GetOpArgs<OP::BoolEqualScalarConstArgs>(at);
+            f(args.value );
+            break;
+        }
+
+        //-------------------------------------------------------------------------------------
+        case OP_TYPE::BO_AND:
+        case OP_TYPE::BO_OR:
+        {
+            auto args = program.GetOpArgs<OP::BoolBinaryArgs>(at);
+            f(args.a );
+            f(args.b );
+            break;
+        }
+
+        //-------------------------------------------------------------------------------------
+        case OP_TYPE::BO_NOT:
+        {
+            auto args = program.GetOpArgs<OP::BoolNotArgs>(at);
+            f(args.source );
+            break;
+        }
+
+        //-------------------------------------------------------------------------------------
+        case OP_TYPE::SC_MULTIPLYADD:
+        {
+            auto args = program.GetOpArgs<OP::ScalarMultiplyAddArgs>(at);
+            f(args.factor0 );
+            f(args.factor1 );
+            f(args.add );
+            break;
+        }
+
+        //-------------------------------------------------------------------------------------
+        case OP_TYPE::SC_ARITHMETIC:
+        {
+            auto args = program.GetOpArgs<OP::ArithmeticArgs>(at);
+            f(args.a );
+            f(args.b );
+            break;
+        }
+
+        //-------------------------------------------------------------------------------------
+        case OP_TYPE::CO_SAMPLEIMAGE:
+        {
+            auto args = program.GetOpArgs<OP::ColourSampleImageArgs>(at);
+            f(args.image );
+            f(args.x );
+            f(args.y );
+            break;
+        }
+
+        case OP_TYPE::CO_SWIZZLE:
+        {
+            auto args = program.GetOpArgs<OP::ColourSwizzleArgs>(at);
+            for (int t=0;t<MUTABLE_OP_MAX_SWIZZLE_CHANNELS;++t)
+            {
+                f(args.sources[t] );
+            }
+            break;
+        }
+
+        case OP_TYPE::CO_IMAGESIZE:
+        {
+            auto args = program.GetOpArgs<OP::ColourSampleImageArgs>(at);
+            f(args.image );
+            break;
+        }
+
+        case OP_TYPE::CO_LAYOUTBLOCKTRANSFORM:
+        {
+            auto args = program.GetOpArgs<OP::ColourLayoutBlockTransformArgs>(at);
+            f(args.layout );
+            break;
+        }
+
+        case OP_TYPE::CO_FROMSCALARS:
+        {
+            auto args = program.GetOpArgs<OP::ColourFromScalarsArgs>(at);
+            f(args.x);
+            f(args.y);
+            f(args.z);
+            f(args.w);
+            break;
+        }
+
+        case OP_TYPE::CO_ARITHMETIC:
+        {
+            auto args = program.GetOpArgs<OP::ArithmeticArgs>(at);
+            f(args.a);
+            f(args.b);
+            break;
+        }
+
+        //-------------------------------------------------------------------------------------
+        case OP_TYPE::IM_LAYER:
+        {
+            auto args = program.GetOpArgs<OP::ImageLayerArgs>(at);
+            f(args.base );
+            if ( args.mask )
+            {
+                f(args.mask );
+            }
+            f(args.blended );
+            break;
+        }
+
+        case OP_TYPE::IM_LAYERCOLOUR:
+        {
+            auto args = program.GetOpArgs<OP::ImageLayerColourArgs>(at);
+            f(args.base );
+            if ( args.mask )
+            {
+                f(args.mask );
+            }
+            f(args.colour );
+            break;
+        }
+
+        case OP_TYPE::IM_MULTILAYER:
+        {
+            auto args = program.GetOpArgs<OP::ImageMultiLayerArgs>(at);
+            f(args.rangeSize );
+            f(args.base );
+            if ( args.mask )
+            {
+                f(args.mask );
+            }
+            f(args.blended );
+            break;
+        }
+
+		case OP_TYPE::IM_NORMALCOMPOSITE:
+		{
+			auto args = program.GetOpArgs<OP::ImageNormalCompositeArgs>(at);
+			f(args.base);
+			f(args.normal);
+
+			break;
+		}
+
+        case OP_TYPE::IM_PIXELFORMAT:
+        {
+            auto args = program.GetOpArgs<OP::ImagePixelFormatArgs>(at);
+            f(args.source );
+            break;
+        }
+
+        case OP_TYPE::IM_MIPMAP:
+        {
+            auto args = program.GetOpArgs<OP::ImageMipmapArgs>(at);
+            f(args.source );
+            break;
+        }
+
+        case OP_TYPE::IM_RESIZE:
+        {
+            auto args = program.GetOpArgs<OP::ImageResizeArgs>(at);
+            f(args.source );
+            break;
+        }
+
+        case OP_TYPE::IM_RESIZELIKE:
+        {
+            auto args = program.GetOpArgs<OP::ImageResizeLikeArgs>(at);
+            f(args.source );
+            f(args.sizeSource );
+            break;
+        }
+
+        case OP_TYPE::IM_RESIZEREL:
+        {
+            auto args = program.GetOpArgs<OP::ImageResizeRelArgs>(at);
+            f(args.source );
+            break;
+        }
+
+        case OP_TYPE::IM_BLANKLAYOUT:
+        {
+            auto args = program.GetOpArgs<OP::ImageBlankLayoutArgs>(at);
+            f(args.layout );
+            break;
+        }
+
+        case OP_TYPE::IM_COMPOSE:
+        {
+            auto args = program.GetOpArgs<OP::ImageComposeArgs>(at);
+            f(args.layout );
+            f(args.base );
+            f(args.blockImage );
+            f(args.mask );
+            break;
+        }
+
+        case OP_TYPE::IM_DIFFERENCE:
+        {
+            auto args = program.GetOpArgs<OP::ImageDifferenceArgs>(at);
+            f(args.a );
+            f(args.b );
+            break;
+        }
+
+        case OP_TYPE::IM_INTERPOLATE:
+        {
+            auto args = program.GetOpArgs<OP::ImageInterpolateArgs>(at);
+            f(args.factor );
+
+            for (int t=0;t<MUTABLE_OP_MAX_INTERPOLATE_COUNT;++t)
+            {
+                f(args.targets[t] );
+            }
+            break;
+        }
+
+        case OP_TYPE::IM_INTERPOLATE3:
+        {
+            auto args = program.GetOpArgs<OP::ImageInterpolate3Args>(at);
+            f(args.factor1 );
+            f(args.factor2 );
+            f(args.target0 );
+            f(args.target1 );
+            f(args.target2 );
+            break;
+        }
+
+        case OP_TYPE::IM_SWIZZLE:
+        {
+            auto args = program.GetOpArgs<OP::ImageSwizzleArgs>(at);
+            for (int t=0;t<MUTABLE_OP_MAX_SWIZZLE_CHANNELS;++t)
+            {
+                f(args.sources[t] );
+            }
+            break;
+        }
+
+        case OP_TYPE::IM_SATURATE:
+        {
+            auto args = program.GetOpArgs<OP::ImageSaturateArgs>(at);
+            f(args.base );
+            f(args.factor );
+            break;
+        }
+
+        case OP_TYPE::IM_LUMINANCE:
+        {
+            auto args = program.GetOpArgs<OP::ImageLuminanceArgs>(at);
+            f(args.base );
+            break;
+        }
+
+        case OP_TYPE::IM_SELECTCOLOUR:
+        {
+            auto args = program.GetOpArgs<OP::ImageSelectColourArgs>(at);
+            f(args.base );
+            f(args.colour );
+            break;
+        }
+
+        case OP_TYPE::IM_COLOURMAP:
+        {
+            auto args = program.GetOpArgs<OP::ImageColourMapArgs>(at);
+            f(args.base );
+            f(args.mask );
+            f(args.map );
+            break;
+        }
+
+        case OP_TYPE::IM_GRADIENT:
+        {
+            auto args = program.GetOpArgs<OP::ImageGradientArgs>(at);
+            f(args.colour0 );
+            f(args.colour1 );
+            break;
+        }
+
+        case OP_TYPE::IM_BINARISE:
+        {
+            auto args = program.GetOpArgs<OP::ImageBinariseArgs>(at);
+            f(args.base );
+            f(args.threshold );
+            break;
+        }
+
+        case OP_TYPE::IM_PLAINCOLOUR:
+        {
+            auto args = program.GetOpArgs<OP::ImagePlainColourArgs>(at);
+            f(args.colour );
+            break;
+        }
+
+        case OP_TYPE::IM_CROP:
+        {
+            auto args = program.GetOpArgs<OP::ImageCropArgs>(at);
+            f(args.source );
+            break;
+        }
+
+        case OP_TYPE::IM_PATCH:
+        {
+            auto args = program.GetOpArgs<OP::ImagePatchArgs>(at);
+            f(args.base );
+            f(args.patch );
+            break;
+        }
+
+        case OP_TYPE::IM_RASTERMESH:
+        {
+            auto args = program.GetOpArgs<OP::ImageRasterMeshArgs>(at);
+            f(args.mesh );
+            f(args.image );
+            f(args.mask );
+            f(args.angleFadeProperties );
+            f(args.projector );
+            break;
+        }
+
+        case OP_TYPE::IM_MAKEGROWMAP:
+        {
+            auto args = program.GetOpArgs<OP::ImageMakeGrowMapArgs>(at);
+            f(args.mask );
+            break;
+        }
+
+        case OP_TYPE::IM_DISPLACE:
+        {
+            auto args = program.GetOpArgs<OP::ImageDisplaceArgs>(at);
+            f(args.source );
+            f(args.displacementMap );
+            break;
+        }
+
+		case OP_TYPE::IM_INVERT:
+		{
+			auto args = program.GetOpArgs<OP::ImageInvertArgs>(at);
+			f(args.base);
+			break;
+		}
+
+		case OP_TYPE::IM_TRANSFORM:
+		{
+			OP::ImageTransformArgs Args = program.GetOpArgs<OP::ImageTransformArgs>(at);
+			f(Args.base);
+			f(Args.offsetX);
+			f(Args.offsetY);
+			f(Args.scaleX);
+			f(Args.scaleY);
+			f(Args.rotation);
+			break;
+		}
+
+        //-------------------------------------------------------------------------------------
+        case OP_TYPE::ME_APPLYLAYOUT:
+        {
+            auto args = program.GetOpArgs<OP::MeshApplyLayoutArgs>(at);
+            f(args.layout );
+            f(args.mesh );
+            break;
+        }
+
+        case OP_TYPE::ME_DIFFERENCE:
+        {
+            auto args = program.GetOpArgs<OP::MeshDifferenceArgs>(at);
+            f(args.base );
+            f(args.target );
+            break;
+        }
+
+        case OP_TYPE::ME_MORPH2:
+        {
+            auto args = program.GetOpArgs<OP::MeshMorph2Args>(at);
+            f(args.factor );
+            f(args.base );
+
+            for (int t=0;t<MUTABLE_OP_MAX_MORPH2_TARGETS;++t)
+            {
+                f(args.targets[t] );
+            }
+            break;
+        }
+
+        case OP_TYPE::ME_MERGE:
+        {
+            auto args = program.GetOpArgs<OP::MeshMergeArgs>(at);
+            f(args.base );
+            f(args.added );
+            break;
+        }
+
+        case OP_TYPE::ME_INTERPOLATE:
+        {
+            auto args = program.GetOpArgs<OP::MeshInterpolateArgs>(at);
+            f(args.factor );
+            f(args.base );
+
+            for (int t=0;t<MUTABLE_OP_MAX_INTERPOLATE_COUNT-1;++t)
+            {
+                f(args.targets[t] );
+            }
+            break;
+        }
+
+        case OP_TYPE::ME_MASKCLIPMESH:
+        {
+            auto args = program.GetOpArgs<OP::MeshMaskClipMeshArgs>(at);
+            f(args.source );
+            f(args.clip );
+            break;
+        }
+
+        case OP_TYPE::ME_MASKDIFF:
+        {
+            auto args = program.GetOpArgs<OP::MeshMaskDiffArgs>(at);
+            f(args.source );
+            f(args.fragment );
+            break;
+        }
+
+        case OP_TYPE::ME_SUBTRACT:
+        {
+            auto args = program.GetOpArgs<OP::MeshSubtractArgs>(at);
+            f(args.a );
+            f(args.b );
+            break;
+        }
+
+        case OP_TYPE::ME_REMOVEMASK:
+        {
+            const uint8_t* data = program.GetOpArgsPointer(at);
+            mu::OP::ADDRESS source;
+            memcpy( &source, data, sizeof(OP::ADDRESS) ); data+=sizeof(OP::ADDRESS);
+            f(source);
+
+            uint16_t removes = 0;
+            memcpy( &removes, data, sizeof(uint16_t) ); data+=sizeof(uint16_t);
+            for (uint16_t r=0; r<removes; ++r)
+            {
+                mu::OP::ADDRESS condition;
+                memcpy( &condition, data, sizeof(OP::ADDRESS) ); data+=sizeof(OP::ADDRESS);
+                f(condition);
+
+                mu::OP::ADDRESS mask;
+                memcpy( &mask, data, sizeof(OP::ADDRESS) ); data+=sizeof(OP::ADDRESS);
+                f(mask);
+            }
+            break;
+        }
+
+        case OP_TYPE::ME_FORMAT:
+        {
+            auto args = program.GetOpArgs<OP::MeshFormatArgs>(at);
+            f(args.source );
+            f(args.format );
+            break;
+        }
+
+        case OP_TYPE::ME_TRANSFORM:
+        {
+            auto args = program.GetOpArgs<OP::MeshTransformArgs>(at);
+            f(args.source );
+            break;
+        }
+
+        case OP_TYPE::ME_EXTRACTLAYOUTBLOCK:
+        {
+            const uint8_t* data = program.GetOpArgsPointer(at);
+            mu::OP::ADDRESS source;
+            memcpy( &source, data, sizeof(OP::ADDRESS) );
+            f(source);
+            break;
+        }
+
+        case OP_TYPE::ME_EXTRACTFACEGROUP:
+        {
+            auto args = program.GetOpArgs<OP::MeshExtractFaceGroupArgs>(at);
+            f(args.source );
+            break;
+        }
+
+        case OP_TYPE::ME_CLIPMORPHPLANE:
+        {
+            auto args = program.GetOpArgs<OP::MeshClipMorphPlaneArgs>(at);
+            f(args.source);
+            break;
+        }
+
+        case OP_TYPE::ME_CLIPWITHMESH :
+        {
+            auto args = program.GetOpArgs<OP::MeshClipWithMeshArgs>(at);
+            f(args.source);
+            f(args.clipMesh);
+            break;
+        }
+
+        case OP_TYPE::ME_CLIPDEFORM:
+        {
+            auto args = program.GetOpArgs<OP::MeshClipDeformArgs>(at);
+            f(args.mesh);
+            f(args.clipShape);
+            break;
+        }
+       
+		case OP_TYPE::ME_MORPHRESHAPE:
+		{
+			OP::MeshMorphReshapeArgs Args = program.GetOpArgs<OP::MeshMorphReshapeArgs>(at);
+			f(Args.Morph);
+			f(Args.Reshape);
+
+			break;
+		}
+ 	
+        case OP_TYPE::ME_REMAPINDICES :
+        {
+            auto args = program.GetOpArgs<OP::MeshRemapIndicesArgs>(at);
+            f(args.source);
+            f(args.reference);
+            break;
+        }
+
+        case OP_TYPE::ME_SETSKELETON :
+        {
+            auto args = program.GetOpArgs<OP::MeshSetSkeletonArgs>(at);
+            f(args.source);
+            f(args.skeleton);
+            break;
+        }
+
+        case OP_TYPE::ME_PROJECT :
+        {
+            auto args = program.GetOpArgs<OP::MeshProjectArgs>(at);
+            f(args.mesh);
+            f(args.projector);
+            break;
+        }
+
+		case OP_TYPE::ME_APPLYPOSE:
+		{
+			auto args = program.GetOpArgs<OP::MeshApplyPoseArgs>(at);
+			f(args.base);
+			f(args.pose);
+			break;
+		}
+
+		case OP_TYPE::ME_GEOMETRYOPERATION:
+		{
+			auto args = program.GetOpArgs<OP::MeshGeometryOperationArgs>(at);
+			f(args.meshA);
+			f(args.meshB);
+			f(args.scalarA);
+			f(args.scalarB);
+			break;
+		}
+
+		case OP_TYPE::ME_BINDSHAPE:
+		{
+			auto args = program.GetOpArgs<OP::MeshBindShapeArgs>(at);
+			f(args.mesh);
+			f(args.shape);
+			break;
+		}
+
+		case OP_TYPE::ME_APPLYSHAPE:
+		{
+			auto args = program.GetOpArgs<OP::MeshApplyShapeArgs>(at);
+			f(args.mesh);
+			f(args.shape);
+			break;
+		}
+
+        //-------------------------------------------------------------------------------------
+        case OP_TYPE::IN_ADDMESH:
+        case OP_TYPE::IN_ADDIMAGE:
+        case OP_TYPE::IN_ADDVECTOR:
+        case OP_TYPE::IN_ADDSCALAR:
+        case OP_TYPE::IN_ADDSTRING:
+        case OP_TYPE::IN_ADDCOMPONENT:
+        case OP_TYPE::IN_ADDSURFACE:
+        {
+            auto args = program.GetOpArgs<OP::InstanceAddArgs>(at);
+            f(args.instance );
+            f(args.value );
+            break;
+        }
+
+        case OP_TYPE::IN_ADDLOD:
+        {
+            auto args = program.GetOpArgs<OP::InstanceAddLODArgs>(at);
+            for (int t=0;t<MUTABLE_OP_MAX_ADD_COUNT;++t)
+            {
+                f(args.lod[t] );
+            }
+            break;
+        }
+
+        //-------------------------------------------------------------------------------------
+        case OP_TYPE::LA_PACK:
+        {
+            auto args = program.GetOpArgs<OP::LayoutPackArgs>(at);
+            f(args.layout );
+            break;
+        }
+
+        case OP_TYPE::LA_MERGE:
+        {
+            auto args = program.GetOpArgs<OP::LayoutMergeArgs>(at);
+            f(args.base );
+            f(args.added );
+            break;
+        }
+
+        case OP_TYPE::LA_REMOVEBLOCKS:
+        {
+            auto args = program.GetOpArgs<OP::LayoutRemoveBlocksArgs>(at);
+            f(args.source );
+            f(args.mesh );
+            break;
+        }
+
+        default:
+			check( false );
+            break;
+        }
+
+    }
+
+
+}
+
