[mutable] Move Mutable plugin from the restricted folder to the experimental plugins folder.

diff --git a/Engine/Plugins/Experimental/Mutable/Source/MutableRuntime/Private/Parameters.cpp b/Engine/Plugins/Experimental/Mutable/Source/MutableRuntime/Private/Parameters.cpp
new file mode 100644
--- /dev/null
+++ b/Engine/Plugins/Experimental/Mutable/Source/MutableRuntime/Private/Parameters.cpp
@@ -0,0 +1,1259 @@
+// Copyright Epic Games, Inc. All Rights Reserved.
+
+
+#include "Parameters.h"
+#include "Platform.h"
+
+#include "ParametersPrivate.h"
+
+#include "ModelPrivate.h"
+#include "Operations.h"
+
+
+namespace mu
+{
+
+	//---------------------------------------------------------------------------------------------
+	Parameters::Parameters()
+	{
+		m_pD = new Private();
+	}
+
+
+	//---------------------------------------------------------------------------------------------
+	Parameters::~Parameters()
+	{
+        check( m_pD );
+		delete m_pD;
+		m_pD = 0;
+	}
+
+
+	//---------------------------------------------------------------------------------------------
+	void Parameters::Serialise( const Parameters* p, OutputArchive& arch )
+	{
+		arch << *p->m_pD;
+    }
+
+
+	//---------------------------------------------------------------------------------------------
+	ParametersPtr Parameters::StaticUnserialise( InputArchive& arch )
+	{
+		LLM_SCOPE_BYNAME(TEXT("MutableRuntime"));
+
+		ParametersPtr pResult = new Parameters();
+		arch >> *pResult->m_pD;
+		return pResult;
+	}
+
+
+	//---------------------------------------------------------------------------------------------
+	void Parameters::SetUnserialisedModel( Model* pModel )
+	{
+		m_pD->m_pModel = pModel;
+	}
+
+
+    //---------------------------------------------------------------------------------------------
+    void Parameters::SerialisePortable( const Parameters* params, OutputArchive& arch )
+    {
+        const int32_t ver = 0;
+        arch << ver;
+
+        uint32_t valueCount = uint32_t(params->GetPrivate()->m_values.size());
+        arch << valueCount;
+
+        for (size_t p=0; p<valueCount; ++p)
+        {
+            const auto& desc = params->GetPrivate()->m_pModel->GetPrivate()->m_program.m_parameters[p];
+            arch << desc.m_name;
+            arch << desc.m_uid;
+            arch << desc.m_type;
+            arch << params->GetPrivate()->m_values[p];
+        }
+    }
+
+
+    //---------------------------------------------------------------------------------------------
+    ParametersPtr Parameters::UnserialisePortable( InputArchive& arch, const Model* pModel )
+    {
+		LLM_SCOPE_BYNAME(TEXT("MutableRuntime"));
+
+        ParametersPtr pResult = pModel->NewParameters();
+        size_t modelParameters = pModel->GetPrivate()->m_program.m_parameters.size();
+
+        int32_t ver;
+        arch >> ver;
+        check( ver == 0 );
+
+        uint32_t valueCount = 0;
+        arch >> valueCount;
+
+        for (uint32_t p=0; p<valueCount; ++p)
+        {
+            // Read the parameter data.
+            mu::string name;
+            arch >> name;
+
+            mu::string uid;
+            arch >> uid;
+
+            PARAMETER_TYPE type;
+            arch >> type;
+
+            PARAMETER_VALUE value;
+            arch >> value;
+
+            int modelParam = -1;
+
+            // UIDs have higher priority
+            if ( !uid.empty() )
+            {
+                for (size_t mp=0; mp<modelParameters; ++mp)
+                {
+                    const auto& desc = pModel->GetPrivate()->m_program.m_parameters[mp];
+                    if (desc.m_uid==uid)
+                    {
+                        modelParam = int(mp);
+                        break;
+                    }
+                }
+            }
+
+            // Try with name+type
+            if (modelParam<0)
+            {
+                for (size_t mp=0; mp<modelParameters; ++mp)
+                {
+                    const auto& desc = pModel->GetPrivate()->m_program.m_parameters[mp];
+                    if (desc.m_name==name && desc.m_type==type)
+                    {
+                        modelParam = int(mp);
+                        break;
+                    }
+                }
+            }
+
+            // Try with name only, and maybe adapt type
+            if (modelParam<0)
+            {
+                for (size_t mp=0; mp<modelParameters; ++mp)
+                {
+                    const auto& desc = pModel->GetPrivate()->m_program.m_parameters[mp];
+                    if (desc.m_name==name)
+                    {
+                        modelParam = int(mp);
+                        break;
+                    }
+                }
+            }
+
+            if (modelParam>=0)
+            {
+                // We found something. \todo: convert the type if necessary?
+                pResult->GetPrivate()->m_values[modelParam] = value;
+            }
+        }
+
+        return pResult;
+    }
+
+
+	//---------------------------------------------------------------------------------------------
+	Parameters::Private* Parameters::GetPrivate() const
+	{
+		return m_pD;
+	}
+
+
+	//---------------------------------------------------------------------------------------------
+	ParametersPtr Parameters::Clone() const
+	{
+		LLM_SCOPE_BYNAME(TEXT("MutableRuntime"));
+
+		ParametersPtr pRes = new Parameters();
+
+		pRes->m_pD->m_pModel = m_pD->m_pModel;
+		pRes->m_pD->m_values = m_pD->m_values;
+		pRes->m_pD->m_multiValues = m_pD->m_multiValues;
+
+		return pRes;
+	}
+
+
+	//---------------------------------------------------------------------------------------------
+	int Parameters::GetCount() const
+	{
+		return (int)m_pD->m_values.size();
+	}
+
+
+	//---------------------------------------------------------------------------------------------
+	const char* Parameters::GetName( int index ) const
+	{
+		const PROGRAM& program = m_pD->m_pModel->GetPrivate()->m_program;
+		check( index>=0 && index<(int)program.m_parameters.size() );
+
+		return program.m_parameters[index].m_name.c_str();
+	}
+
+
+	//---------------------------------------------------------------------------------------------
+	const char* Parameters::GetUid( int index ) const
+	{
+		const PROGRAM& program = m_pD->m_pModel->GetPrivate()->m_program;
+		check( index>=0 && index<(int)program.m_parameters.size() );
+
+		return program.m_parameters[index].m_uid.c_str();
+	}
+
+
+	//---------------------------------------------------------------------------------------------
+	int Parameters::Find( const char* strName ) const
+	{
+		return m_pD->Find( strName );
+	}
+
+
+	//---------------------------------------------------------------------------------------------
+	PARAMETER_TYPE Parameters::GetType( int index ) const
+	{
+		const PROGRAM& program = m_pD->m_pModel->GetPrivate()->m_program;
+		check( index>=0 && index<(int)program.m_parameters.size() );
+
+		return program.m_parameters[index].m_type;
+	}
+
+
+	//---------------------------------------------------------------------------------------------
+	PARAMETER_DETAILED_TYPE Parameters::GetDetailedType( int index ) const
+	{
+		const PROGRAM& program = m_pD->m_pModel->GetPrivate()->m_program;
+		check( index>=0 && index<(int)program.m_parameters.size() );
+
+		return program.m_parameters[index].m_detailedType;
+	}
+
+
+    //---------------------------------------------------------------------------------------------
+    Ptr<RangeIndex> Parameters::NewRangeIndex( int paramIndex ) const
+    {
+		LLM_SCOPE_BYNAME(TEXT("MutableRuntime"));
+
+        const PROGRAM& program = m_pD->m_pModel->GetPrivate()->m_program;
+        check( paramIndex>=0 && paramIndex<int(program.m_parameters.size()) );
+        Ptr<RangeIndex> range;
+
+        if ( paramIndex>=0 && paramIndex<int(program.m_parameters.size()) )
+        {
+            size_t rangeCount = program.m_parameters[paramIndex].m_ranges.size();
+            if ( rangeCount>0 )
+            {
+                range = new RangeIndex;
+                range->m_pD->m_pParameters = this;
+                range->m_pD->m_parameter = paramIndex;
+                range->m_pD->m_values.resize( rangeCount, 0 );
+            }
+        }
+
+        return range;
+    }
+
+
+    //---------------------------------------------------------------------------------------------
+    int Parameters::GetValueCount( int paramIndex ) const
+    {
+        if (paramIndex<0 || paramIndex >= int(m_pD->m_multiValues.size()) )
+        {
+            return 0;
+        }
+
+        return int( m_pD->m_multiValues[paramIndex].size() );
+    }
+
+
+    //---------------------------------------------------------------------------------------------
+    Ptr<RangeIndex> Parameters::GetValueIndex( int paramIndex, int valueIndex ) const
+    {
+        if (paramIndex<0 || paramIndex >= int(m_pD->m_multiValues.size()) )
+        {
+            return nullptr;
+        }
+
+        if (valueIndex<0 || valueIndex >= int(m_pD->m_multiValues[paramIndex].size()) )
+        {
+            return nullptr;
+        }
+
+        auto it = m_pD->m_multiValues[paramIndex].begin();
+        for ( int i=0; i<valueIndex; ++i )
+        {
+            ++it;
+        }
+
+        Ptr<RangeIndex> result = NewRangeIndex( paramIndex );
+        result->m_pD->m_values = it->first;
+
+        return result;
+    }
+
+
+    //---------------------------------------------------------------------------------------------
+    void Parameters::ClearAllValues( int paramIndex )
+    {
+        if (paramIndex<0 || paramIndex >= int(m_pD->m_multiValues.size()) )
+        {
+            return;
+        }
+
+        m_pD->m_multiValues[paramIndex].clear();
+    }
+
+
+	//---------------------------------------------------------------------------------------------
+	int Parameters::GetAdditionalImageCount( int index ) const
+	{
+		const PROGRAM& program = m_pD->m_pModel->GetPrivate()->m_program;
+		check( index>=0 && index<(int)program.m_parameters.size() );
+		return (int)program.m_parameters[index].m_descImages.size();
+	}
+
+
+	//---------------------------------------------------------------------------------------------
+    bool Parameters::GetBoolValue( int index,
+                                   const Ptr<const RangeIndex>& pos ) const
+	{
+		check( index>=0 && index<(int)m_pD->m_values.size() );
+        check( GetType(index)==PARAMETER_TYPE::T_BOOL );
+
+        // Early out in case of invalid parameters
+        if ( index < 0
+             ||
+             index >= (int)m_pD->m_values.size()
+             ||
+             GetType(index) != PARAMETER_TYPE::T_BOOL )
+        {
+            return false;
+        }
+
+        // Single value case
+        if (!pos)
+        {
+            // Return the single value
+            return m_pD->m_values[index].m_bool;
+        }
+
+        // Multivalue case
+        check( pos->m_pD->m_parameter==index );
+
+        if ( index<int(m_pD->m_multiValues.size()))
+        {
+            const auto& m = m_pD->m_multiValues[index];
+            auto it = m.find(pos->m_pD->m_values);
+            if (it!=m.end())
+            {
+                return it->second.m_bool;
+            }
+        }
+
+        // Multivalue parameter, but no multivalue set. Return single value.
+        return m_pD->m_values[index].m_bool;
+	}
+
+
+	//---------------------------------------------------------------------------------------------
+    void Parameters::SetBoolValue( int index, bool value,
+                                   const Ptr<const RangeIndex>& pos )
+	{
+		check( index>=0 && index<(int)m_pD->m_values.size() );
+		check( GetType(index)== PARAMETER_TYPE::T_BOOL );
+
+        // Early out in case of invalid parameters
+        if ( index < 0
+             ||
+             index >= (int)m_pD->m_values.size()
+             ||
+             GetType(index) != PARAMETER_TYPE::T_BOOL )
+        {
+            return;
+        }
+
+        // Single value case
+        if (!pos)
+        {
+            // Clear multivalue, if set.
+            if (index<int(m_pD->m_multiValues.size()))
+            {
+                m_pD->m_multiValues[index].clear();
+            }
+
+            m_pD->m_values[index].m_bool = value;
+        }
+
+        // Multivalue case
+        else
+        {
+            check( pos->m_pD->m_parameter==index );
+
+            if ( index>=int(m_pD->m_multiValues.size()))
+            {
+                m_pD->m_multiValues.resize(index+1);
+            }
+
+            auto& m = m_pD->m_multiValues[index];
+            auto it = m.find(pos->m_pD->m_values);
+            if (it!=m.end())
+            {
+                it->second.m_bool = value;
+            }
+            else
+            {
+                PARAMETER_VALUE valueData;
+                valueData.m_bool = value;
+                m.insert( std::make_pair<>( pos->m_pD->m_values, valueData ) );
+            }
+        }
+	}
+
+
+	//---------------------------------------------------------------------------------------------
+	int Parameters::GetIntPossibleValueCount( int paramIndex ) const
+	{
+		const PROGRAM& program = m_pD->m_pModel->GetPrivate()->m_program;
+		check( paramIndex>=0 && paramIndex<(int)program.m_parameters.size() );
+		return (int)program.m_parameters[paramIndex].m_possibleValues.size();
+	}
+
+
+	//---------------------------------------------------------------------------------------------
+	int Parameters::GetIntPossibleValue( int paramIndex, int valueIndex ) const
+	{
+		const PROGRAM& program = m_pD->m_pModel->GetPrivate()->m_program;
+		check( paramIndex>=0
+				&& paramIndex<(int)program.m_parameters.size() );
+		check( valueIndex>=0
+				&& valueIndex<(int)program.m_parameters[paramIndex].m_possibleValues.size() );
+
+		return (int)program.m_parameters[paramIndex].m_possibleValues[valueIndex].m_value;
+	}
+
+
+	//---------------------------------------------------------------------------------------------
+	int Parameters::GetIntValueIndex(int paramIndex, const char* valueName) const
+	{
+		const PROGRAM& program = m_pD->m_pModel->GetPrivate()->m_program;
+		check(paramIndex >= 0
+			&& paramIndex < (int)program.m_parameters.size());
+
+		int result = -1;
+		for (size_t v = 0; v < program.m_parameters[paramIndex].m_possibleValues.size(); ++v)
+		{
+			if (program.m_parameters[paramIndex].m_possibleValues[v].m_name == valueName)
+			{
+				result = (int)v;
+				break;
+			}
+		}
+		return result;
+	}
+
+
+	//---------------------------------------------------------------------------------------------
+	int Parameters::GetIntValueIndex(int paramIndex, int32 Value) const
+	{
+		const PROGRAM& program = m_pD->m_pModel->GetPrivate()->m_program;
+		check(paramIndex >= 0
+			&& paramIndex < (int)program.m_parameters.size());
+
+		for (size_t v = 0; v < program.m_parameters[paramIndex].m_possibleValues.size(); ++v)
+		{
+			if (program.m_parameters[paramIndex].m_possibleValues[v].m_value == Value)
+			{
+				return (int)v;
+			}
+		}
+		return -1;
+	}
+
+
+	//---------------------------------------------------------------------------------------------
+	const char* Parameters::GetIntPossibleValueName( int paramIndex, int valueIndex ) const
+	{
+		const PROGRAM& program = m_pD->m_pModel->GetPrivate()->m_program;
+		check( paramIndex>=0
+				&& paramIndex<(int)program.m_parameters.size() );
+		check( valueIndex>=0
+				&& valueIndex<(int)program.m_parameters[paramIndex].m_possibleValues.size() );
+
+		return program.m_parameters[paramIndex].m_possibleValues[valueIndex].m_name.c_str();
+	}
+
+
+	//---------------------------------------------------------------------------------------------
+    int Parameters::GetIntValue( int index,
+                                 const Ptr<const RangeIndex>& pos ) const
+	{
+		check( index>=0 && index<(int)m_pD->m_values.size() );
+		check( GetType(index)== PARAMETER_TYPE::T_INT );
+
+        // Early out in case of invalid parameters
+        if ( index < 0
+             ||
+             index >= (int)m_pD->m_values.size()
+             ||
+             GetType(index) != PARAMETER_TYPE::T_INT )
+        {
+            return 0;
+        }
+
+        // Single value case
+        if (!pos)
+        {
+            // Return the single value
+            return m_pD->m_values[index].m_int;
+        }
+
+        // Multivalue case
+        check( pos->m_pD->m_parameter==index );
+
+        if ( index<int(m_pD->m_multiValues.size()))
+        {
+            const auto& m = m_pD->m_multiValues[index];
+            auto it = m.find(pos->m_pD->m_values);
+            if (it!=m.end())
+            {
+                return it->second.m_int;
+            }
+        }
+
+        // Multivalue parameter, but no multivalue set. Return single value.
+        return m_pD->m_values[index].m_int;
+    }
+
+
+	//---------------------------------------------------------------------------------------------
+    void Parameters::SetIntValue( int index, int value,
+                                  const Ptr<const RangeIndex>& pos )
+	{
+		check( index>=0 && index<(int)m_pD->m_values.size() );
+		check( GetType(index)== PARAMETER_TYPE::T_INT );
+
+        // Early out in case of invalid parameters
+        if ( index < 0
+             ||
+             index >= (int)m_pD->m_values.size()
+             ||
+             GetType(index) != PARAMETER_TYPE::T_INT )
+        {
+            return;
+        }
+
+        // Single value case
+        if (!pos)
+        {
+            // Clear multivalue, if set.
+            if (index<int(m_pD->m_multiValues.size()))
+            {
+                m_pD->m_multiValues[index].clear();
+            }
+
+            m_pD->m_values[index].m_int = value;
+        }
+
+        // Multivalue case
+        else
+        {
+            check( pos->m_pD->m_parameter==index );
+
+            if ( index>=int(m_pD->m_multiValues.size()))
+            {
+                m_pD->m_multiValues.resize(index+1);
+            }
+
+            auto& m = m_pD->m_multiValues[index];
+            auto it = m.find(pos->m_pD->m_values);
+            if (it!=m.end())
+            {
+                it->second.m_int = value;
+            }
+            else
+            {
+                PARAMETER_VALUE valueData;
+                valueData.m_int = value;
+                m.insert( std::make_pair<>( pos->m_pD->m_values, valueData ) );
+            }
+        }
+    }
+
+
+	//---------------------------------------------------------------------------------------------
+    float Parameters::GetFloatValue( int index,
+                                     const Ptr<const RangeIndex>& pos ) const
+	{
+		check( index>=0 && index<(int)m_pD->m_values.size() );
+		check( GetType(index)== PARAMETER_TYPE::T_FLOAT );
+
+        // Early out in case of invalid parameters
+        if ( index < 0
+             ||
+             index >= (int)m_pD->m_values.size()
+             ||
+             GetType(index) != PARAMETER_TYPE::T_FLOAT )
+        {
+            return false;
+        }
+
+        // Single value case
+        if (!pos)
+        {
+            // Return the single value
+            return m_pD->m_values[index].m_float;
+        }
+
+        // Multivalue case
+        check( pos->m_pD->m_parameter==index );
+
+        if ( index<int(m_pD->m_multiValues.size()))
+        {
+            const auto& m = m_pD->m_multiValues[index];
+            auto it = m.find(pos->m_pD->m_values);
+            if (it!=m.end())
+            {
+                return it->second.m_float;
+            }
+        }
+
+        // Multivalue parameter, but no multivalue set. Return single value.
+        return m_pD->m_values[index].m_float;
+    }
+
+
+	//---------------------------------------------------------------------------------------------
+    void Parameters::SetFloatValue( int index, float value,
+                                    const Ptr<const RangeIndex>& pos )
+	{
+		check( index>=0 && index<(int)m_pD->m_values.size() );
+		check( GetType(index)== PARAMETER_TYPE::T_FLOAT );
+
+        // Early out in case of invalid parameters
+        if ( index < 0
+             ||
+             index >= (int)m_pD->m_values.size()
+             ||
+             GetType(index) != PARAMETER_TYPE::T_FLOAT )
+        {
+            return;
+        }
+
+        // Single value case
+        if (!pos)
+        {
+            // Clear multivalue, if set.
+            if (index<int(m_pD->m_multiValues.size()))
+            {
+                m_pD->m_multiValues[index].clear();
+            }
+
+            m_pD->m_values[index].m_float = value;
+        }
+
+        // Multivalue case
+        else
+        {
+            check( pos->m_pD->m_parameter==index );
+
+            if ( index>=int(m_pD->m_multiValues.size()))
+            {
+                m_pD->m_multiValues.resize(index+1);
+            }
+
+            auto& m = m_pD->m_multiValues[index];
+            auto it = m.find(pos->m_pD->m_values);
+            if (it!=m.end())
+            {
+                it->second.m_float = value;
+            }
+            else
+            {
+                PARAMETER_VALUE valueData;
+                valueData.m_float = value;
+                m.insert( std::make_pair<>( pos->m_pD->m_values, valueData ) );
+            }
+        }
+    }
+
+
+    //---------------------------------------------------------------------------------------------
+    void Parameters::GetColourValue( int index, float* pR, float* pG, float* pB,
+                                     const Ptr<const RangeIndex>& pos ) const
+    {
+        check( index>=0 && index<(int)m_pD->m_values.size() );
+        check( GetType(index)== PARAMETER_TYPE::T_COLOUR );
+
+        // Early out in case of invalid parameters
+        if ( index < 0
+             ||
+             index >= (int)m_pD->m_values.size()
+             ||
+             GetType(index) != PARAMETER_TYPE::T_COLOUR )
+        {
+            return;
+        }
+
+        // Single value case
+        if (!pos)
+        {
+            // Return the single value
+            if (pR) *pR = m_pD->m_values[index].m_colour[0];
+            if (pG) *pG = m_pD->m_values[index].m_colour[1];
+            if (pB) *pB = m_pD->m_values[index].m_colour[2];
+            return;
+        }
+
+        // Multivalue case
+        check( pos->m_pD->m_parameter==index );
+
+        if ( index<int(m_pD->m_multiValues.size()))
+        {
+            const auto& m = m_pD->m_multiValues[index];
+            auto it = m.find(pos->m_pD->m_values);
+            if (it!=m.end())
+            {
+                if (pR) *pR = it->second.m_colour[0];
+                if (pG) *pG = it->second.m_colour[1];
+                if (pB) *pB = it->second.m_colour[2];
+                return;
+            }
+        }
+
+        // Multivalue parameter, but no multivalue set. Return single value.
+        if (pR) *pR = m_pD->m_values[index].m_colour[0];
+        if (pG) *pG = m_pD->m_values[index].m_colour[1];
+        if (pB) *pB = m_pD->m_values[index].m_colour[2];
+        return;
+    }
+
+
+    //---------------------------------------------------------------------------------------------
+    void Parameters::SetColourValue( int index, float r, float g, float b,
+                                     const Ptr<const RangeIndex>& pos )
+    {
+        check( index>=0 && index<(int)m_pD->m_values.size() );
+        check( GetType(index)== PARAMETER_TYPE::T_COLOUR );
+
+        // Early out in case of invalid parameters
+        if ( index < 0
+             ||
+             index >= (int)m_pD->m_values.size()
+             ||
+             GetType(index) != PARAMETER_TYPE::T_COLOUR )
+        {
+            return;
+        }
+
+        // Single value case
+        if (!pos)
+        {
+            // Clear multivalue, if set.
+            if (index<int(m_pD->m_multiValues.size()))
+            {
+                m_pD->m_multiValues[index].clear();
+            }
+
+            m_pD->m_values[index].m_colour[0] = r;
+            m_pD->m_values[index].m_colour[1] = g;
+            m_pD->m_values[index].m_colour[2] = b;
+        }
+
+        // Multivalue case
+        else
+        {
+            check( pos->m_pD->m_parameter==index );
+
+            if ( index>=int(m_pD->m_multiValues.size()))
+            {
+                m_pD->m_multiValues.resize(index+1);
+            }
+
+            auto& m = m_pD->m_multiValues[index];
+            auto it = m.find(pos->m_pD->m_values);
+            if (it!=m.end())
+            {
+                it->second.m_colour[0] = r;
+                it->second.m_colour[1] = g;
+                it->second.m_colour[2] = b;
+            }
+            else
+            {
+                PARAMETER_VALUE valueData;
+                valueData.m_colour[0] = r;
+                valueData.m_colour[1] = g;
+                valueData.m_colour[2] = b;
+                m.insert( std::make_pair<>( pos->m_pD->m_values, valueData ) );
+            }
+        }
+    }
+
+
+    //---------------------------------------------------------------------------------------------
+    EXTERNAL_IMAGE_ID Parameters::GetImageValue( int index ) const
+    {
+        check( index >= 0 && index < (int)m_pD->m_values.size() );
+        check( GetType( index ) == PARAMETER_TYPE::T_IMAGE );
+
+        return m_pD->m_values[index].m_image;
+    }
+
+
+    //---------------------------------------------------------------------------------------------
+    void Parameters::SetImageValue( int index, EXTERNAL_IMAGE_ID id )
+    {
+        check( index >= 0 && index < (int)m_pD->m_values.size() );
+        check( GetType( index ) == PARAMETER_TYPE::T_IMAGE );
+
+        m_pD->m_values[index].m_image = id;
+    }
+
+
+    //---------------------------------------------------------------------------------------------
+    const char* Parameters::GetStringValue( int index, const Ptr<const RangeIndex>& pos ) const
+    {
+        check( index >= 0 && index < (int)m_pD->m_values.size() );
+        check( GetType( index ) == PARAMETER_TYPE::T_STRING );
+
+        // Early out in case of invalid parameters
+        if ( index < 0 || index >= (int)m_pD->m_values.size() ||
+             GetType( index ) != PARAMETER_TYPE::T_STRING )
+        {
+            return "";
+        }
+
+        // Single value case
+        if ( !pos )
+        {
+            // Return the single value
+            return m_pD->m_values[index].m_text;
+        }
+
+        // Multivalue case
+        check( pos->m_pD->m_parameter == index );
+
+        if ( index < int( m_pD->m_multiValues.size() ) )
+        {
+            const auto& m = m_pD->m_multiValues[index];
+            auto it = m.find( pos->m_pD->m_values );
+            if ( it != m.end() )
+            {
+                return it->second.m_text;
+            }
+        }
+
+        // Multivalue parameter, but no multivalue set. Return single value.
+        return m_pD->m_values[index].m_text;
+    }
+
+
+    //---------------------------------------------------------------------------------------------
+    void Parameters::SetStringValue( int index, const char* value, const Ptr<const RangeIndex>& pos )
+    {
+		LLM_SCOPE_BYNAME(TEXT("MutableRuntime"));
+
+        check( index >= 0 && index < (int)m_pD->m_values.size() );
+        check( GetType( index ) == PARAMETER_TYPE::T_FLOAT );
+
+        // Early out in case of invalid parameters
+        if ( index < 0 || index >= (int)m_pD->m_values.size() ||
+             GetType( index ) != PARAMETER_TYPE::T_FLOAT )
+        {
+            return;
+        }
+
+        int len = std::min( MUTABLE_MAX_STRING_PARAM_LENGTH,
+                            int( strnlen( value, MUTABLE_MAX_STRING_PARAM_LENGTH ) ) );
+
+        // Single value case
+        if ( !pos )
+        {
+            // Clear multivalue, if set.
+            if ( index < int( m_pD->m_multiValues.size() ) )
+            {
+                m_pD->m_multiValues[index].clear();
+            }
+
+            //m_pD->m_values[index].m_float = value;
+			FMemory::Memcpy( m_pD->m_values[index].m_text, value, len );
+            m_pD->m_values[index].m_text[len] = 0;
+        }
+
+        // Multivalue case
+        else
+        {
+            check( pos->m_pD->m_parameter == index );
+
+            if ( index >= int( m_pD->m_multiValues.size() ) )
+            {
+                m_pD->m_multiValues.resize( index + 1 );
+            }
+
+            auto& m = m_pD->m_multiValues[index];
+            auto it = m.find( pos->m_pD->m_values );
+            if ( it != m.end() )
+            {
+                //it->second.m_float = value;
+				FMemory::Memcpy( it->second.m_text, value, len );
+                it->second.m_text[len] = 0;
+            }
+            else
+            {
+                PARAMETER_VALUE valueData;
+                //valueData.m_float = value;
+				FMemory::Memcpy( valueData.m_text, value, len );
+                valueData.m_text[len] = 0;
+                m.insert( std::make_pair<>( pos->m_pD->m_values, valueData ) );
+            }
+        }
+    }
+
+
+    //---------------------------------------------------------------------------------------------
+    PROJECTOR Parameters::Private::GetProjectorValue( int index,
+                                                      const Ptr<const RangeIndex>& pos ) const
+    {
+
+        const PROGRAM& program = m_pModel->GetPrivate()->m_program;
+
+        // Early out in case of invalid parameters
+        if ( index < 0
+             ||
+             index >= (int)m_values.size()
+			 || 
+			 index >= (int)program.m_parameters.size()
+             ||
+             program.m_parameters[index].m_type != PARAMETER_TYPE::T_PROJECTOR )
+        {
+			check(false);
+            return PROJECTOR();
+        }
+
+        const PROJECTOR* result = nullptr;
+
+        // Single value case
+        if (!pos)
+        {
+            // Return the single value
+            result = &m_values[index].m_projector;
+        }
+
+        // Multivalue case
+        if (!result)
+        {
+            check( pos->m_pD->m_parameter==index );
+
+            if ( index<int(m_multiValues.size()))
+            {
+                const auto& m = m_multiValues[index];
+                auto it = m.find(pos->m_pD->m_values);
+                if (it!=m.end())
+                {
+                    result = &it->second.m_projector;
+                }
+            }
+        }
+
+        // Multivalue parameter, but no multivalue set. Return single value.
+        if (!result)
+        {
+            result = &m_values[index].m_projector;
+        }
+
+        return *result;
+    }
+
+
+	//---------------------------------------------------------------------------------------------
+	void Parameters::GetProjectorValue( int index,
+                                        PROJECTOR_TYPE* pType,
+										float* pPosX, float* pPosY, float* pPosZ,
+										float* pDirX, float* pDirY, float* pDirZ,
+										float* pUpX, float* pUpY, float* pUpZ,
+                                        float* pScaleX, float* pScaleY, float* pScaleZ,
+                                        float* pProjectionAngle,
+                                        const Ptr<const RangeIndex>& pos ) const
+	{
+		check( index>=0 && index<(int)m_pD->m_values.size() );
+		check( GetType(index)== PARAMETER_TYPE::T_PROJECTOR );
+
+        // Early out in case of invalid parameters
+        if ( index < 0
+             ||
+             index >= (int)m_pD->m_values.size()
+             ||
+             GetType(index) != PARAMETER_TYPE::T_PROJECTOR )
+        {
+            return;
+        }
+
+        PROJECTOR result = m_pD->GetProjectorValue( index, pos );
+
+        // Copy results
+        if (pType) *pType = result.type;
+
+        if (pPosX) *pPosX = result.position[0];
+        if (pPosY) *pPosY = result.position[1];
+        if (pPosZ) *pPosZ = result.position[2];
+
+        if (pDirX) *pDirX = result.direction[0];
+        if (pDirY) *pDirY = result.direction[1];
+        if (pDirZ) *pDirZ = result.direction[2];
+
+        if (pUpX) *pUpX = result.up[0];
+        if (pUpY) *pUpY = result.up[1];
+        if (pUpZ) *pUpZ = result.up[2];
+
+        if (pScaleX) *pScaleX = result.scale[0];
+        if (pScaleY) *pScaleY = result.scale[1];
+        if (pScaleZ) *pScaleZ = result.scale[2];
+
+        if (pProjectionAngle) *pProjectionAngle = result.projectionAngle;
+    }
+
+
+	//---------------------------------------------------------------------------------------------
+	void Parameters::SetProjectorValue( int index,
+										float posX, float posY, float posZ,
+										float dirX, float dirY, float dirZ,
+										float upX, float upY, float upZ,
+                                        float scaleX, float scaleY, float scaleZ,
+                                        float projectionAngle,
+                                        const Ptr<const RangeIndex>& pos )
+	{
+		check( index>=0 && index<(int)m_pD->m_values.size() );
+		check( GetType(index)== PARAMETER_TYPE::T_PROJECTOR );
+
+        // Early out in case of invalid parameters
+        if ( index < 0
+             ||
+             index >= (int)m_pD->m_values.size()
+             ||
+             GetType(index) != PARAMETER_TYPE::T_PROJECTOR )
+        {
+            return;
+        }
+
+        PROJECTOR* result = nullptr;
+
+        // Single value case
+        if (!pos)
+        {
+            // Clear multivalue, if set.
+            if (index<int(m_pD->m_multiValues.size()))
+            {
+                m_pD->m_multiValues[index].clear();
+            }
+
+            result = &m_pD->m_values[index].m_projector;
+        }
+
+        // Multivalue case
+        else
+        {
+            check( pos->m_pD->m_parameter==index );
+
+            if ( index>=int(m_pD->m_multiValues.size()))
+            {
+                m_pD->m_multiValues.resize(index+1);
+            }
+
+            auto& m = m_pD->m_multiValues[index];
+            auto it = m.find(pos->m_pD->m_values);
+            if (it!=m.end())
+            {
+                result = &it->second.m_projector;
+            }
+            else
+            {
+                auto it2 = m.insert( std::make_pair<>( pos->m_pD->m_values, PARAMETER_VALUE() ) );
+                result = &it2.first->second.m_projector;
+            }
+        }
+
+        check( result );
+
+        // Parameters cannot change the projector type anymore
+        result->type = PROJECTOR_TYPE::COUNT;
+        if (m_pD->m_pModel)
+        {
+            const auto& program = m_pD->m_pModel->GetPrivate()->m_program;
+            result->type = program.m_parameters[index].m_defaultValue.m_projector.type;
+        }
+
+        result->position[0] = posX;
+        result->position[1] = posY;
+        result->position[2] = posZ;
+
+        result->direction[0] = dirX;
+        result->direction[1] = dirY;
+        result->direction[2] = dirZ;
+
+        result->up[0] = upX;
+        result->up[1] = upY;
+        result->up[2] = upZ;
+
+        result->scale[0] = scaleX;
+        result->scale[1] = scaleY;
+        result->scale[2] = scaleZ;
+
+        result->projectionAngle = projectionAngle;
+    }
+
+
+    //---------------------------------------------------------------------------------------------
+    bool Parameters::HasSameValue( int thisParamIndex,
+                                   const ParametersPtrConst& other,
+                                   int otherParamIndex ) const
+    {
+        if ( GetType(thisParamIndex) != other->GetType(otherParamIndex) )
+        {
+            return false;
+        }
+
+        if ( !(m_pD->m_values[thisParamIndex]==other->m_pD->m_values[otherParamIndex] ) )
+        {
+            return false;
+        }
+
+        size_t thisNumMultiValues = 0;
+        bool thisHasMultiValues = int(m_pD->m_multiValues.size()) > thisParamIndex;
+        if (thisHasMultiValues)
+        {
+            thisNumMultiValues = m_pD->m_multiValues[thisParamIndex].size();
+        }
+
+        size_t otherNumMultiValues = 0;
+        bool otherHasMultiValues = int(other->m_pD->m_multiValues.size()) > otherParamIndex;
+        if (otherHasMultiValues)
+        {
+            otherNumMultiValues = other->m_pD->m_multiValues[otherParamIndex].size();
+        }
+
+        if ( thisNumMultiValues != otherNumMultiValues )
+        {
+            return false;
+        }
+
+        if ( thisHasMultiValues
+             &&
+             otherHasMultiValues
+             &&
+             m_pD->m_multiValues[thisParamIndex]!=other->m_pD->m_multiValues[otherParamIndex] )
+        {
+            return false;
+        }
+
+        return true;
+    }
+
+
+	//---------------------------------------------------------------------------------------------
+	int Parameters::Private::Find( const char* strName ) const
+	{
+		const PROGRAM& program = m_pModel->GetPrivate()->m_program;
+
+		int result = -1;
+
+		for( int i=0; result<0 && i<(int)program.m_parameters.size(); ++i )
+		{
+			const PARAMETER_DESC& p = program.m_parameters[i];
+
+			if ( p.m_name == strName )
+			{
+				result = i;
+			}
+		}
+
+		return result;
+	}
+
+
+    //---------------------------------------------------------------------------------------------
+    //---------------------------------------------------------------------------------------------
+    //---------------------------------------------------------------------------------------------
+    RangeIndex::RangeIndex()
+    {
+        m_pD = new Private();
+    }
+
+
+    //---------------------------------------------------------------------------------------------
+    RangeIndex::~RangeIndex()
+    {
+        check( m_pD );
+        delete m_pD;
+        m_pD = 0;
+    }
+
+
+    //---------------------------------------------------------------------------------------------
+    RangeIndex::Private* RangeIndex::GetPrivate() const
+    {
+        return m_pD;
+    }
+
+
+    //---------------------------------------------------------------------------------------------
+    int RangeIndex::GetRangeCount() const
+    {
+        return int( m_pD->m_values.size() );
+    }
+
+
+    //---------------------------------------------------------------------------------------------
+    const char* RangeIndex::GetRangeName( int index ) const
+    {
+        check( index >= 0 && index < GetRangeCount() );
+        check( index < int(m_pD->m_pParameters->GetPrivate()->m_values.size()) );
+        ModelPtrConst m_pModel = m_pD->m_pParameters->GetPrivate()->m_pModel;
+        check( m_pModel );
+        check( index < int(m_pModel->GetPrivate()->m_program.m_ranges.size()) );
+
+        return m_pModel->GetPrivate()->m_program.m_ranges[index].m_name.c_str();
+    }
+
+
+    //---------------------------------------------------------------------------------------------
+    const char* RangeIndex::GetRangeUid( int index ) const
+    {
+        check( index >= 0 && index < GetRangeCount() );
+        check( index < int(m_pD->m_pParameters->GetPrivate()->m_values.size()) );
+        ModelPtrConst m_pModel = m_pD->m_pParameters->GetPrivate()->m_pModel;
+        check( m_pModel );
+        check( index < int(m_pModel->GetPrivate()->m_program.m_ranges.size()) );
+
+        return m_pModel->GetPrivate()->m_program.m_ranges[index].m_uid.c_str();
+    }
+
+
+    //---------------------------------------------------------------------------------------------
+    void RangeIndex::SetPosition( int index, int position )
+    {
+        check( index >= 0 && index < GetRangeCount() );
+        if ( index >= 0 && index < GetRangeCount() )
+        {
+            m_pD->m_values[index] = position;
+        }
+    }
+
+
+    //---------------------------------------------------------------------------------------------
+    int RangeIndex::GetPosition( int index ) const
+    {
+        check( index >= 0 && index < GetRangeCount() );
+        if ( index >= 0 && index < GetRangeCount() )
+        {
+            return m_pD->m_values[index];
+        }
+        return 0;
+    }
+
+}
+
diff --git a/Engine/Plugins/Experimental/Mutable/Source/MutableRuntime/Private/ParametersPrivate.h b/Engine/Plugins/Experimental/Mutable/Source/MutableRuntime/Private/ParametersPrivate.h
new file mode 100644
--- /dev/null
+++ b/Engine/Plugins/Experimental/Mutable/Source/MutableRuntime/Private/ParametersPrivate.h
@@ -0,0 +1,504 @@
+// Copyright Epic Games, Inc. All Rights Reserved.
+
+#pragma once
+
+#include "MutableRuntime/Public/Parameters.h"
+
+#include "MutableRuntime/Private/SerialisationPrivate.h"
+#include "MutableRuntime/Private/MemoryPrivate.h"
+#include "MutableRuntime/Private/Operations.h"
+#include "MutableRuntime/Private/ImagePrivate.h"
+
+#include "MutableRuntime/Public/Model.h"
+
+
+namespace mu
+{
+
+    MUTABLE_DEFINE_ENUM_SERIALISABLE( PARAMETER_TYPE )
+    MUTABLE_DEFINE_ENUM_SERIALISABLE( PARAMETER_DETAILED_TYPE )
+    MUTABLE_DEFINE_ENUM_SERIALISABLE( PROJECTOR_TYPE )
+
+
+    //---------------------------------------------------------------------------------------------
+    //! Information to project an image on a mesh.
+    //---------------------------------------------------------------------------------------------
+    struct PROJECTOR
+    {
+        PROJECTOR_TYPE type = PROJECTOR_TYPE::PLANAR;
+        float position[3] = {0,0,0};
+        float direction[3] = {0,0,0};
+        float up[3] = {0,0,0};
+        float scale[3] = {0,0,0};
+        float projectionAngle = 0.0f;
+
+        //!
+        inline void GetDirectionSideUp( vec3f& d,
+                                        vec3f& s,
+                                        vec3f& u ) const
+        {
+            d = vec3f( direction[0], direction[1], direction[2] );
+            u = vec3f( up[0], up[1], up[2] );
+            s = cross( u, d );
+            normalise(s);
+        }
+
+
+        //!
+        void Serialise( OutputArchive& arch ) const
+        {
+            arch << type;
+            arch << position[0];
+            arch << position[1];
+            arch << position[2];
+            arch << direction[0];
+            arch << direction[1];
+            arch << direction[2];
+            arch << up[0];
+            arch << up[1];
+            arch << up[2];
+            arch << scale[0];
+            arch << scale[1];
+            arch << scale[2];
+            arch << projectionAngle;
+        }
+
+        //!
+        void Unserialise( InputArchive& arch )
+        {
+            arch >> type;
+            arch >> position[0];
+            arch >> position[1];
+            arch >> position[2];
+            arch >> direction[0];
+            arch >> direction[1];
+            arch >> direction[2];
+            arch >> up[0];
+            arch >> up[1];
+            arch >> up[2];
+            arch >> scale[0];
+            arch >> scale[1];
+            arch >> scale[2];
+            arch >> projectionAngle;
+        }
+
+        bool operator==( const PROJECTOR& o ) const
+        {
+            return     type==o.type
+                    && position[0]==o.position[0]
+                    && position[1]==o.position[1]
+                    && position[2]==o.position[2]
+                    && up[0]==o.up[0]
+                    && up[1]==o.up[1]
+                    && up[2]==o.up[2]
+                    && direction[0]==o.direction[0]
+                    && direction[1]==o.direction[1]
+                    && direction[2]==o.direction[2]
+                    && scale[0]==o.scale[0]
+                    && scale[1]==o.scale[1]
+                    && scale[2]==o.scale[2]
+                    && projectionAngle==o.projectionAngle;
+        }
+
+    };
+
+	//---------------------------------------------------------------------------------------------
+	//! Information about a generic shape in space.
+	//---------------------------------------------------------------------------------------------
+	struct SHAPE
+	{
+		// Transform
+		vec3f position;
+		vec3f up;
+		vec3f side;
+		
+		vec3f size;
+
+		// 
+		enum class Type : uint8_t
+		{
+			None = 0,
+			Ellipse,
+			AABox
+		};
+        uint8_t type = 0;
+		
+
+		//!
+		void Serialise(OutputArchive& arch) const
+		{
+			arch << position;
+			arch << up;
+			arch << side;
+			arch << size;
+			arch << type;
+		}
+
+		//!
+		void Unserialise(InputArchive& arch)
+		{
+			arch >> position;
+			arch >> up;
+			arch >> side;
+			arch >> size;
+			arch >> type;
+		}
+
+        bool operator==( const SHAPE& o ) const
+        {
+            return type==o.type
+                    && position==o.position
+                    && up==o.up
+                    && side==o.side
+                    && size==o.size;
+        }
+	};
+
+
+	//---------------------------------------------------------------------------------------------
+	//! Information about a curve interpolated from keyframes.
+	//---------------------------------------------------------------------------------------------
+	struct CurveKeyFrame
+	{
+		float time;
+		float value;
+		float in_tangent;
+		float in_tangent_weight;
+		float out_tangent;
+		float out_tangent_weight;
+
+		enum class InterpMode : uint8_t
+		{
+			Linear = 0,
+			Constant,
+			Cubic,
+			None,
+
+			//! Utility enumeration value, not really a Interp mode.
+			Count
+		};
+        uint8_t interp_mode;
+
+
+		enum class TangentMode : uint8_t
+		{
+			Auto = 0,
+			Manual,
+			Broken,
+			None,
+
+			//! Utility enumeration value, not really a Tangent mode.
+			Count
+		};
+        uint8_t tangent_mode;
+
+
+		enum class TangentWeightMode : uint8_t
+		{
+			Auto = 0,
+			Manual,
+			Broken,
+			None,
+
+			//! Utility enumeration value, not really a Tangent weight mode.
+			Count
+		};
+        uint8_t tangent_weight_mode;
+
+        //!
+        bool operator==( const CurveKeyFrame& other ) const
+        {
+            return time==other.time &&
+                    value==other.value&&
+                    in_tangent==other.in_tangent&&
+                    in_tangent_weight==other.in_tangent_weight&&
+                    out_tangent==other.out_tangent&&
+                    out_tangent_weight==other.out_tangent_weight&&
+                    interp_mode==other.interp_mode&&
+                    tangent_mode==other.tangent_mode&&
+                    tangent_weight_mode==other.tangent_weight_mode;
+        }
+
+		//!
+		void Serialise(OutputArchive& arch) const
+		{
+			arch << time;
+			arch << value;
+			arch << in_tangent;
+			arch << in_tangent_weight;
+			arch << out_tangent;
+			arch << out_tangent_weight;
+
+			arch << interp_mode;
+			arch << tangent_mode;
+			arch << tangent_weight_mode;
+		}
+
+		//!
+		void Unserialise(InputArchive& arch)
+		{
+			arch >> time;
+			arch >> value;
+			arch >> in_tangent;
+			arch >> in_tangent_weight;
+			arch >> out_tangent;
+			arch >> out_tangent_weight;
+
+			arch >> interp_mode;
+			arch >> tangent_mode;
+			arch >> tangent_weight_mode;
+		}
+	};
+
+
+	struct Curve
+	{
+		vector<CurveKeyFrame> keyFrames;
+		float defaultValue = 0.f;
+
+        //!
+        bool operator==( const Curve& other ) const
+        {
+            return keyFrames == other.keyFrames && defaultValue == other.defaultValue;
+        }
+
+		//!
+		void Serialise(OutputArchive& arch) const
+		{
+            arch << keyFrames;
+            arch << defaultValue;
+		}
+
+		//!
+		void Unserialise(InputArchive& arch)
+		{
+			arch >> keyFrames;
+            arch >> defaultValue;
+        }
+	};
+
+
+    union PARAMETER_VALUE
+    {
+        PARAMETER_VALUE()
+			: m_projector()
+        {
+            memset( this, 0, sizeof(PARAMETER_VALUE) );
+        }
+
+        bool m_bool;
+
+        int m_int;
+
+        float m_float;
+
+        float m_colour[3];
+
+        PROJECTOR m_projector;
+
+        EXTERNAL_IMAGE_ID m_image;
+
+        char m_text[MUTABLE_MAX_STRING_PARAM_LENGTH+1];
+
+        //!
+        bool operator==( const PARAMETER_VALUE& other ) const
+        {
+            return mutable_memcmp(this, &other, sizeof(PARAMETER_VALUE))==0;
+        }
+
+    };
+
+    MUTABLE_DEFINE_POD_SERIALISABLE( PARAMETER_VALUE )
+
+
+    struct PARAMETER_DESC
+    {
+        string m_name;
+
+        //! Unique id (provided externally, so no actual guarantee that it is unique.)
+		string m_uid;
+
+        PARAMETER_TYPE m_type = PARAMETER_TYPE::T_NONE;
+
+        PARAMETER_DETAILED_TYPE m_detailedType = PARAMETER_DETAILED_TYPE::UNKNOWN;
+
+        PARAMETER_VALUE m_defaultValue;
+
+        //! Ranges, if the parameter is multi-dimensional. The indices refer to the Model's program
+        //! vector of range descriptors.
+        vector<uint32_t> m_ranges;
+
+        //! Additional parameter description information
+        vector<OP::ADDRESS> m_descImages;
+
+        //! Possible values of the parameter in case of being an integer, and its names
+        struct INT_VALUE_DESC
+        {
+            int16_t m_value;
+            string m_name;
+
+            //!
+            bool operator==( const INT_VALUE_DESC& other ) const
+            {
+                return m_value==other.m_value &&
+                        m_name==other.m_name;
+            }
+
+            //!
+            void Serialise( OutputArchive& arch ) const
+            {
+                arch << m_value;
+                arch << m_name;
+            }
+
+            //!
+            void Unserialise( InputArchive& arch )
+            {
+                arch >> m_value;
+                arch >> m_name;
+            }
+        };
+
+        //! For integer parameters, this contains the description of the possible values.
+        //! If empty, the integer may have any value.
+        vector<INT_VALUE_DESC> m_possibleValues;
+
+        //!
+        bool operator==( const PARAMETER_DESC& other ) const
+        {
+            return m_name == other.m_name && m_uid == other.m_uid && m_type == other.m_type &&
+                   m_defaultValue == other.m_defaultValue &&
+                   m_ranges == other.m_ranges && m_descImages == other.m_descImages &&
+                   m_possibleValues == other.m_possibleValues;
+        }
+
+        //!
+        void Serialise( OutputArchive& arch ) const
+        {
+            const int32_t ver = 5;
+            arch << ver;
+
+			arch << m_name;
+			arch << m_uid;
+            arch << m_type;
+            arch << m_defaultValue;
+            arch << m_ranges;
+            arch << m_descImages;
+            arch << m_possibleValues;
+        }
+
+        //!
+        void Unserialise( InputArchive& arch )
+        {
+            int32_t ver;
+            arch >> ver;
+            check( ver == 5 );
+
+			arch >> m_name;
+            arch >> m_uid;
+            arch >> m_type;
+            arch >> m_defaultValue;
+            arch >> m_ranges;
+            arch >> m_descImages;
+            arch >> m_possibleValues;
+        }
+    };
+
+
+    struct RANGE_DESC
+    {
+        string m_name;
+        string m_uid;
+
+        //!
+        bool operator==( const RANGE_DESC& other ) const
+        {
+            return m_name==other.m_name
+                    &&
+                    m_uid==other.m_uid;
+        }
+
+        //!
+        void Serialise( OutputArchive& arch ) const
+        {
+            const int32_t ver = 1;
+            arch << ver;
+
+            arch << m_name;
+            arch << m_uid;
+        }
+
+        //!
+        void Unserialise( InputArchive& arch )
+        {
+            int32_t ver;
+            arch >> ver;
+            check( ver == 1 );
+
+            arch >> m_name;
+            arch >> m_uid;
+        }
+    };
+
+
+    class MUTABLERUNTIME_API RangeIndex::Private : public Base
+    {
+    public:
+
+        //! Run-time data
+        Ptr<const Parameters> m_pParameters;
+
+        //! Index of the parameter for which we are a range index
+        int m_parameter = -1;
+
+        //! Position in the several dimension of the range, as defined in m_pParameters
+        vector<int> m_values;
+    };
+
+
+    class MUTABLERUNTIME_API Parameters::Private : public Base
+    {
+    public:
+
+        //! Warning: update Parameters::Clone method if this members change.
+
+        //! Run-time data
+        ModelPtrConst m_pModel;
+
+        //! Values for the parameters if they are not multidimensional.
+        vector<PARAMETER_VALUE> m_values;
+
+        //! If the parameter is multidemensional, the values are stored here.
+        //! The key of the map is the vector of values stored in a RangeIndex
+        vector< map< vector<int32_t>, PARAMETER_VALUE > > m_multiValues;
+
+
+        //!
+        void Serialise( OutputArchive& arch ) const
+        {
+            const uint32_t ver = 1;
+            arch << ver;
+
+            arch << m_values;
+            arch << m_multiValues;
+        }
+
+        //!
+        void Unserialise( InputArchive& arch )
+        {
+            uint32_t ver;
+            arch >> ver;
+            check( ver == 1 );
+
+            arch >> m_values;
+            arch>>m_multiValues;
+        }
+
+        //!
+        int Find( const char* strName ) const;
+
+        //!
+        PROJECTOR GetProjectorValue( int index, const Ptr<const RangeIndex>& pos ) const;
+    };
+
+}
diff --git a/Engine/Plugins/Experimental/Mutable/Source/MutableRuntime/Private/PhysicsBody.cpp b/Engine/Plugins/Experimental/Mutable/Source/MutableRuntime/Private/PhysicsBody.cpp
new file mode 100644
--- /dev/null
+++ b/Engine/Plugins/Experimental/Mutable/Source/MutableRuntime/Private/PhysicsBody.cpp
@@ -0,0 +1,463 @@
+// Copyright Epic Games, Inc. All Rights Reserved.
+
+#include "PhysicsBody.h"
+
+using namespace mu;
+
+//-------------------------------------------------------------------------------------------------
+void PhysicsBody::Serialise( const PhysicsBody* p, OutputArchive& arch )
+{
+	arch << *p;
+}
+
+
+mu::Ptr<PhysicsBody> PhysicsBody::StaticUnserialise( InputArchive& arch )
+{
+	mu::Ptr<PhysicsBody> pResult = new PhysicsBody();
+	arch >> *pResult;
+	return pResult;
+}
+
+//-------------------------------------------------------------------------------------------------
+mu::Ptr<PhysicsBody> PhysicsBody::Clone() const
+{
+	mu::Ptr<PhysicsBody> pResult = new PhysicsBody();
+
+    pResult->Bodies = Bodies;
+	pResult->Bones = Bones; 
+	pResult->CustomIds = CustomIds; 
+	
+	return pResult;
+}
+
+//-------------------------------------------------------------------------------------------------
+void PhysicsBody::SetBodyCount( int32 Count )
+{
+	Bodies.SetNum(Count);
+	Bones.SetNum(Count);
+	CustomIds.Init(-1, Count);
+}
+
+//-------------------------------------------------------------------------------------------------
+int32 PhysicsBody::GetBodyCount( ) const
+{
+	return Bodies.Num();
+}
+
+//-------------------------------------------------------------------------------------------------
+const char* PhysicsBody::GetBodyBoneName( int32 B ) const
+{
+	check(B >= 0 && B < Bones.Num());
+
+	return Bones[B].c_str();
+}
+
+//-------------------------------------------------------------------------------------------------
+void PhysicsBody::SetBodyBoneName( int32 B, const char* BoneName )
+{
+	check(B >= 0 && B < Bones.Num());
+
+	Bones[B] = BoneName;
+}
+
+//-------------------------------------------------------------------------------------------------
+int32 PhysicsBody::GetBodyCustomId( int32 B ) const
+{
+	check(B >= 0 && B < CustomIds.Num());
+
+	return CustomIds[B];
+}
+
+//-------------------------------------------------------------------------------------------------
+void PhysicsBody::SetBodyCustomId( int32 B, int32 CustomId )
+{
+	check(B >= 0 && B < CustomIds.Num());
+
+	CustomIds[B] = CustomId;
+}
+
+//-------------------------------------------------------------------------------------------------
+int32 PhysicsBody::GetSphereCount( int32 B ) const
+{
+	return Bodies[B].Spheres.Num();
+}
+
+//-------------------------------------------------------------------------------------------------
+int32 PhysicsBody::GetBoxCount( int32 B ) const
+{
+	return Bodies[B].Boxes.Num();	
+}
+
+//-------------------------------------------------------------------------------------------------
+int32 PhysicsBody::GetConvexCount( int32 B ) const
+{
+	return Bodies[B].Convex.Num();
+}
+
+//-------------------------------------------------------------------------------------------------
+int32 PhysicsBody::GetSphylCount( int32 B ) const
+{
+	return Bodies[B].Sphyls.Num();
+}
+
+//-------------------------------------------------------------------------------------------------
+int32 PhysicsBody::GetTaperedCapsuleCount( int32 B ) const
+{
+	return Bodies[B].TaperedCapsules.Num();
+}
+
+//-------------------------------------------------------------------------------------------------
+void PhysicsBody::SetSphereCount( int32 B, int32 Count )
+{
+	Bodies[B].Spheres.SetNum(Count);
+}
+
+//-------------------------------------------------------------------------------------------------
+void PhysicsBody::SetBoxCount( int32 B, int32 Count )
+{
+	Bodies[B].Boxes.SetNum(Count);	
+}
+
+//-------------------------------------------------------------------------------------------------
+void PhysicsBody::SetConvexCount( int32 B, int32 Count)
+{
+    Bodies[B].Convex.SetNum(Count);	
+}
+
+//-------------------------------------------------------------------------------------------------
+void PhysicsBody::SetSphylCount( int32 B, int32 Count )
+{
+	Bodies[B].Sphyls.SetNum(Count);	
+}
+
+//-------------------------------------------------------------------------------------------------
+void PhysicsBody::SetTaperedCapsuleCount( int32 B, int32 Count )
+{
+	Bodies[B].TaperedCapsules.SetNum(Count);	
+}
+		
+//-------------------------------------------------------------------------------------------------
+void PhysicsBody::SetSphere( 
+		int32 B, int32 I, 
+		FVector3f Position, float Radius )
+{
+	check(B >= 0 && B < Bodies.Num());
+	check(I >= 0 && I < Bodies[B].Spheres.Num());
+	
+	Bodies[B].Spheres[I].Position = Position;
+	Bodies[B].Spheres[I].Radius = Radius;
+}
+
+//-------------------------------------------------------------------------------------------------
+void PhysicsBody::SetBox( 
+		int32 B, int32 I, 
+		FVector3f Position, FQuat4f Orientation, FVector3f Size )
+{
+	check(B >= 0 && B < Bodies.Num());
+	check(I >= 0 && I < Bodies[B].Boxes.Num());
+
+	Bodies[B].Boxes[I].Position = Position;
+	Bodies[B].Boxes[I].Orientation = Orientation;
+	Bodies[B].Boxes[I].Size = Size;
+}
+
+//-------------------------------------------------------------------------------------------------
+void PhysicsBody::SetConvex( 
+		int32 B, int32 I, 
+		const FVector3f* Vertices, int32 VerticesCount, 
+		const int32* Indices, int32 IndicesCount, 
+		const FTransform3f& Transform )
+{
+	check(B >= 0 && B < Bodies.Num());
+	check(I >= 0 && I < Bodies[B].Convex.Num());
+
+	Bodies[B].Convex[I].Vertices = TArray<FVector3f>( Vertices, VerticesCount );
+	Bodies[B].Convex[I].Indices = TArray<int32>( Indices, IndicesCount );
+	Bodies[B].Convex[I].Transform = Transform;
+	
+}
+
+//-------------------------------------------------------------------------------------------------
+void PhysicsBody::SetSphyl( 
+		int32 B, int32 I, 
+		FVector3f Position, FQuat4f Orientation, float Radius, float Length )
+{
+	check(B >= 0 && B < Bodies.Num());
+	check(I >= 0 && I < Bodies[B].Sphyls.Num());
+	
+	Bodies[B].Sphyls[I].Position = Position;
+	Bodies[B].Sphyls[I].Orientation = Orientation;
+	Bodies[B].Sphyls[I].Radius = Radius;
+	Bodies[B].Sphyls[I].Length = Length;
+}
+
+//-------------------------------------------------------------------------------------------------
+void PhysicsBody::SetSphereFlags( int32 B, int32 I, uint32 Flags )
+{
+	check(B >= 0 && B < Bodies.Num());
+	check(I >= 0 && I < Bodies[B].Spheres.Num());
+	
+	Bodies[B].Spheres[I].Flags = Flags;
+}
+
+//-------------------------------------------------------------------------------------------------
+void PhysicsBody::SetBoxFlags( int32 B, int32 I, uint32 Flags )
+{
+	check(B >= 0 && B < Bodies.Num());
+	check(I >= 0 && I < Bodies[B].Boxes.Num());
+
+	Bodies[B].Boxes[I].Flags = Flags;
+}
+
+//-------------------------------------------------------------------------------------------------
+void PhysicsBody::SetConvexFlags( int32 B, int32 I, uint32 Flags )
+{
+	check(B >= 0 && B < Bodies.Num());
+	check(I >= 0 && I < Bodies[B].Convex.Num());
+
+	Bodies[B].Convex[I].Flags = Flags;	
+}
+
+//-------------------------------------------------------------------------------------------------
+void PhysicsBody::SetSphylFlags( int32 B, int32 I, uint32 Flags )
+{
+	check(B >= 0 && B < Bodies.Num());
+	check(I >= 0 && I < Bodies[B].Sphyls.Num());
+	
+	Bodies[B].Sphyls[I].Flags = Flags;
+}
+
+//-------------------------------------------------------------------------------------------------
+void PhysicsBody::SetTaperedCapsuleFlags( int32 B, int32 I, uint32 Flags )
+{
+	check(B >= 0 && B < Bodies.Num());
+	check(I >= 0 && I < Bodies[B].TaperedCapsules.Num());
+	
+	Bodies[B].TaperedCapsules[I].Flags = Flags;
+}
+
+//-------------------------------------------------------------------------------------------------
+void PhysicsBody::SetSphereName( int32 B, int32 I, const char* Name )
+{
+	check(B >= 0 && B < Bodies.Num());
+	check(I >= 0 && I < Bodies[B].Spheres.Num());
+	
+	Bodies[B].Spheres[I].Name = Name;
+}
+
+//-------------------------------------------------------------------------------------------------
+void PhysicsBody::SetBoxName( int32 B, int32 I, const char* Name )
+{
+	check(B >= 0 && B < Bodies.Num());
+	check(I >= 0 && I < Bodies[B].Boxes.Num());
+
+	Bodies[B].Boxes[I].Name = Name;
+}
+
+//-------------------------------------------------------------------------------------------------
+void PhysicsBody::SetConvexName( int32 B, int32 I, const char* Name )
+{
+	check(B >= 0 && B < Bodies.Num());
+	check(I >= 0 && I < Bodies[B].Convex.Num());
+
+	Bodies[B].Convex[I].Name = Name;	
+}
+
+//-------------------------------------------------------------------------------------------------
+void PhysicsBody::SetSphylName( int32 B, int32 I, const char* Name )
+{
+	check(B >= 0 && B < Bodies.Num());
+	check(I >= 0 && I < Bodies[B].Sphyls.Num());
+	
+	Bodies[B].Sphyls[I].Name = Name;
+}
+
+//-------------------------------------------------------------------------------------------------
+void PhysicsBody::SetTaperedCapsuleName( int32 B, int32 I, const char* Name )
+{
+	check(B >= 0 && B < Bodies.Num());
+	check(I >= 0 && I < Bodies[B].TaperedCapsules.Num());
+	
+	Bodies[B].TaperedCapsules[I].Name = Name;
+}
+
+//-------------------------------------------------------------------------------------------------
+void PhysicsBody::SetTaperedCapsule( 
+		int32 B, int32 I, 
+		FVector3f Position, FQuat4f Orientation, 
+		float Radius0, float Radius1, float Length )
+{
+	check(B >= 0 && B < Bodies.Num());
+	check(I >= 0 && I < Bodies[B].TaperedCapsules.Num());
+	
+	Bodies[B].TaperedCapsules[I].Position = Position;
+	Bodies[B].TaperedCapsules[I].Orientation = Orientation;
+	Bodies[B].TaperedCapsules[I].Radius0 = Radius0;
+	Bodies[B].TaperedCapsules[I].Radius1 = Radius1;
+	Bodies[B].TaperedCapsules[I].Length = Length;
+}
+
+//-------------------------------------------------------------------------------------------------
+void PhysicsBody::GetSphere( 
+		int32 B, int32 I, 
+		FVector3f& OutPosition, float& OutRadius) const
+{
+	check(B >= 0 && B < Bodies.Num());
+	check(I >= 0 && I < Bodies[B].Spheres.Num());
+
+	OutPosition = Bodies[B].Spheres[I].Position;
+	OutRadius = Bodies[B].Spheres[I].Radius;
+}
+
+//-------------------------------------------------------------------------------------------------
+void PhysicsBody::GetBox( 
+		int32 B, int32 I, 
+		FVector3f& OutPosition, FQuat4f& OutOrientation, FVector3f& OutSize) const
+{
+	check(B >= 0 && B < Bodies.Num());
+	check(I >= 0 && I < Bodies[B].Boxes.Num());
+
+	OutPosition = Bodies[B].Boxes[I].Position;
+	OutOrientation = Bodies[B].Boxes[I].Orientation;
+	OutSize = Bodies[B].Boxes[I].Size;
+}
+
+//-------------------------------------------------------------------------------------------------
+void PhysicsBody::GetConvex( 
+		int32 B, int32 I, 
+		FVector3f const*& OutVertices, int32& OutVerticesCount, 
+		int32 const*& OutIndices, int32& OutIndicesCount, 
+		FTransform3f& OutTransform) const
+{
+	check(B >= 0 && B < Bodies.Num());
+	check(I >= 0 && I < Bodies[B].Convex.Num());
+	
+	OutVertices = Bodies[B].Convex[I].Vertices.GetData();
+	OutVerticesCount = Bodies[B].Convex[I].Vertices.Num();
+
+	OutIndices = Bodies[B].Convex[I].Indices.GetData();
+	OutIndicesCount = Bodies[B].Convex[I].Indices.Num();
+	
+	OutTransform = Bodies[B].Convex[I].Transform;
+}
+
+//-------------------------------------------------------------------------------------------------
+void PhysicsBody::GetSphyl( 
+		int32 B, int32 I, 
+		FVector3f& OutPosition, FQuat4f& OutOrientation, float& OutRadius, float& OutLength ) const
+{
+	check(B >= 0 && B < Bodies.Num());
+	check(I >= 0 && I < Bodies[B].Sphyls.Num());
+
+	OutPosition = Bodies[B].Sphyls[I].Position;
+	OutOrientation = Bodies[B].Sphyls[I].Orientation;
+	OutRadius = Bodies[B].Sphyls[I].Radius;
+	OutLength = Bodies[B].Sphyls[I].Length;
+}
+
+//-------------------------------------------------------------------------------------------------
+void PhysicsBody::GetTaperedCapsule( 
+		int32 B, int32 I, 
+		FVector3f& OutPosition, FQuat4f& OutOrientation, 
+		float& OutRadius0, float& OutRadius1, float& OutLength ) const
+{
+	check(B >= 0 && B < Bodies.Num());
+	check(I >= 0 && I < Bodies[B].TaperedCapsules.Num());
+
+	OutPosition = Bodies[B].TaperedCapsules[I].Position;
+	OutOrientation = Bodies[B].TaperedCapsules[I].Orientation;
+	OutRadius0 = Bodies[B].TaperedCapsules[I].Radius0;
+	OutRadius1 = Bodies[B].TaperedCapsules[I].Radius1;
+	OutLength = Bodies[B].TaperedCapsules[I].Length;
+}
+
+
+//-------------------------------------------------------------------------------------------------
+uint32 PhysicsBody::GetSphereFlags( int32 B, int32 I ) const
+{
+	check(B >= 0 && B < Bodies.Num());
+	check(I >= 0 && I < Bodies[B].Spheres.Num());
+	
+	return Bodies[B].Spheres[I].Flags;
+}
+
+//-------------------------------------------------------------------------------------------------
+uint32 PhysicsBody::GetBoxFlags( int32 B, int32 I ) const
+{
+	check(B >= 0 && B < Bodies.Num());
+	check(I >= 0 && I < Bodies[B].Boxes.Num());
+
+	return Bodies[B].Boxes[I].Flags;
+}
+
+//-------------------------------------------------------------------------------------------------
+uint32 PhysicsBody::GetConvexFlags( int32 B, int32 I ) const
+{
+	check(B >= 0 && B < Bodies.Num());
+	check(I >= 0 && I < Bodies[B].Convex.Num());
+
+	return Bodies[B].Convex[I].Flags;	
+}
+
+//-------------------------------------------------------------------------------------------------
+uint32 PhysicsBody::GetSphylFlags( int32 B, int32 I ) const
+{
+	check(B >= 0 && B < Bodies.Num());
+	check(I >= 0 && I < Bodies[B].Sphyls.Num());
+	
+	return Bodies[B].Sphyls[I].Flags;
+}
+
+//-------------------------------------------------------------------------------------------------
+uint32 PhysicsBody::GetTaperedCapsuleFlags( int32 B, int32 I ) const
+{
+	check(B >= 0 && B < Bodies.Num());
+	check(I >= 0 && I < Bodies[B].TaperedCapsules.Num());
+	
+	return Bodies[B].TaperedCapsules[I].Flags;
+}
+
+//-------------------------------------------------------------------------------------------------
+const char* PhysicsBody::GetSphereName( int32 B, int32 I ) const
+{
+	check(B >= 0 && B < Bodies.Num());
+	check(I >= 0 && I < Bodies[B].Spheres.Num());
+	
+	return Bodies[B].Spheres[I].Name.c_str();
+}
+
+//-------------------------------------------------------------------------------------------------
+const char* PhysicsBody::GetBoxName( int32 B, int32 I ) const
+{
+	check(B >= 0 && B < Bodies.Num());
+	check(I >= 0 && I < Bodies[B].Boxes.Num());
+
+	return Bodies[B].Boxes[I].Name.c_str();
+}
+
+//-------------------------------------------------------------------------------------------------
+const char* PhysicsBody::GetConvexName( int32 B, int32 I ) const
+{
+	check(B >= 0 && B < Bodies.Num());
+	check(I >= 0 && I < Bodies[B].Convex.Num());
+
+	return Bodies[B].Convex[I].Name.c_str();
+}
+
+//-------------------------------------------------------------------------------------------------
+const char* PhysicsBody::GetSphylName( int32 B, int32 I ) const
+{
+	check(B >= 0 && B < Bodies.Num());
+	check(I >= 0 && I < Bodies[B].Sphyls.Num());
+	
+	return Bodies[B].Sphyls[I].Name.c_str();
+}
+
+//-------------------------------------------------------------------------------------------------
+const char* PhysicsBody::GetTaperedCapsuleName( int32 B, int32 I ) const
+{
+	check(B >= 0 && B < Bodies.Num());
+	check(I >= 0 && I < Bodies[B].TaperedCapsules.Num());
+	
+	return Bodies[B].TaperedCapsules[I].Name.c_str();
+}
diff --git a/Engine/Plugins/Experimental/Mutable/Source/MutableRuntime/Private/Platform.h b/Engine/Plugins/Experimental/Mutable/Source/MutableRuntime/Private/Platform.h
new file mode 100644
--- /dev/null
+++ b/Engine/Plugins/Experimental/Mutable/Source/MutableRuntime/Private/Platform.h
@@ -0,0 +1,212 @@
+// Copyright Epic Games, Inc. All Rights Reserved.
+
+#pragma once
+
+// This file contains some hacks to solve differences between platforms
+// This file also contains wrappers for all the libc methods used in the runtime.
+// No other place in the runtime should include any libc header.
+
+#include "HAL/PlatformTime.h"
+#include "Logging/LogMacros.h"
+#include "Stats/Stats.h"
+
+//! There is a define MUTABLE_PROFILE that can be passed on the compiler options to enable
+//! the internal profiling methods. See Config.h for more information.
+
+#include "Types.h"
+#include "HAL/UnrealMemory.h"
+#include <cstdint>
+
+MUTABLERUNTIME_API DECLARE_LOG_CATEGORY_EXTERN(LogMutableCore, Log, All);
+
+
+#ifdef __ANDROID__
+    #define MUTABLE_PLATFORM_ANDROID
+#elif defined(_XBOX_ONE)
+	#define MUTABLE_PLATFORM_XBOXONE
+#elif defined(__PS4__)
+	#define MUTABLE_PLATFORM_PS4
+#elif defined(__PS5__)
+	#define MUTABLE_PLATFORM_PS5
+#elif defined(__SWITCH__)
+    #define MUTABLE_PLATFORM_SWITCH
+#elif defined(_WIN64) && _WIN64
+    #define MUTABLE_PLATFORM_WIN64
+#elif defined(_WIN32) && _WIN32
+    #define MUTABLE_PLATFORM_WIN32
+#elif defined(__APPLE__) && __APPLE__
+    #include "TargetConditionals.h"
+    #if TARGET_IPHONE_SIMULATOR
+        #define MUTABLE_PLATFORM_IOS
+    #elif TARGET_OS_IPHONE
+        #define MUTABLE_PLATFORM_IOS
+    #elif TARGET_OS_MAC
+        #define MUTABLE_PLATFORM_OSX
+    #else
+        // Unsupported platform
+        #warning "Unidentified platform."
+    #endif
+#elif defined(__linux) && __linux
+    #define MUTABLE_PLATFORM_LINUX
+#elif defined(__unix) && __unix // all unices not caught above
+    #define MUTABLE_PLATFORM_LINUX
+#elif defined(__posix) && __posix
+    #define MUTABLE_PLATFORM_LINUX
+#endif
+
+
+#if defined(MUTABLE_PLATFORM_XBOXONE)
+
+	#define NOMINMAX
+	#include <xdk.h>
+	#include <wrl.h>
+
+#elif defined(MUTABLE_PLATFORM_WIN64) || defined(MUTABLE_PLATFORM_WIN32)
+
+	#ifdef _MSC_VER
+
+		#ifndef _HAS_EXCEPTIONS
+			#define _HAS_EXCEPTIONS 0
+		#endif
+
+		#ifndef WIN32_LEAN_AND_MEAN
+			#define WIN32_LEAN_AND_MEAN
+		#endif
+
+        #define NOMINMAX
+        //#include <windows.h>
+
+		// Some warnings that are more annoying than useful
+		#pragma warning( disable : 4127 )
+		#pragma warning( disable : 6385 )
+		#pragma warning( disable : 6287 )
+
+
+    #else
+
+        // Some unix environment under windows.
+        #include <sys/time.h>
+
+	#endif
+
+#elif defined(MUTABLE_PLATFORM_ANDROID)
+
+    #include <sys/time.h>
+
+#else
+
+    #include <sys/time.h>
+
+#endif
+
+
+//! Unify debug defines
+#if !defined(MUTABLE_DEBUG)
+    #if !defined(NDEBUG) || defined(_DEBUG)
+        #define MUTABLE_DEBUG
+    #endif
+#endif
+
+
+namespace mu
+{
+
+    inline void* mutable_system_malloc( size_t size, uint32_t alignment )
+    {        
+        return FMemory::Malloc(size,alignment);
+    }
+
+
+    inline void mutable_system_free( void* ptr )
+    {
+		FMemory::Free(ptr);
+    }
+
+
+	//---------------------------------------------------------------------------------------------
+	inline int mutable_memcmp ( const void * a, const void* b, size_t num )
+	{
+		return memcmp( a, b, num );
+	}
+
+
+	//---------------------------------------------------------------------------------------------
+	inline int mutable_vsnprintf
+		( char *buffer, size_t sizeOfBuffer, const char *format, va_list argptr )
+	{
+		#ifdef _MSC_VER
+			return vsnprintf_s( buffer, sizeOfBuffer, _TRUNCATE, format, argptr );
+		#else
+			return vsnprintf( buffer, sizeOfBuffer, format, argptr );
+		#endif
+	}
+
+
+	//---------------------------------------------------------------------------------------------
+	inline int mutable_snprintf
+		( char *buffer, size_t sizeOfBuffer, const char *format, ... )
+	{
+		va_list ap;
+		va_start( ap, format );
+		int res = mutable_vsnprintf( buffer, sizeOfBuffer, format, ap );
+		va_end(ap);
+		return res;
+	}
+
+
+	//---------------------------------------------------------------------------------------------
+	// Disgracefully halt the program.
+	inline void Halt()
+	{
+		#ifdef _MSC_VER
+            __debugbreak();
+		#else
+			__builtin_trap();
+		#endif
+	}
+
+
+	//---------------------------------------------------------------------------------------------
+	//! Simple profiling timer.
+	//! It should only be used in debug and profile builds.
+	//---------------------------------------------------------------------------------------------
+	struct TIMER
+	{
+	public:
+
+		//!
+		inline void Start()
+		{
+#ifdef MUTABLE_PROFILE
+			m_start = int64( FPlatformTime::Seconds()*1000000.0 );
+#endif
+		}
+
+		//! Stop counting and return the ellapsed time since Start in milliseconds.
+		inline int GetMilliseconds() const
+		{
+#ifdef MUTABLE_PROFILE
+			return int( (int64(FPlatformTime::Seconds() * 1000000.0) - m_start)/1000 );
+#else
+			return 0;
+#endif
+		}
+
+		//! Stop counting and return the ellapsed time since Start in microseconds.
+		inline int GetMicroseconds() const
+		{
+#ifdef MUTABLE_PROFILE
+			return int((int64(FPlatformTime::Seconds() * 1000000.0) - m_start));
+#else
+			return 0;
+#endif
+		}
+
+	private:
+
+#ifdef MUTABLE_PROFILE
+	int64_t m_start=0;
+#endif
+	};
+
+}
diff --git a/Engine/Plugins/Experimental/Mutable/Source/MutableRuntime/Private/Raster.h b/Engine/Plugins/Experimental/Mutable/Source/MutableRuntime/Private/Raster.h
new file mode 100644
--- /dev/null
+++ b/Engine/Plugins/Experimental/Mutable/Source/MutableRuntime/Private/Raster.h
@@ -0,0 +1,255 @@
+// Copyright Epic Games, Inc. All Rights Reserved.
+
+#pragma once
+
+namespace mu
+{
+
+    inline int32_t min( int32_t a, int32_t b ) { return a<b? a : b; }
+    inline int32_t max( int32_t a, int32_t b ) { return a<b? b : a; }
+    inline int32_t min( int32_t a, int32_t b, int32_t c ) { return a<b? ( a<c? a : c ) : ( b<c? b : c ); }
+    inline int32_t max( int32_t a, int32_t b, int32_t c ) { return a>b? ( a>c? a : c ) : ( b>c? b : c ); }
+
+    inline int64_t min( int64_t a, int64_t b ) { return a<b? a : b; }
+    inline int64_t max( int64_t a, int64_t b ) { return a<b? b : a; }
+    inline int64_t min( int64_t a, int64_t b, int64_t c ) { return a<b? ( a<c? a : c ) : ( b<c? b : c ); }
+    inline int64_t max( int64_t a, int64_t b, int64_t c ) { return a>b? ( a>c? a : c ) : ( b>c? b : c ); }
+
+    inline float min( float a, float b ) { return a<b? a : b; }
+    inline float max( float a, float b ) { return a<b? b : a; }
+
+
+	template<int NUM_INTERPOLATORS>
+	class RasterVertex
+	{
+	public:
+		RasterVertex() : x(0), y(0) {}
+		RasterVertex( float aX, float aY ) : x(aX), y(aY) {}
+
+		//! Raster coordinates of the vertex in image space.
+		float x,y;
+
+		//! Extra channels to interpolate accross de triangles.
+		float interpolators[NUM_INTERPOLATORS];
+	};
+
+
+	//! Example rasterizer pixel processor.
+	//! It only takes the color interpolated and assigns it to the pixel.
+	class ColorPixelProcessor
+	{
+	public:
+		inline void ProcessPixel( unsigned char* pBufferPos, float interpolatedValues[3] ) const
+		{
+			pBufferPos[0] = (unsigned char)(interpolatedValues[0]*0xFF);
+			pBufferPos[1] = (unsigned char)(interpolatedValues[1]*0xFF);
+			pBufferPos[2] = (unsigned char)(interpolatedValues[2]*0xFF);
+		}
+	};
+
+	//! Templatized rasterizer. It is not very efficient on the interpolating side, but the cost
+    //! per pixel is supposed to be huge compared to the raster process (for lightmaps or
+    //! normalmaps).
+	//! TODO Optimize anyway.
+	//! TODO Option for two sided? right now it culls back facing.
+    template<class PIXEL_PROCESSOR, int NUM_INTERPOLATORS, typename INT=int64_t>
+	void Triangle( unsigned char* aBuffer, int aWidth, int aHeight, int aPixelSize
+				 , const RasterVertex<NUM_INTERPOLATORS> &av1
+				 , const RasterVertex<NUM_INTERPOLATORS> &av2
+				 , const RasterVertex<NUM_INTERPOLATORS> &av3
+				 , PIXEL_PROCESSOR &processor
+				 , bool cullBackface = true )
+	{
+		float interpolatedValues[NUM_INTERPOLATORS];
+
+		int stride = aWidth*aPixelSize;
+
+		// Check backface culling
+		// TODO better option than resorting vertices?
+		float lCrossZ = (av2.x-av1.x) * (av3.y-av1.y) - (av2.y-av1.y) * (av3.x-av1.x);
+		const RasterVertex<NUM_INTERPOLATORS>& v1 = av1;
+		const RasterVertex<NUM_INTERPOLATORS>& v2 = cullBackface ? av2 : ( (lCrossZ>0) ? av3 : av2);
+		const RasterVertex<NUM_INTERPOLATORS>& v3 = cullBackface ? av3 : ( (lCrossZ>0) ? av2 : av3);
+
+	    // 28.4 fixed-point coordinates
+        const INT Y1 = (INT)(16.0f * v1.y + 0.5f);
+        const INT Y2 = (INT)(16.0f * v2.y + 0.5f);
+        const INT Y3 = (INT)(16.0f * v3.y + 0.5f);
+
+        const INT X1 = (INT)(16.0f * v1.x + 0.5f);
+        const INT X2 = (INT)(16.0f * v2.x + 0.5f);
+        const INT X3 = (INT)(16.0f * v3.x + 0.5f);
+
+	    // Deltas
+        const INT DX12 = X1 - X2;
+        const INT DX23 = X2 - X3;
+        const INT DX31 = X3 - X1;
+
+        const INT DY12 = Y1 - Y2;
+        const INT DY23 = Y2 - Y3;
+        const INT DY31 = Y3 - Y1;
+
+	    // Fixed-point deltas
+        const INT FDX12 = DX12 * 16;
+        const INT FDX23 = DX23 * 16;
+        const INT FDX31 = DX31 * 16;
+
+        const INT FDY12 = DY12 * 16;
+        const INT FDY23 = DY23 * 16;
+        const INT FDY31 = DY31 * 16;
+
+	    // Bounding rectangle
+        INT minx = max(INT(0),        ( min(X1, X2, X3) + 0xF) / 16 );
+        INT maxx = min(INT(aWidth-1), ( max(X1, X2, X3) + 0xF) / 16 );
+        INT miny = max(INT(0),		  ( min(Y1, Y2, Y3) + 0xF) / 16 );
+        INT maxy = min(INT(aHeight-1),( max(Y1, Y2, Y3) + 0xF) / 16 );
+
+	    // Block size, standard 8x8 (must be power of two)
+        const INT q = 8;
+
+	    // Start in corner of 8x8 block
+	    minx &= ~(q - 1);
+	    miny &= ~(q - 1);
+
+		aBuffer += miny * stride;
+
+	    // Half-edge constants
+        INT C1 = DY12 * X1 - DX12 * Y1;
+        INT C2 = DY23 * X2 - DX23 * Y2;
+        INT C3 = DY31 * X3 - DX31 * Y3;
+
+	    // Correct for fill convention
+	    if(DY12 < 0 || (DY12 == 0 && DX12 > 0)) C1++;
+	    if(DY23 < 0 || (DY23 == 0 && DX23 > 0)) C2++;
+	    if(DY31 < 0 || (DY31 == 0 && DX31 > 0)) C3++;
+
+		// Interpolation process constant stuff
+        float ux = v2.x-v1.x;
+        float uy = v2.y-v1.y;
+        float vx = v3.x-v1.x;
+        float vy = v3.y-v1.y;
+        float det = ux*vy-vx*uy;
+        ux /= det;
+        uy /= det;
+        vx /= det;
+        vy /= det;
+
+	    // Loop through blocks
+        for(INT y = miny; y < maxy; y += q)
+	    {
+            for(INT x = minx; x < maxx; x += q)
+	        {
+	            // Corners of block
+                INT x0 = x * 16;
+                INT x1 = (x + q - 1) * 16;
+                INT y0 = y * 16;
+                INT y1 = (y + q - 1) * 16;
+
+	            // Evaluate half-space functions
+	            bool a00 = C1 + DX12 * y0 - DY12 * x0 > 0;
+	            bool a10 = C1 + DX12 * y0 - DY12 * x1 > 0;
+	            bool a01 = C1 + DX12 * y1 - DY12 * x0 > 0;
+	            bool a11 = C1 + DX12 * y1 - DY12 * x1 > 0;
+                INT a = (INT(a00) * 1) | (INT(a10) * 2) | (INT(a01) * 4) | (INT(a11) * 8);
+
+	            bool b00 = C2 + DX23 * y0 - DY23 * x0 > 0;
+	            bool b10 = C2 + DX23 * y0 - DY23 * x1 > 0;
+	            bool b01 = C2 + DX23 * y1 - DY23 * x0 > 0;
+	            bool b11 = C2 + DX23 * y1 - DY23 * x1 > 0;
+                INT b = (INT(b00) * 1) | (INT(b10) * 2) | (INT(b01) * 4) | (INT(b11) * 8);
+
+	            bool c00 = C3 + DX31 * y0 - DY31 * x0 > 0;
+	            bool c10 = C3 + DX31 * y0 - DY31 * x1 > 0;
+	            bool c01 = C3 + DX31 * y1 - DY31 * x0 > 0;
+	            bool c11 = C3 + DX31 * y1 - DY31 * x1 > 0;
+                INT c = (INT(c00) * 1) | (INT(c10) * 2) | (INT(c01) * 4) | (INT(c11) * 8);
+
+	            // Skip block when outside an edge
+	            if(a == 0x0 || b == 0x0 || c == 0x0) continue;
+
+				unsigned char *buffer = aBuffer;
+
+	            // Accept whole block when totally covered
+	            if(a == 0xF && b == 0xF && c == 0xF)
+	            {
+                    for(INT iy = 0; iy < q; iy++)
+	                {
+                        for(INT ix = x; ix < x + q; ix++)
+	                    {
+	                        // Interpolate values
+	                        float px = ix-v1.x;
+	                        float py = y+iy-v1.y;
+							float alpha = px*vy+py*-vx;
+							float beta = px*-uy+py*ux;
+							alpha = max(0.0,min(1.0f, alpha ));
+							beta = max(0.0,min(1.0f, beta ));
+							float gamma = max(0.0, 1.0f-alpha-beta );
+
+	                        for (int ii=0; ii<NUM_INTERPOLATORS; ii++)
+							{
+                                interpolatedValues[ii] = gamma * v1.interpolators[ii]
+													   + alpha * v2.interpolators[ii]
+													   + beta * v3.interpolators[ii];
+							}
+
+	                        // Process pixel
+							processor.ProcessPixel( &buffer[ix*aPixelSize], interpolatedValues );
+	                    }
+
+	                    buffer += stride;
+	                }
+	            }
+	            else // Partially covered block
+	            {
+                    INT CY1 = C1 + DX12 * y0 - DY12 * x0;
+                    INT CY2 = C2 + DX23 * y0 - DY23 * x0;
+                    INT CY3 = C3 + DX31 * y0 - DY31 * x0;
+
+                    for(INT iy = y; iy < y + q; iy++)
+	                {
+                        INT CX1 = CY1;
+                        INT CX2 = CY2;
+                        INT CX3 = CY3;
+
+                        for(INT ix = x; ix < x + q; ix++)
+	                    {
+	                        if(CX1 > 0 && CX2 > 0 && CX3 > 0)
+							{
+		                        // Interpolate values
+		                        float px = ix-v1.x;
+		                        float py = iy-v1.y;
+								float alpha = px*vy+py*-vx;
+								float beta = px*-uy+py*ux;
+								alpha = max(0.0,min(1.0f, alpha ));
+								beta = max(0.0,min(1.0f, beta ));
+								float gamma = max(0.0, 1.0f-alpha-beta );
+
+		                        for (int ii=0; ii<NUM_INTERPOLATORS; ii++)
+								{
+									interpolatedValues[ii] = gamma * v1.interpolators[ii]
+															+ alpha * v2.interpolators[ii]
+															+ beta * v3.interpolators[ii];
+								}
+
+								processor.ProcessPixel( &buffer[ix*aPixelSize], interpolatedValues );
+	                        }
+
+	                        CX1 -= FDY12;
+	                        CX2 -= FDY23;
+	                        CX3 -= FDY31;
+	                    }
+
+	                    CY1 += FDX12;
+	                    CY2 += FDX23;
+	                    CY3 += FDX31;
+
+	                    buffer += stride;
+	                }
+	            }
+	        }
+
+			aBuffer += q * stride;
+	    }
+	}
+
+}
diff --git a/Engine/Plugins/Experimental/Mutable/Source/MutableRuntime/Private/Serialisation.cpp b/Engine/Plugins/Experimental/Mutable/Source/MutableRuntime/Private/Serialisation.cpp
new file mode 100644
--- /dev/null
+++ b/Engine/Plugins/Experimental/Mutable/Source/MutableRuntime/Private/Serialisation.cpp
@@ -0,0 +1,271 @@
+// Copyright Epic Games, Inc. All Rights Reserved.
+
+
+#include "Serialisation.h"
+#include "Platform.h"
+
+#include "SerialisationPrivate.h"
+
+#include "Image.h"
+
+
+namespace mu
+{
+
+
+    //---------------------------------------------------------------------------------------------
+    //---------------------------------------------------------------------------------------------
+    //---------------------------------------------------------------------------------------------
+    InputMemoryStream::InputMemoryStream( const void* pBuffer, uint64 size )
+    {
+        m_pD = new Private();
+        m_pD->m_pBuffer = pBuffer;
+        m_pD->m_size = size;
+    }
+
+
+    //---------------------------------------------------------------------------------------------
+    InputMemoryStream::~InputMemoryStream()
+    {
+        check( m_pD );
+        delete m_pD;
+        m_pD = 0;
+    }
+
+
+    //---------------------------------------------------------------------------------------------
+    void InputMemoryStream::Read( void* pData, uint64 size )
+    {
+        if (size)
+        {
+            check( m_pD->m_pos + size <= m_pD->m_size );
+
+            // This could happen in 32-bit platforms
+            check( size<std::numeric_limits<size_t>::max() );
+
+            const uint8_t* pSource = ((const uint8_t*)(m_pD->m_pBuffer))+m_pD->m_pos;
+            memcpy( pData, pSource, (size_t)size );
+            m_pD->m_pos += size;
+        }
+    }
+
+
+    //---------------------------------------------------------------------------------------------
+    //---------------------------------------------------------------------------------------------
+    //---------------------------------------------------------------------------------------------
+    OutputMemoryStream::OutputMemoryStream(uint64 reserve )
+    {
+        m_pD = new Private();
+        if (reserve)
+        {
+            // This could happen in 32-bit platforms
+            check( reserve<std::numeric_limits<size_t>::max() );
+            m_pD->m_buffer.Reserve( reserve );
+        }
+    }
+
+
+    //---------------------------------------------------------------------------------------------
+    OutputMemoryStream::~OutputMemoryStream()
+    {
+        check( m_pD );
+        delete m_pD;
+        m_pD = 0;
+    }
+
+
+    //---------------------------------------------------------------------------------------------
+    void OutputMemoryStream::Write( const void* pData, uint64 size )
+    {
+        if (size)
+        {
+            // This could happen in 32-bit platforms
+            check( size<std::numeric_limits<size_t>::max() );
+
+            uint64 pos = m_pD->m_buffer.Num();
+            m_pD->m_buffer.SetNum( pos + size, false );
+            memcpy( &m_pD->m_buffer[pos], pData, size );
+        }
+    }
+
+
+    //---------------------------------------------------------------------------------------------
+    const void* OutputMemoryStream::GetBuffer() const
+    {
+        const void* pResult = 0;
+
+        if ( m_pD->m_buffer.Num() )
+        {
+            pResult = &m_pD->m_buffer[0];
+        }
+
+        return pResult;
+    }
+
+
+    //---------------------------------------------------------------------------------------------
+	uint64 OutputMemoryStream::GetBufferSize() const
+    {
+        return m_pD->m_buffer.Num();
+    }
+
+
+    //-------------------------------------------------------------------------------------------------
+    //-------------------------------------------------------------------------------------------------
+    //-------------------------------------------------------------------------------------------------
+    OutputSizeStream::OutputSizeStream()
+    {
+        m_writtenBytes = 0;
+    }
+
+
+    //-------------------------------------------------------------------------------------------------
+    void OutputSizeStream::Write( const void*, uint64 size )
+    {
+        m_writtenBytes += size;
+    }
+
+
+    //-------------------------------------------------------------------------------------------------
+	uint64 OutputSizeStream::GetBufferSize() const
+    {
+        return m_writtenBytes;
+    }
+
+
+    //---------------------------------------------------------------------------------------------
+	//---------------------------------------------------------------------------------------------
+	//---------------------------------------------------------------------------------------------
+	InputArchive::InputArchive( InputStream* pStream )
+	{
+		m_pD = new Private();
+		m_pD->m_pStream = pStream;
+	}
+
+
+	//---------------------------------------------------------------------------------------------
+	InputArchive::~InputArchive()
+	{
+        check( m_pD );
+		delete m_pD;
+		m_pD = 0;
+	}
+
+
+	//---------------------------------------------------------------------------------------------
+	InputArchive::Private* InputArchive::GetPrivate() const
+	{
+		return m_pD;
+	}
+
+
+    //---------------------------------------------------------------------------------------------
+    Ptr<ResourceProxy<Image>> InputArchive::NewImageProxy()
+    {
+        return nullptr;
+    }
+
+
+	//---------------------------------------------------------------------------------------------
+	//---------------------------------------------------------------------------------------------
+	//---------------------------------------------------------------------------------------------
+	OutputArchive::OutputArchive( OutputStream* pStream )
+	{
+		m_pD = new Private();
+		m_pD->m_pStream = pStream;
+	}
+
+
+	//---------------------------------------------------------------------------------------------
+	OutputArchive::~OutputArchive()
+	{
+        check( m_pD );
+		delete m_pD;
+		m_pD = 0;
+	}
+
+
+	//---------------------------------------------------------------------------------------------
+	OutputArchive::Private* OutputArchive::GetPrivate() const
+	{
+		return m_pD;
+	}
+
+
+    //---------------------------------------------------------------------------------------------
+    //---------------------------------------------------------------------------------------------
+    //---------------------------------------------------------------------------------------------
+    class InputArchiveWithProxies::Private : public Base
+    {
+    public:
+        vector< Ptr<ResourceProxy<Image>> > m_proxyHistory;
+        ProxyFactory* m_pFactory = nullptr;
+    };
+
+
+    //---------------------------------------------------------------------------------------------
+    //---------------------------------------------------------------------------------------------
+    //---------------------------------------------------------------------------------------------
+    InputArchiveWithProxies::InputArchiveWithProxies( InputStream* s, ProxyFactory* f )
+        : InputArchive( s )
+    {
+        m_pD = new Private();
+        m_pD->m_pFactory = f;
+    }
+
+
+    //---------------------------------------------------------------------------------------------
+    InputArchiveWithProxies::~InputArchiveWithProxies()
+    {
+        delete m_pD;
+    }
+
+
+    //---------------------------------------------------------------------------------------------
+    Ptr<ResourceProxy<Image>> InputArchiveWithProxies::NewImageProxy()
+    {
+        // Similar to Ptr serisalisation in SerialisationPrivate
+        Ptr<ResourceProxy<Image>> p;
+        {
+            int32 id;
+            (*this) >> id;
+
+            if ( id == -1 )
+            {
+                // We consumed the serialisation, so we need to return something.
+                class ImageProxyNull : public ResourceProxy<Image>
+                {
+                public:
+                    ImagePtrConst Get() override
+                    {
+                        return nullptr;
+                    }
+                };
+                p = new ImageProxyNull();
+            }
+            else
+            {
+                if ( (std::size_t)id < m_pD->m_proxyHistory.size() )
+                {
+                    p = m_pD->m_proxyHistory[id];
+
+                    // If the pointer was nullptr it means the position in history is used, but not set
+                    // yet: we have a smart pointer loop which is very bad.
+                    check( p );
+                }
+                else
+                {
+                    // Ids come in order.
+                    m_pD->m_proxyHistory.resize(id+1);
+
+                    p = m_pD->m_pFactory->NewImageProxy(*this);
+                    m_pD->m_proxyHistory[id] = p;
+                }
+            }
+        }
+
+        return p;
+    }
+
+}
+
