[mutable] Move Mutable plugin from the restricted folder to the experimental plugins folder.

diff --git a/Engine/Plugins/Experimental/Mutable/Source/MutableRuntime/Private/SerialisationPrivate.h b/Engine/Plugins/Experimental/Mutable/Source/MutableRuntime/Private/SerialisationPrivate.h
new file mode 100644
--- /dev/null
+++ b/Engine/Plugins/Experimental/Mutable/Source/MutableRuntime/Private/SerialisationPrivate.h
@@ -0,0 +1,547 @@
+// Copyright Epic Games, Inc. All Rights Reserved.
+
+#pragma once
+
+#include "Serialisation.h"
+
+#include "Ptr.h"
+#include "RefCounted.h"
+#include "MutableMemory.h"
+#include "Types.h"
+#include "MemoryPrivate.h"
+#include "Math/Quat.h"
+#include "Math/Vector.h"
+#include "Math/Transform.h"
+
+#include <vector>
+
+
+namespace mu
+{
+
+    //!
+    class MUTABLERUNTIME_API InputMemoryStream : public InputStream
+    {
+    public:
+
+        //-----------------------------------------------------------------------------------------
+        // Life cycle
+        //-----------------------------------------------------------------------------------------
+
+        //! Create the stream using an external buffer.
+        //! The buffer will not be owned by this object, so it cannot be deallocated while this
+        //! objects is in use.
+        InputMemoryStream( const void* pBuffer, uint64 size );
+
+        ~InputMemoryStream();
+
+
+        //-----------------------------------------------------------------------------------------
+        // InputStream interface
+        //-----------------------------------------------------------------------------------------
+        void Read( void* pData, uint64 size ) override;
+
+
+    private:
+
+        class Private;
+        Private* m_pD;
+
+    };
+
+
+    //! This stream doesn-t store any data, it just counts de amonut of data serialised.
+    class MUTABLERUNTIME_API OutputSizeStream : public OutputStream
+    {
+    public:
+
+        //-----------------------------------------------------------------------------------------
+        // Life cycle
+        //-----------------------------------------------------------------------------------------
+
+        //!
+        OutputSizeStream();
+
+        //-----------------------------------------------------------------------------------------
+        // OutputStream interface
+        //-----------------------------------------------------------------------------------------
+        void Write( const void* pData, uint64 size ) override;
+
+        //-----------------------------------------------------------------------------------------
+        // Own interface
+        //-----------------------------------------------------------------------------------------
+
+        //! Get the amount of data serialised, in bytes.
+        uint64 GetBufferSize() const;
+
+    private:
+
+        uint64 m_writtenBytes;
+
+    };
+
+
+    //---------------------------------------------------------------------------------------------
+    class MUTABLERUNTIME_API InputMemoryStream::Private : public Base
+    {
+    public:
+
+        Private()
+        {
+            m_pBuffer = 0;
+            m_size = 0;
+            m_pos = 0;
+        }
+
+        const void* m_pBuffer;
+
+		uint64 m_size;
+
+		uint64 m_pos;
+    };
+
+
+    //---------------------------------------------------------------------------------------------
+    class MUTABLERUNTIME_API OutputMemoryStream::Private : public Base
+    {
+    public:
+        TArray<uint8,TSizedHeapAllocator<64>> m_buffer;
+    };
+
+
+	//---------------------------------------------------------------------------------------------
+	class MUTABLERUNTIME_API InputArchive::Private : public Base
+	{
+	public:
+
+		//! Not owned
+		InputStream* m_pStream;
+
+		//! Already read pointers
+		TArray< Ptr<RefCounted> > m_history;
+
+	};
+
+	//---------------------------------------------------------------------------------------------
+	class MUTABLERUNTIME_API OutputArchive::Private : public Base
+	{
+	public:
+
+		//! Not owned
+		OutputStream* m_pStream;
+
+		//! Already written pointers and their ids
+        TMap< const void*, int32 > m_history;
+
+	};
+
+
+	//---------------------------------------------------------------------------------------------
+	template< typename T >
+	inline void operator<< ( OutputArchive& arch, const T& t )
+	{
+        t.Serialise( arch );
+	}
+
+	template< typename T >
+	inline void operator>> ( InputArchive& arch, T& t )
+	{
+        t.Unserialise( arch );
+	}
+
+
+	//---------------------------------------------------------------------------------------------
+#define MUTABLE_DEFINE_POD_SERIALISABLE(T)								\
+		template<>														\
+		inline void operator<< <T>( OutputArchive& arch, const T& t )	\
+		{																\
+			arch.GetPrivate()->m_pStream->Write( &t, sizeof(T) );		\
+		}																\
+																		\
+		template<>														\
+		inline void operator>> <T>( InputArchive& arch, T& t )			\
+		{																\
+			arch.GetPrivate()->m_pStream->Read( &t, sizeof(T) );		\
+		}
+
+	MUTABLE_DEFINE_POD_SERIALISABLE(float);
+	MUTABLE_DEFINE_POD_SERIALISABLE(double);
+
+    MUTABLE_DEFINE_POD_SERIALISABLE( int8 );
+    MUTABLE_DEFINE_POD_SERIALISABLE( int16 );
+    MUTABLE_DEFINE_POD_SERIALISABLE( int32 );
+    MUTABLE_DEFINE_POD_SERIALISABLE( int64 );
+
+    MUTABLE_DEFINE_POD_SERIALISABLE( uint8 );
+    MUTABLE_DEFINE_POD_SERIALISABLE( uint16 );
+    MUTABLE_DEFINE_POD_SERIALISABLE( uint32 );
+    MUTABLE_DEFINE_POD_SERIALISABLE( uint64 );
+
+	//---------------------------------------------------------------------------------------------
+	// Bool size is not a standard
+	//---------------------------------------------------------------------------------------------
+	template<>														
+	inline void operator<< <bool>( OutputArchive& arch, const bool& t )
+	{
+        uint8 s = t ? 1 : 0;
+        arch.GetPrivate()->m_pStream->Write( &s, sizeof(uint8) );
+	}
+
+	template<>
+	inline void operator>> <bool>( InputArchive& arch, bool& t )
+	{
+        uint8 s;
+        arch.GetPrivate()->m_pStream->Read( &s, sizeof(uint8) );
+		t = s!=0;
+	}
+
+
+	//---------------------------------------------------------------------------------------------
+#define MUTABLE_DEFINE_ENUM_SERIALISABLE(T)								\
+		template<>														\
+        inline void operator<< <T>( OutputArchive& arch, const T& t )	\
+		{																\
+            auto v = (uint32)t;                                       \
+            arch.GetPrivate()->m_pStream->Write( &v, sizeof(uint32) );\
+		}																\
+																		\
+		template<>														\
+        inline void operator>> <T>( InputArchive& arch, T& t )          \
+		{																\
+            uint32 v;													\
+            arch.GetPrivate()->m_pStream->Read( &v, sizeof(uint32) );	\
+			t = (T)v;													\
+		}
+
+	//---------------------------------------------------------------------------------------------
+	template< typename T >
+	inline void operator<< ( OutputArchive& arch, const vector<T>& v )
+	{
+		// TODO: Optimise for vectors of PODs
+        arch << (uint32)v.size();
+        for( std::size_t i=0; i<v.size(); ++i )
+		{
+			arch << v[i];
+		}
+	}
+
+	template< typename T >
+	inline void operator>> ( InputArchive& arch, vector<T>& v )
+	{
+		// TODO: Optimise for vectors of PODs
+        uint32 size;
+		arch >> size;
+		v.resize( size );
+        for( uint32 i=0; i<size; ++i )
+		{
+			arch >> v[i];
+		}
+	}
+
+    //---------------------------------------------------------------------------------------------
+    template< typename T0, typename T1 >
+    inline void operator<< ( OutputArchive& arch, const std::pair<T0,T1>& v )
+    {
+        arch << v.first;
+        arch << v.second;
+    }
+
+    template< typename T0, typename T1 >
+    inline void operator>> ( InputArchive& arch, std::pair<T0,T1>& v )
+    {
+        arch >> v.first;
+        arch >> v.second;
+    }
+
+
+    //---------------------------------------------------------------------------------------------
+    template< typename K, typename T >
+    inline void operator<< ( OutputArchive& arch, const map<K,T>& v )
+    {
+        arch << (uint32_t)v.size();
+        for( const auto& p: v )
+        {
+            arch << p.first;
+            arch << p.second;
+        }
+    }
+
+    template< typename K, typename T >
+    inline void operator>> ( InputArchive& arch, map<K,T>& v )
+    {
+        uint32_t size;
+        arch >> size;
+
+        for( uint32_t i=0; i<size; ++i )
+        {
+            K k;
+            T t;
+            arch >> k;
+            arch >> t;
+            v.insert( std::make_pair<>(k,t) );
+        }
+    }
+
+	
+	//---------------------------------------------------------------------------------------------
+	template<typename T> inline void operator<<(OutputArchive& arch, const TArray<T>& v)
+	{
+		// TODO: Optimise for vectors of PODs
+		arch << (uint32_t)v.Num();
+		for (std::size_t i = 0; i < v.Num(); ++i)
+		{
+			arch << v[i];
+		}
+	}
+
+	template<typename T> inline void operator>>(InputArchive& arch, TArray<T>& v)
+	{
+		// TODO: Optimise for vectors of PODs
+		uint32_t size;
+		arch >> size;
+		v.SetNum(size);
+		for (std::size_t i = 0; i < size; ++i)
+		{
+			arch >> v[i];
+		}
+	}
+
+
+	//---------------------------------------------------------------------------------------------
+#define MUTABLE_DEFINE_POD_VECTOR_SERIALISABLE(T)										\
+	template<>																			\
+	inline void operator<< <T>( OutputArchive& arch, const vector<T>& v )				\
+	{																					\
+        arch << (uint32_t)v.size();														\
+		if (v.size())																	\
+		{																				\
+            arch.GetPrivate()->m_pStream->Write( &v[0], (uint32_t)v.size()*sizeof(T) );	\
+		}																				\
+	}																					\
+																						\
+	template<>																			\
+	inline void operator>> <T>( InputArchive& arch, vector<T>& v )						\
+	{																					\
+        uint32_t size;																	\
+		arch >> size;																	\
+		v.resize( size );																\
+		if (size)																		\
+		{																				\
+            arch.GetPrivate()->m_pStream->Read( &v[0], (uint32_t)size*sizeof(T) );		\
+		}																				\
+	}																					\
+																						\
+	template<> inline void operator<<<T>(OutputArchive& arch, const TArray<T>& v)		\
+	{                                                                                   \
+		uint32 Num = uint32(v.Num());													\
+		arch << Num;																	\
+		if (Num)																		\
+		{                                                                               \
+			arch.GetPrivate()->m_pStream->Write(&v[0], Num * sizeof(T));				\
+		}                                                                               \
+	}                                                                                   \
+																						\
+	template<> inline void operator>><T>(InputArchive& arch, TArray<T>& v)				\
+	{                                                                                   \
+		uint32 Num;																		\
+		arch >> Num;																	\
+		v.SetNum(Num);																	\
+		if (Num)																		\
+		{                                                                               \
+			arch.GetPrivate()->m_pStream->Read(&v[0], Num * sizeof(T));					\
+		}                                                                               \
+	}
+
+	MUTABLE_DEFINE_POD_VECTOR_SERIALISABLE(float)
+	MUTABLE_DEFINE_POD_VECTOR_SERIALISABLE(double)
+	MUTABLE_DEFINE_POD_VECTOR_SERIALISABLE(uint8)
+	MUTABLE_DEFINE_POD_VECTOR_SERIALISABLE(uint16)
+	MUTABLE_DEFINE_POD_VECTOR_SERIALISABLE(uint32)
+	MUTABLE_DEFINE_POD_VECTOR_SERIALISABLE(uint64)
+	MUTABLE_DEFINE_POD_VECTOR_SERIALISABLE(int8)
+	MUTABLE_DEFINE_POD_VECTOR_SERIALISABLE(int16)
+	MUTABLE_DEFINE_POD_VECTOR_SERIALISABLE(int32)
+	MUTABLE_DEFINE_POD_VECTOR_SERIALISABLE(int64)
+
+
+	//---------------------------------------------------------------------------------------------
+	template<>
+	inline void operator<< <string>( OutputArchive& arch, const string& v )
+	{
+        arch << (uint32_t)v.size();
+		if ( v.size() )
+		{
+			arch.GetPrivate()->m_pStream->Write( &v[0], (unsigned)v.size()*sizeof(char) );
+		}
+	}
+
+	template<>
+	inline void operator>> <string>( InputArchive& arch, string& v )
+	{
+        uint32_t size;
+		arch >> size;
+		v.resize( size );
+		if (size)
+		{
+			arch.GetPrivate()->m_pStream->Read( &v[0], (unsigned)size*sizeof(char) );
+		}
+	}
+
+
+	//---------------------------------------------------------------------------------------------
+	template< typename T >
+	inline void operator<< ( OutputArchive& arch, const Ptr<const T>& p )
+	{
+		if (!p.get())
+		{
+            arch << (int32)-1;
+		}
+		else
+		{
+            int32* it = arch.GetPrivate()->m_history.Find(p.get());
+
+			if ( !it )
+			{
+                int32 id = arch.GetPrivate()->m_history.Num();
+                arch.GetPrivate()->m_history.Add( p.get(), id );
+				arch << id;
+				T::Serialise( p.get(), arch );
+			}
+			else
+			{
+				arch << *it;
+			}
+		}
+	}
+
+	template< typename T >
+	inline void operator>> ( InputArchive& arch, Ptr<const T>& p )
+	{
+        int32 id;
+		arch >> id;
+
+		if ( id == -1 )
+		{
+			p = 0;
+		}
+		else
+		{
+			if ( id < arch.GetPrivate()->m_history.Num() )
+			{
+				p = static_cast<T*>( arch.GetPrivate()->m_history[id].get() );
+
+				// If the pointer was 0 it means the position in history is used, but not set yet
+				// option 1: we have a smart pointer loop which is very bad.
+				// option 2: the resource in this Ptr is also pointed by a Proxy that has absorbed it
+				//			 and this reference should also be a proxy instead of a ptr.
+                check( p );
+			}
+			else
+			{
+                // Ids come in order, but they may have been absorbed outside in some serialisations
+                // like proxies.
+                //check( id == (int)arch.GetPrivate()->m_history.size() );
+                arch.GetPrivate()->m_history.SetNum( id+1 );
+
+				Ptr<T> t = T::StaticUnserialise( arch );
+				p = t;
+				arch.GetPrivate()->m_history[id] = t;
+			}
+		}
+	}
+
+
+	//---------------------------------------------------------------------------------------------
+	template< typename T >
+	inline void operator<< ( OutputArchive& arch, const Ptr<T>& p )
+	{
+		operator<<( arch, (const Ptr<const T>&) p );
+	}
+
+	template< typename T >
+	inline void operator>> ( InputArchive& arch, Ptr<T>& p )
+	{
+		operator>>( arch, (Ptr<const T>&) p );
+	}
+
+
+	//---------------------------------------------------------------------------------------------
+	//---------------------------------------------------------------------------------------------
+	//---------------------------------------------------------------------------------------------
+	template<typename T0, typename T1> inline void operator<<(OutputArchive& arch, const TPair<T0,T1>& v)
+	{
+		arch << v.Key;
+		arch << v.Value;
+	}
+
+	template<typename T0, typename T1> inline void operator>>(InputArchive& arch, TPair<T0, T1>& v)
+	{
+		arch >> v.Key;
+		arch >> v.Value;
+	}
+
+
+	//---------------------------------------------------------------------------------------------
+	// TODO: As POD?
+	template< typename T >
+	inline void operator<< (OutputArchive& arch, const UE::Math::TQuat<T>& v)
+	{		
+		arch << v.X;
+		arch << v.Y;
+		arch << v.Z;
+		arch << v.W;
+	}
+
+	template< typename T >
+	inline void operator>> (InputArchive& arch, UE::Math::TQuat<T>& v)
+	{
+		arch >> v.X;
+		arch >> v.Y;
+		arch >> v.Z;
+		arch >> v.W;
+	}
+
+
+	//---------------------------------------------------------------------------------------------
+	// TODO: As POD?
+	template< typename T >
+	inline void operator<< (OutputArchive& arch, const UE::Math::TVector<T>& v)
+	{
+		arch << v.X;
+		arch << v.Y;
+		arch << v.Z;
+	}
+
+	template< typename T >
+	inline void operator>> (InputArchive& arch, UE::Math::TVector<T>& v)
+	{
+		arch >> v.X;
+		arch >> v.Y;
+		arch >> v.Z;
+	}
+
+
+	//---------------------------------------------------------------------------------------------
+	template< typename T >
+	inline void operator<< (OutputArchive& arch, const UE::Math::TTransform<T>& v)
+	{
+		arch << v.GetRotation();
+		arch << v.GetTranslation();
+		arch << v.GetScale3D();
+	}
+
+	template< typename T >
+	inline void operator>> (InputArchive& arch, UE::Math::TTransform<T>& v)
+	{
+		UE::Math::TQuat<T> Rot;
+		UE::Math::TVector<T> Trans;
+		UE::Math::TVector<T> Scale;
+
+		arch >> Rot;
+		arch >> Trans;
+		arch >> Scale;
+
+		v.SetComponents(Rot, Trans, Scale);
+	}
+
+}
diff --git a/Engine/Plugins/Experimental/Mutable/Source/MutableRuntime/Private/Settings.cpp b/Engine/Plugins/Experimental/Mutable/Source/MutableRuntime/Private/Settings.cpp
new file mode 100644
--- /dev/null
+++ b/Engine/Plugins/Experimental/Mutable/Source/MutableRuntime/Private/Settings.cpp
@@ -0,0 +1,56 @@
+// Copyright Epic Games, Inc. All Rights Reserved.
+
+#include "Settings.h"
+#include "Platform.h"
+
+#include "SettingsPrivate.h"
+#include "Config.h"
+
+namespace mu
+{
+    //---------------------------------------------------------------------------------------------
+    //!
+    //---------------------------------------------------------------------------------------------
+    Settings::Settings()
+    {
+        m_pD = new Settings::Private();
+    }
+
+
+    //---------------------------------------------------------------------------------------------
+    Settings::~Settings()
+    {
+        check( m_pD );
+        delete m_pD;
+        m_pD = 0;
+    }
+
+
+    //---------------------------------------------------------------------------------------------
+    Settings::Private* Settings::GetPrivate() const
+    {
+        return m_pD;
+    }
+
+
+    //---------------------------------------------------------------------------------------------
+    void Settings::SetProfile( bool bEnabled )
+    {
+        m_pD->m_profile = bEnabled;
+    }
+
+
+    //---------------------------------------------------------------------------------------------
+    void Settings::SetStreamingCache( uint64_t bytes )
+    {
+        m_pD->m_streamingCacheBytes = bytes;
+    }
+
+
+    //---------------------------------------------------------------------------------------------
+    void Settings::SetImageCompressionQuality( int quality )
+    {
+        m_pD->m_imageCompressionQuality = quality;
+    }
+
+}
diff --git a/Engine/Plugins/Experimental/Mutable/Source/MutableRuntime/Private/SettingsPrivate.h b/Engine/Plugins/Experimental/Mutable/Source/MutableRuntime/Private/SettingsPrivate.h
new file mode 100644
--- /dev/null
+++ b/Engine/Plugins/Experimental/Mutable/Source/MutableRuntime/Private/SettingsPrivate.h
@@ -0,0 +1,23 @@
+// Copyright Epic Games, Inc. All Rights Reserved.
+
+#pragma once
+
+#include "Settings.h"
+
+
+namespace mu
+{
+
+    //---------------------------------------------------------------------------------------------
+    //! Private Settings implementation
+    //---------------------------------------------------------------------------------------------
+    class Settings::Private : public Base
+    {
+    public:
+        bool m_profile = false;
+        uint64_t m_streamingCacheBytes = 0;
+        int m_imageCompressionQuality = 0;
+    };
+
+
+}
diff --git a/Engine/Plugins/Experimental/Mutable/Source/MutableRuntime/Private/Skeleton.cpp b/Engine/Plugins/Experimental/Mutable/Source/MutableRuntime/Private/Skeleton.cpp
new file mode 100644
--- /dev/null
+++ b/Engine/Plugins/Experimental/Mutable/Source/MutableRuntime/Private/Skeleton.cpp
@@ -0,0 +1,133 @@
+// Copyright Epic Games, Inc. All Rights Reserved.
+
+
+#include "Skeleton.h"
+
+#include "Platform.h"
+#include "MutableTrace.h"
+
+
+using namespace mu;
+
+//---------------------------------------------------------------------------------------------
+void Skeleton::Serialise( const Skeleton* p, OutputArchive& arch )
+{
+    arch << *p;
+}
+
+
+//---------------------------------------------------------------------------------------------
+mu::Ptr<Skeleton> Skeleton::StaticUnserialise( InputArchive& arch )
+{
+    SkeletonPtr pResult = new Skeleton();
+    arch >> *pResult;
+    return pResult;
+}
+
+
+//---------------------------------------------------------------------------------------------
+SkeletonPtr Skeleton::Clone() const
+{
+    SkeletonPtr pResult = new Skeleton();
+
+	pResult->m_bones = m_bones;
+	pResult->m_boneParents = m_boneParents;
+	pResult->m_boneIds = m_boneIds;
+
+    return pResult;
+}
+
+
+//---------------------------------------------------------------------------------------------
+int32 Skeleton::GetBoneCount() const
+{
+    return m_bones.Num();
+}
+
+
+//---------------------------------------------------------------------------------------------
+void Skeleton::SetBoneCount( int32 b )
+{
+	m_bones.SetNum( b );
+	m_boneParents.Init(INDEX_NONE, b);
+	m_boneIds.Init(INDEX_NONE, b);
+}
+
+
+//---------------------------------------------------------------------------------------------
+const char* Skeleton::GetBoneName(int32 boneIndex ) const
+{
+    check( boneIndex>=0 && boneIndex<GetBoneCount() );
+    if (boneIndex >= 0 && boneIndex < GetBoneCount())
+    {
+        return m_bones[boneIndex].c_str();
+    }
+    return "";
+}
+
+
+//---------------------------------------------------------------------------------------------
+void Skeleton::SetBoneName(int32 boneIndex, const char* strName )
+{
+    check( boneIndex>=0 && boneIndex<GetBoneCount() );
+    if (boneIndex >= 0 && boneIndex < GetBoneCount())
+    {
+        m_bones[boneIndex] = strName;
+    }
+}
+
+
+//---------------------------------------------------------------------------------------------
+int Skeleton::GetBoneParent(int32 boneIndex ) const
+{
+    if ( boneIndex >= 0 && boneIndex < m_boneParents.Num() )
+    {
+        return m_boneParents[boneIndex];
+    }
+
+    return -1;
+}
+
+
+//-------------------------------------------------------------------------------------------------
+void Skeleton::SetBoneParent(int32 boneIndex, int32 parentBoneIndex )
+{
+    check( boneIndex >= 0 && boneIndex < GetBoneCount() );
+    check( parentBoneIndex >= -1 && parentBoneIndex < GetBoneCount() &&
+                    parentBoneIndex < 0xffff );
+    if ( boneIndex >= 0 && boneIndex < m_boneParents.Num() )
+    {
+        m_boneParents[boneIndex] = (int16_t)parentBoneIndex;
+    }
+}
+
+
+//-------------------------------------------------------------------------------------------------
+int32 Skeleton::GetBoneId(int32 boneIndex ) const
+{
+    if ( boneIndex >= 0 && boneIndex < m_boneIds.Num() )
+    {
+        return int(m_boneIds[boneIndex]);
+    }
+
+    return 0;
+}
+
+
+//-------------------------------------------------------------------------------------------------
+void Skeleton::SetBoneId(int32 boneIndex, int32 id )
+{
+    int boneCount = GetBoneCount();
+    check( boneIndex >= 0 && boneIndex < boneCount );
+    if ( boneIndex >= 0 && boneIndex < boneCount )
+    {
+        // Enlarge the id buffer if necessary, but enlarge it to the total number of bones to avoid
+        // repetitive allocation.
+        if ( boneIndex >= m_boneIds.Num() )
+        {
+            m_boneIds.SetNumZeroed( boneCount );
+        }
+
+        m_boneIds[boneIndex] = id;
+    }
+}
diff --git a/Engine/Plugins/Experimental/Mutable/Source/MutableRuntime/Private/System.cpp b/Engine/Plugins/Experimental/Mutable/Source/MutableRuntime/Private/System.cpp
new file mode 100644
--- /dev/null
+++ b/Engine/Plugins/Experimental/Mutable/Source/MutableRuntime/Private/System.cpp
@@ -0,0 +1,1005 @@
+// Copyright Epic Games, Inc. All Rights Reserved.
+
+#include "System.h"
+#include "Platform.h"
+
+#include "CodeRunner.h"
+#include "SystemPrivate.h"
+
+#include "ModelPrivate.h"
+#include "InstancePrivate.h"
+#include "ImagePrivate.h"
+#include "MeshPrivate.h"
+#include "MemoryPrivate.h"
+#include "ModelPrivate.h"
+#include "SettingsPrivate.h"
+#include "Operations.h"
+#include "CodeVisitor.h"
+
+#include "Config.h"
+#include "MutableTrace.h"
+
+namespace mu
+{
+    static_assert( sizeof(vec4f) == 16, "UNEXPECTED_STRUCT_PACKING" );
+    static_assert( sizeof(mat4f) == 64, "UNEXPECTED_STRUCT_PACKING" );
+
+    //!
+    static std::atomic<mu::Error> s_error(mu::Error::None);
+
+
+    //---------------------------------------------------------------------------------------------
+    //!
+    //---------------------------------------------------------------------------------------------
+    Error GetError()
+    {
+        return s_error;
+    }
+
+
+    //---------------------------------------------------------------------------------------------
+    //!
+    //---------------------------------------------------------------------------------------------
+    void ClearError()
+    {
+        s_error = Error::None;
+    }
+
+
+    //---------------------------------------------------------------------------------------------
+    //---------------------------------------------------------------------------------------------
+    //---------------------------------------------------------------------------------------------
+    System::System( const SettingsPtr& pInSettings )
+    {
+		LLM_SCOPE_BYNAME(TEXT("MutableRuntime"));
+
+        SettingsPtr pSettings = pInSettings;
+
+        if ( !pSettings )
+        {
+            pSettings = new Settings;
+        }
+
+        // Choose the implementation
+        m_pD = new System::Private( pSettings );
+    }
+
+
+    //---------------------------------------------------------------------------------------------
+    System::~System()
+    {
+		LLM_SCOPE_BYNAME(TEXT("MutableRuntime"));
+
+        check( m_pD );
+
+        delete m_pD;
+        m_pD = nullptr;
+    }
+
+
+    //---------------------------------------------------------------------------------------------
+    System::Private* System::GetPrivate() const
+    {
+        return m_pD;
+    }
+
+
+    //---------------------------------------------------------------------------------------------
+    System::Private::Private( SettingsPtr pSettings )
+    {
+		LLM_SCOPE_BYNAME(TEXT("MutableRuntime"));
+
+        m_pSettings = pSettings;
+        m_pStreamInterface = nullptr;
+        m_pImageParameterGenerator = nullptr;
+        m_maxMemory = 0;
+
+        for(uint8 p=0; p< uint8(ProfileMetric::_Count); ++p)
+        {
+            m_profileMetrics[p]=0;
+        }
+	
+		m_modelCache.m_romBudget = pSettings->GetPrivate()->m_streamingCacheBytes;
+	}
+
+
+    //---------------------------------------------------------------------------------------------
+    System::Private::~Private()
+    {
+		LLM_SCOPE_BYNAME(TEXT("MutableRuntime"));
+
+        delete m_pStreamInterface;
+        m_pStreamInterface = nullptr;
+
+        delete m_pImageParameterGenerator;
+        m_pImageParameterGenerator = nullptr;
+    }
+
+
+    //---------------------------------------------------------------------------------------------
+    void System::SetStreamingInterface( ModelStreamer *pInterface )
+    {
+		LLM_SCOPE_BYNAME(TEXT("MutableRuntime"));
+
+        (void)pInterface;
+       
+        delete m_pD->m_pStreamInterface;
+        m_pD->m_pStreamInterface = pInterface;
+    }
+
+
+    //---------------------------------------------------------------------------------------------
+    void System::SetStreamingCache( uint64 bytes )
+    {
+		LLM_SCOPE_BYNAME(TEXT("MutableRuntime"));
+
+        m_pD->SetStreamingCache( bytes );
+    }
+
+
+    //---------------------------------------------------------------------------------------------
+    void System::SetImageParameterGenerator( ImageParameterGenerator* pInterface )
+    {
+		LLM_SCOPE_BYNAME(TEXT("MutableRuntime"));
+
+        delete m_pD->m_pImageParameterGenerator;
+        m_pD->m_pImageParameterGenerator = pInterface;
+    }
+
+
+    //---------------------------------------------------------------------------------------------
+    void System::SetMemoryLimit( uint32 mem )
+    {
+		LLM_SCOPE_BYNAME(TEXT("MutableRuntime"));
+
+        m_pD->m_maxMemory = mem;
+
+        // TODO: Clear cache if we are over the new limit
+    }
+
+
+    //---------------------------------------------------------------------------------------------
+    void System::ClearCaches()
+    {
+		LLM_SCOPE_BYNAME(TEXT("MutableRuntime"));
+
+        // TODO
+    }
+
+
+    //---------------------------------------------------------------------------------------------
+    uint64 System::GetProfileMetric( ProfileMetric m ) const
+    {
+        size_t metric = size_t(m);
+
+        if (metric>=size_t(ProfileMetric::_Count)) return 0;
+
+        return m_pD->m_profileMetrics[metric];
+    }
+
+
+    //---------------------------------------------------------------------------------------------
+    Instance::ID System::NewInstance( const ModelPtrConst& pModel )
+    {
+		LLM_SCOPE_BYNAME(TEXT("MutableRuntime"));
+		MUTABLE_CPUPROFILER_SCOPE(NewInstance);
+
+		Private::FLiveInstance instanceData;
+		instanceData.m_instanceID = ++m_pD->m_lastInstanceID;
+		instanceData.m_pInstance = nullptr;
+		instanceData.m_pModel = pModel;
+		instanceData.m_state = -1;
+		instanceData.m_memory = MakeShared<FProgramCache>();
+		m_pD->m_liveInstances.Add(instanceData);
+
+		m_pD->m_profileMetrics[size_t(System::ProfileMetric::LiveInstanceCount)]++;
+
+		return instanceData.m_instanceID;
+	}
+
+
+    //---------------------------------------------------------------------------------------------
+    const Instance* System::BeginUpdate( Instance::ID instanceID,
+                                     const ParametersPtrConst& pParams,
+                                     int32 stateIndex,
+                                     uint32 lodMask )
+    {
+		LLM_SCOPE_BYNAME(TEXT("MutableRuntime"));
+		MUTABLE_CPUPROFILER_SCOPE(SystemBeginUpdate);
+
+		if (!pParams)
+		{
+			UE_LOG(LogMutableCore, Error, TEXT("Invalid parameters in mutable update."));
+			return nullptr;
+		}
+
+		Private::FLiveInstance* pLiveInstance = m_pD->FindLiveInstance(instanceID);
+		if (!pLiveInstance)
+		{
+			UE_LOG(LogMutableCore, Error, TEXT("Invalid instance id in mutable update."));
+			return nullptr;
+		}
+
+		m_pD->m_memory = pLiveInstance->m_memory;
+
+		PROGRAM& program = pLiveInstance->m_pModel->GetPrivate()->m_program;
+
+		bool validState = stateIndex >= 0 && stateIndex < (int)program.m_states.size();
+		if (!validState)
+		{
+			UE_LOG(LogMutableCore, Error, TEXT("Invalid state in mutable update."));
+			return nullptr;
+		}
+
+		// This may free resources that allow us to use less memory.
+		pLiveInstance->m_pInstance = nullptr;
+
+		bool fullBuild = (stateIndex != pLiveInstance->m_state);
+
+		pLiveInstance->m_state = stateIndex;
+
+		// If we changed parameters that are not in this state, we need to rebuild all.
+		if (!fullBuild)
+		{
+			fullBuild = m_pD->CheckUpdatedParameters(pLiveInstance, pParams.get(), pLiveInstance->m_updatedParameters);
+		}
+
+		// Remove cached data
+		pLiveInstance->m_memory->ClearCacheLayer0();
+		if (fullBuild)
+		{
+			pLiveInstance->m_memory->ClearCacheLayer1();
+		}
+
+		OP::ADDRESS rootAt = pLiveInstance->m_pModel->GetPrivate()->m_program.m_states[stateIndex].m_root;
+
+		m_pD->PrepareCache(pLiveInstance->m_pModel, stateIndex);
+		pLiveInstance->m_pOldParameters = pParams->Clone();
+
+		m_pD->RunCode(pLiveInstance->m_pModel.get(), pParams.get(), rootAt, lodMask);
+
+		InstancePtrConst pResult = pLiveInstance->m_memory->GetInstance(CACHE_ADDRESS(rootAt, 0, 0));
+
+		pLiveInstance->m_pInstance = pResult;
+		if (pResult)
+		{
+			pResult->GetPrivate()->m_id = pLiveInstance->m_instanceID;
+		}
+
+		m_pD->m_profileMetrics[size_t(System::ProfileMetric::InstanceUpdateCount)]++;
+
+		m_pD->m_memory = nullptr;
+
+		return pResult.get();
+	}
+
+
+	//---------------------------------------------------------------------------------------------
+	ImagePtrConst System::GetImage(Instance::ID instanceID, RESOURCE_ID imageId, int32 MipsToSkip)
+	{
+		LLM_SCOPE_BYNAME(TEXT("MutableRuntime"));
+		MUTABLE_CPUPROFILER_SCOPE(SystemGetImage);
+
+		ImagePtrConst pResult;
+
+		// Find the live instance
+		Private::FLiveInstance* pLiveInstance = m_pD->FindLiveInstance(instanceID);
+		check(pLiveInstance);
+		m_pD->m_memory = pLiveInstance->m_memory;
+
+		// Find the resource id in the model's resource cache
+		for (const Model::Private::RESOURCE_KEY& res : pLiveInstance->m_pModel->GetPrivate()->m_generatedResources)
+		{
+			if (res.m_id == imageId)
+			{
+				pResult = m_pD->BuildImage(pLiveInstance->m_pModel,
+					pLiveInstance->m_pOldParameters.get(),
+					res.m_rootAddress, MipsToSkip);
+
+				// We always need ot return something valid.
+				if (!pResult)
+				{
+					pResult = new mu::Image(16, 16, 1, EImageFormat::IF_RGBA_UBYTE);
+				}
+
+				pResult->m_internalId = imageId;
+
+				break;
+			}
+		}
+
+		m_pD->m_memory = nullptr;
+
+		return pResult;
+	}
+
+
+	// Temporarily make the Image DescCache clear at every image because otherwise it makes some textures 
+	// not evaluate their layout and be of size 0 and 0 lods, making them incorrectly evaluate MipsToSkip
+	static TAutoConsoleVariable<int32> CVarClearImageDescCache(
+		TEXT("mutable.ClearImageDescCache"),
+		1,
+		TEXT("If different than 0, clear the image desc cache at every image."),
+		ECVF_Scalability);
+
+
+	//---------------------------------------------------------------------------------------------
+	void System::GetImageDesc(Instance::ID instanceID, RESOURCE_ID imageId, FImageDesc& OutDesc)
+	{
+		LLM_SCOPE_BYNAME(TEXT("MutableRuntime"));
+		MUTABLE_CPUPROFILER_SCOPE(SystemGetImageDesc);
+
+		OutDesc = FImageDesc();
+
+		// Find the live instance
+		Private::FLiveInstance* pLiveInstance = m_pD->FindLiveInstance(instanceID);
+		check(pLiveInstance);
+		m_pD->m_memory = pLiveInstance->m_memory;
+
+		// Find the resource id in the model's resource cache
+		for (const Model::Private::RESOURCE_KEY& res : pLiveInstance->m_pModel->GetPrivate()->m_generatedResources)
+		{
+			if (res.m_id == imageId)
+			{
+				const mu::Model* Model = pLiveInstance->m_pModel.get();
+				const mu::PROGRAM& program = Model->GetPrivate()->m_program;
+
+				int32 VarValue = CVarClearImageDescCache.GetValueOnAnyThread();
+				if (VarValue != 0)
+				{
+					m_pD->m_memory->m_descCache.Reset();
+				}
+
+				m_pD->m_memory->m_descCache.SetNum(program.m_opAddress.size());
+
+				OP::ADDRESS at = res.m_rootAddress;
+				mu::OP_TYPE opType = program.GetOpType(at);
+				if (GetOpDataType(opType) == DT_IMAGE)
+				{
+					int8 executionOptions = 0;
+					CodeRunner Runner(m_pD->m_pSettings, m_pD, Model, pLiveInstance->m_pOldParameters.get(), at, System::AllLODs, executionOptions, SCHEDULED_OP::EType::ImageDesc);
+					Runner.Run();
+					Runner.GetImageDescResult(OutDesc);
+				}
+
+				break;
+			}
+		}
+
+		m_pD->m_memory = nullptr;
+	}
+
+
+    //---------------------------------------------------------------------------------------------
+    MeshPtrConst System::GetMesh( Instance::ID instanceID,
+                                   RESOURCE_ID meshId )
+    {
+		LLM_SCOPE_BYNAME(TEXT("MutableRuntime"));
+		MUTABLE_CPUPROFILER_SCOPE(SystemGetMesh);
+
+		MeshPtrConst pResult;
+
+		// Find the live instance
+		Private::FLiveInstance* pLiveInstance = m_pD->FindLiveInstance(instanceID);
+		check(pLiveInstance);
+		m_pD->m_memory = pLiveInstance->m_memory;
+
+		// Find the resource id in the model's resource cache
+		for (const Model::Private::RESOURCE_KEY& res : pLiveInstance->m_pModel->GetPrivate()->m_generatedResources)
+		{
+			if (res.m_id == meshId)
+			{
+				pResult = m_pD->BuildMesh(pLiveInstance->m_pModel,
+					pLiveInstance->m_pOldParameters.get(),
+					res.m_rootAddress);
+
+				// If the mesh is null it means empty, but we still need to return a valid one
+				if (!pResult)
+				{
+					pResult = new Mesh();
+				}
+
+				pResult->m_internalId = meshId;
+
+				break;
+			}
+		}
+
+		m_pD->m_memory = nullptr;
+		return pResult;
+	}
+
+
+    //---------------------------------------------------------------------------------------------
+    void System::EndUpdate( Instance::ID instanceID )
+    {
+		LLM_SCOPE_BYNAME(TEXT("MutableRuntime"));
+		MUTABLE_CPUPROFILER_SCOPE(EndUpdate);
+
+		// Reduce the cache until it fits the limit.
+		uint64 totalMemory = m_pD->m_modelCache.EnsureCacheBelowBudget(0);
+		m_pD->m_profileMetrics[size_t(System::ProfileMetric::StreamingCacheBytes)] = totalMemory;
+
+
+		Private::FLiveInstance* pLiveInstance = m_pD->FindLiveInstance(instanceID);
+		if (pLiveInstance)
+		{
+			pLiveInstance->m_pInstance = nullptr;
+			pLiveInstance->m_memory->ClearCacheLayer1();
+		}
+	}
+
+
+    //---------------------------------------------------------------------------------------------
+    void System::ReleaseInstance( Instance::ID instanceID )
+    {
+		LLM_SCOPE_BYNAME(TEXT("MutableRuntime"));
+		MUTABLE_CPUPROFILER_SCOPE(ReleaseInstance);
+
+		int Removed = m_pD->m_liveInstances.RemoveAllSwap(
+			[instanceID](const Private::FLiveInstance& Instance)
+			{
+				return (Instance.m_instanceID == instanceID);
+			});
+		m_pD->m_profileMetrics[int32(System::ProfileMetric::LiveInstanceCount)]-=Removed;
+	}
+
+
+    //---------------------------------------------------------------------------------------------
+    ImagePtrConst System::BuildParameterAdditionalImage( const ModelPtrConst& pModel,
+                                                         const ParametersPtrConst& pParams,
+                                                         int parameter,
+                                                         int image )
+    {
+		LLM_SCOPE_BYNAME(TEXT("MutableRuntime"));
+		MUTABLE_CPUPROFILER_SCOPE(BuildParameterAdditionalImage);
+
+		check(parameter >= 0
+			&&
+			parameter < (int)pModel->GetPrivate()->m_program.m_parameters.size());
+
+		const PARAMETER_DESC& param =
+			pModel->GetPrivate()->m_program.m_parameters[parameter];
+		check(image >= 0 && image < (int)param.m_descImages.size());
+
+		OP::ADDRESS imageAddress = param.m_descImages[image];
+
+		ImagePtrConst pResult;
+
+		if (imageAddress)
+		{
+			m_pD->BeginBuild(pModel);
+			pResult = m_pD->BuildImage( pModel, pParams.get(), imageAddress, 0 );
+			m_pD->EndBuild();
+		}
+
+		return pResult;
+	}
+
+
+    //---------------------------------------------------------------------------------------------
+    class RelevantParameterVisitor : public UniqueDiscreteCoveredCodeVisitor<>
+    {
+    public:
+
+        RelevantParameterVisitor
+            (
+                System::Private* pSystem,
+                const Ptr<const Model>& pModel,
+                const Ptr<const Parameters>& pParams,
+                bool* pFlags
+            )
+            : UniqueDiscreteCoveredCodeVisitor<>( pSystem, pModel, pParams, 0xffffffff )
+        {
+            m_pFlags = pFlags;
+
+            memset( pFlags, 0, sizeof(bool)*pParams->GetCount() );
+
+            OP::ADDRESS at = pModel->GetPrivate()->m_program.m_states[0].m_root;
+
+            pSystem->BeginBuild(pModel);
+
+            Run( at );
+
+			pSystem->EndBuild();
+        }
+
+
+        bool Visit( OP::ADDRESS at, PROGRAM& program ) override
+        {
+            switch ( program.GetOpType(at) )
+            {
+            case OP_TYPE::BO_PARAMETER:
+            case OP_TYPE::NU_PARAMETER:
+            case OP_TYPE::SC_PARAMETER:
+            case OP_TYPE::CO_PARAMETER:
+            case OP_TYPE::PR_PARAMETER:
+            case OP_TYPE::IM_PARAMETER:
+            {
+				OP::ParameterArgs args = program.GetOpArgs<OP::ParameterArgs>(at);
+				OP::ADDRESS param = args.variable;
+                m_pFlags[param] = true;
+                break;
+            }
+
+            default:
+                break;
+
+            }
+
+            return UniqueDiscreteCoveredCodeVisitor<>::Visit( at, program );
+        }
+
+    private:
+
+        //! Non-owned result buffer
+        bool* m_pFlags;
+    };
+
+
+    //---------------------------------------------------------------------------------------------
+    void System::GetParameterRelevancy( const ModelPtrConst& pModel,
+                                        const ParametersPtrConst& pParameters,
+                                        bool* pFlags )
+    {
+		LLM_SCOPE_BYNAME(TEXT("MutableRuntime"));
+
+        RelevantParameterVisitor visitor( m_pD, pModel, pParameters, pFlags );
+    }
+
+
+    //---------------------------------------------------------------------------------------------
+    //---------------------------------------------------------------------------------------------
+    //---------------------------------------------------------------------------------------------
+    bool System::Private::CheckUpdatedParameters( const FLiveInstance* pLiveInstance,
+                                 const Ptr<const Parameters>& pParams,
+                                 uint64& updatedParameters )
+    {
+        bool fullBuild = false;
+
+		if (!pLiveInstance->m_pOldParameters)
+		{
+			updatedParameters = 0xffffffff;
+			return true;
+		}
+
+        // check what parameters have changed
+        updatedParameters = 0;
+        const PROGRAM& program = pLiveInstance->m_pModel->GetPrivate()->m_program;
+        const vector<int>& runtimeParams = program.m_states[ pLiveInstance->m_state ].m_runtimeParameters;
+
+        check( pParams->GetCount() == (int)program.m_parameters.size() );
+        check( !pLiveInstance->m_pOldParameters
+                        ||
+                        pParams->GetCount() == pLiveInstance->m_pOldParameters->GetCount() );
+
+        for ( int p=0; p<(int)program.m_parameters.size() && !fullBuild; ++p )
+        {
+            bool isRuntime = std::find( runtimeParams.begin(), runtimeParams.end(), p )
+                    !=
+                    runtimeParams.end() ;
+
+            bool changed = !pParams->HasSameValue( p, pLiveInstance->m_pOldParameters, p );
+
+            if (changed && isRuntime)
+            {
+				uint64 runtimeIndex =
+                        std::find( runtimeParams.begin(), runtimeParams.end(), p )
+                        -
+                        runtimeParams.begin();
+                updatedParameters |= uint64(1) << runtimeIndex;
+            }
+            else if (changed)
+            {
+                // A non-runtime parameter has changed, we need a full build.
+                // TODO: report, or log somehow.
+                fullBuild = true;
+                updatedParameters = 0xffffffffffffffff;
+            }
+        }
+
+        return fullBuild;
+    }
+
+
+	//---------------------------------------------------------------------------------------------
+	void System::Private::SetStreamingCache(uint64 bytes)
+	{
+		m_modelCache.m_romBudget = bytes;
+		m_modelCache.EnsureCacheBelowBudget(0);
+	}
+
+
+	//---------------------------------------------------------------------------------------------
+	void System::Private::BeginBuild(const Ptr<const Model>& pModel)
+	{
+		// We don't have a FLiveInstance, let's create the memory
+		// \TODO: There is no clear moment to remove this... EndBuild?
+		m_memory = MakeShared<FProgramCache>();
+		m_memory->Init(pModel->GetPrivate()->m_program.m_opAddress.size());
+
+		PrepareCache(pModel, -1);
+	}
+
+
+	//---------------------------------------------------------------------------------------------
+	void System::Private::EndBuild()
+	{
+		m_memory = nullptr;
+	}
+
+
+	//---------------------------------------------------------------------------------------------
+	void System::Private::RunCode(const Ptr<const Model>& InModel,
+		const Parameters* InParameters, OP::ADDRESS InCodeRoot, uint32 InLODs, uint8 executionOptions)
+	{
+		CodeRunner Runner(m_pSettings, this, InModel.get(), InParameters, InCodeRoot, InLODs, 
+			executionOptions, SCHEDULED_OP::EType::Full);
+		Runner.Run();
+		bUnrecoverableError = Runner.bUnrecoverableError;
+	}
+
+
+	//---------------------------------------------------------------------------------------------
+	bool System::Private::BuildBool(const Ptr<const Model>& pModel,
+		const Parameters* pParams,
+		OP::ADDRESS at)
+	{
+		RunCode(pModel.get(), pParams, at);
+		if (bUnrecoverableError)
+		{
+			return false;
+		}
+		return m_memory->GetBool(CACHE_ADDRESS(at, 0, 0));
+	}
+
+
+	//---------------------------------------------------------------------------------------------
+	float System::Private::BuildScalar(const Ptr<const Model>& pModel,
+		const Parameters* pParams,
+		OP::ADDRESS at)
+	{
+		RunCode(pModel.get(), pParams, at);
+		if (bUnrecoverableError)
+		{
+			return 0.0f;
+		}
+		return m_memory->GetScalar(CACHE_ADDRESS(at, 0, 0));
+	}
+
+
+	//---------------------------------------------------------------------------------------------
+	int System::Private::BuildInt(const Ptr<const Model>& pModel,
+		const Parameters* pParams,
+		OP::ADDRESS at)
+	{
+		RunCode(pModel.get(), pParams, at);
+		if (bUnrecoverableError)
+		{
+			return 0;
+		}
+
+		return m_memory->GetInt(CACHE_ADDRESS(at, 0, 0));
+	}
+
+
+	//---------------------------------------------------------------------------------------------
+	void System::Private::BuildColour(const Ptr<const Model>& pModel,
+		const Parameters* pParams,
+		OP::ADDRESS at,
+		float* pR,
+		float* pG,
+		float* pB)
+	{
+		vec4f col;
+
+		mu::OP_TYPE opType = pModel->GetPrivate()->m_program.GetOpType(at);
+		if (GetOpDataType(opType) == DT_COLOUR)
+		{
+			RunCode(pModel.get(), pParams, at);
+			if (bUnrecoverableError)
+			{
+				if (pR) *pR = 0.0f;
+				if (pG) *pG = 0.0f;
+				if (pB) *pB = 0.0f;
+			}
+
+			col = m_memory->GetColour(CACHE_ADDRESS(at, 0, 0));
+		}
+
+		if (pR) *pR = col[0];
+		if (pG) *pG = col[1];
+		if (pB) *pB = col[2];
+	}
+
+	
+	//---------------------------------------------------------------------------------------------
+	Ptr<const Projector> System::Private::BuildProjector(const Ptr<const Model>& pModel, const Parameters* pParams, OP::ADDRESS at)
+	{
+    	RunCode(pModel.get(), pParams, at);
+		if (bUnrecoverableError)
+		{
+			return nullptr;
+		}
+    	return m_memory->GetProjector(CACHE_ADDRESS(at, 0, 0));
+	}
+
+	
+	//---------------------------------------------------------------------------------------------
+	Ptr<const Image> System::Private::BuildImage(const Ptr<const Model>& pModel,
+		const Parameters* pParams,
+		OP::ADDRESS at, int MipsToSkip)
+	{
+		mu::OP_TYPE opType = pModel->GetPrivate()->m_program.GetOpType(at);
+		if (GetOpDataType(opType) == DT_IMAGE)
+		{
+			m_memory->ClearCacheLayer0();
+			RunCode(pModel.get(), pParams, at, System::AllLODs, uint8(MipsToSkip));
+			if (bUnrecoverableError)
+			{
+				return nullptr;
+			}
+			ImagePtrConst Result = m_memory->GetImage(CACHE_ADDRESS(at, 0, MipsToSkip));
+			return Result;
+		}
+
+		return nullptr;
+	}
+
+
+	//---------------------------------------------------------------------------------------------
+	MeshPtrConst System::Private::BuildMesh(const Ptr<const Model>& pModel,
+		const Parameters* pParams,
+		OP::ADDRESS at)
+	{
+		mu::OP_TYPE opType = pModel->GetPrivate()->m_program.GetOpType(at);
+		if (GetOpDataType(opType) == DT_MESH)
+		{
+			m_memory->ClearCacheLayer0();
+			RunCode(pModel.get(), pParams, at);
+			if (bUnrecoverableError)
+			{
+				return nullptr;
+			}
+			MeshPtrConst pResult = m_memory->GetMesh(CACHE_ADDRESS(at, 0, 0));
+			return pResult;
+		}
+
+		return nullptr;
+	}
+
+
+	//---------------------------------------------------------------------------------------------
+	LayoutPtrConst System::Private::BuildLayout(const Ptr<const Model>& pModel,
+		const Parameters* pParams,
+		OP::ADDRESS at)
+	{
+		LayoutPtrConst  pResult;
+
+		if (pModel->GetPrivate()->m_program.m_states[0].m_root)
+		{
+			mu::OP_TYPE opType = pModel->GetPrivate()->m_program.GetOpType(at);
+			if (GetOpDataType(opType) == DT_LAYOUT)
+			{
+				RunCode(pModel.get(), pParams, at);
+				if (bUnrecoverableError)
+				{
+					return nullptr;
+				}
+				pResult = m_memory->GetLayout(CACHE_ADDRESS(at, 0, 0));
+			}
+		}
+
+		return pResult;
+	}
+
+
+	//---------------------------------------------------------------------------------------------
+	Ptr<const String> System::Private::BuildString(const Ptr<const Model>& pModel,
+		const Parameters* pParams,
+		OP::ADDRESS at)
+	{
+		Ptr<const String> pResult;
+
+		if (pModel->GetPrivate()->m_program.m_states[0].m_root)
+		{
+			mu::OP_TYPE opType = pModel->GetPrivate()->m_program.GetOpType(at);
+			if (GetOpDataType(opType) == DT_STRING)
+			{
+				RunCode(pModel.get(), pParams, at);
+				if (bUnrecoverableError)
+				{
+					return nullptr;
+				}
+				pResult = m_memory->GetString(CACHE_ADDRESS(at, 0, 0));
+			}
+		}
+
+		return pResult;
+	}
+
+
+	//---------------------------------------------------------------------------------------------
+	void System::Private::PrepareCache(const Ptr<const Model>& pModel, int state)
+	{
+		MUTABLE_CPUPROFILER_SCOPE(PrepareCache);
+
+		PROGRAM& program = pModel->GetPrivate()->m_program;
+		size_t opCount = program.m_opAddress.size();
+		m_memory->m_opHitCount.clear();
+		m_memory->Init(opCount);
+
+		// Mark the resources that have to be cached to update the instance in this state
+		if (state >= 0)
+		{
+			const PROGRAM::STATE& s = program.m_states[state];
+			for (uint32 a : s.m_updateCache)
+			{
+				m_memory->SetForceCached(a);
+			}
+		}
+	}
+
+
+    //---------------------------------------------------------------------------------------------
+    //---------------------------------------------------------------------------------------------
+    //---------------------------------------------------------------------------------------------
+	FModelCache::FModelCacheEntry& FModelCache::GetModelCache( const Model* m )
+    {
+        check(m);
+
+        for(FModelCacheEntry& c:m_cachePerModel)
+        {
+            mu::Ptr<mu::Model> pCandidate = c.m_pModel.Pin();
+            if (pCandidate)
+            {
+                if (pCandidate==m)
+                {
+                    return c;
+                }
+            }
+            else
+            {
+                // Free stray data. TODO: remove vector entry.
+                c.m_romWeight.Empty();
+            }
+        }
+
+        // Not found. Add new
+		FModelCacheEntry n;
+        n.m_pModel = WeakPtr<Model>(m);
+        m_cachePerModel.Add(n);
+        return m_cachePerModel.Last();
+    }
+
+
+    //---------------------------------------------------------------------------------------------
+    uint64 FModelCache::EnsureCacheBelowBudget( uint64 additionalMemory,
+                                              std::function<bool(const Model*,int)> isRomLockedFunc )
+    {
+
+		uint64 totalMemory = 0;
+        for (FModelCacheEntry& m : m_cachePerModel)
+        {
+			mu::Ptr<mu::Model> pCacheModel = m.m_pModel.Pin();
+            if (pCacheModel)
+            {
+				mu::PROGRAM& program = pCacheModel->GetPrivate()->m_program;
+				for (int32 RomIndex = 0; RomIndex < program.m_roms.Num(); ++RomIndex)
+				{
+					if (program.IsRomLoaded(RomIndex))
+					{
+						totalMemory += program.m_roms[RomIndex].Size;
+					}
+				}
+			}
+        }
+
+        if (totalMemory>0)
+        {
+            totalMemory += additionalMemory;
+
+            bool finished = totalMemory < m_romBudget;
+            while (!finished)
+            {
+                ModelPtr lowestPriorityModel;
+                int32 lowestPriorityRom = -1;
+                float lowestPriority = 0.0f;
+                for (FModelCacheEntry& modelCache : m_cachePerModel)
+                {
+					mu::Ptr<mu::Model> pCacheModel = modelCache.m_pModel.Pin();
+                    if (pCacheModel)
+                    {
+						mu::PROGRAM& program = pCacheModel->GetPrivate()->m_program;
+                        check( modelCache.m_romWeight.Num() == program.m_roms.Num());
+
+                        for (int32 RomIndex=0; RomIndex <program.m_roms.Num(); ++RomIndex)
+                        {
+							const FRomData& Rom = program.m_roms[RomIndex];
+							bool bIsLoaded = program.IsRomLoaded(RomIndex);
+
+                            if (bIsLoaded
+                                &&
+                                (!isRomLockedFunc || !isRomLockedFunc(pCacheModel.get(),RomIndex)) )
+                            {
+                                constexpr float factorWeight = 100.0f;
+                                constexpr float factorTime = -1.0f;
+                                float priority =
+                                        factorWeight * float(modelCache.m_romWeight[RomIndex].Key)
+                                        +
+                                        factorTime * float((m_romTick-modelCache.m_romWeight[RomIndex].Value));
+
+                                if (lowestPriorityRom<0 || priority<lowestPriority)
+                                {
+                                    lowestPriorityRom = RomIndex;
+                                    lowestPriority = priority;
+                                    lowestPriorityModel = pCacheModel;
+                                }
+                            }
+                        }
+                    }
+                }
+
+                if (lowestPriorityRom<0)
+                {
+                    // None found
+                    finished = true;
+                    UE_LOG(LogMutableCore,Log, TEXT("EnsureCacheBelowBudget failed to free all the necessary memory."));
+                }
+                else
+                {
+                    //UE_LOG(LogMutableCore,Log, "Unloading rom because of memory budget: %d.", lowestPriorityRom);
+					const FRomData& Rom = lowestPriorityModel->GetPrivate()->m_program.m_roms[lowestPriorityRom];
+                    lowestPriorityModel->GetPrivate()->m_program.UnloadRom(lowestPriorityRom);
+                    totalMemory -= lowestPriorityModel->GetPrivate()->m_program.m_roms[lowestPriorityRom].Size;
+                    finished = totalMemory < m_romBudget;
+                }
+            }
+        }
+
+        return totalMemory;
+    }
+
+
+    //---------------------------------------------------------------------------------------------
+    void FModelCache::MarkRomUsed( int romIndex, const Model* pModel )
+    {
+        check(pModel);
+
+        PROGRAM& program = pModel->GetPrivate()->m_program;
+
+        // If budget is zero, we don't unload anything here, and we assume it is managed
+        // somewhere else.
+        if ( !m_romBudget )
+        {
+            return;
+        }
+
+        ++m_romTick;
+
+        // Update current cache
+        {
+			FModelCacheEntry& modelCache = GetModelCache(pModel);
+
+            while (modelCache.m_romWeight.Num()<program.m_roms.Num())
+            {
+				modelCache.m_romWeight.Add({ 0,0 });
+            }
+
+            modelCache.m_romWeight[romIndex].Key++;
+            modelCache.m_romWeight[romIndex].Value = m_romTick;
+        }
+    }
+
+
+    //---------------------------------------------------------------------------------------------
+    void FModelCache::UpdateForLoad( int romIndex, const Model* pModel,
+                                     std::function<bool(const Model*,int)> isRomLockedFunc )
+    {
+		MarkRomUsed( romIndex, pModel);
+
+		PROGRAM& program = pModel->GetPrivate()->m_program;
+        EnsureCacheBelowBudget(program.m_roms[romIndex].Size, isRomLockedFunc);
+    }
+
+}
diff --git a/Engine/Plugins/Experimental/Mutable/Source/MutableRuntime/Private/SystemPrivate.h b/Engine/Plugins/Experimental/Mutable/Source/MutableRuntime/Private/SystemPrivate.h
new file mode 100644
--- /dev/null
+++ b/Engine/Plugins/Experimental/Mutable/Source/MutableRuntime/Private/SystemPrivate.h
@@ -0,0 +1,1002 @@
+// Copyright Epic Games, Inc. All Rights Reserved.
+
+#pragma once
+
+#include "MutableRuntime/Public/System.h"
+
+#include "MutableRuntime/Public/Settings.h"
+#include "MutableRuntime/Private/Operations.h"
+#include "MutableRuntime/Private/ImagePrivate.h"
+#include "MutableRuntime/Public/MutableString.h"
+#include "MutableRuntime/Private/MeshPrivate.h"
+#include "MutableRuntime/Private/InstancePrivate.h"
+#include "MutableRuntime/Private/ParametersPrivate.h"
+
+
+namespace mu
+{
+
+	//! Reference-counted colour to be stored in the cache
+	class Colour : public RefCounted
+	{
+	public:
+		Colour(vec4f v = vec4f()) : m_colour(v) {}
+		vec4f m_colour;
+	};
+	typedef Ptr<Colour> ColourPtr;
+
+
+	//! Reference-counted bool to be stored in the cache
+	class Bool : public RefCounted
+	{
+	public:
+		Bool(bool v = false) : m_value(v) {}
+		bool m_value;
+	};
+	typedef Ptr<Bool> BoolPtr;
+
+
+	//! Reference-counted scalar to be stored in the cache
+	class Scalar : public RefCounted
+	{
+	public:
+		Scalar(float v = 0.0f) : m_value(v) {}
+		float m_value;
+	};
+	typedef Ptr<Scalar> ScalarPtr;
+
+
+	//! Reference-counted scalar to be stored in the cache
+	class Int : public RefCounted
+	{
+	public:
+		Int(int32 v = 0) : m_value(v) {}
+		int32 m_value;
+	};
+	typedef Ptr<Int> IntPtr;
+
+
+	//! Reference-counted scalar to be stored in the cache
+	class Projector : public RefCounted
+	{
+	public:
+		PROJECTOR m_value;
+	};
+	typedef Ptr<Projector> ProjectorPtr;
+
+
+	//! ExecutinIndex stores the location inside all ranges of the execution of a specific
+	//! operation. The first integer on each pair is the dimension/range index in the program
+	//! array of ranges, and the second integer is the value inside that range.
+	//! The vector order is undefined.
+	class ExecutionIndex : public  TArray<TPair<int32, int32>>
+	{
+	public:
+		//! Set or add a value to the index
+		void SetFromModelRangeIndex(uint16_t rangeIndex, int rangeValue)
+		{
+			auto Index = IndexOfByPredicate([=](const ElementType& v) { return v.Key >= rangeIndex; });
+			if (Index != INDEX_NONE && (*this)[Index].Key == rangeIndex)
+			{
+				// Update
+				(*this)[Index].Value = rangeValue;
+			}
+			else
+			{
+				// Add new
+				Push(ElementType(rangeIndex, rangeValue));
+			}
+		}
+
+		//! Get the value of the index from the range index in the model.
+		int GetFromModelRangeIndex(int modelRangeIndex) const
+		{
+			for (const auto& e : *this)
+			{
+				if (e.Key == modelRangeIndex)
+				{
+					return e.Value;
+				}
+			}
+			return 0;
+		}
+	};
+
+
+	struct SCHEDULED_OP
+	{
+		SCHEDULED_OP()
+		{
+			stage = 0;
+			Type = EType::Full;
+		}
+
+		SCHEDULED_OP(OP::ADDRESS _at,
+			const SCHEDULED_OP& opTemplate,
+			uint8 _stage = 0,
+			uint32 _customState = 0)
+		{
+			at = _at;
+			executionOptions = opTemplate.executionOptions;
+			executionIndex = opTemplate.executionIndex;
+			stage = _stage;
+			customState = _customState;
+			Type = opTemplate.Type;
+		}
+
+		static inline SCHEDULED_OP FromOpAndOptions(OP::ADDRESS _at,
+			const SCHEDULED_OP& opTemplate,
+			uint8 _executionOptions)
+		{
+			SCHEDULED_OP r;
+			r.at = _at;
+			r.executionOptions = _executionOptions;
+			r.executionIndex = opTemplate.executionIndex;
+			r.stage = 0;
+			r.customState = opTemplate.customState;
+			r.Type = opTemplate.Type;
+			return r;
+		}
+
+		//! Address of the operation
+		OP::ADDRESS at = 0;
+
+		//! Additional custom state data that the operation can store. This is usually used to pass information
+		//! between execution stages of an operation.
+		uint32 customState = 0;
+
+		//! Index of the operation execution: This is used for iteration of different ranges.
+		//! It is an index into the CodeRunner::GetMemory()::m_rangeIndex vector.
+		//! executionIndex 0 is always used for empty ExecutionIndex, which is the most common
+		//! one.
+		uint16 executionIndex = 0;
+
+		//! Additional execution options. Set externally to this op, it usually alters the result.
+		//! For example, this is used to keep track of the mipmaps to skip in image operations.
+		uint8 executionOptions = 0;
+
+		//! Internal stage of the operation.
+		//! Stage 0 is usually scheduling of children, and 1 is execution. Some instructions
+		//! may have more steges to schedule children that are optional for execution, etc.
+		uint8 stage : 4;
+
+		//!
+		enum class EType : uint8
+		{
+			//! Execute the operation to calculate the full result
+			Full,
+
+			//! Execute the operation to obtain the descriptor of an image.
+			ImageDesc
+		};
+		EType Type : 4;
+	};
+
+	inline uint32 GetTypeHash(const SCHEDULED_OP& op)
+	{
+		return HashCombine(::GetTypeHash(op.at), HashCombine(::GetTypeHash(op.stage), ::GetTypeHash(op.executionIndex)));
+	}
+
+
+	//! A cache address is the operation plus the context of execution (iteration indices, etc...).
+	struct CACHE_ADDRESS
+	{
+		OP::ADDRESS at = 0;
+		uint16 executionIndex = 0;
+		uint8 executionOptions = 0;
+		SCHEDULED_OP::EType Type = SCHEDULED_OP::EType::Full;
+
+		CACHE_ADDRESS() {}
+
+		CACHE_ADDRESS(OP::ADDRESS _at, uint16 _executionIndex, uint8 _executionOptions)
+		{
+			at = _at;
+			executionIndex = _executionIndex;
+			executionOptions = _executionOptions;
+		}
+
+		CACHE_ADDRESS(OP::ADDRESS _at, const SCHEDULED_OP& item)
+		{
+			at = _at;
+			executionIndex = item.executionIndex;
+			executionOptions = item.executionOptions;
+			Type = item.Type;
+		}
+
+		CACHE_ADDRESS(const SCHEDULED_OP& item)
+		{
+			at = item.at;
+			executionIndex = item.executionIndex;
+			executionOptions = item.executionOptions;
+			Type = item.Type;
+		}
+	};
+
+
+	inline uint32 GetTypeHash(const CACHE_ADDRESS& a)
+	{
+		return HashCombine(::GetTypeHash(a.at), a.executionIndex);
+	}
+
+
+	//! Container that stores data per executable code operation (indexed by address and execution
+	//! index).
+	template<class DATA>
+	class CodeContainer
+	{
+	public:
+
+		void resize(size_t s)
+		{
+			m_index0.SetNum(s);
+			FMemory::Memzero(m_index0.GetData(), s * sizeof(DATA));
+		}
+
+		uint32 size_code() const
+		{
+			return uint32(m_index0.Num());
+		}
+
+		void clear()
+		{
+			m_index0.Empty();
+			m_otherIndex.Empty();
+		}
+
+		inline void erase(const CACHE_ADDRESS& at)
+		{
+			if (at.executionIndex == 0 && at.executionOptions == 0)
+			{
+				m_index0[at.at] = nullptr;
+			}
+			else
+			{
+				m_otherIndex.Remove(at);
+			}
+		}
+
+		inline DATA get(const CACHE_ADDRESS& at) const
+		{
+			if (at.executionIndex == 0 && at.executionOptions == 0)
+			{
+				if (at.at < uint32(m_index0.Num()))
+				{
+					return m_index0[at.at];
+				}
+				else
+				{
+					return 0;
+				}
+			}
+			else
+			{
+				const DATA* it = m_otherIndex.Find(at);
+				if (it)
+				{
+					return *it;
+				}
+			}
+			return 0;
+		}
+
+		inline const DATA* get_ptr(const CACHE_ADDRESS& at) const
+		{
+			if (at.executionIndex == 0 && at.executionOptions == 0)
+			{
+				if (at.at < uint32(m_index0.Num()))
+				{
+					return &m_index0[at.at];
+				}
+				else
+				{
+					return nullptr;
+				}
+			}
+			else
+			{
+				const DATA* it = m_otherIndex.Find(at);
+				if (it)
+				{
+					return it;
+				}
+			}
+			return nullptr;
+		}
+
+		inline DATA& operator[](const CACHE_ADDRESS& at)
+		{
+			if (at.executionIndex == 0 && at.executionOptions == 0)
+			{
+				return m_index0[at.at];
+			}
+			else
+			{
+				return m_otherIndex.FindOrAdd(at);
+			}
+		}
+
+		struct iterator
+		{
+		private:
+			friend class CodeContainer<DATA>;
+			const CodeContainer<DATA>* container;
+			typename TArray<DATA>::TIterator it0;
+			typename TMap<CACHE_ADDRESS, DATA>::TIterator it1;
+
+			iterator(CodeContainer<DATA>* InContainer)
+				: container(InContainer)
+				, it0(InContainer->m_index0.CreateIterator())
+				, it1(InContainer->m_otherIndex.CreateIterator())
+			{
+			}
+
+		public:
+
+			inline void operator++(int)
+			{
+				if (it0)
+				{
+					++it0;
+				}
+				else
+				{
+					++it1;
+				}
+			}
+
+			iterator operator++()
+			{
+				if (it0)
+				{
+					++it0;
+				}
+				else
+				{
+					++it1;
+				}
+				return *this;
+			}
+
+			inline bool operator!=(const iterator& o) const
+			{
+				return it0 != o.it0 || it1 != o.it1;
+			}
+
+			inline DATA& operator*()
+			{
+				if (it0)
+				{
+					return *it0;
+				}
+				else
+				{
+					return it1->Value;
+				}
+			}
+
+			inline CACHE_ADDRESS get_address() const
+			{
+				if (it0)
+				{
+					return { uint32_t(it0.GetIndex()), 0, 0};
+				}
+				else
+				{
+					return it1->Key;
+				}
+			}
+
+			inline bool IsValid() const
+			{
+				return bool(it0) && bool(it1);
+			}
+		};
+
+		inline iterator begin()
+		{
+			iterator it(this);
+			return it;
+		}
+
+	private:
+		// For index 0
+		TArray<DATA> m_index0;
+
+		// For index>0
+		TMap<CACHE_ADDRESS, DATA> m_otherIndex;
+	};
+
+
+	/** Interface for storage of data while Mutable code is being executed. */
+	class FProgramCache
+	{
+	private:
+
+		TArray< ExecutionIndex, TInlineAllocator<4> > m_usedRangeIndices;
+
+		// first value of the pair:
+		// 0 : value not valid (not set).
+		// 1 : valid, not worth freeing for memory
+		// 2 : valid, worth freeing
+		CodeContainer< TPair<int32, Ptr<const RefCounted>> > m_resources;
+
+		// TODO
+	public:
+
+		//! Addressed with OP::ADDRESS. It is true if value for an image desc is valid.
+		TArray<bool> m_descCache;
+
+		//! The number of times an operation will be run for the current build operation.
+		//! \todo: this could be optimised by merging in other CodeContainers here.
+		CodeContainer<int> m_opHitCount;
+
+	public:
+
+
+		inline const ExecutionIndex& GetRageIndex(uint32_t i)
+		{
+			// Make sure we have the default element.
+			if (m_usedRangeIndices.IsEmpty())
+			{
+				m_usedRangeIndices.Push(ExecutionIndex());
+			}
+
+			check(i < uint32_t(m_usedRangeIndices.Num()));
+			return m_usedRangeIndices[i];
+		}
+
+		//!
+		inline uint32_t GetRageIndexIndex(const ExecutionIndex& rangeIndex)
+		{
+			if (rangeIndex.IsEmpty())
+			{
+				return 0;
+			}
+
+			// Make sure we have the default element.
+			if (m_usedRangeIndices.IsEmpty())
+			{
+				m_usedRangeIndices.Push(ExecutionIndex());
+			}
+
+			// Look for or add the new element.
+			auto ElemIndex = m_usedRangeIndices.Find(rangeIndex);
+			if (ElemIndex != INDEX_NONE)
+			{
+				return ElemIndex;
+			}
+
+			m_usedRangeIndices.Push(rangeIndex);
+			return uint32_t(m_usedRangeIndices.Num()) - 1;
+		}
+
+		void Init(size_t size)
+		{
+			m_resources.resize(size);
+			m_opHitCount.resize(size);
+		}
+
+		void SetUnused(CACHE_ADDRESS at)
+		{
+			//UE_LOG(LogMutableCore, Log, TEXT("memory SetUnused : %5d "), at.at);
+			if (m_resources[at].Key >= 2)
+			{
+				// Keep the result anyway if it doesn't use any memory.
+				if (m_resources[at].Value)
+				{
+					m_resources[at].Value = nullptr;
+					m_resources[at].Key = 0;
+				}
+			}
+		}
+
+
+		bool IsValid(CACHE_ADDRESS at)
+		{
+			if (at.at == 0) return false;
+
+			// Is it a desc data query?
+			if (at.Type == SCHEDULED_OP::EType::ImageDesc)
+			{
+				if (uint32(m_descCache.Num()) > at.at)
+				{
+					return m_descCache[at.at];
+				}
+				else
+				{
+					return false;
+				}
+			}
+
+			// It's a full data query.
+			auto d = m_resources.get_ptr(at);
+			return d && d->Key != 0;
+		}
+
+
+		void ClearCacheLayer0()
+		{
+			CodeContainer<int>::iterator it = m_opHitCount.begin();
+			for (; it.IsValid(); ++it)
+			{
+				int32 Count = *it;
+				if (Count < 3000000)
+				{
+					SetUnused(it.get_address());
+					*it = 0;
+				}
+			}
+		}
+
+
+		void ClearCacheLayer1()
+		{
+			CodeContainer<int>::iterator it = m_opHitCount.begin();
+			for (; it.IsValid(); ++it)
+			{
+				m_resources[it.get_address()].Value = nullptr;
+				m_resources[it.get_address()].Key = 0;
+			}
+
+			m_descCache.SetNum(0);
+		}
+
+
+		void Clear()
+		{
+			size_t codeSize = m_resources.size_code();
+			m_resources.clear();
+			m_resources.resize(codeSize);
+			m_descCache.SetNum(0);
+			m_opHitCount.clear();
+			m_opHitCount.resize(codeSize);
+		}
+
+
+		bool GetBool(CACHE_ADDRESS at)
+		{
+			if (!at.at) return false;
+			auto d = m_resources.get_ptr(at);
+			if (!d) return false;
+
+			Ptr<const Bool> pResult;
+			if (at.at)
+			{
+				pResult = (const Bool*)d->Value.get();
+			}
+			return pResult ? pResult->m_value : false;
+		}
+
+		float GetScalar(CACHE_ADDRESS at)
+		{
+			if (!at.at) return 0.0f;
+			auto d = m_resources.get_ptr(at);
+			if (!d) return 0.0f;
+
+			Ptr<const Scalar> pResult;
+			if (at.at)
+			{
+				pResult = (const Scalar*)d->Value.get();
+			}
+			return pResult ? pResult->m_value : 0.0f;
+		}
+
+		int GetInt(CACHE_ADDRESS at)
+		{
+			if (!at.at) return 0;
+			auto d = m_resources.get_ptr(at);
+			if (!d) return 0;
+
+			Ptr<const Int> pResult;
+			if (at.at)
+			{
+				pResult = (const Int*)d->Value.get();
+			}
+			return pResult ? pResult->m_value : 0;
+		}
+
+		vec4f GetColour(CACHE_ADDRESS at)
+		{
+			if (!at.at) return vec4f();
+			auto d = m_resources.get_ptr(at);
+			if (!d) return vec4f();
+
+			Ptr<const Colour> pResult;
+			if (at.at)
+			{
+				pResult = (const Colour*)d->Value.get();
+			}
+			return pResult ? pResult->m_colour : vec4f();
+		}
+
+		Ptr<const Projector> GetProjector(CACHE_ADDRESS at)
+		{
+			if (!at.at) return nullptr;
+			auto d = m_resources.get_ptr(at);
+			if (!d) return nullptr;
+
+			Ptr<const Projector> pResult;
+			if (at.at)
+			{
+				pResult = (const Projector*)d->Value.get();
+			}
+			return pResult;
+		}
+
+		Ptr<const Instance> GetInstance(CACHE_ADDRESS at)
+		{
+			if (!at.at) return nullptr;
+			auto d = m_resources.get_ptr(at);
+			if (!d) return nullptr;
+
+			Ptr<const Instance> pResult;
+			if (at.at)
+			{
+				pResult = (const Instance*)d->Value.get();
+			}
+			return pResult;
+		}
+
+
+		Ptr<const Layout> GetLayout(CACHE_ADDRESS at)
+		{
+			if (!at.at) return nullptr;
+			auto d = m_resources.get_ptr(at);
+			if (!d) return nullptr;
+
+			Ptr<const Layout> pResult;
+			if (at.at)
+			{
+				pResult = (const Layout*)d->Value.get();
+			}
+			return pResult;
+		}
+
+		Ptr<const String> GetString(CACHE_ADDRESS at)
+		{
+			if (!at.at)
+				return nullptr;
+			auto d = m_resources.get_ptr(at);
+			if (!d)
+				return nullptr;
+
+			Ptr<const String> pResult;
+			if (at.at)
+			{
+				pResult = (const String*)d->Value.get();
+			}
+			return pResult;
+		}
+
+		void SetBool(CACHE_ADDRESS at, bool v)
+		{
+			check(at.at < m_resources.size_code());
+
+			Ptr<Bool> pResult = new Bool;
+			pResult->m_value = v;
+			m_resources[at] = TPair<int, Ptr<const RefCounted>>(1, pResult);
+		}
+
+		void SetValidDesc(CACHE_ADDRESS at)
+		{
+			check(at.Type == SCHEDULED_OP::EType::ImageDesc);
+			check(at.at < uint32(m_descCache.Num()));
+
+			m_descCache[at.at] = true;
+		}
+
+		void SetScalar(CACHE_ADDRESS at, float v)
+		{
+			check(at.at < m_resources.size_code());
+
+			Ptr<Scalar> pResult = new Scalar;
+			pResult->m_value = v;
+			m_resources[at] = TPair<int, Ptr<const RefCounted>>(1, pResult);
+		}
+
+		void SetInt(CACHE_ADDRESS at, int v)
+		{
+			check(at.at < m_resources.size_code());
+
+			Ptr<Int> pResult = new Int;
+			pResult->m_value = v;
+			m_resources[at] = TPair<int, Ptr<const RefCounted>>(1, pResult);
+		}
+
+		void SetColour(CACHE_ADDRESS at, const vec4f& v)
+		{
+			check(at.at < m_resources.size_code());
+
+			Ptr<Colour> pResult = new Colour;
+			pResult->m_colour = v;
+			m_resources[at] = TPair<int, Ptr<const RefCounted>>(1, pResult);
+		}
+
+		void SetProjector(CACHE_ADDRESS at, Ptr<const Projector> v)
+		{
+			check(at.at < m_resources.size_code());
+			m_resources[at] = TPair<int, Ptr<const RefCounted>>(1, v);
+		}
+
+		void SetInstance(CACHE_ADDRESS at, Ptr<const Instance> v)
+		{
+			check(at.at < m_resources.size_code());
+			m_resources[at] = TPair<int, Ptr<const RefCounted>>(1, v);
+		}
+
+		void SetImage(CACHE_ADDRESS at, Ptr<const Image> v)
+		{
+			check(at.at < m_resources.size_code());
+			m_resources[at] = TPair<int, Ptr<const RefCounted>>(2, v);
+		}
+
+		void SetMesh(CACHE_ADDRESS at, Ptr<const Mesh> v)
+		{
+			// debug
+//            if (v)
+//            {
+//                v->GetPrivate()->CheckIntegrity();
+//            }
+
+			check(at.at < m_resources.size_code());
+			m_resources[at] = TPair<int, Ptr<const RefCounted>>(2, v);
+		}
+
+		void SetLayout(CACHE_ADDRESS at, Ptr<const Layout> v)
+		{
+			check(at.at < m_resources.size_code());
+			m_resources[at] = TPair<int, Ptr<const RefCounted>>(1, v);
+		}
+
+		void SetString(CACHE_ADDRESS at, Ptr<const String> v)
+		{
+			check(at.at < m_resources.size_code());
+			m_resources[at] = TPair<int, Ptr<const RefCounted>>(1, v);
+		}
+
+
+		inline void IncreaseHitCount(CACHE_ADDRESS at)
+		{
+			m_opHitCount[at] += 1;
+		}
+
+		inline void SetForceCached(OP::ADDRESS at)
+		{
+			m_opHitCount[{at, 0, 0}] = 0xffffff;
+		}
+
+
+		Ptr<const Image> GetImage(CACHE_ADDRESS at)
+		{
+			if (!at.at) return nullptr;
+			if (size_t(at.at) >= m_resources.size_code()) return nullptr;
+			auto d = m_resources.get_ptr(at);
+			if (!d) return nullptr;
+
+			Ptr<const Image> pResult = (const Image*)d->Value.get();
+
+			// We need to decrease the hit count even if the result is null.
+			// Lower hit counts means we shouldn't clear the value
+			if ( m_opHitCount[at] > 0)
+			{
+				--m_opHitCount[at];
+				if (m_opHitCount[at] <= 0)
+				{
+					SetUnused(at);
+				}
+			}
+
+			return pResult;
+		}
+
+		Ptr<const Mesh> GetMesh(CACHE_ADDRESS at)
+		{
+			if (!at.at) return nullptr;
+			if (size_t(at.at) >= m_resources.size_code()) return nullptr;
+			auto d = m_resources.get_ptr(at);
+			if (!d) return nullptr;
+
+			Ptr<const Mesh> pResult = (const Mesh*)d->Value.get();
+
+			// We need to decrease the hitcount even if the result is null.
+			// Lower hit counts means we shouldn't clear the value
+			if ( m_opHitCount[at] > 0)
+			{
+				--m_opHitCount[at];
+
+				if (m_opHitCount[at] <= 0)
+				{
+					SetUnused(at);
+				}
+			}
+
+			return pResult;
+		}
+	};
+
+
+	inline bool operator==(const CACHE_ADDRESS& a, const CACHE_ADDRESS& b)
+	{
+		return a.at == b.at
+			&&
+			a.executionIndex == b.executionIndex
+			&&
+			a.executionOptions == b.executionOptions
+			&&
+			a.Type == b.Type;
+	}
+
+	inline bool operator<(const CACHE_ADDRESS& a, const CACHE_ADDRESS& b)
+	{
+		if (a.at < b.at) return true;
+		if (a.at > b.at) return false;
+		if (a.executionIndex < b.executionIndex) return true;
+		if (a.executionIndex > b.executionIndex) return false;
+		if (a.executionOptions < b.executionOptions) return true;
+		if (a.executionOptions > b.executionOptions) return false;
+		return a.Type < b.Type;
+	}
+
+
+
+    //---------------------------------------------------------------------------------------------
+    //! Struct to manage models for which we have streamed rom data.
+    //---------------------------------------------------------------------------------------------
+    struct FModelCache
+    {
+        struct FModelCacheEntry
+        {
+            //!
+            WeakPtr<Model> m_pModel;
+
+            // Rom streaming management
+            TArray< TPair<uint64,uint64> > m_romWeight;
+        };
+
+        // Rom streaming management
+        uint64 m_romBudget = 0;
+
+        uint64 m_romTick = 0;
+        TArray< FModelCacheEntry > m_cachePerModel;
+
+        //! Make sure the cached memory is below the internal budget, even counting with the
+        //! passed additional memory.
+        uint64 EnsureCacheBelowBudget( uint64 additionalMemory,
+                                         std::function<bool(const Model*,int)> isRomLockedFunc = nullptr );
+
+        //!
+        void UpdateForLoad( int romIndex, const Model* pModel,
+                            std::function<bool(const Model*,int)> isRomLockedFunc = nullptr );
+        void MarkRomUsed( int romIndex, const Model* pModel );
+
+        //! Private helper
+		FModelCacheEntry& GetModelCache( const Model* m );
+    };
+
+
+    //---------------------------------------------------------------------------------------------
+    //! Abstract private system interface
+    //---------------------------------------------------------------------------------------------
+    class System::Private : public Base
+    {
+    public:
+
+        Private( SettingsPtr pSettings );
+        virtual ~Private();
+
+        //-----------------------------------------------------------------------------------------
+        //! Own interface
+        //-----------------------------------------------------------------------------------------
+
+		/** This method can be used to internally prepare for code execution. */
+		MUTABLERUNTIME_API void BeginBuild(const Ptr<const Model>&);
+		MUTABLERUNTIME_API void EndBuild();
+
+		MUTABLERUNTIME_API bool BuildBool(const Ptr<const Model>&, const Parameters* pParams, OP::ADDRESS at) ;
+		MUTABLERUNTIME_API int BuildInt(const Ptr<const Model>&, const Parameters* pParams, OP::ADDRESS at) ;
+		MUTABLERUNTIME_API float BuildScalar(const Ptr<const Model>&, const Parameters* pParams, OP::ADDRESS at) ;
+		MUTABLERUNTIME_API void BuildColour(const Ptr<const Model>&, const Parameters* pParams, OP::ADDRESS at, float* pR, float* pG, float* pB) ;
+		MUTABLERUNTIME_API Ptr<const String> BuildString(const Ptr<const Model>&, const Parameters* pParams, OP::ADDRESS at) ;
+		MUTABLERUNTIME_API Ptr<const Image> BuildImage(const Ptr<const Model>&, const Parameters* pParams, OP::ADDRESS at, int32 MipsToSkip) ;
+		MUTABLERUNTIME_API Ptr<const Mesh> BuildMesh(const Ptr<const Model>&, const Parameters* pParams, OP::ADDRESS at) ;
+		MUTABLERUNTIME_API Ptr<const Layout> BuildLayout(const Ptr<const Model>&, const Parameters* pParams, OP::ADDRESS at) ;
+    	MUTABLERUNTIME_API Ptr<const Projector> BuildProjector(const Ptr<const Model>&, const Parameters* pParams, OP::ADDRESS at) ;
+		void ClearCache() ;
+		void SetStreamingCache(uint64 bytes) ;
+
+        //!
+        SettingsPtrConst m_pSettings;
+
+        //! Data streaming interface, if any.
+        //! It is owned by this system
+        ModelStreamer* m_pStreamInterface;
+
+        //! It is not owned by this system
+        ImageParameterGenerator* m_pImageParameterGenerator;
+
+        //! Maximum amount of memory (bytes) if a limit has been set. If 0, it means there is no
+        //! limit.
+        uint32 m_maxMemory;
+
+        //! Current instances management
+        //! @{
+
+		/** Data for an instance that is currently being processed in the mutable system. This means it is
+		* between a BeginUpdate and EndUpdate, or during an atmoic operation (like generate a single resource).
+		*/
+        struct FLiveInstance
+        {
+            Instance::ID m_instanceID;
+            InstancePtrConst m_pInstance;
+            ModelPtrConst m_pModel;
+
+            int m_state;
+            ParametersPtr m_pOldParameters;
+
+			//! Mask of the parameters that have changed since the last update.
+			//! Every bit represents a state parameter.
+			uint64 m_updatedParameters = 0;
+
+			//! An entry for every instruction in the program to cache resources (meshes, images) if necessary.
+			TSharedPtr<FProgramCache> m_memory;
+        };
+
+        TArray<FLiveInstance> m_liveInstances;
+
+		/** Temporary reference to the memory of the current instance being updated. */
+		TSharedPtr<FProgramCache> m_memory;
+
+		/** Counter used to generate unique IDs for every new instance created in the system. */
+        Instance::ID m_lastInstanceID = 0;
+
+		/** The pointer returned by this function is only valid for the duration of the current mutable operation. */
+        inline FLiveInstance* FindLiveInstance( Instance::ID id )
+        {
+            for ( int32 i=0; i<m_liveInstances.Num(); ++i )
+            {
+                if (m_liveInstances[i].m_instanceID==id)
+                {
+                    return &m_liveInstances[i];
+                }
+            }
+            return nullptr;
+        }
+
+        //! @}
+
+        //!
+        bool CheckUpdatedParameters( const FLiveInstance* pLiveInstance,
+			const Ptr<const Parameters>& pParams,
+			uint64& OutUpdatedParameters );
+
+        // Profile metrics
+        std::atomic<uint64> m_profileMetrics[ size_t(ProfileMetric::_Count) ];
+
+
+	public:
+
+		void RunCode(const Ptr<const Model>& pModel,
+			const Parameters* pParams,
+			OP::ADDRESS at, uint32 LODs = System::AllLODs, uint8 executionOptions = 0);
+
+		//!
+		void PrepareCache(const Ptr<const Model>& pModel, int state);
+
+		//! Cache related members
+		//! @{
+
+		//! Data for each used model.
+		FModelCache m_modelCache;
+
+		//! @}
+
+	private:
+
+		/** This flag is turned on when a streaming error or similar happens. Results are not usable.
+		* This should only happen in-editor.
+		*/
+		bool bUnrecoverableError = false;
+
+    };
+
+}
diff --git a/Engine/Plugins/Experimental/Mutable/Source/MutableRuntime/Public/Image.h b/Engine/Plugins/Experimental/Mutable/Source/MutableRuntime/Public/Image.h
new file mode 100644
--- /dev/null
+++ b/Engine/Plugins/Experimental/Mutable/Source/MutableRuntime/Public/Image.h
@@ -0,0 +1,368 @@
+// Copyright Epic Games, Inc. All Rights Reserved.
+
+#pragma once
+
+#include "MutableRuntime/Public/Types.h"
+#include "MutableRuntime/Public/Ptr.h"
+#include "MutableRuntime/Public/RefCounted.h"
+#include "MutableRuntime/Public/Serialisation.h"
+#include "MutableRuntime/Public/MutableTrace.h"
+#include "MutableRuntime/Private/MutableMath.h"
+
+namespace mu
+{
+
+	// Forward declarations
+	class Image;
+	typedef Ptr<Image> ImagePtr;
+	typedef Ptr<const Image> ImagePtrConst;
+
+	//!
+	using FImageSize = vec2<uint16>;
+	using FImageRect = box<vec2<uint16>>;
+
+	//! Pixel formats supported by the images.
+	//! \ingroup runtime
+	enum class EImageFormat : uint8
+	{
+		IF_NONE,
+		IF_RGB_UBYTE,
+		IF_RGBA_UBYTE,
+		IF_L_UBYTE,
+
+        //! Deprecated formats
+        IF_PVRTC2_DEPRECATED,
+        IF_PVRTC4_DEPRECATED,
+        IF_ETC1_DEPRECATED,
+        IF_ETC2_DEPRECATED,
+
+		IF_L_UBYTE_RLE,
+		IF_RGB_UBYTE_RLE,
+		IF_RGBA_UBYTE_RLE,
+		IF_L_UBIT_RLE,
+
+        //! Common S3TC formats
+        IF_BC1,
+        IF_BC2,
+        IF_BC3,
+        IF_BC4,
+        IF_BC5,
+
+        //! Not really supported yet
+        IF_BC6,
+        IF_BC7,
+
+        //! Swizzled versions, engineers be damned.
+        IF_BGRA_UBYTE,
+
+        //! The new standard
+        IF_ASTC_4x4_RGB_LDR,
+        IF_ASTC_4x4_RGBA_LDR,
+        IF_ASTC_4x4_RG_LDR,
+
+        IF_COUNT
+	};
+
+
+	struct MUTABLERUNTIME_API FImageDesc
+	{
+		FImageDesc()
+		{
+			m_format = EImageFormat::IF_NONE;
+			m_size = FImageSize(0, 0);
+			m_lods = 0;
+		}
+
+		FImageDesc(const FImageSize& size, EImageFormat format, uint8 lods)
+		{
+			m_format = format;
+			m_size = size;
+			m_lods = lods;
+		}
+
+		//!
+		FImageSize m_size;
+
+		//!
+		EImageFormat m_format;
+
+		//! Levels of detail (mipmaps)
+		uint8 m_lods;
+
+		//!
+		inline bool operator==(const FImageDesc& o) const
+		{
+			return  (m_format == o.m_format) &&
+				(m_lods == o.m_lods) &&
+				(m_size == o.m_size);
+		}
+	};
+
+
+	//! List of supported modes in generic image layering operations.
+	enum class EBlendType
+	{
+		BT_SOFTLIGHT = 0,
+		BT_HARDLIGHT,
+		BT_BURN,
+		BT_DODGE,
+		BT_SCREEN,
+		BT_OVERLAY,
+		BT_BLEND,
+		BT_MULTIPLY,
+		BT_ALPHA_OVERLAY, // TODO: This name is not descriptive.
+		BT_NORMAL_COMBINE,
+		_BT_COUNT
+	};
+
+	enum class EAddressMode
+	{
+		AM_NONE,
+		AM_WRAP,
+		AM_CLAMP,
+		AM_BLACK_BORDER,
+		_AM_COUNT
+
+	};
+
+	enum class EMipmapFilterType
+	{
+		MFT_Unfiltered,
+		MFT_SimpleAverage,
+		MFT_Sharpen,
+		_MFT_COUNT
+	};
+
+	enum class ECompositeImageMode
+	{
+		CIM_Disabled,
+		CIM_NormalRoughnessToRed,
+		CIM_NormalRoughnessToGreen,
+		CIM_NormalRoughnessToBlue,
+		CIM_NormalRoughnessToAlpha,
+		_CIM_COUNT
+	};
+
+    //! \brief 2D image resource with mipmaps.
+	//! \ingroup runtime
+    class MUTABLERUNTIME_API Image : public RefCounted
+	{
+	public:
+
+		//-----------------------------------------------------------------------------------------
+		// Life cycle
+		//-----------------------------------------------------------------------------------------
+
+		Image();
+
+		//! Constructor from the size and format.
+		//! It will allocate the memory for the image, but its content will be undefined.
+		//! \param sizeX Width of the image.
+        //! \param sizeY Height of the image.
+        //! \param lods Number of levels of detail (mipmaps) to include in the image, inclduing the
+        //!         base level. It must be a number between 1 and the maximum possible levels, which
+        //!         depends on the image size.
+        //! \param format Pixel format.
+        Image( uint32 sizeX, uint32 sizeY, uint32 lods, EImageFormat format );
+
+		/** */
+		ImagePtr ExtractMip(int32 Mip) const;
+
+		//! Serialisation
+		static void Serialise( const Image* p, OutputArchive& arch );
+		static ImagePtr StaticUnserialise( InputArchive& arch );
+
+		//-----------------------------------------------------------------------------------------
+		// Own interface
+		//-----------------------------------------------------------------------------------------
+
+		//! Return the width of the image.
+        uint16 GetSizeX() const;
+
+		//! Return the height of the image.
+        uint16 GetSizeY() const;
+
+		const FImageSize& GetSize() const;
+
+		//! Return the pixel format of the image.
+		EImageFormat GetFormat() const;
+
+		//! Return the number of levels of detail (mipmaps) in the texture. The base lavel is also
+		//! counted, so the minimum is 1.
+		int GetLODCount() const;
+
+		//! Return a pointer to a instance-owned buffer where the image pixels are. The pixels are
+		//! stored without any padding. All the LODs are stored in the same buffer, contiguously.
+        const uint8* GetData() const;
+        uint8* GetData();
+
+        //! Return the size in bytes of all the LODs of the image.
+        int32 GetDataSize() const;
+
+        //! Return the size in bytes of a specific LOD of the image.
+        int32 GetLODDataSize( int lod ) const;
+
+		//! Get an internal identifier used to reference this image in operations like deferred
+		//! image building, or instance updating.
+        uint32 GetId() const;
+
+	protected:
+
+		//! Forbidden. Manage with the Ptr<> template.
+		~Image() {}
+
+	public:
+
+		// This used to be the data in the private implementation of the image interface
+		//-----------------------------------------------------------------------------------------
+
+		//!
+		FImageSize m_size;
+
+		//! Non-persistent internal id
+		mutable uint32 m_internalId = 0;
+
+		//!
+		EImageFormat m_format;
+
+		//! Levels of detail (mipmaps)
+		uint8 m_lods;
+
+		//! These set of flags are used to cache information for images at runtime.
+		typedef enum
+		{
+			// Set if the next flag has been calculated and its value is valid.
+			IF_IS_PLAIN_COLOUR_VALID = 1 << 0,
+
+			// If the previous flag is set and this one too, the image is single colour.
+			IF_IS_PLAIN_COLOUR = 1 << 1,
+
+			// If this is set, the image shouldn't be scaled: it's contents is resoultion-dependent.
+			IF_CANNOT_BE_SCALED = 2 << 1
+		} EImageFlags;
+
+		//! Persistent flags with some image properties. The meaning will depend of every
+		//! context.
+		mutable uint8 m_flags = 0;
+
+		//! Pixel data
+		TArray<uint8> m_data;
+
+
+		// This used to be the methods in the private implementation of the image interface
+		//-----------------------------------------------------------------------------------------
+
+		//! Deep clone.
+		ImagePtr Clone() const
+		{
+			LLM_SCOPE_BYNAME(TEXT("MutableRuntime"));
+			MUTABLE_CPUPROFILER_SCOPE(ImageClone)
+
+			ImagePtr pResult = new Image();
+
+			pResult->m_size = m_size;
+			pResult->m_internalId = m_internalId;
+			pResult->m_format = m_format;
+			pResult->m_lods = m_lods;
+			pResult->m_flags = m_flags;
+			pResult->m_data = m_data;
+
+			return pResult;
+		}
+
+
+		//!
+		void Serialise(OutputArchive& arch) const
+		{
+			uint32 ver = 3;
+			arch << ver;
+
+			arch << m_size;
+			arch << m_lods;
+			arch << (uint8)m_format;
+			arch << m_data;
+			arch << m_flags;
+		}
+
+		//!
+		void Unserialise(InputArchive& arch)
+		{
+			uint32 ver;
+			arch >> ver;
+			check(ver == 3);
+
+			arch >> m_size;
+			arch >> m_lods;
+
+			uint8 format;
+			arch >> format;
+			m_format = (EImageFormat)format;
+			arch >> m_data;
+
+			arch >> m_flags;
+		}
+
+		//!
+		inline bool operator==(const Image& o) const
+		{
+			return (m_lods == o.m_lods) &&
+				(m_size == o.m_size) &&
+				(m_format == o.m_format) &&
+				(m_data == o.m_data);
+		}
+
+		//-----------------------------------------------------------------------------------------
+
+		//! Sample the image and return an RGB colour.
+		vec4<float> Sample(vec2<float> coords) const;
+
+		//! Calculate the size of the image data in bytes, regardless of what is allocated in
+		//! m_data, only using the image descriptions. For non-block-compressed images, it returns
+		//! 0.
+		int32 CalculateDataSize() const;
+
+		//! Calculate the size of a lod of the image data in bytes, regardless of what is allocated
+		//! in m_data, only using the image descriptions. For non-block-compressed images, it
+		//! returns 0.
+		int32 CalculateDataSize(int lod) const;
+
+		//! Calculate the number of pixels of the image, regardless of what is allocated in
+		//! m_data, only using the image descriptions.
+		//! For block-compressed images, it includes the wasted pixels in the blocks, in case
+		//! the size is not a multiple of it in all lods.
+		//! It includes the pixels in all lods
+		int32 CalculatePixelCount() const;
+
+		//! Same as above, but only calculates the pixels of an lod
+		int32 CalculatePixelCount(int lod) const;
+
+		//! Calculate the size in pixels of a particular mipmap of this image. The size doesn't
+		//! include pixels necessary for completing blocks in block-compressed formats.
+		vec2<int> CalculateMipSize(int lod) const;
+
+		//! Return a pointer to the beginning of the data for a particular mip.
+		uint8* GetMipData(int mip);
+		const uint8* GetMipData(int mip) const;
+
+		//! See if all the pixels in the image are equal and return the colour.
+		bool IsPlainColour(vec4<float>& colour) const;
+
+		//! See if all the pixels in the alpha channel are the max value of the pixel format
+		//! (white).
+		bool IsFullAlpha() const;
+
+		//! Reduce the LODs of the image to the given amount. Don't do anything if there are already 
+		//! less LODs than specified.
+		void ReduceLODsTo(int32 NewLODCount);
+		void ReduceLODs(int32 LODsToSkip);
+
+		//! Calculate the number of mipmaps for a particular image size.
+		static int GetMipmapCount(int SizeX, int SizeY);
+
+		//! Get the rect inside the image bounding the non-black content of the image.
+		void GetNonBlackRect(FImageRect& rect) const;
+		void GetNonBlackRect_Reference(FImageRect& rect) const;
+
+	};
+}
+
