[mutable] Move Mutable plugin from the restricted folder to the experimental plugins folder.

diff --git a/Engine/Plugins/Experimental/Mutable/Source/MutableRuntime/Public/Instance.h b/Engine/Plugins/Experimental/Mutable/Source/MutableRuntime/Public/Instance.h
new file mode 100644
--- /dev/null
+++ b/Engine/Plugins/Experimental/Mutable/Source/MutableRuntime/Public/Instance.h
@@ -0,0 +1,227 @@
+// Copyright Epic Games, Inc. All Rights Reserved.
+
+#pragma once
+
+#include "MutableRuntime/Public/Types.h"
+#include "MutableRuntime/Public/RefCounted.h"
+#include "MutableRuntime/Public/Ptr.h"
+
+
+namespace mu
+{
+	// Forward references
+	class Image;
+	class Mesh;
+
+	class Instance;
+	typedef Ptr<Instance> InstancePtr;
+	typedef Ptr<const Instance> InstancePtrConst;
+
+
+    //! \brief A customised object created from a model and a set of parameter values.
+    //!
+    //! It corresponds to an "engine object" but the contents of its data depends on the Model, and
+    //! it may contain any number of LODs, components, surfaces, meshes and images, even none.
+	//! \ingroup runtime
+	class MUTABLERUNTIME_API Instance : public RefCounted
+	{
+	public:
+
+		//-----------------------------------------------------------------------------------------
+		// Life cycle
+		//-----------------------------------------------------------------------------------------
+
+		Instance();
+
+        //! Clone this instance
+        InstancePtr Clone() const;
+
+
+		//-----------------------------------------------------------------------------------------
+		// Own interface
+		//-----------------------------------------------------------------------------------------
+
+        //! Type for the instance unique identifiers.
+        typedef uint32_t ID;
+
+        //! Get a unique identifier for this instance. It doesn't change during the entire
+        //! lifecycle of each instance. This identifier can be used in the System methods to update
+        //! or release the instance.
+        Instance::ID GetId() const;
+
+		//! Get the number of levels-of-detail of this instance.
+		int GetLODCount() const;
+
+        //! Get the number of components in a level-of-detail
+        //! \param lod Index of the level of detail, from 0 to GetLODCount()-1
+        int GetComponentCount( int lod ) const;
+
+        //! Get the name of a component
+        //! \param lod Index of the level of detail, from 0 to GetLODCount()-1
+        //! \param comp Index of the component, from 0 to GetComponentCount(lod)-1
+        const char* GetComponentName( int lod, int comp ) const;
+		
+		//! Get the Id of a component
+		//! \param lod Index of the level of detail, from 0 to GetLODCount()-1
+		//! \param comp Index of the component, from 0 to GetComponentCount(lod)-1
+		uint16 GetComponentId( int lod, int comp ) const;
+
+        //! Get the number of surfaces in a component
+        //! \param lod Index of the level of detail, from 0 to GetLODCount()-1
+        //! \param comp Index of the component, from 0 to GetComponentCount(lod)-1
+        int GetSurfaceCount( int lod, int comp ) const;
+
+        //! Get the name of a surface
+        //! \param lod Index of the level of detail, from 0 to GetLODCount()-1
+        //! \param comp Index of the component, from 0 to GetComponentCount(lod)-1
+        //! \param surf Index of the surface, from 0 to GetSurfaceCount(lod,comp)-1
+        const char* GetSurfaceName( int lod, int comp, int surf ) const;
+
+        //! Get an id that can be used to match the surface data with the mesh surface data.
+        //! \param lod Index of the level of detail, from 0 to GetLODCount()-1
+        //! \param comp Index of the component, from 0 to GetComponentCount(lod)-1
+        //! \param surf Index of the surface, from 0 to GetSurfaceCount(lod,comp)-1
+        uint32_t GetSurfaceId( int lod, int comp, int surf ) const;
+
+        //! Find a surface index from the internal id (as returned by GetSurfaceId).
+        //! \param lod Index of the level of detail, from 0 to GetLODCount()-1
+        //! \param comp Index of the component, from 0 to GetComponentCount(lod)-1
+        //! \param id ID of the surface to look for.
+        int FindSurfaceById( int lod, int comp, uint32_t id ) const;
+
+        //! Get an optional, opaque application-defined identifier for this surface. The meaning of
+        //! this ID depends on each application, and it is specified when creating the source data
+        //! that generates this surface.
+        //! See NodeSurfaceNew::SetCustomID.
+		//! \param lod Index of the level of detail, from 0 to GetLODCount()-1
+		//! \param comp Index of the component, from 0 to GetComponentCount(lod)-1
+        //! \param surf Index of the surface, from 0 to GetSurfaceCount(lod,comp)-1
+        uint32_t GetSurfaceCustomId( int lod, int comp, int surf ) const;
+
+        //! Get the number of meshes in a surface
+		//! \param lod Index of the level of detail, from 0 to GetLODCount()-1
+		//! \param comp Index of the component, from 0 to GetComponentCount(lod)-1
+        int GetMeshCount( int lod, int comp ) const;
+
+        //! Get a mesh resource id from a component
+        //! \param lod Index of the level of detail, from 0 to GetLODCount()-1
+        //! \param comp Index of the component, from 0 to GetComponentCount(lod)-1
+        //! \param mesh Index of the mesh, from 0 to GetMeshCount(lod,comp,surf)-1
+        //! \return an ummodifiable pointer to the requested mesh. The returned object is guaranteed
+        //! to be alive only while this instance is alive. The ownership of the returned object
+        //! remains in the instance, so it should not be deleted.
+        RESOURCE_ID GetMeshId( int lod, int comp, int mesh ) const;
+
+		//! Get the name of a mesh in a component
+		//! \param lod Index of the level of detail, from 0 to GetLODCount()-1
+		//! \param comp Index of the component, from 0 to GetComponentCount(lod)-1
+        //! \param mesh Index of the mesh, from 0 to GetMeshCount(lod,comp,surf)-1
+        const char* GetMeshName( int lod, int comp, int mesh ) const;
+
+		//! Get the number of images in a component
+		//! \param lod Index of the level of detail, from 0 to GetLODCount()
+		//! \param comp Index of the component, from 0 to GetComponentCount(lod)
+        //! \param surf Index of the surface, from 0 to GetSurfaceCount(lod,comp)-1
+        int GetImageCount( int lod, int comp, int surf ) const;
+
+        //! Get an image resource id from a component
+		//! \param lod Index of the level of detail, from 0 to GetLODCount()-1
+		//! \param comp Index of the component, from 0 to GetComponentCount(lod)-1
+        //! \param surf Index of the surface, from 0 to GetSurfaceCount(lod,comp)-1
+        //! \param img Index of the image, from 0 to GetImageCount(lod,comp,surf)-1
+        //! \return an ummodifiable pointer to the requested image. The returned object is
+        //! guaranteed to be alive only while this instance is alive. The ownership of the returned
+        //! object remains in the instance, so it should not be deleted.
+        RESOURCE_ID GetImageId( int lod, int comp, int surf, int img ) const;
+
+		//! Get the name of an image in a component
+		//! \param lod Index of the level of detail, from 0 to GetLODCount()-1
+		//! \param comp Index of the component, from 0 to GetComponentCount(lod)-1
+        //! \param surf Index of the surface, from 0 to GetSurfaceCount(lod,comp)-1
+        //! \param img Index of the image, from 0 to GetImageCount(lod,comp,surf)-1
+        const char* GetImageName( int lod, int comp, int surf, int img ) const;
+
+		//! Get the number of vectors in a component
+		//! \param lod Index of the level of detail, from 0 to GetLODCount()
+		//! \param comp Index of the component, from 0 to GetComponentCount(lod)
+        //! \param surf Index of the surface, from 0 to GetSurfaceCount(lod,comp)-1
+        int GetVectorCount( int lod, int comp, int surf ) const;
+
+		//! Get a vector from a component
+		//! \param lod Index of the level of detail, from 0 to GetLODCount()-1
+		//! \param comp Index of the component, from 0 to GetComponentCount(lod)-1
+        //! \param surf Index of the surface, from 0 to GetSurfaceCount(lod,comp)-1
+        //! \param vec Index of the vector, from 0 to GetVectorCount(lod,comp,surf)-1
+		//! \param pX, pY, pZ, pW are optional to floats to store the components of the vector
+		//!		   value. They can be null.
+        void GetVector( int lod, int comp, int surf, int vec,
+						float* pX, float* pY, float* pZ, float* pW ) const;
+
+		//! Get the name of a vector in a component
+		//! \param lod Index of the level of detail, from 0 to GetLODCount()-1
+		//! \param comp Index of the component, from 0 to GetComponentCount(lod)-1
+        //! \param surf Index of the surface, from 0 to GetSurfaceCount(lod,comp)-1
+        //! \param vec Index of the vector, from 0 to GetVectorCount(lod,comp)-1
+        const char* GetVectorName( int lod, int comp, int surf, int vec ) const;
+
+        //! Get the number of scalar values in a component
+        //! \param lod Index of the level of detail, from 0 to GetLODCount()
+        //! \param comp Index of the component, from 0 to GetComponentCount(lod)
+        //! \param surf Index of the surface, from 0 to GetSurfaceCount(lod,comp)-1
+        int GetScalarCount( int lod, int comp, int surf ) const;
+
+        //! Get a scalar value from a component
+        //! \param lod Index of the level of detail, from 0 to GetLODCount()-1
+        //! \param comp Index of the component, from 0 to GetComponentCount(lod)-1
+        //! \param surf Index of the surface, from 0 to GetSurfaceCount(lod,comp)-1
+        //! \param sca Index of the scalar, from 0 to GetScalarCount(lod,comp)-1
+        float GetScalar( int lod, int comp, int surf, int sca ) const;
+
+        //! Get the name of a scalar from a component
+        //! \param lod Index of the level of detail, from 0 to GetLODCount()-1
+        //! \param comp Index of the component, from 0 to GetComponentCount(lod)-1
+        //! \param surf Index of the surface, from 0 to GetSurfaceCount(lod,comp)-1
+        //! \param sca Index of the scalar, from 0 to GetScalarCount(lod,comp)-1
+        const char* GetScalarName( int lod, int comp, int surf, int sca ) const;
+
+        //! Get the number of string values in a component
+        //! \param lod Index of the level of detail, from 0 to GetLODCount()
+        //! \param comp Index of the component, from 0 to GetComponentCount(lod)
+        //! \param surf Index of the surface, from 0 to GetSurfaceCount(lod,comp)-1
+        int GetStringCount( int lod, int comp, int surf ) const;
+
+        //! Get a string value from a component
+        //! \param lod Index of the level of detail, from 0 to GetLODCount()-1
+        //! \param comp Index of the component, from 0 to GetComponentCount(lod)-1
+        //! \param surf Index of the surface, from 0 to GetSurfaceCount(lod,comp)-1
+        //! \param sca Index of the string, from 0 to GetStringCount(lod,comp)-1
+        const char* GetString( int lod, int comp, int surf, int str ) const;
+
+        //! Get the name of a string from a component
+        //! \param lod Index of the level of detail, from 0 to GetLODCount()-1
+        //! \param comp Index of the component, from 0 to GetComponentCount(lod)-1
+        //! \param surf Index of the surface, from 0 to GetSurfaceCount(lod,comp)-1
+        //! \param sca Index of the string, from 0 to GetStringCount(lod,comp)-1
+        const char* GetStringName( int lod, int comp, int surf, int str ) const;
+
+
+        //-----------------------------------------------------------------------------------------
+		// Interface pattern
+		//-----------------------------------------------------------------------------------------
+		class Private;
+		Private* GetPrivate() const;
+
+	protected:
+
+		//! Forbidden. Manage with the Ptr<> template.
+		~Instance();
+
+	private:
+
+		Private* m_pD;
+
+	};
+
+
+}
+
diff --git a/Engine/Plugins/Experimental/Mutable/Source/MutableRuntime/Public/Layout.h b/Engine/Plugins/Experimental/Mutable/Source/MutableRuntime/Public/Layout.h
new file mode 100644
--- /dev/null
+++ b/Engine/Plugins/Experimental/Mutable/Source/MutableRuntime/Public/Layout.h
@@ -0,0 +1,244 @@
+// Copyright Epic Games, Inc. All Rights Reserved.
+
+#pragma once
+
+#include "MutableRuntime/Public/Types.h"
+#include "MutableRuntime/Public/Ptr.h"
+#include "MutableRuntime/Public/RefCounted.h"
+#include "MutableRuntime/Public/Serialisation.h"
+#include "MutableRuntime/Private/MutableMath.h"
+#include "Math/IntPoint.h"
+
+namespace mu
+{
+
+	// Forward references
+	class Layout;
+	typedef Ptr<Layout> LayoutPtr;
+	typedef Ptr<const Layout> LayoutPtrConst;
+
+	//! Types of layout packing strategies 
+	enum class EPackStrategy : uint32_t
+	{
+		RESIZABLE_LAYOUT,
+		FIXED_LAYOUT
+	};
+
+
+    //! \brief Image block layout class.
+    //!
+    //! It contains the information about what blocks are defined in a texture layout (texture
+    //! coordinates set from a mesh).
+    //! It is usually not necessary to use this objects, except for some advanced cases.
+	//! \ingroup runtime
+    class MUTABLERUNTIME_API Layout : public RefCounted
+	{
+	public:
+
+		//-----------------------------------------------------------------------------------------
+		// Life cycle
+		//-----------------------------------------------------------------------------------------
+
+		Layout();
+
+		//! Deep clone this layout.
+		LayoutPtr Clone() const;
+
+		//! Serialisation
+		static void Serialise( const Layout* p, OutputArchive& arch );
+		static LayoutPtr StaticUnserialise( InputArchive& arch );
+
+		//! Full compare
+		bool operator==( const Layout& other ) const;
+
+		//-----------------------------------------------------------------------------------------
+		// Own interface
+		//-----------------------------------------------------------------------------------------
+
+		//! Get the resolution of the grid where the blocks are defined.
+		FIntPoint GetGridSize() const;
+
+		//! Get the resolution of the grid where the blocks are defined. It must be bigger than 0
+		//! on each axis.
+		//! \param sizeX width of the grid.
+		//! \param sizeY height of the grid.
+		void SetGridSize( int sizeX, int sizeY );
+
+		//! Get the maximum resolution of the grid where the blocks are defined.
+		//! \param[out] pSizeX The integer pointed by this will be set to the width of the grid.
+		//! \param[out] pSizeY The integer pointed by this will be set to the height of the grid.
+		void GetMaxGridSize(int* pSizeX, int* pSizeY) const;
+
+		//! Get the maximum resolution of the grid where the blocks are defined. It must be bigger than 0
+		//! on each axis.
+		//! \param sizeX width of the grid.
+		//! \param sizeY height of the grid.
+		void SetMaxGridSize(int sizeX, int sizeY);
+
+		//! Return the number of blocks in this layout.
+		int32 GetBlockCount() const;
+
+		//! Set the number of blocks in this layout.
+		//! The existing blocks will be kept as much as possible. The new blocks will be undefined.
+		void SetBlockCount( int32 );
+
+		//! Return a block of the layout.
+		//! "Position" here means the lower-left corner of the block.
+		//! \param index Block to get.
+		//! \param[out] pMinX will be set to the x position of the block
+		//! (from 0 to the X size of the grid minus one )
+		//! \param[out] pMinY will be set to the y position of the block
+		//! (from 0 to the Y size of the grid minus one )
+		//! \param[out] pSizeX will be set to the x size of the block
+		//! \param[out] pSizeY will be set to the y size of the block
+		void GetBlock( int index, int* pMinX, int* pMinY, int* pSizeX, int* pSizeY ) const;
+
+		//! Returns the reduction priority of a block.
+		//! \param index Block to get priority.
+		//! \param[out] priority will be set to the reduction priority of the block
+		void GetBlockPriority(int index, int* priority) const;
+
+		//! Set a block of the layout.
+		//! "Position" here means the lower-left corner of the block.
+		//! \param index Block to set.
+		//! \param minX will be set to the x position of the block
+		//! (from 0 to the X size of the grid minus one )
+		//! \param minY will be set to the y position of the block
+		//! (from 0 to the Y size of the grid minus one )
+		//! \param sizeX will be set to the x size of the block
+		//! \param sizeY will be set to the y size of the block
+        void SetBlock( int index, int minX, int minY, int sizeX, int sizeY );
+
+		//! Set the reduction priority of a block
+		//! The blocks with the highest values will be the last to be reduced
+		//! \param index Block to set priority
+		//! \param priority will be set to the reduction priority of the block
+		void SetBlockPriority(int index, int priority);
+
+		//! Set the texture layout packing strategy
+		//! By default the texture layout packing strategy is set to resizable layout
+		void SetLayoutPackingStrategy(EPackStrategy _strategy);
+
+		//! Set the texture layout packing strategy
+		EPackStrategy GetLayoutPackingStrategy() const;
+
+	protected:
+
+		//! Forbidden. Manage with the Ptr<> template.
+		~Layout() {}
+
+	public:
+
+		// This used to be the data in the private implementation of the image interface
+		//-----------------------------------------------------------------------------------------
+		struct BLOCK
+		{
+			BLOCK(vec2<uint16_t> min = vec2<uint16_t>(), vec2<uint16_t> size = vec2<uint16_t>())
+			{
+				m_min = min;
+				m_size = size;
+				m_id = -1;
+				m_priority = 0;
+			}
+
+			vec2<uint16_t> m_min;
+			vec2<uint16_t> m_size;
+
+			//! Absolute id used to control merging of various layouts
+			int32_t m_id;
+
+			//! Priority value to control the shrink texture layout strategy
+			int32_t m_priority;
+
+
+			//!
+			void Serialise(OutputArchive& arch) const
+			{
+				arch << m_min;
+				arch << m_size;
+				arch << m_id;
+				arch << m_priority;
+			}
+
+			//!
+			void Unserialise(InputArchive& arch)
+			{
+				arch >> m_min;
+				arch >> m_size;
+				arch >> m_id;
+				arch >> m_priority;
+			}
+
+			//!
+			inline bool operator==(const BLOCK& o) const
+			{
+				return (m_min == o.m_min) &&
+					(m_size == o.m_size) &&
+					(m_id == o.m_id) &&
+					(m_priority == o.m_priority);
+			}
+
+			inline bool IsSimilar(const BLOCK& o) const
+			{
+				// All but ids
+				return (m_min == o.m_min) &&
+					(m_size == o.m_size) &&
+					(m_priority == o.m_priority);
+			}
+		};
+
+
+		//!
+		vec2<uint16_t> m_size;
+
+		vec2<uint16_t> m_maxsize;
+
+		//!
+		TArray<BLOCK> m_blocks;
+
+		//! Packing strategy
+		EPackStrategy m_strategy = EPackStrategy::RESIZABLE_LAYOUT;
+
+		//!
+		void Serialise(OutputArchive& arch) const
+		{
+			uint32_t ver = 3;
+			arch << ver;
+
+			arch << m_size;
+			arch << m_blocks;
+
+			arch << m_maxsize;
+			arch << uint32_t(m_strategy);
+		}
+
+		//!
+		void Unserialise(InputArchive& arch)
+		{
+			uint32_t ver;
+			arch >> ver;
+			check(ver == 3);
+
+			arch >> m_size;
+
+			arch >> m_blocks;
+			arch >> m_maxsize;
+
+			uint32_t temp;
+			arch >> temp;
+			m_strategy = EPackStrategy(temp);
+		}
+
+		//!
+		bool IsSimilar(const Layout& o) const;
+
+		//! Find a block by id. This converts the "absolute" id to a relative index to the layout
+		//! blocks. Return -1 if not found.
+		int32 FindBlock(int32_t id) const;
+
+		//! Return true if the layout is a single block filling all area.
+		bool IsSingleBlockAndFull() const;
+	};
+
+}
+
diff --git a/Engine/Plugins/Experimental/Mutable/Source/MutableRuntime/Public/Mesh.h b/Engine/Plugins/Experimental/Mutable/Source/MutableRuntime/Public/Mesh.h
new file mode 100644
--- /dev/null
+++ b/Engine/Plugins/Experimental/Mutable/Source/MutableRuntime/Public/Mesh.h
@@ -0,0 +1,574 @@
+// Copyright Epic Games, Inc. All Rights Reserved.
+
+#pragma once
+
+#include "MutableRuntime/Public/Types.h"
+#include "MutableRuntime/Public/Ptr.h"
+#include "MutableRuntime/Public/RefCounted.h"
+#include "MutableRuntime/Public/Serialisation.h"
+#include "MutableRuntime/Public/MeshBufferSet.h"
+#include "MutableRuntime/Public/Skeleton.h"
+#include "MutableRuntime/Public/PhysicsBody.h"
+#include "MutableRuntime/Public/Layout.h"
+#include "Math/Transform.h"
+
+#include "Misc/EnumClassFlags.h"
+
+namespace mu
+{
+
+	// Forward references
+	class Layout;
+	typedef Ptr<Layout> LayoutPtr;
+	typedef Ptr<const Layout> LayoutPtrConst;
+
+    class Skeleton;
+    typedef Ptr<Skeleton> SkeletonPtr;
+    typedef Ptr<const Skeleton> SkeletonPtrConst;
+
+    class PhysicsBody;
+    typedef Ptr<PhysicsBody> PhysicsBodyPtr;
+    typedef Ptr<const PhysicsBody> PhysicsBodyPtrConst;
+
+    class Mesh;
+    typedef Ptr<Mesh> MeshPtr;
+    typedef Ptr<const Mesh> MeshPtrConst;
+
+	struct MESH_SURFACE
+	{
+		MESH_SURFACE()
+		{
+			m_firstVertex = 0;
+			m_vertexCount = 0;
+			m_firstIndex = 0;
+			m_indexCount = 0;
+			m_id = 0;
+		}
+
+		int32_t m_firstVertex;
+		int32_t m_vertexCount;
+		int32_t m_firstIndex;
+		int32_t m_indexCount;
+		uint32_t m_id;
+
+
+		//!
+		inline bool operator==(const MESH_SURFACE& o) const
+		{
+			return m_firstVertex == o.m_firstVertex
+				&& m_vertexCount == o.m_vertexCount
+				&& m_firstIndex == o.m_firstIndex
+				&& m_indexCount == o.m_indexCount
+				&& m_id == o.m_id;
+		}
+
+	};
+
+	MUTABLE_DEFINE_POD_VECTOR_SERIALISABLE(MESH_SURFACE);
+
+
+	//!
+	enum class EMeshBufferType
+	{
+		None,
+		SkeletonDeformBinding,
+		PhysicsBodyDeformBinding,
+		PhysicsBodyDeformSelection
+	};
+	MUTABLE_DEFINE_ENUM_SERIALISABLE(EMeshBufferType)
+
+	//!
+	enum class EShapeBindingMethod : uint32
+	{
+		ReshapeClosestProject = 0,
+		ClipDeformClosestProject = 1,
+		ClipDeformClosestToSurface = 2,
+		ClipDeformNormalProject = 3	
+	};
+	MUTABLE_DEFINE_ENUM_SERIALISABLE(EShapeBindingMethod)
+
+	enum class EMeshCloneFlags : uint32
+	{
+		None = 0,
+		WithSkeletalMesh      = 1 << 1,
+		WithSurfaces          = 1 << 2,
+		WithSkeleton          = 1 << 3,
+		WithPhysicsBody       = 1 << 4,
+		WithFaceGroups        = 1 << 5,
+		WithTags              = 1 << 6,
+		WithVertexBuffers     = 1 << 7,
+		WithIndexBuffers      = 1 << 8,
+		WithFaceBuffers       = 1 << 9,
+		WithAdditionalBuffers = 1 << 10,
+		WithLayouts           = 1 << 11,
+		WithPoses			  = 1 << 12
+	};
+	ENUM_CLASS_FLAGS(EMeshCloneFlags);
+	
+
+    //! \brief Mesh object containing any number of buffers with any number of channels.
+    //! The buffers can be per-index or per-vertex.
+    //! The mesh also includes layout information for every texture channel for internal usage, and
+    //! it can be ignored.
+    //! The meshes are always assumed to be triangle list primitives.
+    //! \ingroup runtime
+    class MUTABLERUNTIME_API Mesh : public RefCounted
+    {
+    public:
+
+        //-----------------------------------------------------------------------------------------
+        // Life cycle
+        //-----------------------------------------------------------------------------------------
+
+        //! Deep clone this mesh.
+        MeshPtr Clone() const;
+		
+		// Clone with flags allowing to not include some parts in the cloned mesh
+		MeshPtr Clone(EMeshCloneFlags Flags) const;
+
+        //! Serialisation
+        static void Serialise( const Mesh* p, OutputArchive& arch );
+        static MeshPtr StaticUnserialise( InputArchive& arch );
+
+        //-----------------------------------------------------------------------------------------
+        // Own interface
+        //-----------------------------------------------------------------------------------------
+
+
+        //! \name Buffers
+        //! \{
+
+        //!
+        int GetIndexCount() const;
+
+        //! Index buffers. They are owned by this mesh.
+        FMeshBufferSet& GetIndexBuffers();
+        const FMeshBufferSet& GetIndexBuffers() const;
+
+        //
+        int GetVertexCount() const;
+
+        //! Vertex buffers. They are owned by this mesh.
+        FMeshBufferSet& GetVertexBuffers();
+        const FMeshBufferSet& GetVertexBuffers() const;
+
+        //
+        int GetFaceCount() const;
+
+        //! Face buffers. They are owned by this mesh.
+        FMeshBufferSet& GetFaceBuffers();
+        const FMeshBufferSet& GetFaceBuffers() const;
+
+        //! Get the number of surfaces defined in this mesh. Surfaces are buffer-contiguous mesh
+        //! fragments that share common properties (usually material)
+        int GetSurfaceCount() const;
+        void GetSurface( int surfaceIndex,
+                         int* firstVertex, int* vertexCount,
+                         int* firstIndex, int* indexCount ) const;
+
+        //! Return an internal id that can be used to match mesh surfaces and instance surfaces.
+        //! Only valid for meshes that are part of instances.
+        uint32_t GetSurfaceId( int surfaceIndex ) const;
+
+        //! \}
+
+
+        //! \name Texture layouts
+        //! \{
+
+        //!
+        void AddLayout( Ptr<const Layout> pLayout );
+
+        //!
+        int GetLayoutCount() const;
+
+        //!
+        const Layout* GetLayout( int i ) const;
+
+        //!
+        void SetLayout( int i, Ptr<const Layout> );
+        //! \}
+
+        //! \name Skeleton information
+        //! \{
+
+        void SetSkeleton( Ptr<const Skeleton> );
+        Ptr<const Skeleton> GetSkeleton() const;
+
+        //! \}
+
+        //! \name PhysicsBody information
+        //! \{
+
+        void SetPhysicsBody( Ptr<const PhysicsBody> );
+        Ptr<const PhysicsBody> GetPhysicsBody() const;
+
+        //! \}
+
+        //! \name Face groups
+        //! \{
+
+        //!
+        void SetFaceGroupCount( int count );
+
+        //!
+        int GetFaceGroupCount() const;
+
+        //!
+        const char* GetFaceGroupName( int group ) const;
+
+        //!
+        void SetFaceGroupName( int group, const char* strName );
+
+        //!
+        int GetFaceGroupFaceCount( int group ) const;
+
+        //!
+        const int32* GetFaceGroupFaces( int group ) const;
+
+        //!
+        void SetFaceGroupFaces( int group, int count, const int32* faces );
+
+        //! \}
+
+
+        //! \name Tags
+        //! \{
+
+        //!
+        void SetTagCount( int count );
+
+        //!
+        int GetTagCount() const;
+
+        //!
+        const char* GetTag( int tagIndex ) const;
+
+        //!
+        void SetTag( int tagIndex, const char* strName );
+
+		//!
+		const char* GetBonePoseName( int32 boneIndex ) const;
+
+		//!
+		int32 FindBonePose(const char* strName) const;
+		
+		//!
+		void SetBonePoseCount(int32 count);
+
+		//!
+		int32 GetBonePoseCount() const;
+
+		//!
+		void SetBonePose(int32 boneIndex, const char* strName, FTransform3f transform, bool bSkinned );
+
+		//! Return a matrix stored per bone. It is a set of 16-float values.
+		void GetBoneTransform(int32 boneIndex, FTransform3f& transform) const;
+
+        //! \}
+
+
+        //! Get an internal identifier used to reference this mesh in operations like deferred
+        //! mesh building, or instance updating.
+        uint32_t GetId() const;
+
+
+    protected:
+
+        //! Forbidden. Manage with the Ptr<> template.
+		~Mesh() {}
+
+    
+	public:
+
+		//! Non-persistent internal id unique for a mesh generated for a specific state and
+		//! parameter values.
+		mutable uint32_t m_internalId = 0;
+
+		//!
+		FMeshBufferSet m_VertexBuffers;
+
+		//!
+		FMeshBufferSet m_IndexBuffers;
+
+		//!
+		FMeshBufferSet m_FaceBuffers;
+
+		//! Additional buffers used for temporary or custom data in different algorithms.
+		TArray< TPair<EMeshBufferType, FMeshBufferSet> > m_AdditionalBuffers;
+
+		//! This is bit-mask on the STATIC_MESH_FORMATS enumeration, marking what static formats
+		//! are compatible with this one. Usually precalculated at model compilation time.
+		//! It should be reset after any operation that modifies the format.
+		mutable uint32_t m_staticFormatFlags = 0;
+
+		TArray<MESH_SURFACE> m_surfaces;
+
+		//! This skeleton and physics body are not owned and may be used by other meshes, so it cannot be modified
+		//! once the mesh has been fully created.
+		Ptr<const Skeleton> m_pSkeleton;
+		Ptr<const PhysicsBody> m_pPhysicsBody;
+
+		//! Texture Layout blocks attached to this mesh. They are const because they could be shared with
+		//! other meshes, so they need to be cloned and replaced if a modification is needed.
+		TArray<Ptr<const Layout>> m_layouts;
+
+		struct FACE_GROUP
+		{
+			string m_name;
+			TArray<int32> m_faces;
+
+			inline void Serialise(OutputArchive& arch) const
+			{
+				const int32_t ver = 0;
+				arch << ver;
+
+				arch << m_name;
+				arch << m_faces;
+			}
+
+			inline void Unserialise(InputArchive& arch)
+			{
+				int32_t ver = 0;
+				arch >> ver;
+				check(ver == 0);
+
+				arch >> m_name;
+				arch >> m_faces;
+			}
+
+			//!
+			inline bool operator==(const FACE_GROUP& o) const
+			{
+				return m_name == o.m_name
+					&& m_faces == o.m_faces;
+			}
+
+		};
+		TArray<FACE_GROUP> m_faceGroups;
+
+		//!
+		TArray<string> m_tags;
+
+		struct FBonePose
+		{
+			string m_boneName;
+			uint8 m_boneSkinned : 1;
+			FTransform3f m_boneTransform;
+
+			FBonePose() : m_boneSkinned(0) {};
+
+			inline void Serialise(OutputArchive& arch) const
+			{
+				const int32_t ver = 0;
+				arch << ver;
+
+				arch << m_boneName;
+				arch << m_boneSkinned;
+				arch << m_boneTransform;
+			}
+
+			inline void Unserialise(InputArchive& arch)
+			{
+				int32_t ver = 0;
+				arch >> ver;
+				check(ver == 0);
+
+				arch >> m_boneName;
+
+				uint8 skinned = 0;
+				arch >> skinned;
+				m_boneSkinned = skinned;
+
+				arch >> m_boneTransform;
+			}
+
+			//!
+			inline bool operator==(const FBonePose& o) const
+			{
+				return m_boneSkinned == o.m_boneSkinned
+					&& m_boneName == o.m_boneName;
+			}
+
+		};
+		// This is the pose used by this mesh fragment, used to update the transforms of the final skeleton
+		// taking into consideration the meshes being used.
+		TArray<FBonePose> m_bonePoses;
+
+		//!
+		inline void Serialise(OutputArchive& arch) const
+		{
+			uint32_t ver = 13;
+			arch << ver;
+
+			arch << m_IndexBuffers;
+			arch << m_VertexBuffers;
+			arch << m_FaceBuffers;
+			arch << m_AdditionalBuffers;
+			arch << m_layouts;
+
+			arch << m_pSkeleton;
+			arch << m_pPhysicsBody;
+
+			arch << m_staticFormatFlags;
+			arch << m_surfaces;
+			arch << m_faceGroups;
+
+			arch << m_tags;
+
+			arch << m_bonePoses;
+		}
+
+		//!
+		inline void Unserialise(InputArchive& arch)
+		{
+			uint32_t ver;
+			arch >> ver;
+			check(ver <= 13);
+
+			arch >> m_IndexBuffers;
+			arch >> m_VertexBuffers;
+			arch >> m_FaceBuffers;
+			arch >> m_AdditionalBuffers;
+			arch >> m_layouts;
+
+			arch >> m_pSkeleton;
+			if (ver >= 12)
+			{ 
+				arch >> m_pPhysicsBody;
+			}
+			else
+			{
+				m_pPhysicsBody = nullptr;
+			}
+
+			arch >> m_staticFormatFlags;
+			arch >> m_surfaces;
+			arch >> m_faceGroups;
+
+			arch >> m_tags;
+
+			if (ver >= 13)
+			{
+				arch >> m_bonePoses;
+			}
+			else if (m_pSkeleton)
+			{
+				const int32 NumBones = m_pSkeleton->GetBoneCount();
+				m_bonePoses.SetNum(NumBones);
+				check(m_pSkeleton->m_boneTransforms_DEPRECATED.Num() == NumBones);
+
+				for (int32 BoneIndex = 0; BoneIndex < NumBones; ++BoneIndex)
+				{
+					m_bonePoses[BoneIndex].m_boneName = m_pSkeleton->GetBoneName(BoneIndex);
+					m_bonePoses[BoneIndex].m_boneSkinned = true;
+					m_bonePoses[BoneIndex].m_boneTransform = m_pSkeleton->m_boneTransforms_DEPRECATED[BoneIndex];
+				}
+			}
+		}
+
+
+		//!
+		inline bool operator==(const Mesh& o) const
+		{
+			bool equal = (m_IndexBuffers == o.m_IndexBuffers);
+			if (equal) equal = (m_VertexBuffers == o.m_VertexBuffers);
+			if (equal) equal = (m_FaceBuffers == o.m_FaceBuffers);
+			if (equal) equal = (m_layouts.Num() == o.m_layouts.Num());
+			if (equal) equal = (m_bonePoses.Num() == o.m_bonePoses.Num());
+			if (equal && m_pSkeleton != o.m_pSkeleton)
+			{
+				if (m_pSkeleton && o.m_pSkeleton)
+				{
+					equal = (*m_pSkeleton == *o.m_pSkeleton);
+				}
+				else
+				{
+					equal = false;
+				}
+			}
+			if (equal) equal = (m_surfaces == o.m_surfaces);
+			if (equal) equal = (m_faceGroups == o.m_faceGroups);
+			if (equal) equal = (m_tags == o.m_tags);
+
+			for (int32 i = 0; equal && i < m_layouts.Num(); ++i)
+			{
+				equal &= (*m_layouts[i]) == (*o.m_layouts[i]);
+			}
+
+			for (int32 i = 0; equal && i < m_AdditionalBuffers.Num(); ++i)
+			{
+				equal &= m_AdditionalBuffers[i] == o.m_AdditionalBuffers[i];
+			}
+
+			for (int32 i = 0; equal && i < m_bonePoses.Num(); ++i)
+			{
+				equal &= m_bonePoses[i] == o.m_bonePoses[i];
+			}
+
+			return equal;
+		}
+
+		//! Compare the mesh with another one, but ignore internal data like generated vertex
+		//! indices.
+		bool IsSimilar(const Mesh& o) const;
+
+		//! Compare the vertex attributes to check if they match.
+		bool IsSameVertex( uint32_t v, const Mesh& other, uint32_t otherVertexIndex, float tolerance = 1e-5f) const;
+
+
+		//! Make a map from the vertices in this mesh to thefirst matching vertex of the given
+		//! mesh. If non is found, the index is set to -1.
+		struct VERTEX_MATCH_MAP
+		{
+			//! One for every vertex
+			vector<int> m_firstMatch;
+
+			//! The matches of every vertex in a sequence
+			vector<int> m_matches;
+
+			//!
+			bool Matches(int v, int ov) const;
+		};
+
+		void GetVertexMap
+		(
+			const Mesh& other,
+			VERTEX_MATCH_MAP& vertexMap,
+			float tolerance = 1e-3f//1e-5f //std::numeric_limits<float>::min()
+		) const;
+
+		//! Return true if the mesh has a face with the same vertices than the face in the given
+		//! mesh. The actual vertex data is used in the comparison.
+		bool HasFace( const Mesh& other, int otherFaceIndex, const VERTEX_MATCH_MAP& vertexMap ) const;
+
+		//! Compare the vertex attributes to check if they match.
+		vec3<uint32_t> GetFaceVertexIndices(int f) const;
+
+		//! Return true if the given mesh has the same vertex and index formats, and in the same
+		//! buffer structure.
+		bool HasCompatibleFormat(const Mesh* pOther) const;
+
+		//! Update the flags identifying the mesh format as some of the optimised formats.
+		void ResetStaticFormatFlags() const;
+
+		//! Get the total memory size of the buffers
+		size_t GetDataSize() const;
+
+		//! Create the surface data if not present.
+		void EnsureSurfaceData();
+
+		//! Check mesh buffer data for possible inconsistencies
+		void CheckIntegrity() const;
+
+		//! Change the buffer descriptions so that all buffer indices start at 0 and are in the
+		//! same order than memory.
+		void ResetBufferIndices();
+
+		//! Debug: get a text representation of the mesh
+		void Log(FString& out, int32 VertrexLimit);
+    };
+
+
+}
+
diff --git a/Engine/Plugins/Experimental/Mutable/Source/MutableRuntime/Public/MeshBufferSet.h b/Engine/Plugins/Experimental/Mutable/Source/MutableRuntime/Public/MeshBufferSet.h
new file mode 100644
--- /dev/null
+++ b/Engine/Plugins/Experimental/Mutable/Source/MutableRuntime/Public/MeshBufferSet.h
@@ -0,0 +1,407 @@
+// Copyright Epic Games, Inc. All Rights Reserved.
+
+#pragma once
+
+#include "MutableRuntime/Public/Types.h"
+#include "MutableRuntime/Public/Ptr.h"
+#include "MutableRuntime/Public/RefCounted.h"
+#include "MutableRuntime/Private/SerialisationPrivate.h"
+#include "Math/Transform.h"
+
+
+namespace mu
+{
+
+	//! Supported formats for the elements in mesh buffers.
+	//! \ingroup runtime
+	typedef enum
+	{
+
+		MBF_NONE,
+		MBF_FLOAT16,
+		MBF_FLOAT32,
+
+		MBF_UINT8,
+		MBF_UINT16,
+		MBF_UINT32,
+		MBF_INT8,
+		MBF_INT16,
+		MBF_INT32,
+
+		//! Integers interpreted as being in the range 0.0f to 1.0f
+		MBF_NUINT8,
+		MBF_NUINT16,
+		MBF_NUINT32,
+
+		//! Integers interpreted as being in the range -1.0f to 1.0f
+		MBF_NINT8,
+		MBF_NINT16,
+		MBF_NINT32,
+
+        //! Packed 1 to -1 value using multiply+add (128 is almost zero). Use 8-bit unsigned ints.
+        MBF_PACKEDDIR8,
+
+        //! Same as MBF_PACKEDDIR8, with the w component replaced with the sign of the determinant
+        //! of the vertex basis to define the orientation of the tangent space in UE4 format.
+        //! Use 8-bit unsigned ints.
+        MBF_PACKEDDIR8_W_TANGENTSIGN,
+
+        //! Packed 1 to -1 value using multiply+add (128 is almost zero). Use 8-bit signed ints.
+        MBF_PACKEDDIRS8,
+
+        //! Same as MBF_PACKEDDIRS8, with the w component replaced with the sign of the determinant
+        //! of the vertex basis to define the orientation of the tangent space in UE4 format.
+        //! Use 8-bit signed ints.
+        MBF_PACKEDDIRS8_W_TANGENTSIGN,
+
+		MBF_FLOAT64,
+
+		MBF_COUNT,
+
+		_MBF_FORCE32BITS = 0xFFFFFFFF
+
+	} MESH_BUFFER_FORMAT;
+
+	//!
+	struct FMeshBufferFormatData
+	{
+		//! Size per component
+		unsigned m_size;
+
+		//! log 2 of the max value if integer
+		uint8_t m_maxValueBits;
+	};
+
+	MUTABLERUNTIME_API const FMeshBufferFormatData& GetMeshFormatData(MESH_BUFFER_FORMAT format);
+
+
+	//! Semantics of the mesh buffers
+	//! \ingroup runtime
+	typedef enum
+	{
+
+		MBS_NONE,
+
+		//! For index buffers, and mesh morphs
+		MBS_VERTEXINDEX,
+
+		//! Standard vertex semantics
+		MBS_POSITION,
+		MBS_NORMAL,
+		MBS_TANGENT,
+		MBS_BINORMAL,
+		MBS_TEXCOORDS,
+		MBS_COLOUR,
+		MBS_BONEWEIGHTS,
+		MBS_BONEINDICES,
+
+		//! Internal semantic indicating what layout block each vertex belongs to.
+		//! It can be safely ignored if present in meshes returned by the system.
+		//! It will never be in the same buffer that other vertex semantics.
+		MBS_LAYOUTBLOCK,
+
+		//! Internal semantic indicating what chart does the given vertex or face belong to.
+		//! It can be safely ignored if present in meshes returned by the system.
+		MBS_CHART,
+
+		//! To let users define channels with semantics unknown to the system.
+		//! These channels will never be transformed, and the per-vertex or per-index data will be
+		//! simply copied.
+		MBS_OTHER,
+
+        //! Sign to define the orientation of the tangent space.
+        MBS_TANGENTSIGN,
+
+		//! Semantics usefule for mesh binding.
+		MBS_TRIANGLEINDEX,
+		MBS_BARYCENTRICCOORDS,
+		MBS_DISTANCE,
+
+		//! Utility
+		MBS_COUNT,
+
+        _MBS_FORCE32BITS = 0xFFFFFFFF
+	} MESH_BUFFER_SEMANTIC;
+
+
+	//!
+	struct MESH_BUFFER_CHANNEL
+	{
+		MESH_BUFFER_CHANNEL()
+		{
+			m_semantic = MBS_NONE;
+			m_format = MBF_NONE;
+			m_componentCount = 0;
+			m_semanticIndex = 0;
+			m_offset = 0;
+		}
+
+		//!
+		MESH_BUFFER_SEMANTIC m_semantic;
+
+		//!
+		MESH_BUFFER_FORMAT m_format;
+
+		//! Index of the semantic, in case there are more than one of this type.
+		uint8 m_semanticIndex;
+
+		//! Offset in bytes from the begining of a buffer element
+		uint8 m_offset;
+
+		//! Number of components of the type in m_format for every value in the channel
+		uint16 m_componentCount;
+
+		//!
+		inline bool operator==(const MESH_BUFFER_CHANNEL& o) const
+		{
+			return (m_semantic == o.m_semantic) &&
+				(m_format == o.m_format) &&
+				(m_semanticIndex == o.m_semanticIndex) &&
+				(m_offset == o.m_offset) &&
+				(m_componentCount == o.m_componentCount);
+		}
+
+	};
+
+	MUTABLE_DEFINE_POD_SERIALISABLE(MESH_BUFFER_CHANNEL);
+	MUTABLE_DEFINE_POD_VECTOR_SERIALISABLE(MESH_BUFFER_CHANNEL);
+	MUTABLE_DEFINE_ENUM_SERIALISABLE(MESH_BUFFER_FORMAT);
+	MUTABLE_DEFINE_ENUM_SERIALISABLE(MESH_BUFFER_SEMANTIC);
+
+
+	struct MESH_BUFFER
+	{
+		//!
+		MESH_BUFFER()
+		{
+			m_elementSize = 0;
+		}
+
+		//!
+		TArray<MESH_BUFFER_CHANNEL> m_channels;
+
+		//!
+		TArray<uint8> m_data;
+
+		//!
+		uint32 m_elementSize;
+
+		//!
+		inline void Serialise(OutputArchive& arch) const
+		{
+			arch << m_channels;
+			arch << m_data;
+			arch << m_elementSize;
+		}
+
+		//!
+		inline void Unserialise(InputArchive& arch)
+		{
+			arch >> m_channels;
+			arch >> m_data;
+			arch >> m_elementSize;
+		}
+
+		//!
+		inline bool operator==(const MESH_BUFFER& o) const
+		{
+			bool equal = (m_channels == o.m_channels);
+			if (equal) equal = (m_elementSize == o.m_elementSize);
+			if (equal) equal = (m_data == o.m_data);
+			return equal;
+		}
+
+	};
+
+
+	//! Set of buffers storing mesh element data. Elements can be vertices, indices or faces.
+	class MUTABLERUNTIME_API FMeshBufferSet : public Base
+	{
+	public:
+
+		//!
+		uint32 m_elementCount = 0;
+
+		//!
+		TArray<MESH_BUFFER> m_buffers;
+
+		//!
+		inline void Serialise(OutputArchive& arch) const
+		{
+			arch << m_elementCount;
+			arch << m_buffers;
+		}
+
+		//!
+		inline void Unserialise(InputArchive& arch)
+		{
+			arch >> m_elementCount;
+			arch >> m_buffers;
+		}
+
+		//!
+		inline bool operator==(const FMeshBufferSet& o) const
+		{
+			return (m_elementCount == o.m_elementCount) &&
+				(m_buffers == o.m_buffers);
+		}
+
+	public:
+
+		//! Get the number of elements in the buffers.
+		int32 GetElementCount() const;
+
+		//! Set the number of vertices in the mesh. This will resize the vertex buffers keeping the
+		//! previous data when possible. New data is undefined.
+		void SetElementCount( int32 );
+
+		//! Get the size in bytes of a buffer element.
+		//! \param buffer index of the buffer from 0 to GetBufferCount()-1
+		int32 GetElementSize( int32 buffer ) const;
+
+		//! Get the number of vertex buffers in the mesh
+		int32 GetBufferCount() const;
+
+		//! Set the number of vertex buffers in the mesh.
+		void SetBufferCount( int32 );
+
+		//! Get the number of channels in a vertex buffer.
+		//! \param buffer index of the vertex buffer from 0 to GetBufferCount()-1
+		int32 GetBufferChannelCount( int32 buffer ) const;
+
+		//! Get a channel of a buffer by index
+		//! \param buffer index of the vertex buffer from 0 to GetBufferCount()-1
+		//! \param channel index of the channel from 0 to GetBufferChannelCount( buffer )-1
+		//! \param[out] pSemantic semantic of the channel
+		//! \param[out] pSemanticIndex index of the semantic in case of having more than one of the
+		//!				same type.
+		//! \param[out] pFormat data format of the channel
+		//! \param[out] pComponentCount components of an element of the channel
+		//! \param[out] pOffset offset in bytes from the beginning of an element of the buffer
+		void GetChannel
+			(
+				int32 buffer,
+				int32 channel,
+				MESH_BUFFER_SEMANTIC* pSemantic,
+				int32* pSemanticIndex,
+				MESH_BUFFER_FORMAT* pFormat,
+				int32* pComponentCount,
+				int32* pOffset
+			) const;
+
+		//! Set all the channels of a buffer
+		//! \param buffer index of the buffer from 0 to GetBufferCount()-1
+        //! \param elementSize sizei n bytes of a vertex element in this buffer
+        //! \param channelCount number of channels to set in the buffer
+        //! \param pSemantics buffer of channelCount semantics
+        //! \param pSemanticIndices buffer of indices for the semantic of every channel
+        //! \param pFormats buffer of channelCount formats
+        //! \param pComponentCounts buffer of channelCount component counts
+        //! \param pOffsets offsets in bytes of every particular channel inside the buffer element
+        void SetBuffer
+			(
+				int32 buffer,
+				int32 elementSize,
+				int32 channelCount,
+				const MESH_BUFFER_SEMANTIC* pSemantics=nullptr,
+				const int32* pSemanticIndices=nullptr,
+				const MESH_BUFFER_FORMAT* pFormats=nullptr,
+				const int32* pComponentCounts=nullptr,
+				const int32* pOffsets=nullptr
+			);
+
+		//! Set one  channels of a buffer
+		//! \param buffer index of the buffer from 0 to GetBufferCount()-1
+        //! \param elementSize sizei n bytes of a vertex element in this buffer
+        //! \param channelIndex number of channels to set in the buffer
+        void SetBufferChannel
+			(
+				int32 buffer,
+				int32 channelIndex,
+				MESH_BUFFER_SEMANTIC semantic,
+				int32 semanticIndice,
+				MESH_BUFFER_FORMAT format,
+				int32 componentCount,
+				int32 offset
+			);
+
+		//! Get a pointer to the object-owned data of a buffer.
+		//! Channel data is interleaved for every element and packed in the order it was set
+		//! without any padding.
+		//! \param buffer index of the buffer from 0 to GetBufferCount()-1
+		//! \todo Add padding support for better alignment of buffer elements.
+        uint8* GetBufferData( int32 buffer );
+        const uint8* GetBufferData( int32 buffer ) const;
+
+		//-----------------------------------------------------------------------------------------
+		// Utility methods
+		//-----------------------------------------------------------------------------------------
+
+		//! Find the index of a buffer channel by semantic and relative index inside the semantic.
+        //! \param semantic Semantic of the channel we are searching.
+        //! \param semanticIndex Index of the semantic of the channel we are searching. e.g. if we
+        //!         want the second set of texture coordinates, it should be 1.
+        //! \param[out] pBuffer -1 if the channel is not found, otherwise it will contain the index
+		//! of the buffer where the channel was found.
+		//! \param[out] pChannel -1 if the channel is not found, otherwise it will contain the
+		//! channel index of the channel inside the buffer returned at [buffer]
+		void FindChannel
+			(
+				MESH_BUFFER_SEMANTIC semantic,
+				int semanticIndex,
+				int* pBuffer, int* pChannel
+			) const;
+
+		//! Get the offset in bytes of the data of this channel inside an element data.
+		//! \param buffer index of the buffer from 0 to GetBufferCount()-1
+		//! \param channel index of the channel from 0 to GetBufferChannelCount( buffer )-1
+		int32 GetChannelOffset( int32 buffer, int32 channel ) const;
+
+		//! Add a new buffer by cloning a buffer from another set.
+		//! The buffers must have the same number of elements.
+		void AddBuffer( const FMeshBufferSet& Other, int32 buffer );
+
+		//! Return true if the formats of the two vertex buffers set match.
+		bool HasSameFormat( const FMeshBufferSet& Other ) const;
+
+		//! Remove the buffer at the specified position. This "invalidates" any buffer index that
+		//! was referencing buffers after the removed one.
+		void RemoveBuffer( int32 BufferIndex);
+
+	public:
+
+		//! Copy an element from one position to another, overwriting the other element.
+		//! Both positions must be valid, buffer size won't change.
+		void CopyElement(uint32_t fromIndex, uint32_t toIndex);
+
+		//! Compare the format of the two buffers at index buffer and return true if they match.
+		bool HasSameFormat(int32 buffer, const FMeshBufferSet& pOther) const;
+
+		//! Get the total memory size of the buffers
+		size_t GetDataSize() const;
+
+		//! Compare the mesh buffer with another one, but ignore internal data like generated
+		//! vertex indices.
+		bool IsSpecialBufferToIgnoreInSimilar(const MESH_BUFFER& b) const;
+
+		//! Compare the mesh buffer with another one, but ignore internal data like generated
+		//! vertex indices. Be aware this method compares the data byte by byte without checking
+		//! if the data belong to the buffer components and could give false negatives if unset 
+		//! padding data is present.
+		bool IsSimilar(const FMeshBufferSet& o) const;
+
+		//! Compare the mesh buffer with another one, but ignore internal data like generated
+		//! vertex indices. This version compares the data component-wise, skipping any memory
+		//! not specified in the buffer description.	
+		bool IsSimilarRobust(const FMeshBufferSet& Other) const;
+
+		//! Change the buffer descriptions so that all buffer indices start at 0 and are in the
+		//! same order than memory.
+		void ResetBufferIndices();
+
+		//!
+		void UpdateOffsets(int32 b);
+	};
+
+}
+
diff --git a/Engine/Plugins/Experimental/Mutable/Source/MutableRuntime/Public/Model.h b/Engine/Plugins/Experimental/Mutable/Source/MutableRuntime/Public/Model.h
new file mode 100644
--- /dev/null
+++ b/Engine/Plugins/Experimental/Mutable/Source/MutableRuntime/Public/Model.h
@@ -0,0 +1,177 @@
+// Copyright Epic Games, Inc. All Rights Reserved.
+
+#pragma once
+
+#include "MutableRuntime/Public/RefCounted.h"
+#include "MutableRuntime/Public/Ptr.h"
+#include "MutableRuntime/Public/Types.h"
+#include "MutableRuntime/Public/Serialisation.h"
+#include "MutableRuntime/Public/Parameters.h"
+
+#include <functional>
+
+
+//! This version number changes whenever there is a compatibility-breaking change in the Model
+//! data structures. Compiled models are not necessarily compatible when the runtime is updated,
+//! so this version number can be used externally to verify this. It is not used internally, and
+//! serializing models from different versions than this runtime will probably result in a crash.
+#define MUTABLE_COMPILED_MODEL_CODE_VERSION		uint32_t( 47 )
+#define MUTABLE_PARAMETERS_VERSION              uint32_t( 1 )
+
+
+namespace mu
+{
+
+	//!
+	class Parameters;
+	typedef Ptr<Parameters> ParametersPtr;
+	typedef Ptr<const Parameters> ParametersPtrConst;
+
+    class Model;
+    typedef Ptr<Model> ModelPtr;
+    typedef Ptr<const Model> ModelPtrConst;
+
+    class ModelParametersGenerator;
+    typedef Ptr<ModelParametersGenerator> ModelParametersGeneratorPtr;
+    typedef Ptr<const ModelParametersGenerator> ModelParametersGeneratorPtrConst;
+
+    class System;
+    typedef Ptr<System> SystemPtr;
+    typedef Ptr<const System> SystemPtrConst;
+
+
+    //! \brief A Model represents a customisable object with any number of parameters.
+    //!
+    //! When values are given to the parameters, specific Instances can be built, which hold the
+    //! built application-usable data.
+	//! \ingroup runtime
+    class MUTABLERUNTIME_API Model : public RefCountedWeak
+	{
+	public:
+
+		//-----------------------------------------------------------------------------------------
+		// Life cycle
+		//-----------------------------------------------------------------------------------------
+		Model();
+
+		static void Serialise( const Model* p, OutputArchive& arch );
+		static ModelPtr StaticUnserialise( InputArchive& arch );
+
+		//! Special serialise operation that serialises the data in separate "files". An object
+        //! with the ModelStreamer interface is responsible of storing this data and providing
+		//! the "file" concept.
+        static void Serialise( Model* p, ModelStreamer& arch );
+
+		//! Return true if the model has external data in other files. This kind of models will
+		//! require data streaming when used.
+		bool HasExternalData() const;
+
+        //! Get the streamer-specific location data.
+        //! This is a non-persistent string with meaning specific to every ModelStreamer
+        //! implementation.
+        const char* GetLocation( ) const;
+        void SetLocation( const char* strLocation );
+
+		//-----------------------------------------------------------------------------------------
+		// Own interface
+		//-----------------------------------------------------------------------------------------
+
+		//! Create a set of new parameters of the model with the default values.
+		//! If old parameters are provided, they will be reused when possible instead of the
+		//! default values.
+        ParametersPtr NewParameters( const Parameters* pOld = nullptr ) const;
+
+		//! Get the number of states in the model.
+		int GetStateCount() const;
+
+		//! Get a state name by state index from 0 to GetStateCount-1
+		const char* GetStateName( int stateIndex ) const;
+
+		//! Find a state index by state name
+		int FindState( const char* strName ) const;
+
+		//! Get the number of parameters available in a particular state.
+		int GetStateParameterCount( int stateIndex ) const;
+
+		//! Get the index of one of the parameters in the given state. The index refers to the
+		//! parameters in a Parameters object obtained from this model with NewParameters.
+		int GetStateParameterIndex( int stateIndex, int paramIndex ) const;
+
+        //! Free memory used by streaming data that may be loaded again when needed.
+        void UnloadExternalData();
+
+        //! Free memory used in internal runtime caches. This is useful for long-running processes
+        //! that keep models loaded. This could be called when a game finishes, or a change of
+        //! level.
+        void ClearCaches();
+
+		//-----------------------------------------------------------------------------------------
+		// Interface pattern
+		//-----------------------------------------------------------------------------------------
+		class Private;
+		Private* GetPrivate() const;
+
+	protected:
+
+		//! Forbidden. Manage with the Ptr<> template.
+        ~Model() override;
+
+	private:
+
+		Private* m_pD;
+
+	};
+
+
+    //! \brief This class enumerates all possible model instances and generates the parameters for
+    //! each of them.
+    //!
+    //! This class can be used to generate all the possible instances of a model, or to generate
+    //! random instances. It only takes into account discrete parameters (bools and ints) but not
+    //! continuous parameters like colours and floats.
+    //!
+    //! The use of this class can be very expensive for complex models, since the number of
+    //! combinations grows geometrically. If considerRelevancy is set to false, a brute force
+    //! approach will be used, but this will generate duplicated instances when changing parameters
+    //! that are not relevant.
+    //!
+    //! \ingroup runtime
+    class MUTABLERUNTIME_API ModelParametersGenerator : public RefCounted
+    {
+    public:
+
+        //!
+        ModelParametersGenerator( const Model* pModel, System* pSystem, bool considerRelevancy=true );
+
+        //! Return the number of different possible instances that can be built from the model.
+        int64_t GetInstanceCount();
+
+        //! Return the parameters of one of the possible instances of the model.
+        //! \param index is an index in the range of 0 to GetInstanceCount-1
+        //! \param randomGenerator if not null, this function will be used to generate random values
+        //! for the continuous parameters of the instance.
+        ParametersPtr GetInstance( int64_t index, float (*randomGenerator )() );
+
+        //! Return the parameters of one of a random instance of the model.
+        //! \param randomGenerator This function will be used to generate random values
+        //! for the continuous parameters of the instance.
+        ParametersPtr GetRandomInstance( std::function<float()> randomGenerator );
+
+        //-----------------------------------------------------------------------------------------
+        // Interface pattern
+        //-----------------------------------------------------------------------------------------
+        class Private;
+        Private* GetPrivate() const;
+
+    protected:
+
+        //! Forbidden. Manage with the Ptr<> template.
+        ~ModelParametersGenerator();
+
+    private:
+
+        Private* m_pD;
+    };
+
+}
+
diff --git a/Engine/Plugins/Experimental/Mutable/Source/MutableRuntime/Public/MutableMemory.h b/Engine/Plugins/Experimental/Mutable/Source/MutableRuntime/Public/MutableMemory.h
new file mode 100644
--- /dev/null
+++ b/Engine/Plugins/Experimental/Mutable/Source/MutableRuntime/Public/MutableMemory.h
@@ -0,0 +1,65 @@
+// Copyright Epic Games, Inc. All Rights Reserved.
+
+#pragma once
+
+#include "CoreTypes.h"
+
+#include <cstdint>
+#include <cstring>	// for size_t in some platforms
+
+namespace mu
+{
+	//! \defgroup runtime Run-time classes
+    //! This module contains the only classes necessary to integrate the system into an application
+    //! loading models and building instances by setting parameter values.
+
+    //! \brief Set the memory management methods for the library.
+    //! The library will always manage the heap memory using these methods.
+	//! This should be called only before starting to use the library or any of its data.
+	//! If this method is not called, the standard C library malloc and free will be used.
+	//! \todo: Notes about thread safety
+	//! \param customMalloc Function pointer to a function taking the size in bytes to allocate.
+    //!        The second parameter is the required alignment in bytes.
+	//! 	   This function must return a valid heap pointer or 0 in case of allocation failure.
+	//! \param customFree Function pointer to a function taking a pointer obtained reviously with
+	//! 	   customMalloc, that needs to be freed.
+	//! \ingroup runtime
+	MUTABLERUNTIME_API extern void Initialize
+		(
+            void* (*customMalloc)( size_t, uint32_t ),
+            void (*customFree)( void* )
+		);
+
+    //! \brief Shutdown everything related to mutable.
+    //! This should be called when no other mutable objects will be used, created or destroyed,
+    //! and all threads involving mutable tasks have been terminated.
+	MUTABLERUNTIME_API extern void Finalize();
+
+	//! Memory management functions to be used inside the library. No other memory allocation is
+	//! allowed.
+	MUTABLERUNTIME_API extern void* mutable_malloc(size_t size);
+	MUTABLERUNTIME_API extern void* mutable_malloc_aligned(size_t size, uint32_t alignment);
+	MUTABLERUNTIME_API extern void mutable_free(void* ptr);
+	MUTABLERUNTIME_API extern void mutable_free(void* ptr, size_t size);
+	MUTABLERUNTIME_API extern void mutable_free_aligned(void* ptr, size_t size);
+
+    //! \brief %Base class from which all library classes inherit to control memory allocation.
+	//! For internal use only.
+	//! \ingroup runtime
+	class MUTABLERUNTIME_API Base
+	{
+	public:
+
+		inline void* operator new(size_t s)
+		{
+			return mutable_malloc(s);
+		}
+
+		inline void operator delete(void* p, size_t size)
+		{
+			mutable_free(p, size);
+		}
+    };
+
+}
+
diff --git a/Engine/Plugins/Experimental/Mutable/Source/MutableRuntime/Public/MutableRuntimeModule.h b/Engine/Plugins/Experimental/Mutable/Source/MutableRuntime/Public/MutableRuntimeModule.h
new file mode 100644
--- /dev/null
+++ b/Engine/Plugins/Experimental/Mutable/Source/MutableRuntime/Public/MutableRuntimeModule.h
@@ -0,0 +1,27 @@
+// Copyright Epic Games, Inc. All Rights Reserved.
+
+#pragma once
+
+#include "Modules/ModuleInterface.h"
+
+
+// This will use unreal's malloc instead of system malloc for mutable memory operations
+#define USE_UNREAL_ALLOC_IN_MUTABLE				1
+
+// This will enable the internal memory manager for mutable operations, reducing the number of operations
+// but increasing the memory usage slightly. This only affects packaged builds, and it is not used in editor.
+#define USE_INTERNAL_MEMORY_MANAGER_IN_MUTABLE	1
+
+// This slows down mutable, but gathers some data about memory usage for debug
+#define USE_STAT_MUTABLE_MEMORY					0
+
+
+class MUTABLERUNTIME_API FMutableRuntimeModule : public IModuleInterface
+{
+public:
+
+	// IModuleInterface 
+	virtual void StartupModule() override;
+	virtual void ShutdownModule() override;
+
+};
diff --git a/Engine/Plugins/Experimental/Mutable/Source/MutableRuntime/Public/MutableString.h b/Engine/Plugins/Experimental/Mutable/Source/MutableRuntime/Public/MutableString.h
new file mode 100644
--- /dev/null
+++ b/Engine/Plugins/Experimental/Mutable/Source/MutableRuntime/Public/MutableString.h
@@ -0,0 +1,80 @@
+// Copyright Epic Games, Inc. All Rights Reserved.
+
+#pragma once
+
+#include "MutableRuntime/Public/Types.h"
+#include "MutableRuntime/Public/Ptr.h"
+#include "MutableRuntime/Public/RefCounted.h"
+#include "MutableRuntime/Public/Serialisation.h"
+#include "MutableRuntime/Private/SerialisationPrivate.h"
+#include "MutableRuntime/Private/Platform.h"
+#include "MutableRuntime/Private/MemoryPrivate.h"
+
+namespace mu
+{
+
+    //! \brief Container for text string data.
+	//! \ingroup runtime
+    class MUTABLERUNTIME_API String : public RefCounted
+	{
+	public:
+
+		//-----------------------------------------------------------------------------------------
+		// Life cycle
+		//-----------------------------------------------------------------------------------------
+
+		String(const char* = nullptr);
+
+		//! Deep clone this string.
+		Ptr<String> Clone() const;
+
+		//! Serialisation
+		static void Serialise( const String* p, OutputArchive& arch );
+		static Ptr<String> StaticUnserialise( InputArchive& arch );
+
+		//-----------------------------------------------------------------------------------------
+		// Own interface
+		//-----------------------------------------------------------------------------------------
+
+		//! Get the string data.
+		const char* GetValue() const;
+
+	protected:
+
+		//! Forbidden. Manage with the Ptr<> template.
+		~String() {}
+
+	public:
+
+		//!
+		string m_value;
+
+		//!
+		void Serialise(OutputArchive& arch) const
+		{
+			uint32_t ver = 0;
+			arch << ver;
+
+			arch << m_value;
+		}
+
+		//!
+		void Unserialise(InputArchive& arch)
+		{
+			uint32_t ver;
+			arch >> ver;
+			check(ver <= 0);
+
+			arch >> m_value;
+		}
+
+		//!
+		inline bool operator==(const String& o) const
+		{
+			return (m_value == o.m_value);
+		}
+
+	};
+
+}
+
diff --git a/Engine/Plugins/Experimental/Mutable/Source/MutableRuntime/Public/MutableTrace.h b/Engine/Plugins/Experimental/Mutable/Source/MutableRuntime/Public/MutableTrace.h
new file mode 100644
--- /dev/null
+++ b/Engine/Plugins/Experimental/Mutable/Source/MutableRuntime/Public/MutableTrace.h
@@ -0,0 +1,11 @@
+// Copyright Epic Games, Inc. All Rights Reserved.
+
+#pragma once
+
+#include "Trace/Trace.h"
+
+UE_TRACE_CHANNEL_EXTERN(MutableChannel, MUTABLERUNTIME_API)
+
+/** Custom Mutable profiler scope. */
+#define MUTABLE_CPUPROFILER_SCOPE(Name) \
+	TRACE_CPUPROFILER_EVENT_SCOPE_ON_CHANNEL(Mutable_##Name, MutableChannel)
diff --git a/Engine/Plugins/Experimental/Mutable/Source/MutableRuntime/Public/Parameters.h b/Engine/Plugins/Experimental/Mutable/Source/MutableRuntime/Public/Parameters.h
new file mode 100644
--- /dev/null
+++ b/Engine/Plugins/Experimental/Mutable/Source/MutableRuntime/Public/Parameters.h
@@ -0,0 +1,411 @@
+// Copyright Epic Games, Inc. All Rights Reserved.
+
+#pragma once
+
+#include "MutableRuntime/Public/RefCounted.h"
+#include "MutableRuntime/Public/Ptr.h"
+#include "MutableRuntime/Public/Types.h"
+#include "MutableRuntime/Public/Serialisation.h"
+
+#define MUTABLE_MAX_STRING_PARAM_LENGTH     40
+
+namespace mu
+{
+
+	//!
+	class Parameters;
+    using ParametersPtr = Ptr<Parameters>;
+    using ParametersPtrConst = Ptr<const Parameters>;
+
+    class RangeIndex;
+	using RangeIndexPtr = Ptr<RangeIndex>;
+	using RangeIndexPtrConst = Ptr<const RangeIndex>;
+
+    class Image;
+	class Model;
+
+
+    //! \brief Model parameter types.
+	//! \ingroup runtime
+    enum class PARAMETER_TYPE : uint32
+    {
+        //! Undefined parameter type.
+        T_NONE,
+
+        //! Boolean parameter type (true or false)
+        T_BOOL,
+
+        //! Integer parameter type. It usually has a limited range of possible values that can be
+        //! queried in the Parameters object.
+        T_INT,
+
+        //! Floating point value in the range of 0.0 to 1.0
+        T_FLOAT,
+
+        //! Floating point RGBA colour, with each channel ranging from 0.0 to 1.0
+        T_COLOUR,
+
+        //! 3D Projector type, defining a position, scale and orientation. Basically used for
+        //! projected decals.
+        T_PROJECTOR,
+
+        //! An externally provided image.
+        T_IMAGE,
+
+        //! A text string.
+        T_STRING,
+
+        //! Utility enumeration value, not really a parameter type.
+        T_COUNT
+
+    };
+
+
+    //! \brief Additional information about a particular parameter type.
+    //!
+	//! This information is not strictly relevant, but it can provide hints for dynamic user
+	//! interfaces.
+	enum class PARAMETER_DETAILED_TYPE : uint32
+	{
+		//!
+		UNKNOWN,
+
+		//! The parameter belongs to a triangular barycentric set of parameters
+		//! The parameter name can be used to identify which coordinate it is (ending number)
+		TRIANGLE
+
+	};
+
+
+    //! \brief Types of projectors.
+    //! \ingroup runtime
+    enum class PROJECTOR_TYPE : uint32
+    {
+        //!
+        PLANAR,
+
+        //!
+        CYLINDRICAL,
+
+		//!
+        WRAPPING,
+
+        //! Utility enumeration value, not really a projector type.
+        COUNT
+    };
+
+
+    //! \brief Class used to set or read multi-dimensional parameter values.
+    //! If parameters have multiple values because of ranges, RangeIndex can be used to specify
+    //! what value is read or set in the methods of the Parameter class.
+    //! RangeIndex objects can be reused in multiple calls for the same parameter in the Parameter
+    //! class interface even after changing the position values.
+    //! \ingroup runtime
+    class MUTABLERUNTIME_API RangeIndex : public RefCounted
+    {
+    public:
+
+        //-----------------------------------------------------------------------------------------
+        // Life cycle
+        //-----------------------------------------------------------------------------------------
+
+
+        //-----------------------------------------------------------------------------------------
+        // Own interface
+        //-----------------------------------------------------------------------------------------
+
+        //! Return the number of ranges (or dimensions) used by this index
+        int GetRangeCount() const;
+
+        //! Return the name of a range.
+        //! \param index Index of the range from 0 to GetRangeCount()-1
+        const char* GetRangeName( int index ) const;
+
+        //! Return the Guid of the parameter, resistant to parameter name changes.
+        //! \param index Index of the parameter from 0 to GetRangeCount()-1
+        const char* GetRangeUid( int index ) const;
+
+        //! Set the position in one of the dimensions of the index.
+        //! \param index Index of the parameter from 0 to GetRangeCount()-1
+        void SetPosition( int index, int position );
+
+        //! Return the position in one of the dimensions of the index.
+        //! \param index Index of the parameter from 0 to GetRangeCount()-1
+        int GetPosition( int index ) const;
+
+
+        //-----------------------------------------------------------------------------------------
+        // Interface pattern
+        //-----------------------------------------------------------------------------------------
+        class Private;
+        Private* GetPrivate() const;
+
+    protected:
+
+        //! Forbid creation. Instances of this class must be obtained from a Parameters instance
+        RangeIndex();
+
+        //! Forbidden. Manage with the Ptr<> template.
+        ~RangeIndex() override;
+
+    private:
+
+        friend class Parameters;
+        Private* m_pD;
+
+    };
+
+
+    //! \brief  This class represents the parameters of a model including description, type and
+    //! value, and additional resources associated to the parameter.
+    //!
+    //! \warning Every object of this class holds a reference to the model that it was created from.
+    //! This implies that while any instance of a Parameters object is alive, its model will not
+    //! be freed.
+    //!
+	//! \ingroup runtime
+	class MUTABLERUNTIME_API Parameters : public RefCounted
+	{
+	public:
+
+		//-----------------------------------------------------------------------------------------
+		// Life cycle
+		//-----------------------------------------------------------------------------------------
+
+        //! Serialisation
+        static void Serialise( const Parameters* p, OutputArchive& arch );
+        static ParametersPtr StaticUnserialise( InputArchive& arch );
+        void SetUnserialisedModel( Model* pModel );
+
+        //! Portable serialisation, which tries to keep values when:
+        //! - new parametes have been added
+        //! - parameters have been removed
+        //! - parameter types have changed
+        //! - parameter names have changed, and the application is using UIDs
+        static void SerialisePortable( const Parameters* p, OutputArchive& arch );
+        static ParametersPtr UnserialisePortable( InputArchive& arch, const Model* pModel );
+
+		//! Deep clone this object.
+		ParametersPtr Clone() const;
+
+		//-----------------------------------------------------------------------------------------
+		// Own interface
+		//-----------------------------------------------------------------------------------------
+
+		//! Return the number of parameters
+		int GetCount() const;
+
+		//! Return the name of the parameter.
+		//! \param index Index of the parameter from 0 to GetCount()-1
+		const char* GetName( int index ) const;
+
+		//! Return the Guid of the parameter, resistant to parameter name changes.
+		//! \param index Index of the parameter from 0 to GetCount()-1
+		const char* GetUid( int index ) const;
+
+		//! Find the parameter index by name.
+		//! It returns -1 if the parameter doesn't exist.
+		int Find( const char* strName ) const;
+
+		//! Return the type of the parameter.
+		//! \param index Index of the parameter from 0 to GetCount()-1
+		PARAMETER_TYPE GetType( int index ) const;
+
+		//! Return the additional type data of the parameter.
+		//! \param index Index of the parameter from 0 to GetCount()-1
+		PARAMETER_DETAILED_TYPE GetDetailedType( int index ) const;
+
+        //! Create a new RangeIndex object to use to access a multi-dimensional parameter.
+        //! It will return nullptr if the parameter is not multidimensional.
+        Ptr<RangeIndex> NewRangeIndex( int paramIndex ) const;
+
+        //! Return the number of values other than the default that have been set to a specific
+        //! parameter.
+        int GetValueCount( int paramIndex ) const;
+
+        //! Return the RangeIndex of a value that has been set to a parameter.
+        //! \param paramIndex Index of the parameter from 0 to GetCount()-1
+        //! \param valueIndex Index of the value from 0 to GetValueCount()-1
+		Ptr<RangeIndex> GetValueIndex( int paramIndex, int valueIndex ) const;
+
+        //! Remove all the multidimensional values for a parameter. The current non-dimensional
+        //! value is kept.
+        void ClearAllValues( int paramIndex );
+
+		//! Return the value of a boolean parameter.
+		//! \pre The parameter specified by index is a T_BOOL.
+		//! \param index Index of the parameter from 0 to GetCount()-1
+        //! \param pos Only for multidimensional parametres: relevant position to get in the ranges
+        bool GetBoolValue( int index,
+                           const Ptr<const RangeIndex>& pos=nullptr ) const;
+
+		//! Set the value of a parameter of type boolean.
+		//! \pre The parameter specified by index is a T_BOOL.
+		//! \param index Index of the parameter from 0 to GetCount()-1
+		//! \param value new value of the parameter
+        //! \param pos optional parameter to set a specific value for the given multidimensional
+        //! position. If null, the value is set for all possible positions.
+        void SetBoolValue( int index, bool value,
+                           const Ptr<const RangeIndex>& pos=nullptr );
+
+		//! Get the number of possible values for this integer parameter.
+		//! If the number is zero, it means any integer value is accepted.
+		int GetIntPossibleValueCount( int paramIndex ) const;
+
+		//! Get the value of one of the possible values for this integer.
+		//! The paramIndex is in the reange of 0 to GetIntPossibleValueCount()-1
+		int GetIntPossibleValue( int paramIndex, int valueIndex ) const;
+
+        //! Get the name of one of the possible values for this integer.
+        //! The paramIndex is in the reange of 0 to GetIntPossibleValueCount()-1
+        const char* GetIntPossibleValueName( int paramIndex, int valueIndex ) const;
+
+        //! Get the index of the value of one of the possible values for this integer.
+        //! The paramIndex is in the reange of 0 to GetIntPossibleValueCount()-1
+		int GetIntValueIndex(int paramIndex, const char* valueName) const;
+
+		//! Get the index of the value of one of the possible values for this integer.
+		//! The paramIndex is in the reange of 0 to GetIntPossibleValueCount()-1
+		int GetIntValueIndex(int paramIndex, int32 Value) const;
+
+		//! Return the value of a integer parameter.
+		//! \pre The parameter specified by index is a T_INT.
+		//! \param index Index of the parameter from 0 to GetCount()-1
+        //! \param pos Only for multidimensional parametres: relevant position to get in the ranges
+        int GetIntValue( int index,
+                         const Ptr<const RangeIndex>& pos=nullptr ) const;
+
+		//! If the parameter is of the integer type, set its value.
+		//! \param index Index of the parameter from 0 to GetCount()-1
+		//! \param value new value of the parameter. It must be in the possible values for this
+		//! parameter (see GetIntPossibleValue), or the method will leave the value unchanged.
+        //! \param pos Only for multidimensional parametres: relevant position to set in the ranges
+        void SetIntValue( int index, int value,
+                          const Ptr<const RangeIndex>& pos=nullptr );
+
+		//! Return the value of a float parameter.
+		//! \pre The parameter specified by index is a T_FLOAT.
+		//! \param index Index of the parameter from 0 to GetCount()-1
+        //! \param pos Only for multidimensional parametres: relevant position to get in the ranges
+        float GetFloatValue( int index,
+                             const Ptr<const RangeIndex>& pos=nullptr ) const;
+
+		//! If the parameter is of the float type, set its value.
+		//! \param index Index of the parameter from 0 to GetCount()-1
+		//! \param value new value of the parameter
+        //! \param pos Only for multidimensional parametres: relevant position to set in the ranges
+        void SetFloatValue( int index, float value,
+                            const Ptr<const RangeIndex>& pos=nullptr );
+
+		//! Return the value of a colour parameter.
+		//! \pre The parameter specified by index is a T_FLOAT.
+        //! \param index Index of the parameter from 0 to GetCount()-1
+        //! \param pR,pG,pB Pointers to values where every resulting colour channel will be stored
+        //! \param pos Only for multidimensional parametres: relevant position to get in the ranges
+        void GetColourValue( int index, float* pR, float* pG, float* pB,
+                             const Ptr<const RangeIndex>& pos=nullptr ) const;
+
+		//! If the parameter is of the colour type, set its value.
+		//! \param index Index of the parameter from 0 to GetCount()-1
+        //! \param r,g,b new value of the parameter
+        //! \param pos Only for multidimensional parametres: relevant position to set in the ranges
+        void SetColourValue( int index, float r, float g, float b,
+                             const Ptr<const RangeIndex>& pos=nullptr );
+
+		//! Return the value of a projector parameter, as a 4x4 matrix. The matrix is supposed to be
+		//! a linear transform in column-major.
+		//! \pre The parameter specified by index is a T_PROJECTOR.
+        //! \param index Index of the parameter from 0 to GetCount()-1
+        //! \param pPosX,pPosY,pPosZ Pointers to where the object-space position coordinates of the
+        //!         projector will be stored.
+        //! \param pDirX,pDirY,pDirZ Pointers to where the object-space direction vector of the
+        //!         projector will be stored.
+        //! \param pUpX,pUpY,pUpZ Pointers to where the object-space vertically up direction vector
+        //!         of the projector will be stored. This controls the "roll" angle of the
+        //!         projector.
+        //! \param pScaleX, pScaleY Pointers to the projector-space scaling of the projector.
+        //! \param pos Only for multidimensional parametres: relevant position to get in the ranges
+        void GetProjectorValue( int index,
+                                PROJECTOR_TYPE* pProjectionType,
+								float* pPosX, float* pPosY, float* pPosZ,
+								float* pDirX, float* pDirY, float* pDirZ,
+								float* pUpX, float* pUpY, float* pUpZ,
+                                float* pScaleX, float* pScaleY, float* pScaleZ,
+                                float* pProjectionAngle,
+                                const Ptr<const RangeIndex>& pos=nullptr ) const;
+
+		//! If the parameter is of the projector type, set its value.
+		//! \param index Index of the parameter from 0 to GetCount()-1
+        //! \param posX,posY,posZ Object-space position coordinates of the projector.
+        //! \param dirX,dirY,dirZ Object-space direction vector of the projector.
+        //! \param upX,upY,upZ Object-space vertically up direction vector of the projector.
+        //! \param scaleX, scaleY Projector-space scaling of the projector.
+        //! \param projectionAngle [only for Cylindrical projectors], the angle in radians of the
+        //! projection area on the cylinder surface.
+        //! \param pos Only for multidimensional parametres: relevant position to set in the ranges
+        void SetProjectorValue( int index,
+                                float posX, float posY, float posZ,
+								float dirX, float dirY, float dirZ,
+								float upX, float upY, float upZ,
+                                float scaleX, float scaleY, float scaleZ,
+                                float projectionAngle,
+                                const Ptr<const RangeIndex>& pos=nullptr );
+
+        //! Return the value of an image parameter.
+        //! \pre The parameter specified by index is a T_IMAGE.
+        //! \param index Index of the parameter from 0 to GetCount()-1
+        //! \return The externalId specified when setting the image value (\see SetImageValue)
+        EXTERNAL_IMAGE_ID GetImageValue( int index ) const;
+
+        //! If the parameter is of the image type, set its value.
+        //! \param index Index of the parameter from 0 to GetCount()-1
+        //! \param externalId Application-specific id used to identify this image during replication.
+        void SetImageValue( int index, EXTERNAL_IMAGE_ID externalId );
+
+        //! Return the value of a float parameter.
+        //! \pre The parameter specified by index is a T_FLOAT.
+        //! \param index Index of the parameter from 0 to GetCount()-1
+        //! \param pos Only for multidimensional parametres: relevant position to get in the ranges
+        const char* GetStringValue( int index, const Ptr<const RangeIndex>& pos = nullptr ) const;
+
+        //! If the parameter is of the float type, set its value.
+        //! \param index Index of the parameter from 0 to GetCount()-1
+        //! \param value new value of the parameter
+        //! \param pos Only for multidimensional parametres: relevant position to set in the ranges
+        void SetStringValue( int index, const char* value, const Ptr<const RangeIndex>& pos = nullptr );
+
+        //! Get the number of additional description images defined in the parameter.
+		//! The use of these depends on the game, but it can be used to provide colour bars, etc.
+		//! These images can be built using a system object and the model of these parameters.
+		int GetAdditionalImageCount( int index ) const;
+
+        //! Utility method to compare the values of a specific parameter with the values of another
+        //! Parameters object. It returns false type or values are different.
+        bool HasSameValue( int thisParamIndex,
+                           const ParametersPtrConst& other,
+                           int otherParamIndex ) const;
+
+		//-----------------------------------------------------------------------------------------
+		// Interface pattern
+		//-----------------------------------------------------------------------------------------
+		class Private;
+		Private* GetPrivate() const;
+
+	protected:
+
+		//! Forbid creation. Instances of this class must be obtained from a Model instance
+		Parameters();
+
+		//! Forbidden. Manage with the Ptr<> template.
+		~Parameters();
+
+	private:
+
+		friend class Model;
+		Private* m_pD;
+
+	};
+
+
+}
+
diff --git a/Engine/Plugins/Experimental/Mutable/Source/MutableRuntime/Public/PhysicsBody.h b/Engine/Plugins/Experimental/Mutable/Source/MutableRuntime/Public/PhysicsBody.h
new file mode 100644
--- /dev/null
+++ b/Engine/Plugins/Experimental/Mutable/Source/MutableRuntime/Public/PhysicsBody.h
@@ -0,0 +1,462 @@
+// Copyright Epic Games, Inc. All Rights Reserved.
+
+#pragma once
+
+#include "MutableRuntime/Public/Types.h"
+#include "MutableRuntime/Public/Ptr.h"
+#include "MutableRuntime/Public/RefCounted.h"
+#include "MutableRuntime/Public/Serialisation.h"
+#include "MutableRuntime/Private/SerialisationPrivate.h"
+
+namespace mu
+{
+	struct FBodyShape
+	{
+		string Name;
+		uint32 Flags = 0;
+
+		bool operator==(const FBodyShape& Other) const 
+		{
+			return Flags == Other.Flags && Name == Other.Name;
+		}
+	
+		//!
+		inline void Serialise( OutputArchive& arch ) const
+		{
+			const uint32_t ver = 0;
+			arch << ver;
+
+			arch << Name;
+			arch << Flags;
+		}
+
+		//!
+		inline void Unserialise( InputArchive& arch )
+		{
+			uint32_t ver;
+			arch >> ver;
+			check(ver == 0);
+
+			arch >> Name;
+			arch >> Flags;
+		}
+	};
+
+	struct FSphereBody : FBodyShape
+	{
+		FVector3f Position = FVector3f::ZeroVector;
+		float Radius = 0.0f;
+
+		bool operator==( const FSphereBody& Other ) const
+		{
+			return FBodyShape::operator==(Other) && Position == Other.Position && Radius == Other.Radius;
+		}
+		 
+		//!
+		inline void Serialise( OutputArchive& arch ) const
+		{
+			FBodyShape::Serialise(arch);
+
+			const uint32_t ver = 0;
+			arch << ver;
+
+			arch << Position;
+			arch << Radius;
+		}
+
+		//!
+		inline void Unserialise( InputArchive& arch )
+		{
+			FBodyShape::Unserialise(arch);
+
+			uint32_t ver;
+			arch >> ver;
+			check(ver == 0);
+
+			arch >> Position;
+			arch >> Radius;
+		} 
+	};
+	
+	struct FBoxBody : FBodyShape
+	{
+		FVector3f Position = FVector3f::ZeroVector;
+		FQuat4f Orientation = FQuat4f::Identity;
+		FVector3f Size = FVector3f::ZeroVector;
+		 
+		inline bool operator==( const FBoxBody& Other ) const 
+		{
+			return FBodyShape::operator==(Other) && Position == Other.Position && Orientation == Other.Orientation && Size == Other.Size;
+		}
+
+		//!
+		inline void Serialise( OutputArchive& arch ) const
+		{
+			FBodyShape::Serialise( arch );
+
+			const uint32_t ver = 0;
+			arch << ver;
+
+			arch << Position;
+			arch << Orientation;
+			arch << Size;
+		}
+
+		  //!
+		inline void Unserialise( InputArchive& arch )
+		{
+			FBodyShape::Unserialise( arch );
+
+			uint32_t ver;
+			arch >> ver;
+			check(ver == 0);
+
+			arch >> Position;
+			arch >> Orientation;
+			arch >> Size;
+		} 
+	};
+
+	struct FSphylBody : FBodyShape
+	{
+		FVector3f Position = FVector3f::ZeroVector;
+		FQuat4f Orientation = FQuat4f::Identity;
+		float Radius = 0.0f;
+		float Length = 0.0f;
+
+		inline bool operator==( const FSphylBody& Other ) const 
+		{
+			return FBodyShape::operator==(Other) && 
+				Position == Other.Position && 
+				Orientation == Other.Orientation && 
+				Radius == Other.Radius && 
+				Length == Other.Length;
+		}
+		 
+		//!
+		inline void Serialise( OutputArchive& arch ) const
+		{
+			FBodyShape::Serialise( arch );
+
+			const uint32_t ver = 0;
+			arch << ver;
+
+			arch << Position;
+			arch << Orientation;
+			arch << Radius;
+			arch << Length;
+		}
+
+		//!
+		inline void Unserialise( InputArchive& arch )
+		{
+			FBodyShape::Unserialise( arch );
+
+			uint32_t ver;
+			arch >> ver;
+			check(ver == 0);
+
+			arch >> Position;
+			arch >> Orientation;
+			arch >> Radius;
+			arch >> Length;
+		}
+	};
+	
+	struct FTaperedCapsuleBody : FBodyShape
+	{
+		FVector3f Position;
+		FQuat4f Orientation;
+		float Radius0 = 0.0f;
+		float Radius1 = 0.0f;
+		float Length = 0.0f;
+		 
+		inline bool operator==( const FTaperedCapsuleBody& Other ) const 
+		{
+			return FBodyShape::operator==(Other) && 
+				Position == Other.Position && 
+				Orientation == Other.Orientation && 
+				Radius0 == Other.Radius0 && 
+				Radius1 == Other.Radius1 && 
+				Length == Other.Length;
+		}
+		 
+		//!
+		inline void Serialise( OutputArchive& arch ) const
+		{
+			FBodyShape::Serialise( arch );
+
+			const uint32_t ver = 0;
+			arch << ver;
+
+			arch << Position;
+			arch << Orientation;
+			arch << Radius0;
+			arch << Radius1;
+			arch << Length;
+		}
+
+		 //!
+		inline void Unserialise( InputArchive& arch )
+		{
+			FBodyShape::Unserialise( arch );
+			
+			uint32_t ver;
+			arch >> ver;
+			check(ver == 0);
+
+			arch >> Position;
+			arch >> Orientation;
+			arch >> Radius0;
+			arch >> Radius1;
+			arch >> Length;
+		 }
+	};
+	
+	struct FConvexBody : FBodyShape
+	{
+		TArray<FVector3f> Vertices;
+		TArray<int32> Indices;
+
+		FTransform3f Transform = FTransform3f::Identity;
+		
+		inline bool operator==( const FConvexBody& Other ) const
+		{
+			return FBodyShape::operator==(Other) && 
+				Vertices == Other.Vertices && 
+				Indices == Other.Indices && 
+				Transform.Equals(Other.Transform);
+		}
+
+		//!
+		inline void Serialise( OutputArchive& arch ) const
+		{
+			FBodyShape::Serialise( arch );
+	
+		 	uint32_t ver = 0;
+		 	arch << ver;
+
+		 	arch << Vertices;
+		 	arch << Indices;
+		 	arch << Transform;
+		}
+
+		  //!
+		inline void Unserialise( InputArchive& arch )
+		{
+		 	uint32_t ver;
+		 	arch >> ver;
+		 	check(ver == 0);
+
+		 	arch >> Vertices;
+		 	arch >> Indices;
+		 	arch >> Transform;
+		}
+	};
+	
+	struct FPhysicsBodyAggregate
+	{
+		TArray<FSphereBody> Spheres;
+		TArray<FBoxBody> Boxes;
+		TArray<FConvexBody> Convex;
+		TArray<FSphylBody> Sphyls;
+		TArray<FTaperedCapsuleBody> TaperedCapsules;
+
+		inline bool operator==( const FPhysicsBodyAggregate& Other ) const
+		{
+			return Spheres == Other.Spheres &&
+					Boxes == Other.Boxes &&
+					Sphyls == Other.Sphyls &&
+					TaperedCapsules == Other.TaperedCapsules &&
+					Convex == Other.Convex;
+		}
+
+		 //!
+		inline void Serialise( OutputArchive& arch ) const
+		{
+		 	uint32_t ver = 0;
+		 	arch << ver;
+		
+		 	arch << Spheres;
+		 	arch << Boxes;
+		 	arch << Convex;
+		 	arch << Sphyls;
+		 	arch << TaperedCapsules;
+		}
+
+		 //!
+		inline void Unserialise( InputArchive& arch )
+		{
+		 	uint32_t ver;
+		 	arch >> ver;
+		 	check(ver == 0);
+		
+		 	arch >> Spheres;
+		 	arch >> Boxes;
+		 	arch >> Convex;
+		 	arch >> Sphyls;
+		 	arch >> TaperedCapsules;
+		}
+	};
+
+	class MUTABLERUNTIME_API PhysicsBody : public RefCounted
+	{
+	public:
+
+		PhysicsBody() 
+		{
+		}
+		
+		mu::Ptr<PhysicsBody> Clone() const; 
+
+		//! Serialisation
+        static void Serialise( const PhysicsBody* p, OutputArchive& arch );
+		static mu::Ptr<PhysicsBody> StaticUnserialise( InputArchive& arch );
+
+        //bool operator==( const PhysicsBody& Other ) const;
+
+		void SetBodyCount( int32 B );
+		int32 GetBodyCount( ) const;
+
+		void SetBodyBoneName( int32 B, const char* BoneName );
+		const char* GetBodyBoneName( int32 B ) const;
+	
+		void SetBodyCustomId( int32 B, int32 CustomId );
+		int32 GetBodyCustomId( int32 B ) const;
+		
+		int32 GetSphereCount( int32 B ) const;
+		int32 GetBoxCount( int32 B) const;
+		int32 GetConvexCount( int32 B) const;
+		int32 GetSphylCount( int32 B ) const;
+		int32 GetTaperedCapsuleCount(int32 B ) const;
+		
+		void SetSphereCount(int32 B, int32 Count);
+		void SetBoxCount(int32 B, int32 Count);
+		void SetConvexCount(int32 B, int32 Count);
+		void SetSphylCount(int32 B, int32 Count);
+		void SetTaperedCapsuleCount( int32 B, int32 Count);
+		
+		void SetSphere( 
+				int32 B, int32 I, 
+				FVector3f Position, float Radius );
+
+		void SetBox( 
+				int32 B, int32 I, 
+				FVector3f Position, FQuat4f Orientation, FVector3f Size );
+
+		void SetConvex( 
+				int32 B, int32 I,
+				const FVector3f* Vertices, int32 VerticesCount, 
+				const int32* Indices, int32 IndicesCount, 
+				const FTransform3f& Transform );
+
+		void SetSphyl( 
+				int32 B, int32 I, 
+				FVector3f Position, FQuat4f Orientation, 
+				float Radius, float Length );
+
+		void SetTaperedCapsule( 
+				int32 B, int32 I, 
+				FVector3f Position, FQuat4f Orientation, 
+				float Radius0, float Radius1, float Length );
+
+		void SetSphereFlags( int32 B, int32 I, uint32 Flags );
+		void SetBoxFlags( int32 B, int32 I, uint32 Flags );
+		void SetConvexFlags( int32 B, int32 I, uint32 Flags );
+		void SetSphylFlags( int32 B, int32 I, uint32 Flags );
+		void SetTaperedCapsuleFlags( int32 B, int32 I, uint32 Flags );
+
+		void SetSphereName( int32 B, int32 I, const char* Name );
+		void SetBoxName( int32 B, int32 I, const char* Name );
+		void SetConvexName( int32 B, int32 I, const char* Name );
+		void SetSphylName( int32 B, int32 I, const char* Name );
+		void SetTaperedCapsuleName( int32 B, int32 I, const char* Name );
+	
+		void GetSphere( 
+				int32 B, int32 I, 
+				FVector3f& OutPosition, float& OutRadius) const;
+
+		void GetBox( 
+				int32 B, int32 I, 
+				FVector3f& OutPosition, FQuat4f& OutOrientation, FVector3f& OutSize ) const;
+
+		void GetConvex( 
+				int32 B, int32 I, 
+				FVector3f const*& OutVertices, int32& OutVerticesCount, 
+				int32 const*& OutIndices, int32& OutIndicesCount, FTransform3f& OutTransform ) const;
+
+		void GetSphyl( 
+				int32 B, int32 I, 
+				FVector3f& OutPosition, FQuat4f& OutOrientation, 
+				float& OutRadius, float& OutLength) const;
+
+		void GetTaperedCapsule( 
+				int32 B, int32 I, 
+				FVector3f& OutPosition, FQuat4f& OutOrientation, 
+				float& OutRadius0, float& OutRadius1, float& OutLength) const;
+		
+		uint32 GetSphereFlags( int32 B, int32 I ) const;
+		uint32 GetBoxFlags( int32 B, int32 I ) const;
+		uint32 GetConvexFlags( int32 B, int32 I ) const;
+		uint32 GetSphylFlags( int32 B, int32 I ) const;
+		uint32 GetTaperedCapsuleFlags( int32 B, int32 I ) const;
+
+		const char* GetSphereName( int32 B, int32 I ) const;
+		const char* GetBoxName( int32 B, int32 I ) const;
+		const char* GetConvexName( int32 B, int32 I ) const;
+		const char* GetSphylName( int32 B, int32 I ) const;
+		const char* GetTaperedCapsuleName( int32 B, int32 I ) const;
+
+	protected:
+		//! Forbidden. Manage with the Ptr<> template.
+		~PhysicsBody()
+		{
+		}
+
+	public:
+		// Bone name the physics volume aggregate is bound to. 
+		TArray<string> Bones;
+		TArray<FPhysicsBodyAggregate> Bodies;
+		TArray<int32> CustomIds;
+
+		bool bBodiesModified = false;
+
+        //!
+        inline void Serialise( OutputArchive& arch ) const
+		{
+            uint32_t ver = 1;
+			arch << ver;
+        	
+        	arch << Bodies;
+        	arch << Bones;
+        	arch << CustomIds;
+			arch << bBodiesModified;
+        }
+
+        //!
+        inline void Unserialise( InputArchive& arch )
+		{
+            uint32_t ver;
+			arch >> ver;
+            check(ver <= 1);
+        	
+        	arch >> Bodies;
+        	arch >> Bones;
+        	arch >> CustomIds;
+
+			if (ver >= 1)
+			{
+				arch >> bBodiesModified;
+			}
+		}
+		
+        //!
+		inline bool operator==( const PhysicsBody& o ) const
+        {
+        	return Bones == o.Bones &&
+        		   Bodies == o.Bodies &&
+        	       CustomIds == o.CustomIds &&
+				   bBodiesModified == o.bBodiesModified;
+        }
+	};
+}
diff --git a/Engine/Plugins/Experimental/Mutable/Source/MutableRuntime/Public/RefCounted.h b/Engine/Plugins/Experimental/Mutable/Source/MutableRuntime/Public/RefCounted.h
new file mode 100644
--- /dev/null
+++ b/Engine/Plugins/Experimental/Mutable/Source/MutableRuntime/Public/RefCounted.h
@@ -0,0 +1,232 @@
+// Copyright Epic Games, Inc. All Rights Reserved.
+
+#pragma once
+
+#include "MutableRuntime/Public/MutableMemory.h"
+#include "MutableRuntime/Public/Ptr.h"
+
+#include "HAL/PlatformAtomics.h"
+#include "HAL/LowLevelMemTracker.h"
+
+#include <atomic>
+
+namespace mu
+{
+
+
+    //! \brief %Base class for all reference counted objects.
+    //!
+    //! Any subclass of this class can be managed using smart pointers through the Ptr<T> template.
+    //! \warning This base allow multi-threaded manipulation of smart pointers, since the count
+    //! increments and decrements are atomic.
+    //! \ingroup runtime
+    class MUTABLERUNTIME_API RefCounted : public Base
+    {
+    public:
+
+        inline void IncRef() const
+        {
+			FPlatformAtomics::InterlockedIncrement(&m_refCount);
+        }
+
+        inline void DecRef() const
+        {
+			LLM_SCOPE_BYNAME(TEXT("MutableRuntime"));
+
+            // Warning: This has to be in the same line: we cannot check m_refCount after decrement
+            // since another thread could have decremented it
+            if (FPlatformAtomics::InterlockedDecrement(&m_refCount)==0)
+            {
+                delete this;
+            }
+        }
+
+        inline int GetRefCount() const
+        {
+            return int( FPlatformAtomics::InterlockedExchange(&m_refCount, m_refCount) );
+        }
+
+        RefCounted( const RefCounted& ) = delete;
+        RefCounted( const RefCounted&& ) = delete;
+        RefCounted& operator=( const RefCounted& ) = delete;
+        RefCounted& operator=( const RefCounted&& ) = delete;
+
+    protected:
+
+        RefCounted()
+        {
+            m_refCount = 0;
+        }
+
+        inline virtual ~RefCounted() = default;
+
+    private:
+
+        mutable volatile int32 m_refCount;
+
+    };
+
+
+    inline void mutable_ptr_add_ref( const RefCounted* p )
+    {
+        if (p) p->IncRef();
+    }
+
+
+    inline void mutable_ptr_release( const RefCounted* p )
+    {
+        if (p) p->DecRef();
+    }
+
+
+    //! RefCounted with weak pointer support
+    class MUTABLERUNTIME_API RefCountedWeak : public Base
+    {
+    public:
+
+        inline void IncRef() const
+        {
+			FPlatformAtomics::InterlockedIncrement(&m_refCount);
+        }
+
+        inline void DecRef() const
+        {
+			LLM_SCOPE_BYNAME(TEXT("MutableRuntime"));
+
+            // Lock the control block to prevent creation of pointers while we work
+            Ptr<ControlBlock> pControlBlock = m_pControlBlock;
+            while (pControlBlock->m_lock.test_and_set(std::memory_order_acquire));
+
+            // Warning: This has to be in the same line: we cannot check m_refCount after decrement
+            // since another thread could have decremented it
+            if (FPlatformAtomics::InterlockedDecrement(&m_refCount)==0)
+            {
+                delete this;
+            }
+
+            pControlBlock->m_lock.clear();
+        }
+
+        inline int GetRefCount() const
+        {
+            return int(FPlatformAtomics::InterlockedExchange(&m_refCount, m_refCount));
+        }
+
+        RefCountedWeak( const RefCountedWeak& ) = delete;
+        RefCountedWeak( const RefCountedWeak&& ) = delete;
+        RefCountedWeak& operator=( const RefCountedWeak& ) = delete;
+        RefCountedWeak& operator=( const RefCountedWeak&& ) = delete;
+
+    protected:
+
+        RefCountedWeak()
+        {
+			LLM_SCOPE_BYNAME(TEXT("MutableRuntime"));
+
+            m_refCount = 0;
+            m_pControlBlock = new ControlBlock;
+            m_pControlBlock->m_pObject = this;
+        }
+
+        inline virtual ~RefCountedWeak()
+        {
+            // Control block is locked as per DecRef
+            m_pControlBlock->m_pObject = nullptr;
+        }
+
+
+    private:
+
+        // Normal intrusive smart pointer support
+        mutable volatile int32 m_refCount;
+
+        // Weak pointer support
+        class ControlBlock : public RefCounted
+        {
+        private:
+            RefCountedWeak* m_pObject = nullptr;
+            std::atomic_flag m_lock = ATOMIC_FLAG_INIT;
+
+            friend class RefCountedWeak;
+            template<class T> friend class WeakPtr;
+        };
+
+        Ptr<ControlBlock> m_pControlBlock;
+
+        template<class T> friend class WeakPtr;
+
+    };
+
+
+    inline void mutable_ptr_add_ref( const RefCountedWeak* p )
+    {
+        if (p) p->IncRef();
+    }
+
+
+    inline void mutable_ptr_release( const RefCountedWeak* p )
+    {
+        if (p) p->DecRef();
+    }
+
+
+    //!
+    template<class T>
+    class WeakPtr
+    {
+    private:
+        Ptr<RefCountedWeak::ControlBlock> m_pControlBlock;
+
+    public:
+
+        WeakPtr() = default;
+
+        WeakPtr( const Ptr<const T>& ptr )
+        {
+            if (ptr)
+            {
+                m_pControlBlock = ptr->m_pControlBlock;
+            }
+        }
+
+        inline Ptr<T> Pin()
+        {
+            Ptr<T> result;
+
+            if (m_pControlBlock)
+            {
+                while (m_pControlBlock->m_lock.test_and_set(std::memory_order_acquire));
+                result = (T*)m_pControlBlock->m_pObject;
+                m_pControlBlock->m_lock.clear();
+            }
+
+            return result;
+        }
+
+		inline Ptr<const T> Pin() const
+        {
+            Ptr<const T> result;
+
+            if (m_pControlBlock)
+            {
+                while (m_pControlBlock->m_lock.test_and_set(std::memory_order_acquire));
+                result = (const T*)m_pControlBlock->m_pObject;
+                m_pControlBlock->m_lock.clear();
+            }
+
+            return result;
+        }
+
+		inline void Reset()
+        {
+            m_pControlBlock = nullptr;
+        }
+
+		inline bool IsNull() const
+        {
+            return Pin().get()==nullptr;
+        }
+    };
+
+}
+
