[mutable] Move Mutable plugin from the restricted folder to the experimental plugins folder.

diff --git a/Engine/Plugins/Experimental/Mutable/Source/MutableRuntime/Public/Serialisation.h b/Engine/Plugins/Experimental/Mutable/Source/MutableRuntime/Public/Serialisation.h
new file mode 100644
--- /dev/null
+++ b/Engine/Plugins/Experimental/Mutable/Source/MutableRuntime/Public/Serialisation.h
@@ -0,0 +1,286 @@
+// Copyright Epic Games, Inc. All Rights Reserved.
+
+#pragma once
+
+#include "MutableRuntime/Public/Types.h"
+#include "MutableRuntime/Public/MutableMemory.h"
+#include "MutableRuntime/Public/Ptr.h"
+#include "MutableRuntime/Public/RefCounted.h"
+
+
+namespace mu
+{    
+    class Image;
+	class Model;
+
+    //! \brief
+    //! \ingroup model
+    template<class R>
+    class MUTABLERUNTIME_API ResourceProxy : public RefCounted
+    {
+    public:
+
+        virtual Ptr<const R> Get() = 0;
+
+    };
+
+
+	template<class R>
+	class ResourceProxyMemory : public ResourceProxy<R>
+	{
+	private:
+
+		Ptr<const R> m_resource;
+
+	public:
+
+		ResourceProxyMemory(const Ptr<const R>& i)
+		{
+			m_resource = i;
+		}
+
+		// ImageProxy interface
+		Ptr<const R> Get() override
+		{
+			return m_resource;
+		}
+	};
+
+
+    //! \brief Interface of any class responsible of providing model data.
+    //!
+    //! When building a Model that was generated with split data, an object implementing this
+    //! interface must be provided to the System. The System will call this interface to fetch the
+    //! data it needs from the model.
+    //! \ingroup runtime
+    class MUTABLERUNTIME_API ModelStreamer : public Base
+    {
+    public:
+
+        //-----------------------------------------------------------------------------------------
+        // Life cycle
+        //-----------------------------------------------------------------------------------------
+
+        //! Ensure virtual destruction.
+        virtual ~ModelStreamer() {}
+
+        //-----------------------------------------------------------------------------------------
+        // Reading interface
+        //-----------------------------------------------------------------------------------------
+
+        //! Identifier of reading data operations sent to this interface.
+		//! Negative values indicate an error.
+        typedef int32 OPERATION_ID;
+
+		//! \brief Start a data request operation.
+		//! \param Model.
+        //! \param key0 key identifying the model data fragment that is requested.
+        //!         This key interpretation depends on the implementation of the ModelStreamer,
+		//! \param pBuffer is an already-allocated buffer big enough to receive the expected data.
+		//! \param size is the size of the pBuffer buffer, which must match the size of the data
+		//! requested with the key identifiers.
+		//! \return a previously unused identifier, now used for this operation, that can be used in
+		//! calls to the other methods of this interface. If the return value is negative it indicates
+		//! an unrecoverable error.
+		virtual OPERATION_ID BeginReadBlock(const mu::Model*, uint64 key0, void* pBuffer, uint64 size) = 0;
+
+        //! Check if a data request operation has been completed.
+        //! This is a weak check than *may* return true if the given operation has completed, but
+        //! it is not mandatory. It is used as a hint by the System to optimise its opertaions.
+        //! There is no guarantee that this method will ever be called, and it is safe to always
+        //! return false.
+        virtual bool IsReadCompleted( OPERATION_ID ) = 0;
+
+        //! Complete a data request operation. This method has to block until a data request issued
+        //! with OpenFile has been completed. After returning from this call, the ID cannot be used
+        //! any more to identify the same operation and becomes free.
+        virtual void EndRead( OPERATION_ID ) = 0;
+
+        //-----------------------------------------------------------------------------------------
+        // Writing interface
+        //-----------------------------------------------------------------------------------------
+
+        //! \brief Open a file for writing data.
+        //!
+        //! Only one file can be open a time for writing. This method is only required when
+        //! writing model data from tools.
+        //! \param strModelName Name of the model where the file to open belongs to. This string's
+        //!         meaning depends on the implementation of the ModelStreamer subclasses. It could
+        //!         be a file path, or a resource identifier, etc. It will be the same for all
+        //!         files in a model.
+        //! \param key key identifying the model data fragment that is requested.
+        //!         This key interpretation depends on the implementation of the ModelStreamer,
+        virtual void OpenWriteFile( const char* strModelName, uint64 key0 ) = 0;
+
+        //! \brief Write a piece of data to the currently open file.
+        //!
+        //! There must be a file open with OpenWriteFile before calling this.
+        //! \param pBuffer pointer to the data that will be written to the file.
+        //! \param size size of the data to write on the file, in bytes.
+        virtual void Write( const void* pBuffer, uint64 size ) = 0;
+
+        //! \brief Close the file open for writing in a previous call to OpenWriteFile in this
+        //! object.
+        virtual void CloseWriteFile() = 0;
+
+
+    };
+
+
+
+    //! Interface for any input stream to be use with InputArchives.
+    //! \ingroup tools
+    class MUTABLERUNTIME_API InputStream : public Base
+    {
+    public:
+
+        //! Ensure virtual destruction
+        virtual ~InputStream() {}
+
+        //! Read a byte buffer
+        //! \param pData destination buffer, must have at least size bytes allocated.
+        //! \param size amount of bytes to read from the stream.
+        virtual void Read( void* pData, uint64 size ) = 0;
+    };
+
+
+    //! Interface for any output stream to be used with OutputArchives
+    //! \ingroup tools
+    class MUTABLERUNTIME_API OutputStream : public Base
+    {
+    public:
+
+        //! Ensure virtual destruction
+        virtual ~OutputStream() {}
+
+        //! Write a byte buffer
+        //! \param pData source buffer where data will be read from.
+        //! \param size amount of data to write to the stream.
+        virtual void Write( const void* pData, uint64 size ) = 0;
+
+    };
+
+
+    //! Archive containing data to be deserialised.
+    //! \ingroup tools
+    class MUTABLERUNTIME_API InputArchive : public Base
+    {
+    public:
+
+        //! Construct form an input stream. The stream will not be owned by the archive and the
+        //! caller must make sure it is not modified or destroyed while serialisation is happening.
+        InputArchive( InputStream* );
+
+        //!
+        virtual ~InputArchive();
+
+        //
+        virtual Ptr<ResourceProxy<Image>> NewImageProxy();
+
+        // Interface pattern
+        class Private;
+        Private* GetPrivate() const;
+
+    private:
+
+        Private* m_pD;
+
+    };
+
+
+    //! Archive where data can be serialised to.
+    //! \ingroup tools
+    class MUTABLERUNTIME_API OutputArchive : public Base
+    {
+    public:
+
+        //! Construct form an output stream. The stream will not be owned by the archive and the
+        //! caller must make sure it is not modified or destroyed while serialisation is happening.
+        OutputArchive( OutputStream* );
+
+        //!
+        ~OutputArchive();
+
+        // Interface pattern
+        class Private;
+        Private* GetPrivate() const;
+
+    private:
+
+        Private* m_pD;
+
+    };
+
+    //!
+    //! \ingroup runtime
+    class MUTABLERUNTIME_API InputArchiveWithProxies : public InputArchive
+    {
+    public:
+
+        class ProxyFactory
+        {
+        public:
+          virtual ~ProxyFactory() {}
+          virtual Ptr<ResourceProxy<Image>> NewImageProxy(InputArchive& arch) = 0;
+        };
+
+    public:
+
+        //! Construct form an input stream. The stream will not be owned by the archive and the
+        //! caller must make sure it is not modified or destroyed while serialisation is happening.
+        InputArchiveWithProxies( InputStream*, ProxyFactory*  );
+
+        ~InputArchiveWithProxies() override;
+
+        //
+        Ptr<ResourceProxy<Image>> NewImageProxy() override;
+
+    private:
+
+        class Private;
+        Private* m_pD;
+    };
+
+
+
+    //!
+    class MUTABLERUNTIME_API OutputMemoryStream : public OutputStream
+    {
+    public:
+
+        //-----------------------------------------------------------------------------------------
+        // Life cycle
+        //-----------------------------------------------------------------------------------------
+
+        //! Create the stream with an optional buffer size in bytes.
+        //! The internal buffer will be enlarged as much as necessary.
+        OutputMemoryStream( uint64 reserve = 0 );
+
+        ~OutputMemoryStream();
+
+
+        //-----------------------------------------------------------------------------------------
+        // OutputStream interface
+        //-----------------------------------------------------------------------------------------
+        void Write( const void* pData, uint64 size ) override;
+
+        //-----------------------------------------------------------------------------------------
+        // Own interface
+        //-----------------------------------------------------------------------------------------
+
+        //! Get the serialised data buffer pointer. This pointer invalidates after a Write
+        //! operation has been done, and you need to get it again.
+        const void* GetBuffer() const;
+
+        //! Get the amount of data in the stream, in bytes.
+        uint64 GetBufferSize() const;
+
+    private:
+
+        class Private;
+        Private* m_pD;
+
+    };
+
+}
+
diff --git a/Engine/Plugins/Experimental/Mutable/Source/MutableRuntime/Public/Settings.h b/Engine/Plugins/Experimental/Mutable/Source/MutableRuntime/Public/Settings.h
new file mode 100644
--- /dev/null
+++ b/Engine/Plugins/Experimental/Mutable/Source/MutableRuntime/Public/Settings.h
@@ -0,0 +1,79 @@
+// Copyright Epic Games, Inc. All Rights Reserved.
+
+#pragma once
+
+#include "MutableRuntime/Public/RefCounted.h"
+#include "MutableRuntime/Public/Ptr.h"
+#include "MutableRuntime/Public/Types.h"
+#include "MutableRuntime/Public/Serialisation.h"
+#include "MutableRuntime/Public/Instance.h"
+
+
+namespace mu
+{
+	// Forward references
+    class Settings;
+    using SettingsPtr=Ptr<Settings>;
+    using SettingsPtrConst=Ptr<const Settings>;
+
+
+
+    //! \brief Settings class that can be used to create a customised System instance.
+    //! \ingroup runtime
+    class MUTABLERUNTIME_API Settings : public RefCounted
+    {
+    public:
+
+        //! Create new settings with the default values as explained in each method below.
+        Settings();
+
+        //! Record internal profiling data.
+        //! Disabled by default.
+        void SetProfile( bool bEnabled );
+
+        //! Limit the maximum memory in bytes used to store streaming data. A low value will force
+        //! more streaming and higher instance construction times, but will use less memory while
+        //! building objects.
+        //! Defaults to 0, which disables the limit.
+        void SetStreamingCache( uint64_t bytes );
+
+        //! Set the quality for the image compression algorithms.
+        //! \param quality Quality level of the compression. Higher quality increases the instance
+        //! generation time. The value is in the range 0 to 4 with 0 being the fastest. It defaults
+        //! to 0, which is the fastes for runtime. Tools may want ot set it higher, and low
+        //! performance profiles to lower.
+        //! Genreal rules are:
+        //! 0 - Fastest for runtime
+        //! 1 - Best for runtime
+        //! 2 - Fast for tools
+        //! 3 - Best for tools
+        //! 4 - Maximum, with no time limits.
+        //! Visit https://work.anticto.com/w/mutable/core/guide/ for
+        //! more information.
+        void SetImageCompressionQuality( int quality );
+
+        //-----------------------------------------------------------------------------------------
+        // Interface pattern
+        //-----------------------------------------------------------------------------------------
+        class Private;
+        Private* GetPrivate() const;
+
+        Settings( const Settings& ) = delete;
+        Settings( Settings&& ) = delete;
+        Settings& operator=(const Settings&) = delete;
+        Settings& operator=(Settings&&) = delete;
+
+    protected:
+
+        //! Forbidden. Manage with the Ptr<> template.
+        ~Settings() override;
+
+    private:
+
+        Private* m_pD;
+
+    };
+
+
+}
+
diff --git a/Engine/Plugins/Experimental/Mutable/Source/MutableRuntime/Public/Skeleton.h b/Engine/Plugins/Experimental/Mutable/Source/MutableRuntime/Public/Skeleton.h
new file mode 100644
--- /dev/null
+++ b/Engine/Plugins/Experimental/Mutable/Source/MutableRuntime/Public/Skeleton.h
@@ -0,0 +1,153 @@
+// Copyright Epic Games, Inc. All Rights Reserved.
+
+#pragma once
+
+#include "MutableRuntime/Public/Types.h"
+#include "MutableRuntime/Public/Ptr.h"
+#include "MutableRuntime/Public/RefCounted.h"
+#include "MutableRuntime/Public/Serialisation.h"
+#include "MutableRuntime/Private/SerialisationPrivate.h"
+#include "MutableRuntime/Private/MutableMath.h"
+#include "MutableRuntime/Private/MemoryPrivate.h"
+#include "Math/Transform.h"
+
+
+namespace mu
+{
+
+	// Forward references
+    class Skeleton;
+    typedef Ptr<Skeleton> SkeletonPtr;
+    typedef Ptr<const Skeleton> SkeletonPtrConst;
+
+
+    //! \brief Skeleton object.
+	//! \ingroup runtime
+    class MUTABLERUNTIME_API Skeleton : public RefCounted
+	{
+	public:
+
+		//-----------------------------------------------------------------------------------------
+		// Life cycle
+		//-----------------------------------------------------------------------------------------
+
+		Skeleton() {}
+
+        //! Deep clone this skeleton.
+        Ptr<Skeleton> Clone() const;
+
+		//! Serialisation
+        static void Serialise( const Skeleton* p, OutputArchive& arch );
+        static Ptr<Skeleton> StaticUnserialise( InputArchive& arch );
+
+
+		//-----------------------------------------------------------------------------------------
+		// Own interface
+		//-----------------------------------------------------------------------------------------
+
+        //! Return the number of bones used by this skeleton
+		int32 GetBoneCount() const;
+        void SetBoneCount(int32 c);
+
+        //! Return the name of one of the bones used by this skeleton
+		//! \param boneIndex goes from 0 to GetBoneCount()-1
+		const char* GetBoneName(int32 boneIndex) const;
+        void SetBoneName( int b, const char* strName );
+
+        //! Get and set the parent bone of each bone. The parent can be -1 if the bone is a root.
+        int32 GetBoneParent(int32 boneIndex) const;
+        void SetBoneParent(int32 boneIndex, int32 parentBoneIndex);
+
+        //! Get and set an optional external Id for each bone that can be retrieved from generated
+        //! skeletons. This data is opaque to mutable, so it is copied around and merged as
+        //! necessary, but never changed. By default all indices are 0.
+        int32 GetBoneId(int32 boneIndex) const;
+        void SetBoneId(int32 boneIndex, int32 boneId);
+
+	protected:
+
+		//! Forbidden. Manage with the Ptr<> template.
+		~Skeleton() {}
+
+	public:
+
+		//!
+		TArray<string> m_bones;
+		TArray<FTransform3f> m_boneTransforms_DEPRECATED;
+
+		//! For each bone, index of the parent bone in the bone vectors. -1 means no parent.
+		//! This array must have the same size than the m_bones array.
+		TArray<int16> m_boneParents;
+
+		//! For each bone, an externally provided optional id. This vector can be empty or smaller
+		//! Than the number of bones, and missing bones are considered to have id 0.
+		TArray<int32> m_boneIds;
+
+
+		//!
+		inline void Serialise(OutputArchive& arch) const
+		{
+			uint32_t ver = 4;
+			arch << ver;
+
+			arch << m_bones;
+			arch << m_boneParents;
+			arch << m_boneIds;
+		}
+
+		//!
+		inline void Unserialise(InputArchive& arch)
+		{
+			uint32_t ver;
+			arch >> ver;
+			check(ver >= 3);
+
+			arch >> m_bones;
+
+			if (ver == 3)
+			{
+				arch >> m_boneTransforms_DEPRECATED;
+			}
+
+			arch >> m_boneParents;
+			arch >> m_boneIds;
+
+			if (ver == 3)
+			{
+				// ensure m_boneIds' size is equal to the number of bones
+				m_boneIds.SetNum(m_bones.Num());
+
+				bool bBoneTransformModified;
+				arch >> bBoneTransformModified;
+			}
+		}
+
+
+		//!
+		inline bool operator==(const Skeleton& o) const
+		{
+			return m_bones == o.m_bones 
+				&& m_boneParents == o.m_boneParents
+				&& m_boneIds == o.m_boneIds;
+		}
+
+
+		//-----------------------------------------------------------------------------------------
+		//! Fins a bone index by name
+		//-----------------------------------------------------------------------------------------
+		int FindBone(const char* strName) const
+		{
+			for (int32 i = 0; i < m_bones.Num(); ++i)
+			{
+				if (m_bones[i] == strName)
+				{
+					return (int)i;
+				}
+			}
+			return -1;
+		}
+
+	};
+
+}
+
diff --git a/Engine/Plugins/Experimental/Mutable/Source/MutableRuntime/Public/System.h b/Engine/Plugins/Experimental/Mutable/Source/MutableRuntime/Public/System.h
new file mode 100644
--- /dev/null
+++ b/Engine/Plugins/Experimental/Mutable/Source/MutableRuntime/Public/System.h
@@ -0,0 +1,245 @@
+// Copyright Epic Games, Inc. All Rights Reserved.
+
+#pragma once
+
+#include "MutableRuntime/Public/RefCounted.h"
+#include "MutableRuntime/Public/Ptr.h"
+#include "MutableRuntime/Public/Types.h"
+#include "MutableRuntime/Public/Serialisation.h"
+#include "MutableRuntime/Public/Instance.h"
+#include "MutableRuntime/Public/Image.h"
+#include "MutableRuntime/Public/Settings.h"
+
+// This define will use the newer task graph interface to manage mutable concurrency. 
+// This is currently broken in Switch and maybe other consoles, for some unknown reason.
+//#define MUTABLE_USE_NEW_TASKGRAPH
+
+
+namespace mu
+{
+	// Forward references
+	class Model;
+    using ModelPtr=Ptr<Model>;
+    using ModelPtrConst=Ptr<const Model>;
+
+	class Parameters;
+    using ParametersPtr=Ptr<Parameters>;
+    using ParametersPtrConst=Ptr<const Parameters>;
+
+    class Mesh;
+    using MeshPtr=Ptr<Mesh>;
+    using MeshPtrConst=Ptr<const Mesh>;
+
+    class System;
+    using SystemPtr=Ptr<System>;
+    using SystemPtrConst=Ptr<const System>;
+
+
+
+    //! List of critical errors that may happen during execution of mutable code
+    enum class Error
+    {
+        //! No error happened.
+        None = 0,
+
+        //! A memory allocation failed
+        OutOfMemory,
+
+        //! A file was missing or the data was corrupted.
+        StreamingError,
+
+        //! The necessary functionality is not supported in this version of Mutable.
+        Unsupported,
+
+        //! Utility value with the number of error types.
+        Count
+    };
+
+    //! \brief Get an error code for the most critical error that happened since mutable
+    //! initialisation.
+    //! \ingroup runtime
+	MUTABLERUNTIME_API extern Error GetError();
+
+    //! \brief Remove the last error code, setting it to Error::None. This should obnly be used
+    //! in case of possible recovery from one of the critical errors, which is very unlikely.
+    //! \ingroup runtime
+	MUTABLERUNTIME_API extern void ClearError();
+
+
+    //! \brief Interface to request external images used as parameters.
+    //! \ingroup runtime
+    class MUTABLERUNTIME_API ImageParameterGenerator : public Base
+    {
+    public:
+
+        //! Ensure virtual destruction
+        virtual ~ImageParameterGenerator() = default;
+
+        //!
+        virtual ImagePtr GetImage( EXTERNAL_IMAGE_ID id ) = 0;
+    };
+
+
+    //! \brief Main system class to load models and build instances.
+	//! \ingroup runtime
+	class MUTABLERUNTIME_API System : public RefCounted
+	{
+    public:
+
+        //! This constant can be used in place of the lodMask in methods like BeginUpdate
+        static constexpr unsigned AllLODs = 0xffffffff;
+
+    public:
+
+		//! Constructor of a system object to build data.
+        //! \param settings Optional class with the settings to use in this system. The default
+        //! value configures a production-ready system.
+        System( const SettingsPtr& settings = nullptr );
+
+        //! Set a new provider for model data. The provider will become owned by this instance and
+        //! destroyed when necessary.
+        void SetStreamingInterface( ModelStreamer* );
+
+        //! Overwrite the streaming memory limit set in the settings when the system was created.
+        //! Refer to Settings::SetStreamingCache for more information.
+        void SetStreamingCache( uint64 bytes );
+
+        //! Set a new provider for external image data. This is only necessary if image parameters
+        //! are used in the models.
+        void SetImageParameterGenerator( ImageParameterGenerator* );
+
+        //! Set the maximum memory that this system can use. This memory is used for built data,
+        //! cached data, and streaming. If set to 0, the system will use as much memory as required.
+        //! \warning This limit is ignored if the system was built with SF_REFERENCE.
+        //! \param mem Maximum memory, in bytes
+        void SetMemoryLimit( uint32 mem );
+
+        //! Create a new instance from the given model. The instance can then be configured through
+        //! calls to BeginUpdate/EndUpdate.
+        //! A call to NewInstance must be paired to a call to ReleasesInstance when the instance is
+        //! no longer needed.
+        //! \param pModel Model to build an instance of
+        //! \return An identifier that is always bigger than 0.
+        Instance::ID NewInstance( const ModelPtrConst& pModel );
+
+        //! \brief Update an instance with a new parameter set and/or state.
+        //!
+        //! \warning a call to BeginUpdate must be paired with a call to EndUpdate once the returned
+        //! data has been processed.
+        //! \param instanceID The id of the instance to update, as created by a NewInstance call.
+        //! \param pParams The parameters that customise this instance.
+        //! \param stateIndex The index of the state this instance will be set to. The states range
+        //! from 0 to Model::GetStateCount-1
+        //! \param lodMask Bitmask selecting the levels of detail to build (i-th bit selects i-th lod).
+        //! \return the instance data with all the LOD, components, and ids to generate the meshes 
+        //! and images.  The returned Instance is only valid until the next call to EndUpdate with 
+        //! the same instanceID parameter.
+        const Instance* BeginUpdate( Instance::ID instanceID,
+                                     const ParametersPtrConst& pParams,
+                                     int32 stateIndex,
+                                     uint32 lodMask );
+
+		//! Only valid between BeginUpdate and EndUpdate
+		//! Calculate the description of an image, without generating it.
+		void GetImageDesc(Instance::ID instanceID, RESOURCE_ID imageId, FImageDesc& OutDesc );
+
+		//! Only valid between BeginUpdate and EndUpdate
+		//! \param MipsToSkip Number of mips to skip compared from the full image.
+		//! If 0, all mip levels will be generated. If more levels than possible to discard are specified, 
+		//! the image will still contain a minimum number of mips specified at model compile time.
+		ImagePtrConst GetImage(Instance::ID instanceID, RESOURCE_ID imageId, int32 MipsToSkip = 0);
+
+        //! Only valid between BeginUpdate and EndUpdate
+        MeshPtrConst GetMesh(Instance::ID instanceID, RESOURCE_ID meshId);
+
+        //! Invalidate and free the last Instance data returned by a call to BeginUpdate with
+        //! the same instance index. After a call to this method, that Instance cannot be used any
+        //! more and its content is undefined.
+        //! \param instance The index of the instance whose last data will be invalidated.
+        void EndUpdate( Instance::ID instance );
+
+        //! Completely destroy an instance. After a call to this method the given instance cannot be
+        //! updated any more, and its resources may have been freed.
+        //! \param instance The id of the instance to destroy.
+        void ReleaseInstance( Instance::ID instance );
+
+		//! Build one of the images defined in a model parameter as additional description. These
+		//! images can be used for colour bars, icons, etc...
+        ImagePtrConst BuildParameterAdditionalImage( const ModelPtrConst& pModel,
+                                                     const ParametersPtrConst& pParams,
+                                                     int parameter,
+                                                     int image );
+
+		//! Calculate the relevancy of every parameter. Some parameters may be unused depending on
+		//! the values of other parameters. This method will set to true the flags for parameters
+		//! that are relevant, and to false otherwise. This is useful to hide irrelevant parameters
+		//! in dynamic user interfaces.
+        //! \param pModel The model used to create the Parameters instance.
+        //! \param pParameters Parameter set that we want to find the relevancy of.
+        //! \param pFlags is a pointer to a preallocated array of booleans that contains at least
+		//! pParameters->GetCount() elements.
+        void GetParameterRelevancy( const ModelPtrConst& pModel,
+									const ParametersPtrConst& pParameters,
+									bool* pFlags );
+
+        //! Free memory used in internal runtime caches. This is useful for long-running processes
+        //! that keep models loaded. This could be called when a game finishes, or a change of
+        //! level.
+        void ClearCaches();
+
+        //-----------------------------------------------------------------------------------------
+        // Debugging and profiling
+        //-----------------------------------------------------------------------------------------
+
+        //! Types of metrics that can be queried
+        enum class ProfileMetric : uint8
+        {
+            //! No metric
+            None,
+
+            //! Number of model instances currently being created or updated.
+            LiveInstanceCount,
+
+            //! Number of bytes currently loaded for streaming data.
+            StreamingCacheBytes,
+
+            //! Number of update operations performed since the system was created.
+            InstanceUpdateCount,
+
+            //! Utility value that doesn't really represent a metric.
+            _Count
+        };
+
+        //! Get the values of several profile metrics. This can be called any time in a valid
+        //! system object from any thread.
+        //! These values can change any time.
+        uint64 GetProfileMetric( ProfileMetric ) const;
+
+		//-----------------------------------------------------------------------------------------
+		// Interface pattern
+		//-----------------------------------------------------------------------------------------
+		class Private;
+		Private* GetPrivate() const;
+
+    public:
+
+        // Prevent copy, move and assignment.
+        System( const System& ) = delete;
+        System& operator=( const System& ) = delete;
+        System( System&& ) = delete;
+        System& operator=( System&& ) = delete;
+
+	protected:
+
+		//! Forbidden. Manage with the Ptr<> template.
+        ~System() override;
+
+	private:
+
+		Private* m_pD;
+
+	};
+
+
+}
+
diff --git a/Engine/Plugins/Experimental/Mutable/Source/MutableRuntime/Public/Types.h b/Engine/Plugins/Experimental/Mutable/Source/MutableRuntime/Public/Types.h
new file mode 100644
--- /dev/null
+++ b/Engine/Plugins/Experimental/Mutable/Source/MutableRuntime/Public/Types.h
@@ -0,0 +1,16 @@
+// Copyright Epic Games, Inc. All Rights Reserved.
+
+#pragma once
+
+#include "HAL/Platform.h"
+
+namespace mu
+{
+
+    typedef uint32            RESOURCE_ID;
+
+    //!
+    typedef uint64            EXTERNAL_IMAGE_ID;
+
+}
+
diff --git a/Engine/Plugins/Experimental/Mutable/Source/MutableTools/MutableTools.Build.cs b/Engine/Plugins/Experimental/Mutable/Source/MutableTools/MutableTools.Build.cs
new file mode 100644
--- /dev/null
+++ b/Engine/Plugins/Experimental/Mutable/Source/MutableTools/MutableTools.Build.cs
@@ -0,0 +1,32 @@
+// Copyright Epic Games, Inc. All Rights Reserved.
+
+using System.IO;
+
+namespace UnrealBuildTool.Rules
+{
+    public class MutableTools : ModuleRules
+    {
+        public MutableTools(ReadOnlyTargetRules Target) : base(Target)
+		{
+			ShortName = "MuT";
+
+			DefaultBuildSettings = BuildSettingsVersion.V2;
+
+			// 
+			bUseRTTI = true;
+			bEnableExceptions = true;
+			//bUseUnity = false;
+
+			PrivateIncludePaths.AddRange(new string[] {
+				"MutableRuntime/Private",
+			});
+
+			PublicDependencyModuleNames.AddRange(
+                new string[] {
+					"MutableRuntime", 
+					"Core",
+                }
+            );
+		}
+	}
+}
diff --git a/Engine/Plugins/Experimental/Mutable/Source/MutableTools/Private/AST.cpp b/Engine/Plugins/Experimental/Mutable/Source/MutableTools/Private/AST.cpp
new file mode 100644
--- /dev/null
+++ b/Engine/Plugins/Experimental/Mutable/Source/MutableTools/Private/AST.cpp
@@ -0,0 +1,1897 @@
+// Copyright Epic Games, Inc. All Rights Reserved.
+
+#include "MutableTools/Private/AST.h"
+#include "MutableTools/Private/ASTOpConstantResource.h"
+#include "MutableTools/Private/ASTOpMeshRemoveMask.h"
+
+#include "ErrorLogPrivate.h"
+#include "CodeOptimiser.h"
+#include "StreamsPrivate.h"
+#include "MutableTrace.h"
+
+#include <queue>
+
+
+std::atomic<uint32_t> mu::ASTOp::s_lastTraverseIndex(1);
+
+
+using namespace mu;
+
+
+
+void mu::DebugLogAST(const Ptr<ASTOp>& at, int indent, ASTOpList* done, const char* label )
+{
+    if (!at) return;
+
+    static int s_logTree = 0;
+
+    if (indent==0)
+    {
+		UE_LOG(LogMutableCore, Warning, TEXT("ASTOp tree [%d]:"), s_logTree );
+        ++s_logTree;
+    }
+
+    ASTOpList localDone;
+    if (!done) done = &localDone;
+
+    int index = -1;
+    for ( size_t s=0; s<done->size(); ++s )
+    {
+        if ((*done)[s]==at)
+        {
+            index = int(s);
+            break;
+        }
+    }
+
+    std::string pre;
+    for (int i=0; i<indent; ++i) pre +=" ";
+    if (label)
+    {
+        pre += label;
+    }
+
+    auto opType = at->GetOpType();
+    if (index>=0)
+    {
+        if ( GetOpDataType(opType)==DT_MESH)
+        {
+			UE_LOG(LogMutableCore, Warning, TEXT("%srepeated : %d"),
+                     pre.c_str(), index );
+        }
+        return;
+    }
+    index = (int)done->size();
+    done->push_back(at);
+
+    bool childrenAdded = false;
+    if ( GetOpDataType(opType)==DT_MESH)
+    {
+        if (opType==OP_TYPE::ME_CONSTANT)
+        {
+            ASTOpConstantResource* opc = dynamic_cast<ASTOpConstantResource*>(at.get());
+            auto value = opc->GetValue();
+            auto mesh = (const Mesh*)( value.get());
+            if (!mesh)
+            {
+                UE_LOG(LogMutableCore, Warning, TEXT("%s%s [%d]"),
+                         pre.c_str(),
+                         "mesh-null", index );
+            }
+            else if (mesh->GetIndexCount()==0)
+            {
+                UE_LOG(LogMutableCore, Warning, TEXT("%s%s [%d]"),
+                         pre.c_str(),
+                         "mesh-mask", index );
+            }
+            else
+            {
+                UE_LOG(LogMutableCore, Warning, TEXT("%s%s [%d]"),
+                         pre.c_str(),
+                         "mesh", index );
+            }
+        }
+
+        else if (opType==OP_TYPE::ME_REMOVEMASK)
+        {
+            UE_LOG(LogMutableCore, Warning, TEXT("%s%s [%d]"),
+                     pre.c_str(),
+                     mu::s_opNames[int(opType)], index );
+
+            ASTOpMeshRemoveMask* opc = dynamic_cast<ASTOpMeshRemoveMask*>(at.get());
+
+            DebugLogAST( opc->source.child(), indent+2, done, "source : " );
+
+            for ( const auto& m : opc->removes )
+            {
+                DebugLogAST( m.first.child(), indent+2, done, "condition : " );
+                DebugLogAST( m.second.child(), indent+2, done, "mask : ");
+            }
+            childrenAdded = true;
+        }
+
+        else
+        {
+            UE_LOG(LogMutableCore, Warning, TEXT("%s%s [%d]"),
+                     pre.c_str(),
+                     mu::s_opNames[int(opType)], index );
+        }
+    }
+
+    if (!childrenAdded)
+    {
+        at->ForEachChild([indent,done]( ASTChild& c )
+        {
+            DebugLogAST( c.m_child, indent+2, done );
+        });
+    }
+}
+
+
+//-------------------------------------------------------------------------------------------------
+//-------------------------------------------------------------------------------------------------
+//-------------------------------------------------------------------------------------------------
+ASTChild::ASTChild(ASTOp* p, const Ptr<ASTOp>& c)
+    : m_parent(p)
+    , m_child(c)
+{
+    if (m_parent && m_child.get())
+    {
+        AddParent();
+    }
+}
+
+ASTChild::ASTChild( const Ptr<ASTOp>& p, const Ptr<ASTOp>& c)
+    : ASTChild(p.get(),c)
+{
+}
+
+ASTChild::~ASTChild()
+{
+    if (m_child && m_parent)
+    {
+        ClearParent();
+    }
+}
+
+ASTChild& ASTChild::operator=( const Ptr<ASTOp>& c )
+{
+    if (c!=m_child)
+    {
+        if (m_child && m_parent)
+        {
+            ClearParent();
+        }
+
+        m_child = c;
+
+        if (m_child && m_parent)
+        {
+            AddParent();
+        }
+    }
+
+    return *this;
+}
+
+ASTChild& ASTChild::operator=( ASTChild&& rhs )
+{
+    m_parent = rhs.m_parent;
+    m_parentIndexInChild = rhs.m_parentIndexInChild;
+    m_child = rhs.m_child;
+    rhs.m_parent=nullptr;
+    rhs.m_child.reset();
+
+    return *this;
+}
+
+
+
+//-------------------------------------------------------------------------------------------------
+//-------------------------------------------------------------------------------------------------
+//-------------------------------------------------------------------------------------------------
+void ASTOp::ForEachParent(const std::function<void(ASTOp*)>& f )
+{
+    for( auto& p: m_parents )
+    {
+        if (p)
+        {
+            f(p);
+        }
+    }
+}
+
+
+//-------------------------------------------------------------------------------------------------
+void ASTOp::RemoveChildren()
+{
+    // Actually destroyed when running out of scope
+    vector<Ptr<ASTOp>> toDestroy;
+
+    // Try to make children destruction iterative
+    TArray<ASTOp*> pending;
+	pending.Reserve(1024);
+    pending.Add(this);
+
+    while (pending.Num())
+    {
+        ASTOp* n = pending.Pop(false);
+
+        n->ForEachChild( [&](ASTChild& c)
+        {
+            if (c)
+            {
+                // Are we clearing the last reference?
+                if (c.child()->GetRefCount()==1)
+                {
+                    toDestroy.push_back(c.child());
+                    pending.Add(c.child().get());
+                }
+
+                c = nullptr;
+            }
+        });
+    }
+}
+
+
+//-------------------------------------------------------------------------------------------------
+void ASTOp::Assert()
+{
+    // Check that every valid parent has us a child
+    // TODO: Should count the numbers, since a node may be child of another in multiple connections.
+    ForEachParent( [&](ASTOp*parent)
+    {
+        if(parent)
+        {
+            bool foundInParent = false;
+            parent->ForEachChild([&](ASTChild&c)
+            {
+                if (c && c.m_child.get()==this)
+                {
+                    foundInParent = true;
+                }
+            });
+
+            // If we hit this, we have a parent that doesn't know us.
+            check(foundInParent);
+        }
+    });
+
+    // Validate the children
+    ForEachChild( [=](ASTChild&c)
+    {
+        if(c)
+        {
+            // The child must have us as the parent.
+//            bool found = false;
+//            c.child()->ForEachParent( [&](ASTOp* childParent)
+//            {
+//                if (childParent==this)
+//                {
+//                    found = true;
+//                }
+//            });
+//            check(found);
+            check( c.m_parentIndexInChild < size_t(c.child()->m_parents.Num()) );
+            check( c.child()->m_parents[c.m_parentIndexInChild]==this );
+        }
+     });
+}
+
+
+//-------------------------------------------------------------------------------------------------
+bool ASTOp::operator==( const ASTOp& other ) const
+{
+//    if (typeid(*this) != typeid(other))
+//        return false;
+
+    return IsEqual(other);
+}
+
+
+//---------------------------------------------------------------------------------------------
+void ASTOp::FullAssert( const vector<Ptr<ASTOp>>& roots )
+{
+    MUTABLE_CPUPROFILER_SCOPE(AST_FullAssert);
+    Traverse_TopDown_Unique_Imprecise( roots, [](const Ptr<ASTOp>& n)
+    {
+        n->Assert();
+        return true;
+    });
+}
+
+//-------------------------------------------------------------------------------------------------
+size_t ASTOp::CountNodes( const vector<Ptr<ASTOp>>& roots )
+{
+    MUTABLE_CPUPROFILER_SCOPE(AST_CountNodes);
+    size_t count=0;
+    Traverse_TopRandom_Unique_NonReentrant( roots, [&](const Ptr<ASTOp>&)
+    {
+        ++count;
+        return true;
+    });
+    return count;
+}
+
+
+//-------------------------------------------------------------------------------------------------
+mu::Ptr<ASTOp> ASTOp::DeepClone( const Ptr<ASTOp>& root )
+{
+    MUTABLE_CPUPROFILER_SCOPE(AST_DeepClone);
+
+    std::unordered_map<Ptr<const ASTOp>,Ptr<ASTOp>> visited;
+
+    MapChildFunc m = [&](const Ptr<ASTOp>&n)
+    {
+        if (!n) return Ptr<ASTOp>();
+        auto it = visited.find(n);
+        check(it!=visited.end());
+        return it->second;
+    };
+
+    Ptr<ASTOp> r = root;
+    Traverse_BottomUp_Unique( r, [&](Ptr<ASTOp> n)
+    {
+        Ptr<ASTOp> cloned = n->Clone( m );
+        visited[n] = cloned;
+    });
+
+    auto it = visited.find(r);
+    check(it!=visited.end());
+    return it->second;
+}
+
+
+//-------------------------------------------------------------------------------------------------
+void ASTOp::FullLink( Ptr<ASTOp>& root, PROGRAM& program, const FLinkerOptions* Options )
+{
+    MUTABLE_CPUPROFILER_SCOPE(AST_FullLink);
+    Traverse_BottomUp_Unique( root,
+                              [&](Ptr<ASTOp> n){ n->Link(program, Options); },
+                              [&](Ptr<const ASTOp> n){ return n->linkedAddress==0; });
+}
+
+
+//-------------------------------------------------------------------------------------------------
+void ASTOp::ClearLinkData( Ptr<ASTOp>& root )
+{
+    MUTABLE_CPUPROFILER_SCOPE(AST_ClearLinkData);
+    ASTOpList roots;
+    roots.push_back(root);
+    Traverse_TopDown_Unique_Imprecise( roots,
+                             [&](const Ptr<ASTOp>& n){ n->linkedAddress = 0; return true; });
+}
+
+
+//---------------------------------------------------------------------------------------------
+void ASTOp::LogHistogram( ASTOpList& roots )
+{
+    (void)roots;
+
+#if 0
+    uint64_t countPerType[(int)OP_TYPE::COUNT];
+    FMemory::Memzero(countPerType,sizeof(countPerType));
+
+    size_t count=0;
+
+    Traverse_TopRandom_Unique_NonReentrant( roots,
+                             [&](const Ptr<ASTOp>& n)
+    {
+        ++count;
+        countPerType[(int)n->GetOpType()]++;
+        return true;
+    });
+
+    vector< pair<uint64_t,OP_TYPE> > sorted((int)OP_TYPE::COUNT);
+    for (int i=0; i<(int)OP_TYPE::COUNT; ++i)
+    {
+        sorted[i].second = (OP_TYPE)i;
+        sorted[i].first = countPerType[i];
+    }
+
+    std::sort(sorted.begin(),sorted.end(), []( const pair<uint64_t,OP_TYPE>& a, const pair<uint64_t,OP_TYPE>& b )
+    {
+        return a.first>b.first;
+    });
+
+    UE_LOG(LogMutableCore,Log, TEXT("Op histogram (%llu ops):"), count);
+    for(int i=0; i<8; ++i)
+    {
+        float p = sorted[i].first/float(count)*100.0f;
+        int op = (int)sorted[i].second;
+        UE_LOG(LogMutableCore,Log, TEXT("  %5.2f%% : %s"), p, s_opNames[op] );
+    }
+#endif
+}
+
+
+//-------------------------------------------------------------------------------------------------
+void ASTOp::Traverse_TopDown_Unique( const vector<Ptr<ASTOp>>& roots,
+                                     std::function<bool(Ptr<ASTOp>&)> f )
+{
+    std::queue<Ptr<ASTOp>> pending;
+    for (auto& r:roots) pending.push( r );
+
+    std::unordered_set<Ptr<const ASTOp>> traversed;
+
+    // We record the parents of all roots as traversed
+    for ( const auto& r: roots )
+    {
+        r->ForEachParent( [&]( const ASTOp* parent )
+        {
+            // If the parent is also a root, we will want to process it.
+            if ( std::find(roots.begin(), roots.end(), parent) == roots.end() )
+            {
+                traversed.insert( parent );
+            }
+        });
+    }
+
+    while (pending.size())
+    {
+        Ptr<ASTOp> pCurrent = pending.front();
+        pending.pop();
+        if (!pCurrent)
+        {
+            continue;
+        }
+
+        // Did we traverse all parents?
+        bool parentsTraversed = true;
+
+        pCurrent->ForEachParent( [&]( const ASTOp* parent )
+        {
+            if (traversed.find(parent)==traversed.end())
+            {
+                // \todo Is the parent in the relevant subtree?
+
+
+                parentsTraversed = false;
+            }
+        });
+
+        if (!parentsTraversed)
+        {
+            pending.push(pCurrent);
+        }
+        else if (traversed.find(pCurrent)==traversed.end())
+        {
+            traversed.insert(pCurrent);
+
+            // Process
+            bool recurse = f(pCurrent);
+
+            // Recurse children
+            if (recurse)
+            {
+                pCurrent->ForEachChild([&]( ASTChild& c )
+                {
+                    if (c && traversed.find(c.m_child)==traversed.end())
+                    {
+                        pending.push( c.m_child );
+                    }
+                });
+            }
+        }
+    }
+}
+
+
+//-------------------------------------------------------------------------------------------------
+void ASTOp::Traverse_TopDown_Unique_Imprecise( const vector<Ptr<ASTOp>>& roots,
+                                     std::function<bool(Ptr<ASTOp>&)> f )
+{
+    std::queue<Ptr<ASTOp>> pending;
+    for (auto& r:roots) pending.push( r );
+
+    std::unordered_set<Ptr<const ASTOp>> traversed;
+
+    while (pending.size())
+    {
+        Ptr<ASTOp> pCurrent = pending.front();
+        pending.pop();
+
+        // It could have been completed in another branch
+        if (pCurrent && traversed.find(pCurrent)==traversed.end())
+        {
+            traversed.insert(pCurrent);
+
+            // Process
+            bool recurse = f(pCurrent);
+
+            // Recurse children
+            if (recurse)
+            {
+                pCurrent->ForEachChild([&]( ASTChild& c )
+                {
+                    if (c && traversed.find(c.m_child)==traversed.end())
+                    {
+                        pending.push( c.m_child );
+                    }
+                });
+            }
+        }
+    }
+}
+
+
+//-------------------------------------------------------------------------------------------------
+void ASTOp::Traverse_TopRandom_Unique_NonReentrant( const vector<Ptr<ASTOp>>& roots,
+                                                  std::function<bool(Ptr<ASTOp>&)> f )
+{
+    ASTOpList pending;
+
+    uint32_t traverseIndex = s_lastTraverseIndex++;
+
+    for (auto& r:roots)
+    {
+        if (r && r->m_traverseIndex!=traverseIndex )
+        {
+            r->m_traverseIndex = traverseIndex;
+            pending.push_back( r );
+        }
+    }
+    for( auto& p : pending)
+    {
+        p->m_traverseIndex = traverseIndex-1;
+    }
+
+    while (pending.size())
+    {
+        Ptr<ASTOp> pCurrent = pending.back();
+        pending.pop_back();
+
+        // It could have been completed in another branch
+        if (pCurrent->m_traverseIndex!=traverseIndex)
+        {
+            pCurrent->m_traverseIndex = traverseIndex;
+
+            // Process
+            bool recurse = f(pCurrent);
+
+            // Recurse children
+            if (recurse)
+            {
+                pCurrent->ForEachChild([&]( ASTChild& c )
+                {
+                    if (c && c.m_child->m_traverseIndex!=traverseIndex)
+                    {
+                        pending.push_back( c.m_child );
+                    }
+                });
+            }
+        }
+    }
+}
+
+
+
+//-------------------------------------------------------------------------------------------------
+//-------------------------------------------------------------------------------------------------
+//-------------------------------------------------------------------------------------------------
+void Visitor_TopDown_Unique_Cloning::Traverse( Ptr<ASTOp>& root )
+{
+    // Visit the given root
+    if (root)
+    {
+        m_pending.push_back( std::make_pair(false,root) );
+
+        Process();
+
+        root = GetOldToNew(root);
+    }
+}
+
+
+
+void Visitor_TopDown_Unique_Cloning::Process()
+{
+    while ( m_pending.size() )
+    {
+        std::pair<bool,Ptr<ASTOp>> item = m_pending.back();
+        m_pending.pop_back();
+
+        Ptr<ASTOp> at = item.second;
+
+		ASTOp::MapChildFunc Identity = [](const Ptr<ASTOp>& o) {return o; };
+		
+		if (item.first)
+        {
+            // Item indicating we finished with all the children of this instruction
+			auto cop = at->Clone(Identity);
+
+            // Fix the references to the children
+            bool childChanged = false;
+            cop->ForEachChild( [&](ASTChild& ref)
+            {
+                auto it = m_oldToNew.find(ref.m_child);
+                if ( ref && it!=m_oldToNew.end() && it->second.get()!=nullptr )
+                {
+                    auto oldRef = ref.m_child;
+                    ref=GetOldToNew(ref.m_child);
+                    if (ref.m_child!=oldRef)
+                    {
+                        childChanged = true;
+                    }
+                }
+            });
+
+            // If any child changed, we need to replace this instruction
+            if ( childChanged )
+            {
+                m_oldToNew[at]=cop;
+            }
+
+        }
+        else
+        {
+            auto it = m_oldToNew.find(at);
+            if (it==m_oldToNew.end())
+            {
+                auto initialAt = at;
+
+                // Fix the references to the children, possibly adding a new instruction
+                {
+                    auto cop = at->Clone(Identity);
+                    bool childChanged = false;
+                    cop->ForEachChild( [&](ASTChild& ref)
+                    {
+                        auto ito = m_oldToNew.find(ref.m_child);
+                        if ( ref && ito!=m_oldToNew.end() && ito->second.get()!=nullptr )
+                        {
+                            auto oldRef = ref.m_child;
+                            ref=GetOldToNew(ref.m_child);
+                            if (ref.m_child!=oldRef)
+                            {
+                                childChanged = true;
+                            }
+                        }
+                    });
+
+                    // If any child changed, we need to re-add this instruction
+                    if ( childChanged )
+                    {
+                        m_oldToNew[at]=cop;
+                        at = cop;
+                    }
+                }
+
+                //auto test1 = at->Clone();
+                //check(*test1==*at);
+
+                bool processChildren = true;
+                auto newAt = Visit( at, processChildren );
+                m_oldToNew[initialAt]=newAt;
+
+                //check(*test1==*at);
+
+                // Proceed with children
+                if (processChildren)
+                {
+                    // TODO: Shouldn't we recurse newAt?
+                    m_pending.push_back( std::make_pair(true,at) );
+
+                    at->ForEachChild( [&](ASTChild& ref)
+                    {
+                        if (ref && m_oldToNew.find(ref.m_child)==m_oldToNew.end())
+                        {
+                            m_pending.push_back( std::make_pair(false,ref.m_child) );
+                        }
+                    });
+                }
+            }
+        }
+
+    }
+}
+
+
+//-------------------------------------------------------------------------------------------------
+//-------------------------------------------------------------------------------------------------
+//-------------------------------------------------------------------------------------------------
+void ASTOp::Traverse_TopDown_Repeat( const vector<Ptr<ASTOp>>& roots,
+                                     std::function<bool(Ptr<ASTOp>& node)> f )
+{
+    ASTOpList pending = roots;
+
+    while (pending.size())
+    {
+        Ptr<ASTOp> pCurrent = pending.back();
+        pending.pop_back();
+
+        if (pCurrent)
+        {
+            // Process
+            bool recurse = f(pCurrent);
+
+            // Recurse children
+            if (recurse)
+            {
+                pCurrent->ForEachChild([&]( ASTChild& c )
+                {
+                    if (c)
+                    {
+                        pending.push_back( c.m_child );
+                    }
+                });
+            }
+        }
+    }
+}
+
+
+//-------------------------------------------------------------------------------------------------
+//-------------------------------------------------------------------------------------------------
+//-------------------------------------------------------------------------------------------------
+void ASTOp::Traverse_BottomUp_Unique_NonReentrant
+(
+        ASTOpList& roots,
+        std::function<void(Ptr<ASTOp>&)> f
+)
+{
+    uint32_t traverseIndex = s_lastTraverseIndex++;
+
+    vector< std::pair<Ptr<ASTOp>,int> > pending;
+    for (auto& r:roots)
+    {
+        if (r && r->m_traverseIndex!=traverseIndex )
+        {
+            r->m_traverseIndex=traverseIndex;
+            pending.push_back( std::make_pair<>(r,0) );
+        }
+    }
+    for(auto& p : pending)
+    {
+        p.first->m_traverseIndex = traverseIndex-1;
+    }
+
+    while (pending.size())
+    {
+        int phase = pending.back().second;
+        Ptr<ASTOp> pCurrent = pending.back().first;
+        pending.pop_back();
+
+        // It could have been completed in another branch
+        if (pCurrent->m_traverseIndex!=traverseIndex)
+        {
+            if (phase==0)
+            {
+                // Process this again...
+                pending.push_back( std::make_pair<>(pCurrent,1) );
+
+                // ...after the children are processed
+                pCurrent->ForEachChild([&]( ASTChild& c )
+                {
+                    if (c && c.m_child->m_traverseIndex!=traverseIndex )
+                    {
+                        auto e = std::make_pair<>(c.m_child,0);
+                        pending.push_back( e );
+                    }
+                });
+            }
+            else
+            {
+                pCurrent->m_traverseIndex=traverseIndex;
+
+                // Children have been be completed
+                f(pCurrent);
+            }
+        }
+    }
+}
+
+
+//-------------------------------------------------------------------------------------------------
+//-------------------------------------------------------------------------------------------------
+//-------------------------------------------------------------------------------------------------
+void ASTOp::Traverse_BottomUp_Unique_NonReentrant
+(
+        ASTOpList& roots,
+        std::function<void(Ptr<ASTOp>&)> f,
+        std::function<bool(const ASTOp*)> accept
+)
+{
+    uint32_t traverseIndex = s_lastTraverseIndex++;
+
+    vector< std::pair<Ptr<ASTOp>,int> > pending;
+    for (auto& r:roots)
+    {
+        if (r && r->m_traverseIndex!=traverseIndex)
+        {
+            r->m_traverseIndex=traverseIndex;
+            pending.push_back( std::make_pair<>(r,0) );
+        }
+    }
+    for(auto& p : pending)
+    {
+        p.first->m_traverseIndex = traverseIndex-1;
+    }
+
+    while (pending.size())
+    {
+        int phase = pending.back().second;
+        Ptr<ASTOp> pCurrent = pending.back().first;
+        pending.pop_back();
+
+        // It could have been completed in another branch
+        if (pCurrent->m_traverseIndex!=traverseIndex && accept(pCurrent.get()))
+        {
+            if (phase==0)
+            {
+                // Process this again...
+                pending.push_back( std::make_pair<>(pCurrent,1) );
+
+                // ...after the children are processed
+                pCurrent->ForEachChild([&]( ASTChild& c )
+                {
+                    if (c && accept(c.m_child.get()) && c.m_child->m_traverseIndex!=traverseIndex )
+                    {
+                        pending.push_back( std::make_pair<>(c.m_child,0) );
+                    }
+                });
+            }
+            else
+            {
+                pCurrent->m_traverseIndex=traverseIndex;
+
+                // Children have been be completed
+                f(pCurrent);
+            }
+        }
+    }
+}
+
+
+//-------------------------------------------------------------------------------------------------
+//-------------------------------------------------------------------------------------------------
+//-------------------------------------------------------------------------------------------------
+void ASTOp::Traverse_BottomUp_Unique
+(
+        ASTOpList& roots,
+        std::function<void(Ptr<ASTOp>&)> f,
+        std::function<bool(const ASTOp*)> accept
+)
+{
+    TSet<Ptr<ASTOp>> Traversed;
+    TArray< std::pair<Ptr<ASTOp>,int> > Pending;
+    for (auto& r:roots)
+    {
+        if (r)
+        {
+			auto It = Pending.FindByPredicate([&](const std::pair<Ptr<ASTOp>, int>& p)
+				{
+					return r == p.first;
+				});
+            if ( !It )
+            {
+                Pending.Add( std::make_pair<>(r,0) );
+            }
+        }
+    }
+
+    while (!Pending.IsEmpty())
+    {
+        int phase = Pending.Last().second;
+        Ptr<ASTOp> pCurrent = Pending.Last().first;
+        Pending.Pop();
+
+        // It could have been completed in another branch
+        if (accept(pCurrent.get()) && !Traversed.Find(pCurrent))
+        {
+            if (phase==0)
+            {
+                // Process this again...
+				Pending.Add( std::make_pair<>(pCurrent,1) );
+
+                // ...after the children are processed
+                pCurrent->ForEachChild([&]( ASTChild& c )
+                {
+                    if (c && accept(c.m_child.get()) && !Traversed.Find(c.m_child) )
+                    {
+                        Pending.Add( std::make_pair<>(c.m_child,0) );
+                    }
+                });
+            }
+            else
+            {
+                Traversed.Add(pCurrent);
+
+                // Children have been be completed
+                f(pCurrent);
+            }
+        }
+    }
+}
+
+
+//-------------------------------------------------------------------------------------------------
+//-------------------------------------------------------------------------------------------------
+//-------------------------------------------------------------------------------------------------
+void ASTOp::Traverse_BottomUp_Unique
+(
+        Ptr<ASTOp>& root,
+        std::function<void(Ptr<ASTOp>&)> f,
+        std::function<bool(const ASTOp*)> accept
+)
+{
+    if (root)
+    {
+        ASTOpList roots;
+        roots.push_back(root);
+        Traverse_BottomUp_Unique(roots,f,accept);
+    }
+}
+
+
+
+//-------------------------------------------------------------------------------------------------
+//-------------------------------------------------------------------------------------------------
+//-------------------------------------------------------------------------------------------------
+int ASTOp::GetParentCount()
+{
+    int result=0;
+    ForEachParent( [&](ASTOp* p)
+    {
+        if (p!=nullptr) ++result;
+    });
+    return result;
+}
+
+
+void ASTOp::Replace( const Ptr<ASTOp>& node, const Ptr<ASTOp>& other )
+{
+    if(other==node)
+    {
+        return;
+    }
+
+    auto parentsCopy = node->m_parents;
+
+    for(auto& p:parentsCopy)
+    {
+        if(p)
+        {
+            p->ForEachChild( [=](ASTChild& c)
+            {
+                if(c.m_child==node)
+                {
+                    c = other;
+                }
+            });
+        }
+    }
+}
+
+
+FImageDesc ASTOp::GetImageDesc( bool, GetImageDescContext* )
+{
+    check(false);
+    return FImageDesc();
+}
+
+
+bool ASTOp::IsImagePlainConstant( vec4<float>& ) const
+{
+    check(false);
+    return false;
+}
+
+
+bool ASTOp::IsColourConstant( vec4<float>& ) const
+{
+    check(false);
+    return false;
+}
+
+
+void ASTOp::GetBlockLayoutSize( int, int*, int*, BLOCK_LAYOUT_SIZE_CACHE* )
+{
+    check(false);
+}
+
+void ASTOp::GetBlockLayoutSizeCached( int blockIndex, int* pBlockX, int* pBlockY,
+                               BLOCK_LAYOUT_SIZE_CACHE* cache )
+{
+	BLOCK_LAYOUT_SIZE_CACHE::KeyType key(this,blockIndex);
+	BLOCK_LAYOUT_SIZE_CACHE::ValueType* ValuePtr = cache->Find( key );
+    if (ValuePtr)
+    {
+        *pBlockX = ValuePtr->Key;
+        *pBlockY = ValuePtr->Value;
+        return;
+    }
+
+    GetBlockLayoutSize( blockIndex, pBlockX, pBlockY, cache );
+
+	cache->Add(key, BLOCK_LAYOUT_SIZE_CACHE::ValueType( *pBlockX, *pBlockY) );
+}
+
+
+void ASTOp::GetLayoutBlockSize( int*, int* )
+{
+    check(false);
+}
+
+
+bool ASTOp::GetNonBlackRect( FImageRect& ) const
+{
+    return false;
+}
+
+
+//-------------------------------------------------------------------------------------------------
+//-------------------------------------------------------------------------------------------------
+//-------------------------------------------------------------------------------------------------
+ASTOpFixed::ASTOpFixed()
+    : children({{ASTChild(this,nullptr),ASTChild(this,nullptr),ASTChild(this,nullptr),ASTChild(this,nullptr),ASTChild(this,nullptr),ASTChild(this,nullptr),ASTChild(this,nullptr),ASTChild(this,nullptr)}})
+{
+}
+
+
+ASTOpFixed::~ASTOpFixed()
+{
+    // Explicit call needed to avoid recursive destruction
+    ASTOp::RemoveChildren();
+}
+
+
+void ASTOpFixed::ForEachChild( const std::function<void(ASTChild&)>& f )
+{
+    // ugly but temporary while ASTOpFixed exists
+    ForEachReference( op, [&](OP::ADDRESS* pAt)
+    {
+        if (pAt && *pAt)
+        {
+            check(children.size()>*pAt);
+            f(children[*pAt]);
+        }
+    });
+}
+
+
+void ASTOpFixed::Link( PROGRAM& program, const FLinkerOptions* )
+{
+    if (!linkedAddress)
+    {
+        OP lop = op;
+
+        // Fix the linked op and generate children
+        ForEachReference(lop, [&](OP::ADDRESS* pChildOp)
+        {
+            Ptr<ASTOp> pASTChild = children[*pChildOp].m_child;
+            if (pASTChild)
+            {
+                *pChildOp = pASTChild->linkedAddress;
+            }
+            else
+            {
+                *pChildOp = 0;
+            }
+        });
+
+        linkedAddress = (OP::ADDRESS)program.m_opAddress.size();        
+        program.m_opAddress.push_back((uint32_t)program.m_byteCode.size());
+
+        AppendCode(program.m_byteCode,lop.type);
+        if (lop.type==OP_TYPE::ME_MERGE)
+        {
+            AppendCode(program.m_byteCode,lop.args.MeshMerge);
+        }
+        else
+        {
+            // Generic encoding
+            AppendCode(program.m_byteCode,lop.args);
+        }
+    }
+}
+
+
+bool ASTOpFixed::IsEqual(const ASTOp& otherUntyped) const
+{
+    if ( auto other = dynamic_cast<const ASTOpFixed*>(&otherUntyped) )
+    {
+        return op==other->op && children==other->children;
+    }
+    return false;
+}
+
+
+mu::Ptr<ASTOp> ASTOpFixed::Clone( MapChildFunc& mapChild ) const
+{
+    Ptr<ASTOpFixed> n = new ASTOpFixed();
+    n->op = op;
+    n->childCount = childCount;
+    // Skip the children 0, used to represent null op.
+    for(uint8_t i=1; i<childCount; ++i)
+    {
+        n->children[i] = mapChild(children[i].child());
+    }
+    return n;
+}
+
+
+uint64 ASTOpFixed::Hash() const
+{
+	uint64 res = std::hash<uint64>()(uint64(op.type));
+    for (const auto& c: children)
+    {
+        hash_combine( res, c.child().get() );
+    }
+    return res;
+}
+
+
+//!
+FImageDesc ASTOpFixed::GetImageDesc( bool returnBestOption, GetImageDescContext* context )
+{
+    FImageDesc res;
+
+    // Local context in case it is necessary
+    GetImageDescContext localContext;
+    if (!context)
+    {
+      context = &localContext;
+    }
+    else
+    {
+        // Cached result?
+		FImageDesc* PtrValue = context->m_results.Find(this);
+        if (PtrValue)
+        {
+            return *PtrValue;
+        }
+    }
+
+    OP_TYPE type = (OP_TYPE) op.type;
+
+    switch ( type )
+    {
+
+    case OP_TYPE::NONE:
+        break;
+
+    case OP_TYPE::IM_LAYER:
+        res = GetImageDesc( op.args.ImageLayer.base, returnBestOption, context );
+        if ( res.m_format==EImageFormat::IF_L_UBYTE )
+        {
+            res.m_format = EImageFormat::IF_RGB_UBYTE;
+        }
+        break;
+
+    case OP_TYPE::IM_LAYERCOLOUR:
+        res = GetImageDesc( op.args.ImageLayerColour.base, returnBestOption, context );
+        if ( res.m_format==EImageFormat::IF_L_UBYTE )
+        {
+            res.m_format = EImageFormat::IF_RGB_UBYTE;
+        }
+        break;
+
+    case OP_TYPE::IM_SATURATE:
+        res = GetImageDesc( op.args.ImageSaturate.base, returnBestOption, context );
+        break;
+
+    case OP_TYPE::IM_LUMINANCE:
+        res = GetImageDesc( op.args.ImageLuminance.base,returnBestOption, context );
+        res.m_format = EImageFormat::IF_L_UBYTE;
+        break;
+
+    case OP_TYPE::IM_INTERPOLATE:
+        res = GetImageDesc(op.args.ImageInterpolate.targets[0], returnBestOption, context );
+        break;
+
+    case OP_TYPE::IM_INTERPOLATE3:
+        res = GetImageDesc( op.args.ImageInterpolate3.target0, returnBestOption, context );
+        break;
+
+    case OP_TYPE::IM_DIFFERENCE:
+        res = GetImageDesc( op.args.ImageDifference.a, returnBestOption, context );
+        res.m_format = EImageFormat::IF_L_UBYTE;
+        break;
+
+    case OP_TYPE::IM_PLAINCOLOUR:
+        res.m_format = EImageFormat(op.args.ImagePlainColour.format);
+        res.m_size[0] = op.args.ImagePlainColour.size[0];
+        res.m_size[1] = op.args.ImagePlainColour.size[1];
+        res.m_lods = 1;
+        check( res.m_format != EImageFormat::IF_NONE );
+        break;
+
+    case OP_TYPE::IM_CROP:
+        res = GetImageDesc( op.args.ImageCrop.source, returnBestOption, context );
+
+		check(op.args.ImageCrop.sizeX > 0);
+		check(op.args.ImageCrop.sizeY > 0);
+
+        res.m_size = FImageSize
+            (
+                op.args.ImageCrop.sizeX,
+                op.args.ImageCrop.sizeY
+            );
+        break;
+
+    case OP_TYPE::IM_RESIZE:
+        res = GetImageDesc( op.args.ImageResize.source, returnBestOption, context );
+
+        res.m_size = FImageSize
+            (
+                op.args.ImageResize.size[0],
+                op.args.ImageResize.size[1]
+            );
+        break;
+
+    case OP_TYPE::IM_RESIZEREL:
+        res = GetImageDesc( op.args.ImageResizeRel.source, returnBestOption, context );
+        res.m_size[0] = (uint16_t)( res.m_size[0]*op.args.ImageResizeRel.factor[0] );
+        res.m_size[1] = (uint16_t)( res.m_size[1]*op.args.ImageResizeRel.factor[1] );
+        break;
+
+    case OP_TYPE::IM_RESIZELIKE:
+        res = GetImageDesc( op.args.ImageResizeLike.source, returnBestOption, context );
+        if (op.args.ImageResizeLike.sizeSource)
+        {
+            res.m_size = children[op.args.ImageResizeLike.sizeSource]->GetImageDesc( returnBestOption, context ).m_size;
+        }
+        break;
+
+    case OP_TYPE::IM_SWIZZLE:
+        res = GetImageDesc( op.args.ImageSwizzle.sources[0], returnBestOption, context );
+        res.m_format = op.args.ImageSwizzle.format;
+        check( res.m_format != EImageFormat::IF_NONE );
+        break;
+
+    case OP_TYPE::IM_SELECTCOLOUR:
+        res = GetImageDesc( op.args.ImageSelectColour.base, returnBestOption, context );
+        res.m_format = EImageFormat::IF_L_UBYTE;
+        break;
+
+    case OP_TYPE::IM_GRADIENT:
+        res.m_size[0] = op.args.ImageGradient.size[0];
+        res.m_size[1] = op.args.ImageGradient.size[1];
+        res.m_format = EImageFormat::IF_RGB_UBYTE;
+        break;
+
+    case OP_TYPE::IM_BLANKLAYOUT:
+        // TODO: We would need to process the layout to find the grid size, and then use
+        // the block size with it.
+        res.m_size = FImageSize(0,0);
+        res.m_format = op.args.ImageBlankLayout.format;
+        break;
+
+    case OP_TYPE::IM_BINARISE:
+        res = GetImageDesc( op.args.ImageBinarise.base, returnBestOption, context );
+        res.m_format = EImageFormat::IF_L_UBYTE;
+        break;
+
+    case OP_TYPE::IM_GPU:
+    {
+        check(false);
+        // The size comes from a special instruction
+//                const GPU_PROGRAM& gpuProg = program.m_gpuPrograms[ op.args.ImageGPU.program ];
+//                if (gpuProg.m_size)
+//                {
+//                    res = children[op.args.VolumeLayer.source]->GetImageDesc( returnBestOption, context );
+//                }
+//                // The format is fixed in the program description
+//                res.m_format = gpuProg.m_format;
+        break;
+    }
+
+    case OP_TYPE::IM_RASTERMESH:
+        res = GetImageDesc( op.args.ImageRasterMesh.image, returnBestOption, context );
+        res.m_size[0]=op.args.ImageRasterMesh.sizeX;
+        res.m_size[1]=op.args.ImageRasterMesh.sizeY;
+        break;
+
+    case OP_TYPE::IM_MAKEGROWMAP:
+        res = GetImageDesc( op.args.ImageMakeGrowMap.mask, returnBestOption, context );
+        res.m_format = EImageFormat::IF_L_UBYTE;
+        break;
+
+    case OP_TYPE::IM_DISPLACE:
+        res = GetImageDesc( op.args.ImageDisplace.source, returnBestOption, context );
+        break;
+
+	case OP_TYPE::IM_INVERT:
+		res = GetImageDesc(op.args.ImageInvert.base, returnBestOption, context);
+		break;
+
+    case OP_TYPE::CO_IMAGESIZE:
+        res = GetImageDesc( op.args.ColourImageSize.image, returnBestOption, context );
+        res.m_format = EImageFormat::IF_NONE;
+        break;
+
+	case OP_TYPE::IM_COLOURMAP:
+		res = GetImageDesc( op.args.ImageColourMap.base, returnBestOption, context );
+		break;
+
+    default:
+        check(false);
+    }
+
+    // Cache the result
+    if (context)
+    {
+        context->m_results.Add(this,res);
+    }
+
+    return res;
+}
+
+
+void ASTOpFixed::GetBlockLayoutSize( int blockIndex, int* pBlockX, int* pBlockY,
+                                     BLOCK_LAYOUT_SIZE_CACHE* cache )
+{
+    switch ( op.type )
+    {
+    case OP_TYPE::LA_PACK:
+    {
+        GetBlockLayoutSize( blockIndex, op.args.LayoutPack.layout,
+                            pBlockX, pBlockY, cache );
+        break;
+    }
+
+    case OP_TYPE::LA_MERGE:
+    {
+        GetBlockLayoutSize( blockIndex, op.args.LayoutMerge.base,
+                            pBlockX, pBlockY, cache );
+
+        if (!*pBlockX)
+        {
+            GetBlockLayoutSize( blockIndex, op.args.LayoutMerge.added,
+                                pBlockX, pBlockY, cache );
+        }
+
+        break;
+    }
+
+    case OP_TYPE::LA_REMOVEBLOCKS:
+    {
+        GetBlockLayoutSize( blockIndex, op.args.LayoutRemoveBlocks.source,
+                            pBlockX, pBlockY, cache );
+        break;
+    }
+
+    default:
+        checkf( false, TEXT("Instruction not supported") );
+    }
+}
+
+
+void ASTOpFixed::GetLayoutBlockSize( int* pBlockX, int* pBlockY )
+{
+    switch ( op.type )
+    {
+    
+	case OP_TYPE::IM_RESIZEREL:
+	{
+		GetLayoutBlockSize(op.args.ImageResizeRel.source, pBlockX, pBlockY);
+		(*pBlockX) = int(float(*pBlockX) * op.args.ImageResizeRel.factor[0]);
+		(*pBlockY) = int(float(*pBlockX) * op.args.ImageResizeRel.factor[1]);
+
+		break;
+	}
+
+	case OP_TYPE::IM_RESIZE:
+	{
+		GetLayoutBlockSize(op.args.ImageResize.source, pBlockX, pBlockY);
+
+		if (*pBlockX > 0 && *pBlockY> 0)
+		{
+			FImageDesc sourceDesc = GetImageDesc(op.args.ImageResize.source);
+			if (sourceDesc.m_size[0]>0 && sourceDesc.m_size[1]>0)
+			{
+				float factorX = float(op.args.ImageResize.size[0]) / float(sourceDesc.m_size[0]);
+				float factorY = float(op.args.ImageResize.size[1]) / float(sourceDesc.m_size[1]);
+				*pBlockX = int(*pBlockX * factorX);
+				*pBlockY = int(*pBlockX * factorY);
+			}
+			else
+			{
+				*pBlockX = 0;
+				*pBlockY = 0;
+			}
+		}
+
+		break;
+	}
+
+    case OP_TYPE::IM_SWIZZLE:
+    {
+        GetLayoutBlockSize( op.args.ImageSwizzle.sources[0], pBlockX, pBlockY );
+        break;
+    }
+
+    case OP_TYPE::IM_LAYER:
+    {
+        GetLayoutBlockSize(op.args.ImageLayer.base, pBlockX, pBlockY);
+        break;
+    }
+
+    case OP_TYPE::IM_LAYERCOLOUR:
+    {
+        GetLayoutBlockSize(op.args.ImageLayerColour.base, pBlockX, pBlockY);
+        break;
+    }
+
+	case OP_TYPE::IM_BLANKLAYOUT:
+	{
+		*pBlockX = int(op.args.ImageBlankLayout.blockSize[0]);
+		*pBlockY = int(op.args.ImageBlankLayout.blockSize[1]);
+		break;
+	}
+
+	case OP_TYPE::IM_PLAINCOLOUR:
+	{
+		*pBlockX = 0;
+		*pBlockY = 0;
+		break;
+	}
+
+    default:
+        checkf( false, TEXT("Instruction not supported") );
+    }
+
+    //check( *pBlockX!=0 && *pBlockY!=0 );
+}
+
+
+ASTOp::BOOL_EVAL_RESULT ASTOpFixed::EvaluateBool( ASTOpList& facts, EVALUATE_BOOL_CACHE* cache ) const
+{
+    EVALUATE_BOOL_CACHE localCache;
+    if (!cache)
+    {
+        cache = &localCache;
+    }
+    else
+    {
+        // Is this in the cache?
+        auto it = cache->find(this);
+        if (it!=cache->end())
+        {
+            return it->second;
+        }
+    }
+
+    BOOL_EVAL_RESULT result = BET_UNKNOWN;
+    switch(op.type)
+    {
+    case OP_TYPE::BO_NOT:
+    {
+        if (children[op.args.BoolNot.source])
+        {
+            BOOL_EVAL_RESULT child = children[op.args.BoolNot.source]->EvaluateBool(facts,cache);
+            if (child==BET_TRUE)
+            {
+                result = BET_FALSE;
+            }
+            else if (child==BET_FALSE)
+            {
+                result = BET_TRUE;
+            }
+        }
+        break;
+    }
+
+    case OP_TYPE::BO_EQUAL_INT_CONST:
+    {
+        int intValue = op.args.BoolEqualScalarConst.constant;
+        const auto& intExp = children[op.args.BoolEqualScalarConst.value].child();
+        bool intUnknown = true;
+        int intResult = intExp->EvaluateInt( facts, intUnknown );
+        if (intUnknown)
+        {
+            result = BET_UNKNOWN;
+        }
+        else if (intResult==intValue)
+        {
+            result = BET_TRUE;
+        }
+        else
+        {
+            result = BET_FALSE;
+        }
+        break;
+    }
+
+    case OP_TYPE::BO_AND:
+    {
+        const auto& a = children[op.args.BoolBinary.a].child();
+        const auto& b = children[op.args.BoolBinary.b].child();
+        BOOL_EVAL_RESULT resultA = BET_UNKNOWN;
+        BOOL_EVAL_RESULT resultB = BET_UNKNOWN;
+        for ( size_t f=0; f<facts.size(); ++f )
+        {
+            if ( a && resultA == BET_UNKNOWN )
+            {
+                resultA = a->EvaluateBool( facts, cache );
+
+                if ( resultA==BET_TRUE && resultB==BET_TRUE )
+                {
+                    result = BET_TRUE;
+                    break;
+                }
+                if ( resultA==BET_FALSE || resultB==BET_FALSE )
+                {
+                    result = BET_FALSE;
+                    break;
+                }
+            }
+            if ( b && resultB == BET_UNKNOWN )
+            {
+                resultB = b->EvaluateBool( facts, cache );
+
+                if ( resultA==BET_TRUE && resultB==BET_TRUE )
+                {
+                    result = BET_TRUE;
+                    break;
+                }
+                if ( resultA==BET_FALSE || resultB==BET_FALSE )
+                {
+                    result = BET_FALSE;
+                    break;
+                }
+            }
+        }
+        break;
+    }
+
+    case OP_TYPE::BO_OR:
+    {
+        const auto& a = children[op.args.BoolBinary.a].child();
+        const auto& b = children[op.args.BoolBinary.b].child();
+        BOOL_EVAL_RESULT resultA = BET_UNKNOWN;
+        BOOL_EVAL_RESULT resultB = BET_UNKNOWN;
+        for ( size_t f=0; f<facts.size(); ++f )
+        {
+            if ( a && resultA == BET_UNKNOWN )
+            {
+                resultA = a->EvaluateBool( facts, cache );
+
+                if ( resultA==BET_TRUE || resultB==BET_TRUE )
+                {
+                    result = BET_TRUE;
+                    break;
+                }
+                if ( resultA==BET_FALSE && resultB==BET_FALSE )
+                {
+                    result = BET_FALSE;
+                    break;
+                }
+            }
+            if ( b && resultB == BET_UNKNOWN )
+            {
+                resultB = b->EvaluateBool( facts, cache );
+
+                if ( resultA==BET_TRUE || resultB==BET_TRUE )
+                {
+                    result = BET_TRUE;
+                    break;
+                }
+                if ( resultA==BET_FALSE && resultB==BET_FALSE )
+                {
+                    result = BET_FALSE;
+                    break;
+                }
+            }
+        }
+        break;
+    }
+
+    default:
+        checkf( false, TEXT("Instruction not supported") );
+    }
+
+    (*cache)[this] = result;
+
+    return result;
+}
+
+
+int ASTOpFixed::EvaluateInt( ASTOpList& /*facts*/, bool &unknown ) const
+{
+    unknown = true;
+
+    switch (GetOpType())
+    {
+    case OP_TYPE::NU_CONSTANT:
+        unknown = false;
+        return op.args.IntConstant.value;
+        break;
+
+    case OP_TYPE::SC_CONSTANT:
+        unknown = false;
+        return (int)op.args.ScalarConstant.value;
+        break;
+
+        // TODO
+//    case IntExpression::IET_PARAMETER:
+//        for ( size_t f=0; (result==BET_UNKNOWN) && f<facts.size(); ++f )
+//        {
+//            result = EvaluateIntEquality(facts[f].get(),intExp,value);
+//        }
+//        break;
+
+    default:
+        break;
+
+    }
+
+
+    return 0;
+}
+
+
+bool ASTOpFixed::IsImagePlainConstant( vec4<float>& colour ) const
+{
+    bool res = false;
+    switch( op.type )
+    {
+
+    case OP_TYPE::IM_BLANKLAYOUT:
+        res = true;
+        colour[0] = 0;
+        colour[1] = 0;
+        colour[2] = 0;
+        colour[3] = 0;
+        break;
+
+    case OP_TYPE::IM_RESIZE:
+        res = children[op.args.ImageResize.source]->IsImagePlainConstant( colour );
+        break;
+
+    case OP_TYPE::IM_RESIZEREL:
+        res = children[op.args.ImageResizeRel.source]->IsImagePlainConstant( colour );
+        break;
+
+    case OP_TYPE::IM_RESIZELIKE:
+        res = children[op.args.ImageResizeLike.source]->IsImagePlainConstant(  colour );
+        break;
+
+    case OP_TYPE::IM_PLAINCOLOUR:
+        res = children[op.args.ImagePlainColour.colour]->IsColourConstant( colour );
+        break;
+
+    case OP_TYPE::IM_INTERPOLATE3:
+        res = children[op.args.ImageInterpolate3.target0]->IsColourConstant( colour );
+        if (res)
+        {
+            vec4<float> baseColour;
+            res = children[op.args.ImageInterpolate3.target1]->IsColourConstant( colour );
+            res &= (colour==baseColour);
+        }
+        if (res)
+        {
+            vec4<float> baseColour;
+            res = children[op.args.ImageInterpolate3.target2]->IsColourConstant( colour );
+            res &= (colour==baseColour);
+        }
+        break;
+
+    default:
+        // TODO: Improve this test with more operations
+        //check( false );
+        res = false;
+        break;
+    }
+
+    return res;
+}
+
+
+//-------------------------------------------------------------------------------------------------
+bool ASTOpFixed::IsColourConstant( vec4<float>& colour ) const
+{
+    bool res = false;
+    switch( op.type )
+    {
+    case OP_TYPE::CO_CONSTANT:
+        res = true;
+        colour[0] = op.args.ColourConstant.value[0];
+        colour[1] = op.args.ColourConstant.value[1];
+        colour[2] = op.args.ColourConstant.value[2];
+        colour[3] = op.args.ColourConstant.value[3];
+        break;
+
+    case OP_TYPE::CO_SAMPLEIMAGE:
+    case OP_TYPE::CO_SWIZZLE:
+    case OP_TYPE::CO_IMAGESIZE:
+    case OP_TYPE::CO_LAYOUTBLOCKTRANSFORM:
+    case OP_TYPE::CO_FROMSCALARS:
+    case OP_TYPE::CO_ARITHMETIC:
+
+    default:
+        // TODO: Improve this test with more operations
+        //check( false );
+        res = false;
+        break;
+    }
+
+    return res;
+}
+
+
+//-------------------------------------------------------------------------------------------------
+mu::Ptr<ImageSizeExpression> ASTOpFixed::GetImageSizeExpression() const
+{
+    Ptr<ImageSizeExpression> pRes = new ImageSizeExpression;
+
+    OP_TYPE type = GetOpType();
+    switch ( type )
+    {
+    case OP_TYPE::NONE:
+    {
+        pRes->type = ImageSizeExpression::ISET_CONSTANT;
+        pRes->size[0]=0;
+        pRes->size[1]=0;
+        break;
+    }
+
+    case OP_TYPE::IM_LAYER:
+        if ( children[op.args.ImageLayer.base] )
+        {
+            pRes = children[op.args.ImageLayer.base].child()->GetImageSizeExpression();
+        }
+        break;
+
+    case OP_TYPE::IM_LAYERCOLOUR:
+        if ( children[op.args.ImageLayerColour.base] )
+        {
+            pRes = children[op.args.ImageLayerColour.base].child()->GetImageSizeExpression();
+        }
+        break;
+
+    case OP_TYPE::IM_RESIZE:
+        pRes->type = ImageSizeExpression::ISET_CONSTANT;
+        pRes->size[0] = op.args.ImageResize.size[0];
+        pRes->size[1] = op.args.ImageResize.size[1];
+        break;
+
+    case OP_TYPE::IM_RESIZEREL:
+        if ( children[op.args.ImageResizeRel.source] )
+        {
+            pRes = children[op.args.ImageResizeRel.source].child()->GetImageSizeExpression();
+            if (pRes->type == ImageSizeExpression::ISET_CONSTANT)
+            {
+                pRes->size[0] = uint16_t( pRes->size[0] * op.args.ImageResizeRel.factor[0] );
+                pRes->size[1] = uint16_t( pRes->size[1] * op.args.ImageResizeRel.factor[1] );
+            }
+            else
+            {
+                // TODO: Proportional factor
+                pRes = new ImageSizeExpression();
+                pRes->type = ImageSizeExpression::ISET_UNKNOWN;
+            }
+        }
+        break;
+
+    case OP_TYPE::IM_RESIZELIKE:
+        if ( children[op.args.ImageResizeLike.sizeSource] )
+        {
+            pRes = children[op.args.ImageResizeLike.sizeSource].child()->GetImageSizeExpression();
+        }
+        break;
+
+    case OP_TYPE::IM_PLAINCOLOUR:
+        pRes->type = ImageSizeExpression::ISET_CONSTANT;
+        pRes->size[0] = op.args.ImagePlainColour.size[0];
+        pRes->size[1] = op.args.ImagePlainColour.size[1];
+        break;
+
+    case OP_TYPE::IM_BLANKLAYOUT:
+        pRes->type = ImageSizeExpression::ISET_LAYOUTFACTOR;
+        pRes->layout = children[op.args.ImageBlankLayout.layout].child();
+        pRes->factor[0] = op.args.ImageBlankLayout.blockSize[0];
+        pRes->factor[1] = op.args.ImageBlankLayout.blockSize[1];
+        break;
+
+    case OP_TYPE::IM_DIFFERENCE:
+        if ( children[op.args.ImageDifference.a] )
+        {
+            pRes = children[op.args.ImageDifference.a].child()->GetImageSizeExpression();
+        }
+        break;
+
+    case OP_TYPE::IM_INTERPOLATE:
+        if ( children[op.args.ImageInterpolate.targets[0]] )
+        {
+            pRes = children[op.args.ImageInterpolate.targets[0]].child()->GetImageSizeExpression();
+        }
+        break;
+
+    case OP_TYPE::IM_INTERPOLATE3:
+        if ( children[op.args.ImageInterpolate3.target0] )
+        {
+            pRes = children[op.args.ImageInterpolate3.target0].child()->GetImageSizeExpression();
+        }
+        break;
+
+    case OP_TYPE::IM_SATURATE:
+        if ( children[op.args.ImageSaturate.base] )
+        {
+            pRes = children[op.args.ImageSaturate.base].child()->GetImageSizeExpression();
+        }
+        break;
+
+    case OP_TYPE::IM_LUMINANCE:
+        if ( children[op.args.ImageLuminance.base] )
+        {
+            pRes = children[op.args.ImageLuminance.base].child()->GetImageSizeExpression();
+        }
+        break;
+
+    case OP_TYPE::IM_BINARISE:
+        if ( children[op.args.ImageBinarise.base] )
+        {
+            pRes = children[op.args.ImageBinarise.base].child()->GetImageSizeExpression();
+        }
+        break;
+
+    case OP_TYPE::IM_COLOURMAP:
+        if ( children[op.args.ImageColourMap.base] )
+        {
+            pRes = children[op.args.ImageColourMap.base].child()->GetImageSizeExpression();
+        }
+        break;
+
+    case OP_TYPE::IM_SELECTCOLOUR:
+        if ( children[op.args.ImageSelectColour.colour] )
+        {
+            pRes = children[op.args.ImageSelectColour.colour].child()->GetImageSizeExpression();
+        }
+        break;
+
+    case OP_TYPE::IM_SWIZZLE:
+        if ( children[op.args.ImageSwizzle.sources[0]] )
+        {
+            pRes = children[op.args.ImageSwizzle.sources[0]].child()->GetImageSizeExpression();
+        }
+        break;
+
+    case OP_TYPE::IM_GRADIENT:
+        pRes->type = ImageSizeExpression::ISET_CONSTANT;
+        pRes->size[0] = op.args.ImageGradient.size[0];
+        pRes->size[1] = op.args.ImageGradient.size[1];
+        break;
+
+    case OP_TYPE::IM_MAKEGROWMAP:
+        if ( children[op.args.ImageMakeGrowMap.mask] )
+        {
+            pRes = children[op.args.ImageMakeGrowMap.mask].child()->GetImageSizeExpression();
+        }
+        break;
+
+    case OP_TYPE::IM_DISPLACE:
+        if ( children[op.args.ImageDisplace.source] )
+        {
+            pRes = children[op.args.ImageDisplace.source].child()->GetImageSizeExpression();
+        }
+        break;
+
+	case OP_TYPE::IM_INVERT:
+		if (children[op.args.ImageInvert.base])
+		{
+			pRes = children[op.args.ImageInvert.base].child()->GetImageSizeExpression();
+		}
+		break;
+
+	case OP_TYPE::IM_RASTERMESH:
+		pRes->type = ImageSizeExpression::ISET_CONSTANT;
+		pRes->size[0] = op.args.ImageRasterMesh.sizeX ? op.args.ImageRasterMesh.sizeX : 256;
+		pRes->size[1] = op.args.ImageRasterMesh.sizeY ? op.args.ImageRasterMesh.sizeY : 256;
+		break;
+    	
+    case OP_TYPE::IM_CROP:
+		pRes->type = ImageSizeExpression::ISET_CONSTANT;
+		pRes->size[0] = op.args.ImageCrop.sizeX;
+		pRes->size[1] = op.args.ImageCrop.sizeY;
+    	break;
+
+    default:
+    	check( false );
+    	break;
+    }
+
+    return pRes;
+}
+
