[mutable] Move Mutable plugin from the restricted folder to the experimental plugins folder.

diff --git a/Engine/Plugins/Experimental/Mutable/Source/MutableTools/Private/AST.h b/Engine/Plugins/Experimental/Mutable/Source/MutableTools/Private/AST.h
new file mode 100644
--- /dev/null
+++ b/Engine/Plugins/Experimental/Mutable/Source/MutableTools/Private/AST.h
@@ -0,0 +1,1143 @@
+// Copyright Epic Games, Inc. All Rights Reserved.
+
+#pragma once
+
+#include "Operations.h"
+#include "MutableMath.h"
+#include "ModelPrivate.h"
+#include "ParametersPrivate.h"
+#include "Image.h"
+#include "Containers/Map.h"
+
+#include <unordered_set>
+#include <unordered_map>
+#include <memory>
+#include <array>
+#include <atomic>
+
+namespace std
+{
+
+  template<typename T>
+  struct hash<mu::Ptr<T>>
+  {
+    uint64 operator()(const mu::Ptr<T>& k) const
+    {
+      return hash<const void*>()(k.get());
+    }
+  };
+
+}
+
+
+namespace mu
+{
+	template<typename T>
+	inline uint32 GetTypeHash(const Ptr<const T>& p)
+	{
+		return ::GetTypeHash(p.get());
+	}
+
+
+	template<typename T>
+	inline uint32 GetTypeHash(const Ptr<T>& p)
+	{
+		return ::GetTypeHash(p.get());
+	}
+
+
+    //---------------------------------------------------------------------------------------------
+    //! This class stores the expression that defines the size of an image.
+    //--------------------------------------------------------------------------------------------
+    class ImageSizeExpression : public RefCounted
+    {
+    public:
+
+        enum
+        {
+            ISET_UNKNOWN,
+            ISET_CONSTANT,
+            ISET_LAYOUTFACTOR,
+            ISET_CONDITIONAL
+        } type;
+
+        // For constant sizes
+        FImageSize size;
+
+        // For layout factor sizes
+        Ptr<class ASTOp> layout;
+        uint16_t factor[2];
+
+        // For conditionals
+        Ptr<class ASTOp> condition;
+        Ptr<ImageSizeExpression> yes;
+        Ptr<ImageSizeExpression> no;
+
+        //!
+        ImageSizeExpression()
+        {
+            type = ISET_UNKNOWN;
+            size[0] = 0;
+            size[1] = 0;
+            layout = 0;
+            factor[0] = 0;
+            factor[1] = 0;
+            condition = 0;
+        }
+
+        //!
+        void CopyFrom(const ImageSizeExpression& o)
+        {
+            type = o.type;
+            size = o.size;
+            layout = o.layout;
+            factor[0] = o.factor[0];
+            factor[1] = o.factor[1];
+            condition = o.condition;
+            yes = o.yes;
+            no = o.no;
+        }
+
+        //!
+        bool operator==( const ImageSizeExpression& other ) const
+        {
+            if ( type == other.type )
+            {
+                switch (type)
+                {
+                case ISET_CONSTANT:
+                    return size[0]==other.size[0]
+                            && size[1] == other.size[1];
+
+                case ISET_LAYOUTFACTOR:
+                    return layout==other.layout
+                            && factor[0]==other.factor[0]
+                            && factor[1]==other.factor[1];
+
+                case ISET_CONDITIONAL:
+                    return condition==other.condition
+                            && *yes==*(other.yes)
+                            && *no==*(other.no);
+
+                default:
+                    return false;
+                }
+            }
+
+            return false;
+        }
+
+        void Optimise()
+        {
+            switch (type)
+            {
+            case ISET_CONSTANT:
+                break;
+
+            case ISET_LAYOUTFACTOR:
+                // TODO: See if the layout is constant and so is this expression.
+                break;
+
+            case ISET_CONDITIONAL:
+                yes->Optimise();
+                no->Optimise();
+
+                // TODO: See if the condition is constant
+                if ( *yes==*no )
+                {
+                    CopyFrom( *yes );
+                }
+
+            default:
+                break;
+            }
+        }
+
+    };
+
+
+    class ASTOp;
+    typedef vector<Ptr<ASTOp>> ASTOpList;
+    typedef std::unordered_set<Ptr<ASTOp>> ASTOpSet;
+
+
+    //! Detailed optimization flags
+    struct MODEL_OPTIMIZATION_OPTIONS
+    {
+        bool m_enabled = true;
+        bool m_optimiseOverlappedMasks = false;
+        bool m_constReduction = true;
+
+        //! Preprocess all mesh fragments so that they use the same skeleton, even if not all bones
+        //! are relevant for all fragments.
+        bool m_uniformizeSkeleton = true;
+
+        //! Maximum number of iterations when optimising models. If 0 as many as necessary will be performed.
+        int m_maxOptimisationLoopCount = 8;
+
+        //! Store resource data in disk instead of memory
+        bool m_useDiskCache = false;
+
+        // Additional advanced fine-tuning parameters
+        //---------------------------------------------------------------------
+
+        //! Ratio used to decide if it is worth to generate a crop operation
+        float m_acceptableCropRatio = 0.5f;
+
+        //! Ratio used to decide if it is worth to generate a crop operation
+        float m_minRLECompressionGain = 1.2f;
+    };
+
+
+	struct FLinkerOptions
+	{
+		int MinTextureResidentMipCount = 0;
+	};
+
+
+	class Sink_ImageCropAST
+	{
+	public:
+
+		Ptr<ASTOp> Apply(const ASTOp* root);
+
+	protected:
+
+		const ASTOp* m_root = nullptr;
+		Ptr<ASTOp> m_initialSource;
+		//! For each operation we sink, the map from old instructions to new instructions.
+		map<Ptr<const class ASTOpFixed>, std::unordered_map<Ptr<ASTOp>, Ptr<ASTOp>>> m_oldToNew;
+		vector<Ptr<ASTOp>> m_newOps;
+
+		Ptr<ASTOp> Visit(Ptr<ASTOp> at, const ASTOpFixed* currentCropOp);
+	};
+
+
+	//---------------------------------------------------------------------------------------------
+	//!
+	//---------------------------------------------------------------------------------------------
+	class Sink_ImagePixelFormatAST
+	{
+	public:
+
+		Ptr<ASTOp> Apply(const ASTOp* root);
+
+	protected:
+
+		const class ASTOpImagePixelFormat* m_root = nullptr;
+		Ptr<ASTOp> m_initialSource;
+		//! For each operation we sink, the map from old instructions to new instructions.
+		map<Ptr<const ASTOp>, std::unordered_map<Ptr<ASTOp>, Ptr<ASTOp>>> m_oldToNew;
+		vector<Ptr<ASTOp>> m_newOps;
+
+		Ptr<ASTOp> Visit(Ptr<ASTOp> at, const class ASTOpImagePixelFormat* currentFormatOp);
+	};
+
+
+	//---------------------------------------------------------------------------------------------
+	//!
+	//---------------------------------------------------------------------------------------------
+	class Sink_MeshFormatAST
+	{
+	public:
+
+		// \TODO This is recursive and may cause stack overflows in big models.
+		mu::Ptr<ASTOp> Apply(const ASTOp* root);
+
+	protected:
+
+		const class ASTOpMeshFormat* m_root = nullptr;
+		mu::Ptr<ASTOp> m_initialSource;
+		std::unordered_map<mu::Ptr<ASTOp>, mu::Ptr<ASTOp>> m_oldToNew;
+		vector<mu::Ptr<ASTOp>> m_newOps;
+
+		mu::Ptr<ASTOp> Visit(const mu::Ptr<ASTOp>& at, const class ASTOpMeshFormat* currentFormatOp);
+	};
+
+
+	//---------------------------------------------------------------------------------------------
+	//!
+	//---------------------------------------------------------------------------------------------
+	class Sink_ImageMipmapAST
+	{
+	public:
+
+		Ptr<ASTOp> Apply(const ASTOp* root);
+
+	protected:
+
+		const class ASTOpImageMipmap* m_root = nullptr;
+		Ptr<ASTOp> m_initialSource;
+
+		//! For each operation we sink, the map from old instructions to new instructions.
+		std::unordered_map<Ptr<ASTOp>, Ptr<ASTOp>> m_oldToNew;
+		vector<Ptr<ASTOp>> m_newOps;
+
+		Ptr<ASTOp> Visit(Ptr<ASTOp> at, const class ASTOpImageMipmap* currentMipmapOp);
+
+	};
+
+	struct OPTIMIZE_SINK_CONTEXT
+	{
+		Sink_ImageCropAST ImageCropSinker;
+		Sink_ImagePixelFormatAST ImagePixelFormatSinker;
+		Sink_ImageMipmapAST ImageMipmapSinker;
+		Sink_MeshFormatAST MeshFormatSinker;
+	};
+
+
+    //!
+    class ASTChild
+    {
+    public:
+        explicit ASTChild(ASTOp* parent, const Ptr<ASTOp>& child=Ptr<ASTOp>());
+        ASTChild(const Ptr<ASTOp>& parent, const Ptr<ASTOp>& child);
+        ~ASTChild();
+
+        ASTChild(const ASTChild&) = delete;
+        ASTChild& operator=(const ASTChild&) = delete;
+
+        // move constructor
+        ASTChild(ASTChild&& rhs)
+             : m_parent(rhs.m_parent)
+             , m_child(rhs.m_child)
+             , m_parentIndexInChild(rhs.m_parentIndexInChild)
+        {
+            rhs.m_parent=nullptr;
+            rhs.m_child.reset();
+        }
+
+        // Move assignment
+        ASTChild& operator=( ASTChild&& rhs );
+
+        ASTChild& operator=( const Ptr<ASTOp>& c );
+
+        inline explicit operator bool() const
+        {
+            return m_child.get()!=nullptr;
+        }
+
+        inline Ptr<class ASTOp>& child()
+        {
+            return m_child;
+        }
+
+        inline const Ptr<class ASTOp>& child() const
+        {
+            return m_child;
+        }
+
+        inline const Ptr<class ASTOp>& operator->() const
+        {
+            return m_child;
+        }
+
+        inline bool operator==(const ASTChild& o) const
+        {
+            return m_child==o.m_child;
+        }
+
+        class ASTOp* m_parent;
+        Ptr<class ASTOp> m_child;
+        int32 m_parentIndexInChild = 0;
+
+    private:
+
+        inline void AddParent();
+
+        inline void ClearParent();
+    };
+
+
+    //---------------------------------------------------------------------------------------------
+    //! Abstract Syntax Tree of operations in the mutable virtual machine.
+    //! Avoid any kind of recursivity here, since the hierarchy can be very deep, and it will
+    //! easily cause stack overflows with production models.
+    //---------------------------------------------------------------------------------------------
+    class ASTOp : public RefCounted
+    {
+    private:
+
+        //! Operations referring to this one. They may be null: elements are never removed
+        //! from this vector.
+		TArray<ASTOp*, TInlineAllocator<4> > m_parents;
+
+    public:
+
+        inline ASTOp()
+        {
+            m_traversedTopDownAdapative = 0;
+            m_constantSubtree = false;
+            m_hasSpecialOpInSubtree = false;
+
+        }
+
+        virtual ~ASTOp() {}
+
+        //! Get the operation type
+        virtual OP_TYPE GetOpType() const = 0;
+
+        //! Validate that everything is fine with this tree
+        virtual void Assert();
+
+        //! Run something for each child operation, with a chance to modify it.
+        virtual void ForEachChild( const std::function<void(ASTChild&)>& ) = 0;
+
+        //! Run something for each parent operation+.
+        void ForEachParent( const std::function<void(ASTOp*)>& );
+
+        //! Run something for each child operation, with a chance to modify it.
+        virtual bool operator==( const ASTOp& other ) const;
+
+        //! Hint hash method for op sorting and containers. It's a hash of the actual operation.
+        virtual uint64 Hash() const = 0;
+
+        //! Shallow clone. New node will have no parents but reference to the same children.
+        using MapChildFunc=std::function<Ptr<ASTOp>(const Ptr<ASTOp>&)>;
+		virtual Ptr<ASTOp> Clone(MapChildFunc& mapChild) const = 0;
+
+    protected:
+
+        void RemoveChildren();
+
+        virtual bool IsEqual(const ASTOp& other) const = 0;
+
+    public:
+
+        //---------------------------------------------------------------------------------------------
+        static void FullAssert( const vector<Ptr<ASTOp>>& roots );
+
+        static size_t CountNodes( const vector<Ptr<ASTOp>>& roots );
+
+        //! Deep clone. New node will have no parents and reference new children
+        static Ptr<ASTOp> DeepClone( const Ptr<ASTOp>& );
+
+        // Code optimisation methods
+        //---------------------------------------------------------------------------------------------
+
+        //!
+        virtual Ptr<ASTOp> OptimiseSize() const { return nullptr; }
+        virtual Ptr<ASTOp> OptimiseSemantic(const MODEL_OPTIMIZATION_OPTIONS&) const { return nullptr; }
+        virtual Ptr<ASTOp> OptimiseSink(const MODEL_OPTIMIZATION_OPTIONS&, OPTIMIZE_SINK_CONTEXT& ) const { return nullptr; }
+        virtual Ptr<ImageSizeExpression> GetImageSizeExpression() const
+        {
+            check( false );
+            return nullptr;
+        }
+
+
+        // Code linking
+        //---------------------------------------------------------------------------------------------
+
+        //!
+        static void FullLink( Ptr<ASTOp>& root, PROGRAM&, const FLinkerOptions*);
+
+        //!
+        static void ClearLinkData( Ptr<ASTOp>& root );
+
+        //!
+        static void LogHistogram( ASTOpList& roots );
+
+    private:
+
+        //!
+        virtual void Link( PROGRAM& program, const FLinkerOptions* Options ) = 0;
+
+    protected:
+
+        //---------------------------------------------------------------------------------------------
+        //!
+        //---------------------------------------------------------------------------------------------
+        struct RANGE_DATA
+        {
+            //!
+            ASTChild rangeSize;
+
+            //!
+            string rangeName;
+
+            //!
+            string rangeUID;
+
+            //!
+            RANGE_DATA( ASTOp* parentOp, Ptr<ASTOp> childOp, const string& name, const string& uid )
+                : rangeSize( parentOp, childOp )
+                , rangeName(name)
+                , rangeUID(uid)
+            {
+            }
+
+            RANGE_DATA(const RANGE_DATA&) = delete;
+            RANGE_DATA& operator=(const RANGE_DATA&) = delete;
+            RANGE_DATA& operator=(RANGE_DATA&&) = delete;
+
+
+            // move constructor
+            RANGE_DATA(RANGE_DATA&& rhs)
+                 : rangeSize(std::move(rhs.rangeSize))
+                 , rangeName(std::move(rhs.rangeName))
+                 , rangeUID(std::move(rhs.rangeUID))
+            {
+            }
+
+            //!
+
+            //!
+            bool operator==(const RANGE_DATA& o) const
+            {
+                return rangeSize==o.rangeSize
+                        &&
+                        rangeName==o.rangeName
+                        &&
+                        rangeUID==o.rangeUID;
+            }
+        };
+
+        //!
+        static void LinkRange( PROGRAM& program,
+                               const RANGE_DATA& range,
+                               OP::ADDRESS& rangeSize,
+                               uint16_t& rangeId )
+        {
+            if (range.rangeSize)
+            {
+                if (range.rangeSize->linkedRange<0)
+                {
+                    check( program.m_ranges.size()<255 );
+                    range.rangeSize->linkedRange = int8_t( program.m_ranges.size() );
+                    RANGE_DESC rangeData;
+                    rangeData.m_name = range.rangeName;
+                    rangeData.m_uid = range.rangeUID;
+                    program.m_ranges.push_back( rangeData );
+                }
+                rangeSize = range.rangeSize->linkedAddress;
+                rangeId = uint16_t(range.rangeSize->linkedRange);
+            }
+        }
+
+    public:
+
+        // Generic traversals
+        //---------------------------------------------------------------------------------------------
+
+        //!
+        static void Traverse_TopDown_Unique( const vector<Ptr<ASTOp>>& roots,
+                                             std::function<bool(Ptr<ASTOp>&)> f );
+
+        //! \todo: it is not strictly top down.
+        static void Traverse_TopDown_Unique_Imprecise( const vector<Ptr<ASTOp>>& roots,
+                                             std::function<bool(Ptr<ASTOp>&)> f );
+
+        //! Kind of top-down, but really not.
+        //! This version is slighlty faster, but doesn't support recursive traversals so
+        //! use it only in controlled cases.
+        static void Traverse_TopRandom_Unique_NonReentrant
+            (
+                const vector<Ptr<ASTOp>>& roots,
+                std::function<bool(Ptr<ASTOp>&)> f
+            );
+
+        //! Kind of top-down, but really not.
+        //! This version is slighlty faster, but doesn't support recursive traversals so
+        //! use it only in controlled cases.
+        template<typename STATE>
+        static inline void Traverse_TopDown_Unique_Imprecise_WithState
+            (
+                Ptr<ASTOp>& root, const STATE& initialState,
+                std::function<bool(Ptr<ASTOp>&, STATE&, vector<pair<Ptr<ASTOp>,STATE>>&)> f
+            )
+        {
+            if (!root) { return; }
+
+            vector<pair<Ptr<ASTOp>,STATE>> pending;
+            pending.push_back( std::make_pair<>( root, initialState ) );
+
+
+            struct custom_partial_hash
+            {
+				uint64 operator()(const pair<Ptr<ASTOp>,STATE>& k) const
+                {
+                    return std::hash<const void*>()(k.first.get());
+                }
+            };
+
+            std::unordered_set<pair<Ptr<ASTOp>,STATE>, custom_partial_hash> traversed;
+
+            while (pending.size())
+            {
+                pair<Ptr<ASTOp>,STATE> current = pending.back();
+                pending.pop_back();
+
+                // It could have been completed in another branch
+                auto iti = traversed.insert( current );
+                if (iti.second)
+                {
+                    // It has been inserted, so it wasn't there before.
+
+                    // Process. State in current.second may change
+                    bool recurse = f(current.first,current.second,pending);
+
+                    // Recurse children
+                    if (recurse)
+                    {
+                        current.first->ForEachChild([&]( ASTChild& c )
+                        {
+                            auto child = std::make_pair<>(c.m_child, current.second);
+                            if ( c.m_child && !traversed.count(child) )
+                            {
+                                pending.push_back( child );
+                            }
+                        });
+                    }
+                }
+            }
+        }
+
+
+        //! Kind of top-down, but really not.
+        static void Traverse_TopDown_Repeat( const vector<Ptr<ASTOp>>& roots,
+                                             std::function<bool(Ptr<ASTOp>& node)> f );
+
+        //! This version is slighlty faster, but doesn't support recursive traversals so
+        //! use it only in controlled cases.
+        static void Traverse_BottomUp_Unique_NonReentrant
+        (
+                ASTOpList& roots,
+                std::function<void(Ptr<ASTOp>&)> f
+        );
+
+        //! This version is slighlty faster, but doesn't support recursive traversals so
+        //! use it only in controlled cases.
+        static void Traverse_BottomUp_Unique_NonReentrant
+        (
+                ASTOpList& roots,
+                std::function<void(Ptr<ASTOp>&)> f,
+                std::function<bool(const ASTOp*)> accept
+        );
+
+        //!
+        static void Traverse_BottomUp_Unique
+        (
+                ASTOpList& roots,
+                std::function<void(Ptr<ASTOp>&)> f,
+                std::function<bool(const ASTOp*)> accept = [](const ASTOp*){return true;}
+        );
+
+        //!
+        static void Traverse_BottomUp_Unique
+        (
+                Ptr<ASTOp>& root,
+                std::function<void(Ptr<ASTOp>&)> f,
+                std::function<bool(const ASTOp*)> accept = [](const ASTOp*){return true;}
+        );
+
+        //!
+        OP::ADDRESS linkedAddress = 0;
+
+        //! Generic traverse control counter. It should always be left to 0 after any process for all
+        //! nodes in the hierarchy.
+        uint32_t m_traverseIndex = 0;
+        static std::atomic<uint32_t> s_lastTraverseIndex;
+
+        //!
+        int8_t linkedRange = -1;
+
+        //! special flag for Traverse_TopDown_Adaptative traversal.
+        uint8_t m_traversedTopDownAdapative : 1;
+
+        //! Embedded node data for the constant subtree detection. This flag is only valid if the
+        //! constant detection process has been executed and no relevant AST transformations have
+        //! happened.
+        uint8_t m_constantSubtree : 1;
+        uint8_t m_hasSpecialOpInSubtree : 1;
+
+    private:
+        friend class ASTChild;
+
+        //! Remove a node from the parent list if it is there.
+        //void RemoveParent(const ASTOp* parent);
+
+    public:
+
+        int GetParentCount();
+
+        //! Make all parents of this node point at the other node instead.
+        static void Replace( const Ptr<ASTOp>& node, const Ptr<ASTOp>& other );
+
+        // Other code generation utilities
+        //---------------------------------------------------------------------------------------------
+
+        //! This class contains the support data to accelerate the GetImageDesc recursive function.
+        //! If none is provided in the call, one will be created at that level and used from there on.
+        class GetImageDescContext
+        {
+        public:
+            TMap<const ASTOp*, FImageDesc> m_results;
+        };
+
+        //!
+        virtual FImageDesc GetImageDesc( bool returnBestOption=false, GetImageDescContext* context=nullptr );
+
+        //! Optional cache struct to use int he method below.
+        using BLOCK_LAYOUT_SIZE_CACHE=TMap< const TPair<ASTOp*,int>, TPair<int,int>>;
+
+        //! Return the size in layout blocks of a particular block given by absolute index
+        virtual void GetBlockLayoutSize( int blockIndex, int* pBlockX, int* pBlockY,
+                                         BLOCK_LAYOUT_SIZE_CACHE* cache );
+        void GetBlockLayoutSizeCached( int blockIndex, int* pBlockX, int* pBlockY,
+                                       BLOCK_LAYOUT_SIZE_CACHE* cache );
+
+
+        //! Return the size in pixels of the layout grid block for the image operation
+        virtual void GetLayoutBlockSize( int* pBlockX, int* pBlockY );
+
+        virtual bool IsImagePlainConstant( vec4<float>& colour ) const;
+        virtual bool IsColourConstant( vec4<float>& colour ) const;
+
+        //!
+        virtual bool GetNonBlackRect( FImageRect& maskUsage ) const;
+
+        // Logic expression evaluation
+        //---------------------------------------------------------------------------------------------
+        typedef enum
+        {
+            BET_UNKNOWN,
+            BET_TRUE,
+            BET_FALSE,
+        } BOOL_EVAL_RESULT;
+
+
+        using EVALUATE_BOOL_CACHE=std::unordered_map<const ASTOp*,BOOL_EVAL_RESULT>;
+
+        virtual BOOL_EVAL_RESULT EvaluateBool( ASTOpList& /*facts*/, EVALUATE_BOOL_CACHE* = nullptr ) const
+        {
+            check(false);
+            return BET_UNKNOWN;
+        }
+
+        virtual int EvaluateInt( ASTOpList& /*facts*/, bool &unknown ) const
+        {
+            check(false);
+            unknown = true;
+            return 0;
+        }
+
+    };
+
+    template<class DERIVED>
+    inline Ptr<DERIVED> Clone( const ASTOp* s )
+    {
+		ASTOp::MapChildFunc Identity = [](const Ptr<ASTOp>& o) {return o; };
+		auto c = s->Clone(Identity);
+		Ptr<DERIVED> t = dynamic_cast<DERIVED*>(c.get());
+        check(t);
+        return t;
+    }
+
+    template<class DERIVED>
+    inline Ptr<DERIVED> Clone( const Ptr<const ASTOp>& s )
+    {
+		ASTOp::MapChildFunc Identity = [](const Ptr<ASTOp>& o) {return o; };
+        auto c = s->Clone(Identity);
+        Ptr<DERIVED> t = dynamic_cast<DERIVED*>(c.get());
+        check(t);
+        return t;
+    }
+
+    //---------------------------------------------------------------------------------------------
+    //!
+    //---------------------------------------------------------------------------------------------
+    template<typename STATE>
+    class Visitor_TopDown_Unique_Const : public Base
+    {
+    private:
+
+        //!
+        virtual bool Visit( const Ptr<ASTOp>& node ) = 0;
+
+    private:
+
+        //! States found so far
+        vector<STATE> m_states;
+
+        //! Index of the current state, from the m_states vector.
+        int m_currentState;
+
+        struct PENDING
+        {
+            PENDING()
+            {
+                stateIndex = 0;
+            }
+            PENDING(Ptr<ASTOp> _at, int _stateIndex)
+            {
+                at = _at;
+                stateIndex = _stateIndex;
+            }
+
+            Ptr<ASTOp> at;
+            int stateIndex;
+        };
+
+        std::vector<PENDING> m_pending;
+
+        //! List of traversed nodes with the state in which they were traversed.
+        std::unordered_multimap<Ptr<ASTOp>,int> m_traversed;
+
+    protected:
+
+        //!
+        const STATE& GetCurrentState() const
+        {
+            return m_states[m_currentState];
+        }
+
+        //!
+        STATE GetDefaultState() const
+        {
+            return m_states[0];
+        }
+
+        //! For manual recursion that changes the state for a specific path.
+        void RecurseWithState(const Ptr<ASTOp>& at, const STATE& newState)
+        {
+            if(at)
+            {
+                auto it = std::find(m_states.begin(),m_states.end(),newState);
+                if (it==m_states.end())
+                {
+                    m_states.push_back(newState);
+                    it = m_states.end() - 1;
+                }
+                int stateIndex = (int)(it - m_states.begin());
+
+                m_pending.push_back( PENDING(at,stateIndex) );
+            }
+        }
+
+        //! For manual recursion that doesn't change the state for a specific path.
+        void RecurseWithCurrentState(const Ptr<ASTOp>& at)
+        {
+            if(at)
+            {
+                m_pending.emplace_back( at, m_currentState );
+            }
+        }
+
+        //! Can be called from visit to set the state to visit all children ops
+        void SetCurrentState(const STATE& newState)
+        {
+            auto it = std::find(m_states.begin(),m_states.end(),newState);
+            if (it==m_states.end())
+            {
+                m_states.push_back(newState);
+                it = m_states.end() - 1;
+            }
+            m_currentState = (int)(it - m_states.begin());
+        }
+
+    public:
+
+        //! Ensure virtual destruction
+        virtual ~Visitor_TopDown_Unique_Const() {}
+
+        //!
+        void Traverse(const ASTOpList& roots, const STATE& initialState)
+        {
+            m_pending.clear();
+            m_states.clear();
+            m_states.push_back(initialState);
+            m_currentState = 0;
+
+            for( const auto& r: roots)
+            {
+                if (r)
+                {
+                    m_pending.push_back( PENDING(r,m_currentState) );
+                }
+            }
+
+            while (m_pending.size())
+            {
+                PENDING item = m_pending.back();
+                m_pending.pop_back();
+
+                auto thisNodeRange = m_traversed.equal_range(item.at);
+                bool visitedInThisState = false;
+                for (auto it=thisNodeRange.first; it!=thisNodeRange.second; ++it)
+                {
+                    if (it->second==item.stateIndex)
+                    {
+                        visitedInThisState = true;
+                        break;
+                    }
+                }
+
+                // It could have been completed in another branch
+                if (!visitedInThisState)
+                {
+                    m_traversed.insert( std::make_pair(item.at,item.stateIndex) );
+
+                    // Process
+                    m_currentState = item.stateIndex;
+                    bool recurse = Visit(item.at);
+
+                    // Recurse children
+                    if (recurse)
+                    {
+                        item.at->ForEachChild([&]( ASTChild& c )
+                        {
+                            if (c)
+                            {
+                                m_pending.push_back( PENDING(c.child(),m_currentState) );
+                            }
+                        });
+                    }
+                }
+            }
+        }
+    };
+
+
+    //---------------------------------------------------------------------------------------------
+    //! Stateless top-down code visitor that can change the instructions. Iterative version.
+    //! Once an instruction has changed, all the chain of instructions up to the root will be
+    //! cloned, referencing the new instruction.
+    //---------------------------------------------------------------------------------------------
+    class Visitor_TopDown_Unique_Cloning : public Base
+    {
+    public:
+
+        //! Ensure virtual destruction
+        virtual ~Visitor_TopDown_Unique_Cloning() {}
+
+    protected:
+
+        //! Do the actual work by overriding this in the derived classes.
+        virtual Ptr<ASTOp> Visit( Ptr<ASTOp> at, bool& processChildren ) = 0;
+
+
+        void Traverse( Ptr<ASTOp>& root );
+
+    private:
+
+        //! Operations to be processed
+        vector< std::pair<bool,Ptr<ASTOp>> > m_pending;
+
+        //! Map for visited operations
+        std::unordered_map<Ptr<ASTOp>,Ptr<ASTOp>> m_oldToNew;
+
+        //!
+        inline Ptr<ASTOp> GetOldToNew( const Ptr<ASTOp>& old ) const
+        {
+            Ptr<ASTOp> n;
+            auto it = m_oldToNew.find(old);
+            if (it!=m_oldToNew.end())
+            {
+                n = it->second;
+            }
+
+            it = m_oldToNew.find(n);
+            while (it!=m_oldToNew.end() && it->second.get()!=nullptr && it->second!=n)
+            {
+                n = it->second;
+                it = m_oldToNew.find(n);
+            }
+
+            return n;
+        }
+
+        //! Process all the pending operations and visit all children if necessary
+        void Process();
+
+    };
+
+
+    //---------------------------------------------------------------------------------------------
+    //! This class of operation is used to ease the transition from legacy code generation.
+    //---------------------------------------------------------------------------------------------
+    class ASTOpFixed : public ASTOp
+    {
+    public:
+
+        //! Fixed size op structure. Any code ADDRESS in it, refers to the children vector below.
+        OP op;
+
+        //! Operations used by this one. To be indexed with the ADDRESSes in the m_op.
+        //vector<ASTChild> children;
+        std::array<ASTChild,8> children;
+        uint8_t childCount = 1;
+
+    public:
+        ASTOpFixed();
+
+        ~ASTOpFixed() override;
+
+        OP_TYPE GetOpType() const override { return op.type; }
+
+        Ptr<ASTOp> Clone( MapChildFunc& mapChild ) const override;
+        void ForEachChild( const std::function<void(ASTChild&)>& f ) override;
+        void Link( PROGRAM& program, const FLinkerOptions* Options) override;
+        bool IsEqual(const ASTOp& otherUntyped) const override;
+		uint64 Hash() const override;
+		FImageDesc GetImageDesc( bool returnBestOption, GetImageDescContext* context ) override;
+        void GetBlockLayoutSize( int blockIndex, int* pBlockX, int* pBlockY,
+                                 BLOCK_LAYOUT_SIZE_CACHE* cache ) override;
+        void GetLayoutBlockSize( int* pBlockX, int* pBlockY ) override;
+        Ptr<ASTOp> OptimiseSize() const override;
+        Ptr<ASTOp> OptimiseSemantic(const MODEL_OPTIMIZATION_OPTIONS&) const override;
+        Ptr<ASTOp> OptimiseSink(const MODEL_OPTIMIZATION_OPTIONS&, OPTIMIZE_SINK_CONTEXT&) const override;
+        BOOL_EVAL_RESULT EvaluateBool( ASTOpList& facts, EVALUATE_BOOL_CACHE* cache ) const override;
+        int EvaluateInt( ASTOpList& facts, bool &unknown ) const override;
+        bool IsImagePlainConstant( vec4<float>& colour ) const override;
+        bool IsColourConstant( vec4<float>& colour ) const override;
+        Ptr<ImageSizeExpression> GetImageSizeExpression() const override;
+
+        //---------------------------------------------------------------------------------------------
+        //! Add child for OP fixed nodes.
+        //---------------------------------------------------------------------------------------------
+        inline void SetChild( OP::ADDRESS& at, const Ptr<ASTOp>& child )
+        {
+            // hack check
+            check( ((uint8_t*)&at)>=((uint8_t*)&op.args)
+                            &&
+                            ((uint8_t*)&at)<=((uint8_t*)&op.args)+sizeof(op.args));
+
+            if (!at)
+            {
+                if (child)
+                {
+                    check(childCount<8);
+                    at = (OP::ADDRESS)childCount++;
+                    children[at] = child;
+                }
+            }
+            else
+            {
+                children[at]=child;
+                if (!child)
+                {
+                    at=0;
+                }
+            }
+        }
+
+        inline void SetChild( OP::ADDRESS& at, const ASTChild& child )
+        {
+            SetChild(at,child.child());
+        }
+
+    private:
+
+        //
+        inline FImageDesc GetImageDesc( OP::ADDRESS at, bool returnBestOption=false, GetImageDescContext* context=nullptr )
+        {
+			FImageDesc res;
+
+            if (children[at])
+            {
+                res = children[at]->GetImageDesc( returnBestOption, context );
+            }
+
+            return res;
+        }
+
+        inline void GetBlockLayoutSize( int blockIndex, OP::ADDRESS at, int* pBlockX, int* pBlockY,
+                                        BLOCK_LAYOUT_SIZE_CACHE* cache )
+        {
+            if (children[at])
+            {
+                children[at]->GetBlockLayoutSizeCached(blockIndex,pBlockX,pBlockY,cache);
+            }
+        }
+
+        inline void GetLayoutBlockSize( OP::ADDRESS at, int* pBlockX, int* pBlockY )
+        {
+            if (children[at])
+            {
+                children[at]->GetLayoutBlockSize(pBlockX,pBlockY);
+            }
+        }
+
+        //! \todo move to ast swizzle op class when created
+        Ptr<ASTOp> OptimiseSwizzle( const MODEL_OPTIMIZATION_OPTIONS& ) const;
+    };
+
+
+
+    //---------------------------------------------------------------------------------------------
+    //---------------------------------------------------------------------------------------------
+    //---------------------------------------------------------------------------------------------
+    inline void ASTChild::AddParent()
+    {
+        m_parentIndexInChild = m_child->m_parents.Num();
+        m_child->m_parents.Add(m_parent);
+    }
+
+    //---------------------------------------------------------------------------------------------
+    inline void ASTChild::ClearParent()
+    {
+        check( m_parentIndexInChild<m_child->m_parents.Num() );
+		// Can't do this, because the indices are stored in children.
+		//m_child->m_parents.RemoveAtSwap(m_parentIndexInChild);
+		m_child->m_parents[m_parentIndexInChild] = nullptr;
+	}
+
+
+    //---------------------------------------------------------------------------------------------
+    //---------------------------------------------------------------------------------------------
+    //---------------------------------------------------------------------------------------------
+    class UniqueOpPool
+    {
+    private:
+        struct op_pointer_hash
+        {
+			uint64 operator() (const Ptr<ASTOp>& n) const
+            {
+                return n->Hash();
+            }
+        };
+
+        struct op_pointer_equal
+        {
+            bool operator() (const Ptr<ASTOp>& a, const Ptr<ASTOp>& b) const
+            {
+                return *a==*b;
+            }
+        };
+
+        // Existing ops, per type
+        std::unordered_set<Ptr<ASTOp>,op_pointer_hash,op_pointer_equal> visited[(int)OP_TYPE::COUNT];
+
+    public:
+
+        bool m_disabled = false;
+
+        Ptr<ASTOp> Add( const Ptr<ASTOp>& n )
+        {
+            if (m_disabled) return n;
+
+            if (!n) return nullptr;
+
+            auto& container = visited[(int)n->GetOpType()];
+            auto it = container.insert(n);
+            if( !it.second )
+            {
+                return *it.first;
+            }
+            else
+            {
+                return n;
+            }
+        }
+    };
+
+
+    //
+    extern void DebugLogAST(const Ptr<ASTOp>& at,
+                            int indent=0,
+                            ASTOpList* done=nullptr,
+                            const char* label=nullptr);
+
+}
+
diff --git a/Engine/Plugins/Experimental/Mutable/Source/MutableTools/Private/ASTOpAddLOD.cpp b/Engine/Plugins/Experimental/Mutable/Source/MutableTools/Private/ASTOpAddLOD.cpp
new file mode 100644
--- /dev/null
+++ b/Engine/Plugins/Experimental/Mutable/Source/MutableTools/Private/ASTOpAddLOD.cpp
@@ -0,0 +1,97 @@
+// Copyright Epic Games, Inc. All Rights Reserved.
+
+#include "ASTOpAddLOD.h"
+#include "ErrorLogPrivate.h"
+#include "CodeOptimiser.h"
+#include "StreamsPrivate.h"
+
+
+using namespace mu;
+
+
+//-------------------------------------------------------------------------------------------------
+//-------------------------------------------------------------------------------------------------
+//-------------------------------------------------------------------------------------------------
+ASTOpAddLOD::ASTOpAddLOD()
+{
+}
+
+
+ASTOpAddLOD::~ASTOpAddLOD()
+{
+    // Explicit call needed to avoid recursive destruction
+    ASTOp::RemoveChildren();
+}
+
+
+bool ASTOpAddLOD::IsEqual(const ASTOp& otherUntyped) const
+{
+    if ( auto other = dynamic_cast<const ASTOpAddLOD*>(&otherUntyped) )
+    {
+        return lods==other->lods;
+    }
+    return false;
+}
+
+
+uint64 ASTOpAddLOD::Hash() const
+{
+	uint64 res = std::hash<uint64>()(uint64(OP_TYPE::IN_ADDLOD));
+    for(auto& c:lods)
+    {
+        hash_combine( res, c.child().get() );
+    }
+    return res;
+}
+
+
+mu::Ptr<ASTOp> ASTOpAddLOD::Clone(MapChildFunc& mapChild) const
+{
+    Ptr<ASTOpAddLOD> n = new ASTOpAddLOD();
+    for(auto& c:lods)
+    {
+        n->lods.emplace_back(n,mapChild(c.child()));
+    }
+    return n;
+}
+
+
+void ASTOpAddLOD::ForEachChild( const std::function<void(ASTChild&)>& f )
+{
+    for (auto& l:lods)
+    {
+        f(l);
+    }
+}
+
+
+void ASTOpAddLOD::Link( PROGRAM& program, const FLinkerOptions* )
+{
+    // Already linked?
+    if (!linkedAddress)
+    {
+        OP::InstanceAddLODArgs args;
+        memset( &args,0, sizeof(args) );
+
+        int i=0;
+        for (auto& l:lods)
+        {
+            if ( l )
+            {
+                args.lod[i] = l->linkedAddress;
+                ++i;
+                if (i>=MUTABLE_OP_MAX_ADD_COUNT)
+                {
+                    break;
+                }
+            }
+        }
+
+        linkedAddress = (OP::ADDRESS)program.m_opAddress.size();
+        program.m_opAddress.push_back((uint32_t)program.m_byteCode.size());
+        AppendCode(program.m_byteCode,OP_TYPE::IN_ADDLOD);
+        AppendCode(program.m_byteCode,args);
+    }
+
+}
+
diff --git a/Engine/Plugins/Experimental/Mutable/Source/MutableTools/Private/ASTOpAddLOD.h b/Engine/Plugins/Experimental/Mutable/Source/MutableTools/Private/ASTOpAddLOD.h
new file mode 100644
--- /dev/null
+++ b/Engine/Plugins/Experimental/Mutable/Source/MutableTools/Private/ASTOpAddLOD.h
@@ -0,0 +1,38 @@
+// Copyright Epic Games, Inc. All Rights Reserved.
+
+#pragma once
+
+#include "AST.h"
+
+
+namespace mu
+{
+
+
+	//---------------------------------------------------------------------------------------------
+	//! Operations to add elements to a LOD
+	//---------------------------------------------------------------------------------------------
+	class ASTOpAddLOD : public ASTOp
+	{
+	public:
+
+		vector< ASTChild > lods;
+
+	public:
+
+		ASTOpAddLOD();
+		ASTOpAddLOD(const ASTOpAddLOD&) = delete;
+		~ASTOpAddLOD();
+
+		OP_TYPE GetOpType() const override { return OP_TYPE::IN_ADDLOD; }
+		uint64 Hash() const override;
+		void ForEachChild(const std::function<void(ASTChild&)>&) override;
+		bool IsEqual(const ASTOp& otherUntyped) const override;
+		Ptr<ASTOp> Clone(MapChildFunc& mapChild) const override;
+		void Link(PROGRAM& program, const FLinkerOptions* Options) override;
+	};
+
+
+
+}
+
diff --git a/Engine/Plugins/Experimental/Mutable/Source/MutableTools/Private/ASTOpConditional.cpp b/Engine/Plugins/Experimental/Mutable/Source/MutableTools/Private/ASTOpConditional.cpp
new file mode 100644
--- /dev/null
+++ b/Engine/Plugins/Experimental/Mutable/Source/MutableTools/Private/ASTOpConditional.cpp
@@ -0,0 +1,355 @@
+// Copyright Epic Games, Inc. All Rights Reserved.
+
+#include "ASTOpConditional.h"
+#include "ASTOpConstantBool.h"
+#include "ErrorLogPrivate.h"
+#include "CodeOptimiser.h"
+#include "StreamsPrivate.h"
+
+#include <queue>
+
+
+using namespace mu;
+
+
+//-------------------------------------------------------------------------------------------------
+//-------------------------------------------------------------------------------------------------
+//-------------------------------------------------------------------------------------------------
+ASTOpConditional::ASTOpConditional()
+    : condition(this)
+    , yes(this)
+    , no(this)
+{
+}
+
+
+//-------------------------------------------------------------------------------------------------
+ASTOpConditional::~ASTOpConditional()
+{
+    // Explicit call needed to avoid recursive destruction
+    ASTOp::RemoveChildren();
+}
+
+
+//-------------------------------------------------------------------------------------------------
+bool ASTOpConditional::IsEqual(const ASTOp& otherUntyped) const
+{
+    if ( auto other = dynamic_cast<const ASTOpConditional*>(&otherUntyped) )
+    {
+        return type==other->type && condition==other->condition &&
+                yes==other->yes && no==other->no;
+    }
+    return false;
+}
+
+
+//-------------------------------------------------------------------------------------------------
+mu::Ptr<ASTOp> ASTOpConditional::Clone(MapChildFunc& mapChild) const
+{
+    Ptr<ASTOpConditional> n = new ASTOpConditional();
+    n->type = type;
+    n->condition = mapChild(condition.child());
+    n->yes = mapChild(yes.child());
+    n->no = mapChild(no.child());
+    return n;
+}
+
+
+//-------------------------------------------------------------------------------------------------
+uint64 ASTOpConditional::Hash() const
+{
+	uint64 res = std::hash<void*>()(condition.m_child.get());
+    hash_combine(res, yes.m_child.get() );
+    hash_combine(res, no.m_child.get() );
+    return res;
+}
+
+
+//-------------------------------------------------------------------------------------------------
+void ASTOpConditional::Assert()
+{
+    switch(type)
+    {
+    case OP_TYPE::NU_CONDITIONAL:
+    case OP_TYPE::SC_CONDITIONAL:
+    case OP_TYPE::CO_CONDITIONAL:
+    case OP_TYPE::IM_CONDITIONAL:
+    case OP_TYPE::ME_CONDITIONAL:
+    case OP_TYPE::LA_CONDITIONAL:
+    case OP_TYPE::IN_CONDITIONAL:
+        break;
+    default:
+        // Unexpected type
+        check(false);
+        break;
+    }
+
+    ASTOp::Assert();
+}
+
+
+//-------------------------------------------------------------------------------------------------
+void ASTOpConditional::ForEachChild(const std::function<void(ASTChild&)>& f )
+{
+    f(condition);
+    f(yes);
+    f(no);
+}
+
+
+//-------------------------------------------------------------------------------------------------
+void ASTOpConditional::Link( PROGRAM& program, const FLinkerOptions*)
+{
+    // Already linked?
+    if (!linkedAddress)
+    {
+        OP::ConditionalArgs args;
+        memset( &args,0, sizeof(args) );
+
+        if (condition) args.condition = condition->linkedAddress;
+        if (yes) args.yes = yes->linkedAddress;
+        if (no) args.no = no->linkedAddress;
+
+        linkedAddress = (OP::ADDRESS)program.m_opAddress.size();
+        //program.m_code.push_back(op);
+        program.m_opAddress.push_back((uint32_t)program.m_byteCode.size());
+        AppendCode(program.m_byteCode,type);
+        AppendCode(program.m_byteCode,args);
+
+    }
+}
+
+
+//-------------------------------------------------------------------------------------------------
+FImageDesc ASTOpConditional::GetImageDesc( bool returnBestOption, class GetImageDescContext* context )
+{
+    FImageDesc res;
+
+    // Local context in case it is necessary
+    GetImageDescContext localContext;
+    if (!context)
+    {
+        context = &localContext;
+    }
+    else
+    {
+        // Cached result?
+		FImageDesc* PtrValue = context->m_results.Find(this);
+		if (PtrValue)
+		{
+			return *PtrValue;
+		}
+	}
+
+    if ( type==OP_TYPE::IM_CONDITIONAL )
+    {
+        // In a conditional we cannot guarantee the size and format.
+        // We check both options, and if they are the same we return that.
+        // Otherwise, we return an empty descriptor that will force re-formatting of the image.
+        // The code optimiser will take care then of moving the format operations down to each
+        // branch and remove the unnecessary ones.
+        FImageDesc noDesc;
+        FImageDesc yesDesc;
+
+        if (no.child())
+        {
+            noDesc = no->GetImageDesc(returnBestOption, context);
+        }
+        if (yes.child())
+        {
+            yesDesc = yes->GetImageDesc(returnBestOption, context);
+        }
+
+        if ( yesDesc == noDesc || returnBestOption )
+        {
+            res = yesDesc;
+        }
+        else
+        {
+            res.m_format = (yesDesc.m_format==noDesc.m_format) ? yesDesc.m_format : EImageFormat::IF_NONE;
+            res.m_lods = (yesDesc.m_lods==noDesc.m_lods) ? yesDesc.m_lods : 0;
+            res.m_size = (yesDesc.m_size==noDesc.m_size) ? yesDesc.m_size : FImageSize( 0, 0 );
+        }
+    }
+    else
+    {
+        checkf(false, TEXT("Instruction not supported") );
+    }
+
+
+    // Cache the result
+    if (context)
+    {
+        context->m_results.Add(this,res);
+    }
+
+    return res;
+}
+
+
+//-------------------------------------------------------------------------------------------------
+void ASTOpConditional::GetLayoutBlockSize( int* pBlockX, int* pBlockY )
+{
+    if ( type==OP_TYPE::IM_CONDITIONAL )
+    {
+        yes->GetLayoutBlockSize( pBlockX, pBlockY );
+
+        if (!*pBlockX)
+        {
+            no->GetLayoutBlockSize( pBlockX, pBlockY );
+        }
+    }
+    else
+    {
+        checkf( false, TEXT("Instruction not supported") );
+    }
+}
+
+
+//-------------------------------------------------------------------------------------------------
+void ASTOpConditional::GetBlockLayoutSize( int blockIndex, int* pBlockX, int* pBlockY,
+                                           BLOCK_LAYOUT_SIZE_CACHE* cache )
+{
+    if ( type==OP_TYPE::LA_CONDITIONAL )
+    {
+        yes->GetBlockLayoutSizeCached( blockIndex, pBlockX, pBlockY, cache );
+
+        if (*pBlockX==0)
+        {
+            no->GetBlockLayoutSizeCached( blockIndex, pBlockX, pBlockY, cache );
+        }
+    }
+    else
+    {
+        checkf( false, TEXT("Instruction not supported") );
+    }
+}
+
+
+//-------------------------------------------------------------------------------------------------
+bool ASTOpConditional::GetNonBlackRect( FImageRect& maskUsage ) const
+{
+    if (type == OP_TYPE::IM_CONDITIONAL )
+    {
+        FImageRect local;
+        bool localValid = false;
+        if (yes)
+        {
+            localValid = yes->GetNonBlackRect(local);
+            if (!localValid)
+            {
+                return false;
+            }
+        }
+
+        if (no)
+        {
+            FImageRect noRect;
+            bool validNo = no->GetNonBlackRect(noRect);
+            if (validNo)
+            {
+                if (localValid)
+                {
+                    local.Bound(noRect);
+                }
+                else
+                {
+                    local = noRect;
+                    localValid = true;
+                }
+            }
+            else
+            {
+                return false;
+            }
+        }
+
+        if (localValid)
+        {
+            maskUsage = local;
+            return true;
+        }
+    }
+
+    return false;
+}
+
+
+//-------------------------------------------------------------------------------------------------
+mu::Ptr<ASTOp> ASTOpConditional::OptimiseSemantic( const MODEL_OPTIMIZATION_OPTIONS& ) const
+{
+    if (!condition)
+    {
+        // If there is no expression, we'll assume true.
+        return yes.child();
+    }
+
+    // If the branches are the same, remove the instruction
+    if ( yes.child() == no.child() )
+    {
+        return yes.child();
+    }
+
+    // Constant condition?
+    if ( condition->GetOpType() == OP_TYPE::BO_CONSTANT )
+    {
+        auto typedCondition = dynamic_cast<const ASTOpConstantBool*>(condition.child().get());
+        if ( typedCondition->value )
+        {
+            return yes.child();
+        }
+
+        return no.child();
+    }
+
+    else
+    {
+        // If the yes branch is a conditional with the same condition
+        if (yes && yes->GetOpType() == type)
+        {
+            auto typedYes = dynamic_cast<const ASTOpConditional*>(yes.child().get());
+            if ( condition.child()==typedYes->condition.child()
+                 ||
+                 *condition.child()==*typedYes->condition.child() )
+            {
+                auto op = mu::Clone<ASTOpConditional>(this);
+                op->yes = typedYes->yes.child();
+                return op;
+            }
+        }
+
+        // If the no branch is a conditional with the same condition
+        else if (no && no->GetOpType() == type)
+        {
+            auto typedNo = dynamic_cast<const ASTOpConditional*>(no.child().get());
+            if ( condition.child()==typedNo->condition.child()
+                 ||
+                 *condition.child()==*typedNo->condition.child() )
+            {
+                auto op = mu::Clone<ASTOpConditional>(this);
+                op->no = typedNo->no.child();
+                return op;
+            }
+        }
+    }
+
+    return nullptr;
+}
+
+
+//-------------------------------------------------------------------------------------------------
+mu::Ptr<ImageSizeExpression> ASTOpConditional::GetImageSizeExpression() const
+{
+    Ptr<ImageSizeExpression> pRes = new ImageSizeExpression;
+    pRes->type = ImageSizeExpression::ISET_CONDITIONAL;
+    pRes->yes = yes->GetImageSizeExpression();
+    pRes->no = no->GetImageSizeExpression();
+
+    if ( *pRes->yes == *pRes->no )
+    {
+        pRes = pRes->yes;
+    }
+
+    return pRes;
+}
+
diff --git a/Engine/Plugins/Experimental/Mutable/Source/MutableTools/Private/ASTOpConditional.h b/Engine/Plugins/Experimental/Mutable/Source/MutableTools/Private/ASTOpConditional.h
new file mode 100644
--- /dev/null
+++ b/Engine/Plugins/Experimental/Mutable/Source/MutableTools/Private/ASTOpConditional.h
@@ -0,0 +1,53 @@
+// Copyright Epic Games, Inc. All Rights Reserved.
+
+#pragma once
+
+#include "AST.h"
+
+
+namespace mu
+{
+
+	//---------------------------------------------------------------------------------------------
+	//! Conditional operation.
+	//---------------------------------------------------------------------------------------------
+	class ASTOpConditional : public ASTOp
+	{
+	public:
+
+		//! Type of switch
+		OP_TYPE type;
+
+		//! Boolean expression
+		ASTChild condition;
+
+		//! Branches
+		ASTChild yes;
+		ASTChild no;
+
+	public:
+
+		ASTOpConditional();
+		ASTOpConditional(const ASTOpConditional&) = delete;
+		~ASTOpConditional() override;
+
+		OP_TYPE GetOpType() const override { return type; }
+
+		bool IsEqual(const ASTOp& otherUntyped) const override;
+		Ptr<ASTOp> Clone(MapChildFunc& mapChild) const override;
+		uint64 Hash() const override;
+		void Assert() override;
+		void ForEachChild(const std::function<void(ASTChild&)>& f) override;
+		void Link(PROGRAM& program, const FLinkerOptions* Options) override;
+		FImageDesc GetImageDesc(bool returnBestOption, class GetImageDescContext* context) override;
+		void GetLayoutBlockSize(int* pBlockX, int* pBlockY) override;
+		void GetBlockLayoutSize(int blockIndex, int* pBlockX, int* pBlockY,
+			BLOCK_LAYOUT_SIZE_CACHE* cache) override;
+		Ptr<ASTOp> OptimiseSemantic(const MODEL_OPTIMIZATION_OPTIONS&) const override;
+		bool GetNonBlackRect(FImageRect& maskUsage) const override;
+		Ptr<ImageSizeExpression> GetImageSizeExpression() const override;
+	};
+
+
+}
+
diff --git a/Engine/Plugins/Experimental/Mutable/Source/MutableTools/Private/ASTOpConstantBool.cpp b/Engine/Plugins/Experimental/Mutable/Source/MutableTools/Private/ASTOpConstantBool.cpp
new file mode 100644
--- /dev/null
+++ b/Engine/Plugins/Experimental/Mutable/Source/MutableTools/Private/ASTOpConstantBool.cpp
@@ -0,0 +1,75 @@
+// Copyright Epic Games, Inc. All Rights Reserved.
+
+#include "ASTOpConstantBool.h"
+#include "ErrorLogPrivate.h"
+#include "CodeOptimiser.h"
+#include "Hash/CityHash.h"
+#include "StreamsPrivate.h"
+
+
+using namespace mu;
+
+
+//-------------------------------------------------------------------------------------------------
+//-------------------------------------------------------------------------------------------------
+//-------------------------------------------------------------------------------------------------
+ASTOpConstantBool::ASTOpConstantBool( bool v )
+{
+    value = v;
+}
+
+
+void ASTOpConstantBool::ForEachChild(const std::function<void(ASTChild&)>& )
+{
+}
+
+
+bool ASTOpConstantBool::IsEqual(const ASTOp& otherUntyped) const
+{
+    if ( auto other = dynamic_cast<const ASTOpConstantBool*>(&otherUntyped) )
+    {
+        return value==other->value;
+    }
+    return false;
+}
+
+
+uint64 ASTOpConstantBool::Hash() const
+{
+	uint64 res = std::hash<uint64>()(uint64(OP_TYPE::BO_CONSTANT));
+    hash_combine( res, value );
+    return res;
+}
+
+
+mu::Ptr<ASTOp> ASTOpConstantBool::Clone(MapChildFunc&) const
+{
+    Ptr<ASTOpConstantBool> n = new ASTOpConstantBool();
+    n->value = value;
+    return n;
+}
+
+
+void ASTOpConstantBool::Link( PROGRAM& program, const FLinkerOptions*)
+{
+    if (!linkedAddress)
+    {
+        OP::BoolConstantArgs args;
+        memset( &args,0, sizeof(args) );
+        args.value = value;
+
+        linkedAddress = (OP::ADDRESS)program.m_opAddress.size();
+//        program.m_code.push_back(op);
+        program.m_opAddress.push_back((uint32_t)program.m_byteCode.size());
+        AppendCode(program.m_byteCode,OP_TYPE::BO_CONSTANT);
+        AppendCode(program.m_byteCode,args);
+    }
+}
+
+
+ASTOp::BOOL_EVAL_RESULT ASTOpConstantBool::EvaluateBool( ASTOpList&, EVALUATE_BOOL_CACHE* ) const
+{
+    return value? BET_TRUE : BET_FALSE;
+}
+
+
diff --git a/Engine/Plugins/Experimental/Mutable/Source/MutableTools/Private/ASTOpConstantBool.h b/Engine/Plugins/Experimental/Mutable/Source/MutableTools/Private/ASTOpConstantBool.h
new file mode 100644
--- /dev/null
+++ b/Engine/Plugins/Experimental/Mutable/Source/MutableTools/Private/ASTOpConstantBool.h
@@ -0,0 +1,37 @@
+// Copyright Epic Games, Inc. All Rights Reserved.
+
+#pragma once
+
+#include "AST.h"
+
+
+namespace mu
+{
+
+
+	//---------------------------------------------------------------------------------------------
+	//!
+	//---------------------------------------------------------------------------------------------
+	class ASTOpConstantBool : public ASTOp
+	{
+	public:
+
+		//!
+		bool value;
+
+	public:
+
+		ASTOpConstantBool(bool value = true);
+
+		OP_TYPE GetOpType() const override { return OP_TYPE::BO_CONSTANT; }
+		uint64 Hash() const override;
+		void ForEachChild(const std::function<void(ASTChild&)>&) override;
+		bool IsEqual(const ASTOp& otherUntyped) const override;
+		Ptr<ASTOp> Clone(MapChildFunc& mapChild) const override;
+		void Link(PROGRAM& program, const FLinkerOptions* Options) override;
+		BOOL_EVAL_RESULT EvaluateBool(ASTOpList& facts, EVALUATE_BOOL_CACHE* cache) const override;
+	};
+
+
+}
+
diff --git a/Engine/Plugins/Experimental/Mutable/Source/MutableTools/Private/ASTOpConstantProjector.cpp b/Engine/Plugins/Experimental/Mutable/Source/MutableTools/Private/ASTOpConstantProjector.cpp
new file mode 100644
--- /dev/null
+++ b/Engine/Plugins/Experimental/Mutable/Source/MutableTools/Private/ASTOpConstantProjector.cpp
@@ -0,0 +1,61 @@
+// Copyright Epic Games, Inc. All Rights Reserved.
+
+#include "ASTOpConstantProjector.h"
+#include "ErrorLogPrivate.h"
+#include "CodeOptimiser.h"
+#include "StreamsPrivate.h"
+
+
+using namespace mu;
+
+
+//-------------------------------------------------------------------------------------------------
+//-------------------------------------------------------------------------------------------------
+//-------------------------------------------------------------------------------------------------
+void ASTOpConstantProjector::ForEachChild(const std::function<void(ASTChild&)>& )
+{
+}
+
+
+uint64 ASTOpConstantProjector::Hash() const
+{
+	uint64 res = std::hash<float>()(value.position[0]);
+    hash_combine( res, value.direction[0] );
+    return res;
+}
+
+
+bool ASTOpConstantProjector::IsEqual(const ASTOp& otherUntyped) const
+{
+    if ( auto other = dynamic_cast<const ASTOpConstantProjector*>(&otherUntyped) )
+    {
+        return value==other->value;
+    }
+    return false;
+}
+
+
+mu::Ptr<ASTOp> ASTOpConstantProjector::Clone(MapChildFunc&) const
+{
+    Ptr<ASTOpConstantProjector> n = new ASTOpConstantProjector();
+    n->value = value;
+    return n;
+}
+
+
+void ASTOpConstantProjector::Link( PROGRAM& program, const FLinkerOptions*)
+{
+    if (!linkedAddress)
+    {
+        OP::ResourceConstantArgs args;
+        memset( &args,0, sizeof(args) );
+        args.value = program.AddConstant(value);
+
+        linkedAddress = (OP::ADDRESS)program.m_opAddress.size();
+        //program.m_code.push_back(op);
+        program.m_opAddress.push_back((uint32_t)program.m_byteCode.size());
+        AppendCode(program.m_byteCode,OP_TYPE::PR_CONSTANT);
+        AppendCode(program.m_byteCode,args);
+    }
+}
+
diff --git a/Engine/Plugins/Experimental/Mutable/Source/MutableTools/Private/ASTOpConstantProjector.h b/Engine/Plugins/Experimental/Mutable/Source/MutableTools/Private/ASTOpConstantProjector.h
new file mode 100644
--- /dev/null
+++ b/Engine/Plugins/Experimental/Mutable/Source/MutableTools/Private/ASTOpConstantProjector.h
@@ -0,0 +1,34 @@
+// Copyright Epic Games, Inc. All Rights Reserved.
+
+#pragma once
+
+#include "AST.h"
+
+
+namespace mu
+{
+
+
+	//---------------------------------------------------------------------------------------------
+	//!
+	//---------------------------------------------------------------------------------------------
+	class ASTOpConstantProjector : public ASTOp
+	{
+	public:
+
+		//!
+		PROJECTOR value;
+
+	public:
+
+		OP_TYPE GetOpType() const override { return OP_TYPE::PR_CONSTANT; }
+		uint64 Hash() const override;
+		void ForEachChild(const std::function<void(ASTChild&)>&) override;
+		bool IsEqual(const ASTOp& otherUntyped) const override;
+		Ptr<ASTOp> Clone(MapChildFunc& mapChild) const override;
+		void Link(PROGRAM& program, const FLinkerOptions*) override;
+	};
+
+
+}
+
diff --git a/Engine/Plugins/Experimental/Mutable/Source/MutableTools/Private/ASTOpConstantResource.cpp b/Engine/Plugins/Experimental/Mutable/Source/MutableTools/Private/ASTOpConstantResource.cpp
new file mode 100644
--- /dev/null
+++ b/Engine/Plugins/Experimental/Mutable/Source/MutableTools/Private/ASTOpConstantResource.cpp
@@ -0,0 +1,433 @@
+// Copyright Epic Games, Inc. All Rights Reserved.
+
+#include "ASTOpConstantResource.h"
+#include "ErrorLogPrivate.h"
+#include "CodeOptimiser.h"
+#include "StreamsPrivate.h"
+
+#include "Hash/CityHash.h"
+
+#include "OpMeshMerge.h"
+
+#include <queue>
+
+
+using namespace mu;
+
+
+//-------------------------------------------------------------------------------------------------
+//-------------------------------------------------------------------------------------------------
+//-------------------------------------------------------------------------------------------------
+void ASTOpConstantResource::ForEachChild(const std::function<void (ASTChild&)>& )
+{
+}
+
+
+//-------------------------------------------------------------------------------------------------
+bool ASTOpConstantResource::IsEqual(const ASTOp& otherUntyped) const
+{
+    if ( auto other = dynamic_cast<const ASTOpConstantResource*>(&otherUntyped) )
+    {
+        return type == other->type && hash == other->hash &&
+               loadedValue == other->loadedValue && proxy == other->proxy;
+    }
+    return false;
+}
+
+
+//-------------------------------------------------------------------------------------------------
+mu::Ptr<ASTOp> ASTOpConstantResource::Clone(MapChildFunc&) const
+{
+    Ptr<ASTOpConstantResource> n = new ASTOpConstantResource();
+    n->type = type;
+    n->proxy = proxy;
+    n->loadedValue = loadedValue;
+    n->hash = hash;
+    return n;
+}
+
+
+//-------------------------------------------------------------------------------------------------
+uint64 ASTOpConstantResource::Hash() const
+{
+	uint64 res = std::hash<uint64>()(uint64(type));
+    hash_combine( res, hash );
+    return res;
+}
+
+
+//-------------------------------------------------------------------------------------------------
+void ASTOpConstantResource::Link( PROGRAM& program, const FLinkerOptions* Options)
+{
+    if (!linkedAddress && !bLinkedAndNull)
+    {
+        if (type==OP_TYPE::ME_CONSTANT)
+        {
+            OP::MeshConstantArgs args;
+            memset( &args,0, sizeof(args) );
+
+            Ptr<Mesh> pTyped = static_cast<const Mesh*>(GetValue().get())->Clone();
+            check(pTyped);
+
+            args.skeleton = -1;
+            if (Ptr<const Skeleton> pSkeleton = pTyped->GetSkeleton())
+            {
+                // See if it can be merged into another skeleton
+//                for (const auto& pOtherSkeleton : program.m_constantSkeletons)
+//                {
+//                    if (true)
+//                    //if (AreSkeletonsRelated(pOtherSkeleton, pSkeleton))
+//                    {
+//                        // \todo : ugly const hack
+//                        ExtendSkeleton( const_cast<Skeleton*>(pOtherSkeleton.get()),
+//                                        pSkeleton.get() );
+//                        pSkeleton = pOtherSkeleton;
+
+//                        Ptr<Mesh> pRemapped = MeshRemapSkeleton( pTyped.get(),
+//                                                                 pOtherSkeleton.get() );
+//                        if (pRemapped)
+//                        {
+//                            pTyped = pRemapped;
+//                        }
+//                        break;
+//                    }
+//                }
+
+                args.skeleton = program.AddConstant( pSkeleton.get() );
+                pTyped->SetSkeleton( nullptr );
+            }
+
+			args.physicsBody = -1;
+			if (Ptr<const PhysicsBody> pPhysicsBody = pTyped->GetPhysicsBody())
+			{
+				args.physicsBody = program.AddConstant( pPhysicsBody.get() );
+                pTyped->SetPhysicsBody( nullptr );
+			}
+
+            args.value = program.AddConstant(pTyped.get());
+
+            linkedAddress = (OP::ADDRESS)program.m_opAddress.size();
+            program.m_opAddress.push_back((uint32_t)program.m_byteCode.size());
+            AppendCode(program.m_byteCode,type);
+            AppendCode(program.m_byteCode,args);
+        }
+        else
+        {
+            OP::ResourceConstantArgs args;
+            memset( &args,0, sizeof(args) );
+
+			bool bValidData = true;
+
+            switch (type)
+            {
+            case OP_TYPE::IM_CONSTANT:
+            {
+                Ptr<const Image> pTyped = static_cast<const Image*>(GetValue().get());
+                check(pTyped);
+
+				if (pTyped->GetSizeX()*pTyped->GetSizeY()==0)
+				{
+					// It's an empty or degenerated image, return a null operation.
+					bValidData = false;
+				}
+				else
+				{
+					// TODO: Compilation options?
+					int32 MinTextureResidentMipCount = Options->MinTextureResidentMipCount;
+					args.value = program.AddConstant(pTyped, MinTextureResidentMipCount);
+				}
+                break;
+            }
+            case OP_TYPE::LA_CONSTANT:
+            {
+                Ptr<const Layout> pTyped = static_cast<const Layout*>(GetValue().get());
+                check(pTyped);
+                args.value = program.AddConstant(pTyped);
+                break;
+            }
+            default:
+                check(false);
+            }
+
+			if (bValidData)
+			{
+				linkedAddress = (OP::ADDRESS)program.m_opAddress.size();
+				program.m_opAddress.push_back((uint32_t)program.m_byteCode.size());
+				AppendCode(program.m_byteCode, type);
+				AppendCode(program.m_byteCode, args);
+			}
+			else
+			{
+				// Null op
+				linkedAddress = 0;
+				bLinkedAndNull = true;
+			}
+        }
+    }
+}
+
+
+//-------------------------------------------------------------------------------------------------
+FImageDesc ASTOpConstantResource::GetImageDesc( bool, class GetImageDescContext* )
+{
+    FImageDesc res;
+
+    if (type==OP_TYPE::IM_CONSTANT)
+    {
+        // TODO: cache to avoid disk loading
+        Ptr<const Image> pConst = static_cast<const Image*>(GetValue().get());
+        res.m_format = pConst->m_format;
+		res.m_lods = pConst->m_lods;
+		res.m_size = pConst->m_size;
+	}
+    else
+    {
+        check(false);
+    }
+
+    return res;
+}
+
+
+//-------------------------------------------------------------------------------------------------
+void ASTOpConstantResource::GetBlockLayoutSize( int blockIndex, int* pBlockX, int* pBlockY,
+                                                BLOCK_LAYOUT_SIZE_CACHE* )
+{
+    switch ( type )
+    {
+    case OP_TYPE::LA_CONSTANT:
+    {
+        Ptr<const Layout> pLayout = static_cast<const Layout*>( GetValue().get() );
+        check(pLayout);
+
+        if (pLayout)
+        {
+            int relId = pLayout->FindBlock( blockIndex );
+            if ( relId>=0 )
+            {
+                *pBlockX = pLayout->m_blocks[relId].m_size[0];
+                *pBlockY = pLayout->m_blocks[relId].m_size[1];
+            }
+            else
+            {
+                *pBlockX = 0;
+                *pBlockY = 0;
+            }
+        }
+
+        break;
+    }
+    default:
+        check(false);
+    }
+}
+
+
+//-------------------------------------------------------------------------------------------------
+void ASTOpConstantResource::GetLayoutBlockSize( int* pBlockX, int* pBlockY )
+{
+    switch ( type )
+    {
+
+    case OP_TYPE::IM_CONSTANT:
+    {
+        // We didn't find any layout.
+        *pBlockX = 0;
+        *pBlockY = 0;
+        break;
+    }
+
+    default:
+        checkf( false, TEXT("Instruction not supported") );
+    }
+}
+
+
+//-------------------------------------------------------------------------------------------------
+bool ASTOpConstantResource::GetNonBlackRect( FImageRect& maskUsage ) const
+{
+    if (type==OP_TYPE::IM_CONSTANT)
+    {
+        // TODO: cache
+        Ptr<const Image> pMask = static_cast<const Image*>(GetValue().get());
+        pMask->GetNonBlackRect(maskUsage);
+        return true;
+    }
+
+    return false;
+}
+
+
+//-------------------------------------------------------------------------------------------------
+bool ASTOpConstantResource::IsImagePlainConstant( vec4<float>& colour ) const
+{
+    bool res = false;
+    switch( type )
+    {
+
+    case OP_TYPE::IM_CONSTANT:
+    {
+        Ptr<const Image> pImage = static_cast<const Image*>(GetValue().get());
+		check(pImage->m_size[0] > 0);
+		check(pImage->m_size[1] > 0);
+
+        if ( pImage->m_flags & Image::IF_IS_PLAIN_COLOUR_VALID )
+        {
+            if ( pImage->m_flags & Image::IF_IS_PLAIN_COLOUR )
+            {
+                res = true;
+                colour = pImage->Sample( vec2<float>(0,0) );
+            }
+            else
+            {
+                res = false;
+            }
+        }
+        else
+        {
+            if ( pImage->IsPlainColour( colour ) )
+            {
+                res = true;
+                pImage->m_flags |= Image::IF_IS_PLAIN_COLOUR;
+            }
+
+            pImage->m_flags |= Image::IF_IS_PLAIN_COLOUR_VALID;
+        }
+        break;
+    }
+
+    default:
+        break;
+    }
+
+    return res;
+}
+
+
+//-------------------------------------------------------------------------------------------------
+ASTOpConstantResource::~ASTOpConstantResource()
+{
+}
+
+
+//-------------------------------------------------------------------------------------------------
+uint64 ASTOpConstantResource::GetValueHash() const
+{
+    return hash;
+}
+
+
+//-------------------------------------------------------------------------------------------------
+mu::Ptr<const RefCounted> ASTOpConstantResource::GetValue() const
+{
+    if (loadedValue)
+    {
+        return loadedValue;
+    }
+    else
+    {
+        switch( type )
+        {
+
+        case OP_TYPE::IM_CONSTANT:
+        {
+            Ptr<ResourceProxy<Image>> typedProxy = dynamic_cast< ResourceProxy<Image>* >( proxy.get() );
+            Ptr<const Image> r = typedProxy->Get();
+            return r;
+        }
+
+        default:
+            check(false);
+            break;
+        }
+    }
+
+    return nullptr;
+}
+
+
+//-------------------------------------------------------------------------------------------------
+void ASTOpConstantResource::SetValue( const Ptr<const RefCounted>& v,
+                                      bool useDiskCache )
+{
+    switch( type )
+    {
+    case OP_TYPE::IM_CONSTANT:
+    {
+        Ptr<const Image> r = static_cast<const Image*>(v.get());
+
+        OutputMemoryStream stream(r->GetDataSize()+1024);
+        OutputArchive arch(&stream);
+        Image::Serialise( r.get(), arch );
+
+		hash = CityHash64(static_cast<const char*>(stream.GetBuffer()), stream.GetBufferSize());
+
+        if (useDiskCache)
+        {
+            proxy = new ResourceProxyTempFile<Image>(r.get());
+        }
+        else
+        {
+            loadedValue = r;
+        }
+        break;
+    }
+
+    case OP_TYPE::ME_CONSTANT:
+    {
+        OutputMemoryStream stream;
+        OutputArchive arch(&stream);
+
+        Ptr<const Mesh> r = static_cast<const Mesh*>(v.get());
+        Mesh::Serialise( r.get(), arch );
+
+		hash = CityHash64(static_cast<const char*>(stream.GetBuffer()), stream.GetBufferSize());
+
+        loadedValue = v;
+        break;
+    }
+
+    case OP_TYPE::LA_CONSTANT:
+    {
+        OutputMemoryStream stream;
+        OutputArchive arch(&stream);
+
+        Ptr<const Layout> r = static_cast<const Layout*>(v.get());
+        Layout::Serialise( r.get(), arch );
+
+		hash = CityHash64(static_cast<const char*>(stream.GetBuffer()), stream.GetBufferSize());
+
+        loadedValue = v;
+        break;
+    }
+
+    default:
+        loadedValue = v;
+        break;
+    }
+}
+
+
+//-------------------------------------------------------------------------------------------------
+mu::Ptr<ImageSizeExpression> ASTOpConstantResource::GetImageSizeExpression() const
+{
+    Ptr<ImageSizeExpression> pRes = new ImageSizeExpression;
+    pRes->type = ImageSizeExpression::ISET_CONSTANT;
+
+    switch( type )
+    {
+    case OP_TYPE::IM_CONSTANT:
+    {
+        Ptr<const Image> pConst = static_cast<const Image*>(GetValue().get());
+        pRes->size = pConst->m_size;
+        break;
+    }
+
+    default:
+        check( false );
+        return nullptr;
+    }
+
+    return pRes;
+}
+
diff --git a/Engine/Plugins/Experimental/Mutable/Source/MutableTools/Private/ASTOpConstantResource.h b/Engine/Plugins/Experimental/Mutable/Source/MutableTools/Private/ASTOpConstantResource.h
new file mode 100644
--- /dev/null
+++ b/Engine/Plugins/Experimental/Mutable/Source/MutableTools/Private/ASTOpConstantResource.h
@@ -0,0 +1,73 @@
+// Copyright Epic Games, Inc. All Rights Reserved.
+
+#pragma once
+
+#include "AST.h"
+
+
+namespace mu
+{
+
+
+	//---------------------------------------------------------------------------------------------
+	//! A constant mesh, image, volume or layout
+	//---------------------------------------------------------------------------------------------
+	class ASTOpConstantResource : public ASTOp
+	{
+	private:
+
+		//!
+		//std::FILE* resourceFile = nullptr;
+		//size_t fileSize = 0;
+
+		//!
+		Ptr<const RefCounted> loadedValue;
+		Ptr<RefCounted> proxy;
+
+		//! Value hash
+		uint64 hash;
+
+		//! We tried to link already but the result is a null op.
+		bool bLinkedAndNull = false;
+
+	public:
+
+		//! Type of switch
+		OP_TYPE type;
+
+	public:
+
+		~ASTOpConstantResource() override;
+
+		// Own interface
+
+		//! Get a hash of the stored value.
+		uint64 GetValueHash() const;
+
+		//! Get a copy of the stored value
+		Ptr<const RefCounted> GetValue() const;
+
+		//! Set the value to store in this op
+		void SetValue(const Ptr<const RefCounted>& v,
+			bool useDiskCache);
+
+
+		// ASTOp interface
+		OP_TYPE GetOpType() const override { return type; }
+		void ForEachChild(const std::function<void(ASTChild&)>&) override;
+		bool IsEqual(const ASTOp& otherUntyped) const override;
+		Ptr<ASTOp> Clone(MapChildFunc& mapChild) const override;
+		uint64 Hash() const override;
+		void Link(PROGRAM& program, const FLinkerOptions*) override;
+		FImageDesc GetImageDesc(bool, class GetImageDescContext*) override;
+		void GetBlockLayoutSize(int blockIndex, int* pBlockX, int* pBlockY,
+			BLOCK_LAYOUT_SIZE_CACHE* cache) override;
+		void GetLayoutBlockSize(int* pBlockX, int* pBlockY) override;
+		bool GetNonBlackRect(FImageRect& maskUsage) const override;
+		bool IsImagePlainConstant(vec4<float>& colour) const override;
+		Ptr<ImageSizeExpression> GetImageSizeExpression() const override;
+	};
+
+
+}
+
diff --git a/Engine/Plugins/Experimental/Mutable/Source/MutableTools/Private/ASTOpConstantString.cpp b/Engine/Plugins/Experimental/Mutable/Source/MutableTools/Private/ASTOpConstantString.cpp
new file mode 100644
--- /dev/null
+++ b/Engine/Plugins/Experimental/Mutable/Source/MutableTools/Private/ASTOpConstantString.cpp
@@ -0,0 +1,59 @@
+// Copyright Epic Games, Inc. All Rights Reserved.
+
+#include "ASTOpConstantString.h"
+#include "ErrorLogPrivate.h"
+#include "CodeOptimiser.h"
+#include "StreamsPrivate.h"
+
+
+using namespace mu;
+
+
+//-------------------------------------------------------------------------------------------------
+//-------------------------------------------------------------------------------------------------
+//-------------------------------------------------------------------------------------------------
+void ASTOpConstantString::ForEachChild(const std::function<void(ASTChild&)>& )
+{
+}
+
+
+uint64 ASTOpConstantString::Hash() const
+{
+	uint64 res = std::hash<std::string>()(value.c_str());
+    return res;
+}
+
+
+bool ASTOpConstantString::IsEqual(const ASTOp& otherUntyped) const
+{
+    if ( auto other = dynamic_cast<const ASTOpConstantString*>(&otherUntyped) )
+    {
+        return value==other->value;
+    }
+    return false;
+}
+
+
+mu::Ptr<ASTOp> ASTOpConstantString::Clone(MapChildFunc&) const
+{
+    Ptr<ASTOpConstantString> n = new ASTOpConstantString();
+    n->value = value;
+    return n;
+}
+
+
+void ASTOpConstantString::Link( PROGRAM& program, const FLinkerOptions* )
+{
+    if (!linkedAddress)
+    {
+        OP::ResourceConstantArgs args;
+        memset( &args,0, sizeof(args) );
+        args.value = program.AddConstant(value);
+
+        linkedAddress = (OP::ADDRESS)program.m_opAddress.size();
+        program.m_opAddress.push_back((uint32_t)program.m_byteCode.size());
+        AppendCode(program.m_byteCode,OP_TYPE::ST_CONSTANT);
+        AppendCode(program.m_byteCode,args);
+    }
+}
+
diff --git a/Engine/Plugins/Experimental/Mutable/Source/MutableTools/Private/ASTOpConstantString.h b/Engine/Plugins/Experimental/Mutable/Source/MutableTools/Private/ASTOpConstantString.h
new file mode 100644
--- /dev/null
+++ b/Engine/Plugins/Experimental/Mutable/Source/MutableTools/Private/ASTOpConstantString.h
@@ -0,0 +1,32 @@
+// Copyright Epic Games, Inc. All Rights Reserved.
+
+#pragma once
+
+#include "AST.h"
+
+
+namespace mu
+{
+
+
+	//---------------------------------------------------------------------------------------------
+	//!
+	//---------------------------------------------------------------------------------------------
+	class ASTOpConstantString : public ASTOp
+	{
+	public:
+		//!
+		string value;
+
+	public:
+		OP_TYPE GetOpType() const override { return OP_TYPE::ST_CONSTANT; }
+		uint64 Hash() const override;
+		void ForEachChild(const std::function<void(ASTChild&)>&) override;
+		bool IsEqual(const ASTOp& otherUntyped) const override;
+		Ptr<ASTOp> Clone(MapChildFunc& mapChild) const override;
+		void Link(PROGRAM& program, const FLinkerOptions* Options) override;
+	};
+
+
+}
+
diff --git a/Engine/Plugins/Experimental/Mutable/Source/MutableTools/Private/ASTOpImageCompose.cpp b/Engine/Plugins/Experimental/Mutable/Source/MutableTools/Private/ASTOpImageCompose.cpp
new file mode 100644
--- /dev/null
+++ b/Engine/Plugins/Experimental/Mutable/Source/MutableTools/Private/ASTOpImageCompose.cpp
@@ -0,0 +1,272 @@
+// Copyright Epic Games, Inc. All Rights Reserved.
+
+#include "ASTOpImageCompose.h"
+#include "ErrorLogPrivate.h"
+#include "CodeOptimiser.h"
+#include "StreamsPrivate.h"
+#include "MutableTools/Private/ASTOpImagePixelFormat.h"
+#include "MutableTools/Private/ASTOpConstantResource.h"
+#include "MutableRuntime/Public/Layout.h"
+
+
+using namespace mu;
+
+
+//-------------------------------------------------------------------------------------------------
+//-------------------------------------------------------------------------------------------------
+//-------------------------------------------------------------------------------------------------
+ASTOpImageCompose::ASTOpImageCompose()
+    : Layout(this)
+    , Base(this)
+	, BlockImage(this)
+	, Mask(this)
+{
+}
+
+
+//-------------------------------------------------------------------------------------------------
+ASTOpImageCompose::~ASTOpImageCompose()
+{
+    // Explicit call needed to avoid recursive destruction
+    ASTOp::RemoveChildren();
+}
+
+
+//-------------------------------------------------------------------------------------------------
+bool ASTOpImageCompose::IsEqual(const ASTOp& otherUntyped) const
+{
+    if ( const ASTOpImageCompose* other = dynamic_cast<const ASTOpImageCompose*>(&otherUntyped) )
+    {
+        return Layout ==other->Layout &&
+			Base ==other->Base &&
+			BlockImage == other->BlockImage &&
+			Mask == other->Mask &&
+			BlockIndex == other->BlockIndex;
+    }
+    return false;
+}
+
+
+//-------------------------------------------------------------------------------------------------
+uint64 ASTOpImageCompose::Hash() const
+{
+	uint64 res = std::hash<OP_TYPE>()(OP_TYPE::IM_COMPOSE);
+    hash_combine( res, Layout.child().get() );
+	hash_combine(res, Base.child().get());
+	hash_combine(res, BlockImage.child().get());
+	hash_combine(res, Mask.child().get());
+	hash_combine(res, BlockIndex);
+	return res;
+}
+
+
+//-------------------------------------------------------------------------------------------------
+mu::Ptr<ASTOp> ASTOpImageCompose::Clone(MapChildFunc& mapChild) const
+{
+    Ptr<ASTOpImageCompose> n = new ASTOpImageCompose();
+    n->Layout = mapChild(Layout.child());
+	n->Base = mapChild(Base.child());
+	n->BlockImage = mapChild(BlockImage.child());
+	n->Mask = mapChild(Mask.child());
+	n->BlockIndex = BlockIndex;
+    return n;
+}
+
+
+//-------------------------------------------------------------------------------------------------
+void ASTOpImageCompose::ForEachChild(const std::function<void(ASTChild&)>& f )
+{
+	f(Layout);
+	f(Base);
+	f(BlockImage);
+	f(Mask);
+}
+
+
+//-------------------------------------------------------------------------------------------------
+void ASTOpImageCompose::Link( PROGRAM& program, const FLinkerOptions*)
+{
+    // Already linked?
+    if (!linkedAddress)
+    {
+        OP::ImageComposeArgs args;
+        memset( &args,0, sizeof(args) );
+
+		if (Layout) args.layout = Layout->linkedAddress;
+		if (Base) args.base = Base->linkedAddress;
+		if (BlockImage) args.blockImage = BlockImage->linkedAddress;
+		if (Mask) args.mask = Mask->linkedAddress;
+        args.blockIndex = BlockIndex;
+
+        linkedAddress = (OP::ADDRESS)program.m_opAddress.size();
+        program.m_opAddress.push_back((uint32_t)program.m_byteCode.size());
+        AppendCode(program.m_byteCode,OP_TYPE::IM_COMPOSE);
+        AppendCode(program.m_byteCode,args);
+    }
+
+}
+
+
+FImageDesc ASTOpImageCompose::GetImageDesc( bool returnBestOption, GetImageDescContext* context )
+{
+    FImageDesc res;
+
+    // Local context in case it is necessary
+    GetImageDescContext localContext;
+    if (!context)
+    {
+      context = &localContext;
+    }
+    else
+    {
+        // Cached result?
+		FImageDesc* PtrValue = context->m_results.Find(this);
+		if (PtrValue)
+		{
+			return *PtrValue;
+		}
+	}
+
+    // Actual work
+    if (Base)
+    {
+        res = Base->GetImageDesc( returnBestOption, context );
+    }
+
+
+    // Cache the result
+    if (context)
+    {
+		context->m_results.Add(this, res);
+	}
+
+    return res;
+}
+
+
+mu::Ptr<ImageSizeExpression> ASTOpImageCompose::GetImageSizeExpression() const
+{
+    if ( Base )
+    {
+        return Base->GetImageSizeExpression();
+    }
+
+    return nullptr;
+}
+
+
+bool ASTOpImageCompose::IsImagePlainConstant(vec4<float>& colour) const
+{
+	bool res = false;
+
+	if (BlockImage.child())
+	{
+		BlockImage->IsImagePlainConstant(colour);
+	}
+
+	if (res && Base.child())
+	{
+		vec4<float> baseColour;
+		res = Base->IsImagePlainConstant(baseColour);
+		res &= (colour == baseColour);
+	}
+
+	return res;
+}
+
+
+void ASTOpImageCompose::GetLayoutBlockSize(int* pBlockX, int* pBlockY)
+{
+	// We can find out in two ways: following the block image, or following the base image.
+	// Let's try the block approach: We need the block size and the layout blocks
+	int layoutBlocksX = 0;
+	int layoutBlocksY = 0;
+	if (Layout.child())
+	{
+		MUTABLE_CPUPROFILER_SCOPE(GetLayoutBlockSize_GetBlockLayoutSize);
+		BLOCK_LAYOUT_SIZE_CACHE cache;
+		Layout->GetBlockLayoutSizeCached(BlockIndex,
+			&layoutBlocksX, &layoutBlocksY,
+			&cache);
+	}
+
+	if (layoutBlocksX > 0 && layoutBlocksY > 0 && BlockImage.child())
+	{
+		FImageDesc blockDesc = BlockImage->GetImageDesc();
+
+		*pBlockX = blockDesc.m_size[0] / layoutBlocksX;
+		*pBlockY = blockDesc.m_size[1] / layoutBlocksY;
+	}
+	else
+	{
+		*pBlockX = 0;
+		*pBlockY = 0;
+	}
+}
+
+
+Ptr<ASTOp> ASTOpImageCompose::OptimiseSemantic(const MODEL_OPTIMIZATION_OPTIONS& options) const
+{
+	Ptr<ASTOp> at;
+
+	auto baseAt = Base.child();
+	auto blockAt = BlockImage.child();
+	auto layoutAt = Layout.child();
+	if (layoutAt
+		&&
+		layoutAt->GetOpType() == OP_TYPE::LA_CONSTANT
+		&&
+		baseAt
+		&&
+		blockAt)
+	{
+		auto typedLayout = dynamic_cast<const ASTOpConstantResource*>(layoutAt.get());
+		Ptr<const mu::Layout> pLayout = static_cast<const mu::Layout*>(typedLayout->GetValue().get());
+
+		// Constant single-block full layout?
+		if (pLayout->GetBlockCount() == 1
+			&&
+			pLayout->m_blocks[0].m_min == vec2<uint16_t>(0, 0)
+			&&
+			pLayout->m_blocks[0].m_size == pLayout->m_size
+			&&
+			pLayout->m_blocks[0].m_id == int(BlockIndex)
+			)
+		{
+			// We could only take the block, but we must make sure it will have the format
+			// and size of the base.
+			auto baseDesc = baseAt->GetImageDesc(true);
+			auto blockDesc = blockAt->GetImageDesc(true);
+
+			at = blockAt;
+
+			// \todo isn't this a common operation?
+			if (baseDesc.m_format != blockDesc.m_format
+				&&
+				baseDesc.m_format != EImageFormat::IF_NONE)
+			{
+				Ptr<ASTOpImagePixelFormat> reformat = new ASTOpImagePixelFormat;
+				reformat->Format = baseDesc.m_format;
+				reformat->FormatIfAlpha = EImageFormat::IF_NONE;
+				reformat->Source = at;
+				at = reformat;
+			}
+
+			if (baseDesc.m_size != blockDesc.m_size
+				&&
+				baseDesc.m_size[0] != 0
+				&&
+				baseDesc.m_size[1] != 0)
+			{
+				Ptr<ASTOpFixed> resize = new ASTOpFixed;
+				resize->op.type = OP_TYPE::IM_RESIZE;
+				resize->op.args.ImageResize.size[0] = baseDesc.m_size[0];
+				resize->op.args.ImageResize.size[1] = baseDesc.m_size[1];
+				resize->SetChild(resize->op.args.ImageResize.source, at);
+				at = resize;
+			}
+		}
+	}
+
+	return at;
+}
diff --git a/Engine/Plugins/Experimental/Mutable/Source/MutableTools/Private/ASTOpImageCompose.h b/Engine/Plugins/Experimental/Mutable/Source/MutableTools/Private/ASTOpImageCompose.h
new file mode 100644
--- /dev/null
+++ b/Engine/Plugins/Experimental/Mutable/Source/MutableTools/Private/ASTOpImageCompose.h
@@ -0,0 +1,43 @@
+// Copyright Epic Games, Inc. All Rights Reserved.
+
+#pragma once
+
+#include "MutableTools/Private/AST.h"
+
+
+namespace mu
+{
+
+	class ASTOpImageCompose : public ASTOp
+	{
+	public:
+
+		ASTChild Layout;
+		ASTChild Base;
+		ASTChild BlockImage;
+		ASTChild Mask;
+		uint32 BlockIndex=0;
+
+	public:
+
+		ASTOpImageCompose();
+		ASTOpImageCompose(const ASTOpImageCompose&) = delete;
+		~ASTOpImageCompose();
+
+		OP_TYPE GetOpType() const override { return OP_TYPE::IM_COMPOSE; }
+		uint64 Hash() const override;
+		bool IsEqual(const ASTOp& otherUntyped) const override;
+		Ptr<ASTOp> Clone(MapChildFunc& mapChild) const override;
+		void ForEachChild(const std::function<void(ASTChild&)>&) override;
+		void Link(PROGRAM& program, const FLinkerOptions* Options) override;
+		//TODO: Ptr<ASTOp> OptimiseSink(const MODEL_OPTIMIZATION_OPTIONS& options, OPTIMIZE_SINK_CONTEXT& context) const override;
+		Ptr<ASTOp> OptimiseSemantic(const MODEL_OPTIMIZATION_OPTIONS& options) const;
+		FImageDesc GetImageDesc(bool returnBestOption, GetImageDescContext* context) override;
+		void GetLayoutBlockSize(int* pBlockX, int* pBlockY) override;
+		bool IsImagePlainConstant(vec4<float>& colour) const override;
+		Ptr<ImageSizeExpression> GetImageSizeExpression() const override;
+
+	};
+
+}
+
diff --git a/Engine/Plugins/Experimental/Mutable/Source/MutableTools/Private/ASTOpImageMipmap.cpp b/Engine/Plugins/Experimental/Mutable/Source/MutableTools/Private/ASTOpImageMipmap.cpp
new file mode 100644
--- /dev/null
+++ b/Engine/Plugins/Experimental/Mutable/Source/MutableTools/Private/ASTOpImageMipmap.cpp
@@ -0,0 +1,450 @@
+// Copyright Epic Games, Inc. All Rights Reserved.
+
+#include "ASTOpImageMipmap.h"
+
+#include "ASTOpImagePixelFormat.h"
+#include "ASTOpImageCompose.h"
+#include "ASTOpImagePatch.h"
+#include "ASTOpSwitch.h"
+#include "ASTOpConditional.h"
+
+#include "ErrorLogPrivate.h"
+#include "CodeOptimiser.h"
+#include "StreamsPrivate.h"
+
+
+using namespace mu;
+
+
+//-------------------------------------------------------------------------------------------------
+//-------------------------------------------------------------------------------------------------
+//-------------------------------------------------------------------------------------------------
+ASTOpImageMipmap::ASTOpImageMipmap()
+    : Source(this)
+{
+}
+
+
+ASTOpImageMipmap::~ASTOpImageMipmap()
+{
+    // Explicit call needed to avoid recursive destruction
+    ASTOp::RemoveChildren();
+}
+
+
+bool ASTOpImageMipmap::IsEqual(const ASTOp& otherUntyped) const
+{
+    if (const ASTOpImageMipmap* other = dynamic_cast<const ASTOpImageMipmap*>(&otherUntyped) )
+    {
+        return Source==other->Source && 
+			Levels == other->Levels &&
+			BlockLevels == other->BlockLevels &&
+			bOnlyTail == other->bOnlyTail &&
+			SharpenFactor == other->SharpenFactor &&
+			AddressMode == other->AddressMode &&
+			FilterType == other->FilterType &&
+			DitherMipmapAlpha == other->DitherMipmapAlpha;
+	}
+    return false;
+}
+
+
+uint64 ASTOpImageMipmap::Hash() const
+{
+	uint64 res = std::hash<void*>()(Source.child().get() );
+    hash_combine( res, Levels);
+    return res;
+}
+
+
+mu::Ptr<ASTOp> ASTOpImageMipmap::Clone(MapChildFunc& mapChild) const
+{
+    Ptr<ASTOpImageMipmap> n = new ASTOpImageMipmap();
+    n->Source = mapChild(Source.child());
+	n->Levels = Levels;
+	n->BlockLevels = BlockLevels;
+	n->bOnlyTail = bOnlyTail;
+	n->SharpenFactor = SharpenFactor;
+	n->AddressMode = AddressMode;
+	n->FilterType = FilterType;
+	n->DitherMipmapAlpha = DitherMipmapAlpha;
+
+    return n;
+}
+
+
+void ASTOpImageMipmap::ForEachChild(const std::function<void(ASTChild&)>& f )
+{
+    f( Source );
+}
+
+
+void ASTOpImageMipmap::Link( PROGRAM& program, const FLinkerOptions* )
+{
+    // Already linked?
+    if (!linkedAddress)
+    {
+        OP::ImageMipmapArgs args;
+        FMemory::Memzero( &args, sizeof(args) );
+
+		args.levels = Levels;
+		args.blockLevels = BlockLevels;
+		args.onlyTail = bOnlyTail;
+		args.sharpenFactor = SharpenFactor;
+		args.addressMode = AddressMode;
+		args.filterType = FilterType;
+		args.ditherMipmapAlpha = DitherMipmapAlpha;
+		if (Source) args.source = Source->linkedAddress;
+
+        linkedAddress = (OP::ADDRESS)program.m_opAddress.size();
+        //program.m_code.push_back(op);
+        program.m_opAddress.push_back((uint32_t)program.m_byteCode.size());
+        AppendCode(program.m_byteCode,OP_TYPE::IM_MIPMAP);
+        AppendCode(program.m_byteCode,args);
+    }
+
+}
+
+
+Ptr<ASTOp> ASTOpImageMipmap::OptimiseSink(const MODEL_OPTIMIZATION_OPTIONS& options, OPTIMIZE_SINK_CONTEXT& context) const
+{
+	Ptr<ASTOp> at;
+
+	Ptr<ASTOp> sourceAt = Source.child();
+	switch (sourceAt->GetOpType())
+	{
+
+	case OP_TYPE::IM_BLANKLAYOUT:
+	{
+		// Set the mipmap generation on the blank layout operation
+		auto mop = mu::Clone<ASTOpFixed>(sourceAt);
+		mop->op.args.ImageBlankLayout.generateMipmaps = 1;    // true
+		mop->op.args.ImageBlankLayout.mipmapCount = Levels;
+		at = mop;
+		break;
+	}
+
+	case OP_TYPE::IM_PIXELFORMAT:
+	{
+		// Swap unless the mipmap operation builds only the tail or is compressed.
+		// Otherwise, we could fall in a loop of swapping mipmaps and pixelformats.
+		Ptr<ASTOpImageMipmap> mop = mu::Clone<ASTOpImageMipmap>(this);
+		Ptr<ASTOpImagePixelFormat> fop = mu::Clone<ASTOpImagePixelFormat>(sourceAt);
+		bool isCompressedFormat = IsCompressedFormat(fop->Format);
+		if (isCompressedFormat && !mop->bOnlyTail)
+		{
+			mop->Source = fop->Source.child();
+			fop->Source = mop;
+
+			at = fop;
+		}
+		break;
+	}
+
+	default:
+	{
+		at = context.ImageMipmapSinker.Apply(this);
+
+		break;
+	} // mipmap source default
+
+	} // mipmap source type switch
+
+	
+	return at;
+}
+
+
+//!
+FImageDesc ASTOpImageMipmap::GetImageDesc(bool returnBestOption, GetImageDescContext* context)
+{
+	FImageDesc res;
+
+	// Local context in case it is necessary
+	GetImageDescContext localContext;
+	if (!context)
+	{
+		context = &localContext;
+	}
+	else
+	{
+		// Cached result?
+		FImageDesc* PtrValue = context->m_results.Find(this);
+		if (PtrValue)
+		{
+			return *PtrValue;
+		}
+	}
+
+	if (Source.child())
+	{
+		res = Source.child()->GetImageDesc( returnBestOption, context);
+	}
+
+	int mipLevels = Levels;
+
+	if (mipLevels == 0)
+	{
+		mipLevels = FMath::Max(
+			(int)ceilf(logf((float)res.m_size[0]) / logf(2.0f)),
+			(int)ceilf(logf((float)res.m_size[1]) / logf(2.0f))
+		);
+	}
+
+	res.m_lods = FMath::Max(res.m_lods, (uint8)mipLevels);
+
+
+	// Cache the result
+	if (context)
+	{
+		context->m_results.Add(this, res);
+	}
+
+	return res;
+}
+
+
+void ASTOpImageMipmap::GetLayoutBlockSize(int* pBlockX, int* pBlockY)
+{
+	if (Source.child())
+	{
+		// Assume the block size of the biggest mip
+		Source.child()->GetLayoutBlockSize(pBlockX, pBlockY);
+	}
+}
+
+
+bool ASTOpImageMipmap::IsImagePlainConstant(vec4<float>& colour) const
+{
+	bool res = false;
+	if (Source.child())
+	{
+		Source.child()->IsImagePlainConstant(colour);
+	}
+	return res;
+}
+
+
+mu::Ptr<ImageSizeExpression> ASTOpImageMipmap::GetImageSizeExpression() const
+{
+	Ptr<ImageSizeExpression> pRes;
+
+	if (Source.child())
+	{
+		pRes = Source.child()->GetImageSizeExpression();
+	}
+	else
+	{
+		pRes = new ImageSizeExpression;
+	}
+
+	return pRes;
+}
+
+
+//---------------------------------------------------------------------------------------------
+//---------------------------------------------------------------------------------------------
+//---------------------------------------------------------------------------------------------
+Ptr<ASTOp> Sink_ImageMipmapAST::Apply(const ASTOp* root)
+{
+	check(root->GetOpType() == OP_TYPE::IM_MIPMAP);
+
+	m_root = dynamic_cast<const ASTOpImageMipmap*>(root);
+	m_oldToNew.clear();
+	m_newOps.clear();
+
+	if (m_root->bOnlyTail)
+	{
+		return nullptr;
+	}
+
+	Ptr<ASTOp> newSource;
+	m_initialSource = m_root->Source.child();
+
+	// Before sinking, see if it is worth splitting into miptail and mip.
+	if (!m_root->bOnlyTail
+		// \TODO: Review this: It seems it fails for a BC1 8x4 texture having blockSize=4, and level=0
+			//&&
+			//m_root->Levels
+			//!=
+			//m_root->BlockLevels 
+		)
+	{
+		// the block mipmaps can be done before composition is done.
+		Ptr<ASTOpImageMipmap> newMip = mu::Clone<ASTOpImageMipmap>(m_root);
+		newMip->Levels = m_root->BlockLevels;
+		newMip->BlockLevels = m_root->BlockLevels;
+		newMip->bOnlyTail = false;
+
+		// the smallest mipmaps after the composition is done.
+		Ptr<ASTOpImageMipmap> topMipOp = mu::Clone<ASTOpImageMipmap>(m_root);
+		topMipOp->bOnlyTail = true;
+
+		// Proceed
+		topMipOp->Source = Visit(m_initialSource, newMip.get());
+
+		newSource = topMipOp;
+	}
+	else
+	{
+		// Proceed
+		newSource = Visit(m_initialSource, m_root);
+	}
+
+	m_root = nullptr;
+
+	// If there is any change, it is the new root.
+	if (newSource != m_initialSource)
+	{
+		return newSource;
+	}
+
+	return nullptr;
+}
+
+
+//---------------------------------------------------------------------------------------------
+//---------------------------------------------------------------------------------------------
+//---------------------------------------------------------------------------------------------
+Ptr<ASTOp> Sink_ImageMipmapAST::Visit(Ptr<ASTOp> at, const ASTOpImageMipmap* currentMipmapOp)
+{
+	if (!at) return nullptr;
+
+	// Newly created?
+	if (std::find(m_newOps.begin(), m_newOps.end(), at) != m_newOps.end())
+	{
+		return at;
+	}
+
+	// Already visited?
+	auto cacheIt = m_oldToNew.find(at);
+	if (cacheIt != m_oldToNew.end())
+	{
+		return cacheIt->second;
+	}
+
+	Ptr<ASTOp> newAt = at;
+	switch (at->GetOpType())
+	{
+
+	case OP_TYPE::IM_CONDITIONAL:
+	{
+		// We move the mask creation down the two paths
+		auto newOp = mu::Clone<ASTOpConditional>(at);
+		newOp->yes = Visit(newOp->yes.child(), currentMipmapOp);
+		newOp->no = Visit(newOp->no.child(), currentMipmapOp);
+		newAt = newOp;
+		break;
+	}
+
+	case OP_TYPE::IM_SWITCH:
+	{
+		// We move the mask creation down all the paths
+		auto newOp = mu::Clone<ASTOpSwitch>(at);
+		newOp->def = Visit(newOp->def.child(), currentMipmapOp);
+		for (auto& c : newOp->cases)
+		{
+			c.branch = Visit(c.branch.child(), currentMipmapOp);
+		}
+		newAt = newOp;
+		break;
+	}
+
+	case OP_TYPE::IM_COMPOSE:
+	{
+		const ASTOpImageCompose* typedAt = dynamic_cast<const ASTOpImageCompose*>(at.get());
+		if (!currentMipmapOp->bOnlyTail
+			&&
+			// Don't move the mipmapping if we are composing with a mask.
+			// TODO: allow mipmapping in the masks, RLE formats, etc.
+			!typedAt->Mask
+			)
+		{
+			auto newOp = mu::Clone<ASTOpImageCompose>(at);
+
+			auto baseOp = newOp->Base.child();
+			newOp->Base = Visit(baseOp, currentMipmapOp);
+
+			auto blockOp = newOp->BlockImage.child();
+			newOp->BlockImage = Visit(blockOp, currentMipmapOp);
+
+			newAt = newOp;
+		}
+
+		break;
+	}
+
+	case OP_TYPE::IM_PATCH:
+	{
+		if (!currentMipmapOp->bOnlyTail)
+		{
+			// Special case: we propagate the mipmapping down the patch up to
+			// the level allowed by the patch size and placement. We then leave
+			// a top-level  mipmapping operation to generate the smallest
+			// mipmaps after the patch is done.
+
+			const ASTOpImagePatch* typedSource = dynamic_cast<const ASTOpImagePatch*>(at.get());
+			auto rectOp = typedSource->patch.child();
+			ASTOp::GetImageDescContext context;
+			auto patchDesc = rectOp->GetImageDesc(false, &context);
+
+			uint8_t blockLevels = 0;
+			{
+				uint16_t minX = typedSource->location[0];
+				uint16_t minY = typedSource->location[1];
+				uint16_t sizeX = patchDesc.m_size[0];
+				uint16_t sizeY = patchDesc.m_size[1];
+				while (minX && minY && sizeX && sizeY
+					&&
+					minX % 2 == 0 && minY % 2 == 0 && sizeX % 2 == 0 && sizeY % 2 == 0)
+				{
+					blockLevels++;
+					minX /= 2;
+					minY /= 2;
+					sizeX /= 2;
+					sizeY /= 2;
+				}
+			}
+
+			Ptr<ASTOpImageMipmap> newMip = mu::Clone<ASTOpImageMipmap>(currentMipmapOp);
+			newMip->Levels = blockLevels;
+			newMip->BlockLevels = blockLevels;
+			newMip->bOnlyTail = false;
+
+			auto newOp = mu::Clone<ASTOpImagePatch>(at);
+			newOp->base = Visit(newOp->base.child(), newMip.get());
+			newOp->patch = Visit(newOp->patch.child(), newMip.get());
+			newAt = newOp;
+
+			// We need to add a mipmap on top to finish the mipmapping
+			{
+				Ptr<ASTOpImageMipmap> topMipOp = mu::Clone<ASTOpImageMipmap>(currentMipmapOp);
+				topMipOp->Source = newOp;
+				topMipOp->bOnlyTail = true;
+				newAt = topMipOp;
+			}
+		}
+
+		break;
+	}
+
+	default:
+		break;
+	}
+
+	// end on line, replace with mipmap
+	if (at == newAt && at != m_initialSource)
+	{
+		Ptr<ASTOpImageMipmap> newOp = mu::Clone<ASTOpImageMipmap>(currentMipmapOp);
+		check(newOp->GetOpType() == OP_TYPE::IM_MIPMAP);
+
+		newOp->Source = at;
+
+		newAt = newOp;
+	}
+
+	m_oldToNew[at] = newAt;
+
+	return newAt;
+}
+
