[mutable] Move Mutable plugin from the restricted folder to the experimental plugins folder.

diff --git a/Engine/Plugins/Experimental/Mutable/Source/MutableTools/Private/ASTOpImageMipmap.h b/Engine/Plugins/Experimental/Mutable/Source/MutableTools/Private/ASTOpImageMipmap.h
new file mode 100644
--- /dev/null
+++ b/Engine/Plugins/Experimental/Mutable/Source/MutableTools/Private/ASTOpImageMipmap.h
@@ -0,0 +1,53 @@
+// Copyright Epic Games, Inc. All Rights Reserved.
+
+#pragma once
+
+#include "MutableTools/Private/AST.h"
+
+
+namespace mu
+{
+
+	class ASTOpImageMipmap : public ASTOp
+	{
+	public:
+
+		ASTChild Source;
+
+		uint8_t Levels = 0;
+
+		//! Number of mipmaps that can be generated for a single layout block.
+		uint8_t BlockLevels = 0;
+
+		//! This is true if this operation is supposed to build only the tail mipmaps.
+		//! It is used during the code optimisation phase, and to validate the code.
+		bool bOnlyTail = false;
+
+		//! Mipmap generation settings. 
+		float SharpenFactor = 0.0f;
+		EAddressMode AddressMode = EAddressMode::AM_NONE;
+		EMipmapFilterType FilterType = EMipmapFilterType::MFT_Unfiltered;
+		bool DitherMipmapAlpha = false;
+
+	public:
+
+		ASTOpImageMipmap();
+		ASTOpImageMipmap(const ASTOpImageMipmap&) = delete;
+		~ASTOpImageMipmap();
+
+		OP_TYPE GetOpType() const override { return OP_TYPE::IM_MIPMAP; }
+		uint64 Hash() const override;
+		bool IsEqual(const ASTOp& otherUntyped) const override;
+		Ptr<ASTOp> Clone(MapChildFunc& mapChild) const override;
+		void ForEachChild(const std::function<void(ASTChild&)>&) override;
+		void Link(PROGRAM& program, const FLinkerOptions* Options) override;
+		Ptr<ASTOp> OptimiseSink(const MODEL_OPTIMIZATION_OPTIONS& options, OPTIMIZE_SINK_CONTEXT& context) const override;
+		FImageDesc GetImageDesc(bool returnBestOption, GetImageDescContext* context) override;
+		void GetLayoutBlockSize(int* pBlockX, int* pBlockY) override;
+		bool IsImagePlainConstant(vec4<float>& colour) const override;
+		Ptr<ImageSizeExpression> GetImageSizeExpression() const override;
+
+	};
+
+}
+
diff --git a/Engine/Plugins/Experimental/Mutable/Source/MutableTools/Private/ASTOpImageMultiLayer.cpp b/Engine/Plugins/Experimental/Mutable/Source/MutableTools/Private/ASTOpImageMultiLayer.cpp
new file mode 100644
--- /dev/null
+++ b/Engine/Plugins/Experimental/Mutable/Source/MutableTools/Private/ASTOpImageMultiLayer.cpp
@@ -0,0 +1,169 @@
+// Copyright Epic Games, Inc. All Rights Reserved.
+
+#include "ASTOpImageMultiLayer.h"
+#include "ErrorLogPrivate.h"
+#include "CodeOptimiser.h"
+#include "StreamsPrivate.h"
+
+
+using namespace mu;
+
+
+//-------------------------------------------------------------------------------------------------
+//-------------------------------------------------------------------------------------------------
+//-------------------------------------------------------------------------------------------------
+ASTOpImageMultiLayer::ASTOpImageMultiLayer()
+    : base(this)
+    , blend(this)
+    , mask(this)
+    , range(this,nullptr,string(),string())
+	, blendType(EBlendType::BT_BLEND)
+{
+}
+
+
+//-------------------------------------------------------------------------------------------------
+ASTOpImageMultiLayer::~ASTOpImageMultiLayer()
+{
+    // Explicit call needed to avoid recursive destruction
+    ASTOp::RemoveChildren();
+}
+
+
+//-------------------------------------------------------------------------------------------------
+bool ASTOpImageMultiLayer::IsEqual(const ASTOp& otherUntyped) const
+{
+    if ( auto other = dynamic_cast<const ASTOpImageMultiLayer*>(&otherUntyped) )
+    {
+        return base==other->base &&
+                blend==other->blend &&
+                mask==other->mask &&
+                range==other->range &&
+                blendType==other->blendType;
+    }
+    return false;
+}
+
+
+//-------------------------------------------------------------------------------------------------
+uint64 ASTOpImageMultiLayer::Hash() const
+{
+	uint64 res = std::hash<OP_TYPE>()(OP_TYPE::IM_MULTILAYER);
+    hash_combine( res, base.child().get() );
+    hash_combine( res, blend.child().get() );
+    hash_combine( res, mask.child().get() );
+    return res;
+}
+
+
+//-------------------------------------------------------------------------------------------------
+mu::Ptr<ASTOp> ASTOpImageMultiLayer::Clone(MapChildFunc& mapChild) const
+{
+    Ptr<ASTOpImageMultiLayer> n = new ASTOpImageMultiLayer();
+    n->base = mapChild(base.child());
+    n->blend = mapChild(blend.child());
+    n->mask = mapChild(mask.child());
+    n->range.rangeName = range.rangeName;
+    n->range.rangeUID = range.rangeUID;
+    n->range.rangeSize = mapChild(range.rangeSize.child());
+    n->blendType = blendType;
+    return n;
+}
+
+
+//-------------------------------------------------------------------------------------------------
+void ASTOpImageMultiLayer::ForEachChild( const std::function<void(ASTChild&)>& f )
+{
+    f( base );
+    f( blend );
+    f( mask );
+    f( range.rangeSize );
+}
+
+
+//-------------------------------------------------------------------------------------------------
+void ASTOpImageMultiLayer::Link( PROGRAM& program, const FLinkerOptions*)
+{
+    // Already linked?
+    if (!linkedAddress)
+    {
+        OP::ImageMultiLayerArgs args;
+        memset( &args,0, sizeof(args) );
+
+        args.blendType = (uint16_t)blendType;
+
+        if (base) args.base = base->linkedAddress;
+        if (blend) args.blended = blend->linkedAddress;
+        if (mask) args.mask = mask->linkedAddress;
+        if (range.rangeSize)
+        {
+            LinkRange( program, range, args.rangeSize, args.rangeId );
+        }
+
+        linkedAddress = (OP::ADDRESS)program.m_opAddress.size();
+        program.m_opAddress.push_back((uint32_t)program.m_byteCode.size());
+        AppendCode(program.m_byteCode,OP_TYPE::IM_MULTILAYER);
+        AppendCode(program.m_byteCode,args);
+    }
+}
+
+
+//-------------------------------------------------------------------------------------------------
+FImageDesc ASTOpImageMultiLayer::GetImageDesc( bool returnBestOption, GetImageDescContext* context )
+{
+    FImageDesc res;
+
+    // Local context in case it is necessary
+    GetImageDescContext localContext;
+    if (!context)
+    {
+      context = &localContext;
+    }
+    else
+    {
+        // Cached result?
+		FImageDesc* PtrValue = context->m_results.Find(this);
+		if (PtrValue)
+		{
+			return *PtrValue;
+		}
+	}
+
+    // Actual work
+    if (base)
+    {
+        res = base->GetImageDesc( returnBestOption, context );
+    }
+
+
+    // Cache the result
+    if (context)
+    {
+		context->m_results.Add(this, res);
+	}
+
+    return res;
+}
+
+
+//-------------------------------------------------------------------------------------------------
+void ASTOpImageMultiLayer::GetLayoutBlockSize(int* pBlockX, int* pBlockY)
+{
+	if (base)
+	{
+		base->GetLayoutBlockSize(pBlockX, pBlockY);
+	}
+}
+
+
+//-------------------------------------------------------------------------------------------------
+mu::Ptr<ImageSizeExpression> ASTOpImageMultiLayer::GetImageSizeExpression() const
+{
+    if ( base )
+    {
+        return base->GetImageSizeExpression();
+    }
+
+    return nullptr;
+}
+
diff --git a/Engine/Plugins/Experimental/Mutable/Source/MutableTools/Private/ASTOpImageMultiLayer.h b/Engine/Plugins/Experimental/Mutable/Source/MutableTools/Private/ASTOpImageMultiLayer.h
new file mode 100644
--- /dev/null
+++ b/Engine/Plugins/Experimental/Mutable/Source/MutableTools/Private/ASTOpImageMultiLayer.h
@@ -0,0 +1,43 @@
+// Copyright Epic Games, Inc. All Rights Reserved.
+
+#pragma once
+
+#include "AST.h"
+
+
+namespace mu
+{
+
+
+	//---------------------------------------------------------------------------------------------
+	//!
+	//---------------------------------------------------------------------------------------------
+	class ASTOpImageMultiLayer : public ASTOp
+	{
+	public:
+
+		ASTChild base;
+		ASTChild blend;
+		ASTChild mask;
+		RANGE_DATA range;
+		EBlendType blendType;
+
+	public:
+
+		ASTOpImageMultiLayer();
+		ASTOpImageMultiLayer(const ASTOpImageMultiLayer&) = delete;
+		~ASTOpImageMultiLayer();
+
+		OP_TYPE GetOpType() const override { return OP_TYPE::IM_MULTILAYER; }
+		uint64 Hash() const override;
+		bool IsEqual(const ASTOp& otherUntyped) const override;
+		Ptr<ASTOp> Clone(MapChildFunc& mapChild) const override;
+		void ForEachChild(const std::function<void(ASTChild&)>&) override;
+		void Link(PROGRAM& program, const FLinkerOptions* Options) override;
+		FImageDesc GetImageDesc(bool returnBestOption, GetImageDescContext* context) override;
+		void GetLayoutBlockSize(int* pBlockX, int* pBlockY) override;
+		Ptr<ImageSizeExpression> GetImageSizeExpression() const override;
+	};
+
+}
+
diff --git a/Engine/Plugins/Experimental/Mutable/Source/MutableTools/Private/ASTOpImageNormalComposite.cpp b/Engine/Plugins/Experimental/Mutable/Source/MutableTools/Private/ASTOpImageNormalComposite.cpp
new file mode 100644
--- /dev/null
+++ b/Engine/Plugins/Experimental/Mutable/Source/MutableTools/Private/ASTOpImageNormalComposite.cpp
@@ -0,0 +1,145 @@
+// Copyright Epic Games, Inc. All Rights Reserved.
+
+#include "ASTOpImageNormalComposite.h"
+#include "ErrorLogPrivate.h"
+#include "CodeOptimiser.h"
+#include "StreamsPrivate.h"
+
+using namespace mu;
+
+
+ASTOpImageNormalComposite::ASTOpImageNormalComposite()
+    : Base(this)
+    , Normal(this)
+    , Mode(ECompositeImageMode::CIM_Disabled)
+    , Power(1.0f)
+{
+}
+
+
+ASTOpImageNormalComposite::~ASTOpImageNormalComposite()
+{
+    // Explicit call needed to avoid recursive destruction
+    ASTOp::RemoveChildren();
+}
+
+
+bool ASTOpImageNormalComposite::IsEqual(const ASTOp& OtherUntyped) const
+{
+    if ( const ASTOpImageNormalComposite* Other = dynamic_cast<const ASTOpImageNormalComposite*>( &OtherUntyped) )
+    {
+        return Base == Other->Base && Normal == Other->Normal && Power == Other->Power && Mode == Other->Mode;
+    }
+
+    return false;
+}
+
+
+uint64 ASTOpImageNormalComposite::Hash() const
+{
+	uint64 Res = std::hash<void*>()(Base.child().get() );
+	hash_combine(Res, Normal.child().get());
+
+	return Res;
+}
+
+
+mu::Ptr<ASTOp> ASTOpImageNormalComposite::Clone(MapChildFunc& mapChild) const
+{
+    Ptr<ASTOpImageNormalComposite> N = new ASTOpImageNormalComposite();
+	N->Base = mapChild(Base.child());
+	N->Normal = mapChild(Normal.child());
+    N->Mode = Mode;
+    N->Power = Power;
+
+    return N;
+}
+
+
+void ASTOpImageNormalComposite::ForEachChild(const std::function<void(ASTChild&)>& f )
+{
+	f( Base );
+	f( Normal );
+}
+
+
+void ASTOpImageNormalComposite::Link( PROGRAM& program, const FLinkerOptions* )
+{
+    // Already linked?
+    if (!linkedAddress)
+    {
+        OP::ImageNormalCompositeArgs args;
+        FMemory::Memset( &args, 0, sizeof(args) );
+
+		if (Base) 
+        {
+            args.base = Base->linkedAddress;
+        }
+
+        if (Normal) 
+        {
+            args.normal = Normal->linkedAddress;
+        }
+
+		args.power = Power;
+		args.mode = Mode;
+
+        linkedAddress = (OP::ADDRESS)program.m_opAddress.size();
+        program.m_opAddress.push_back((uint32_t)program.m_byteCode.size());
+        AppendCode(program.m_byteCode,OP_TYPE::IM_NORMALCOMPOSITE);
+        AppendCode(program.m_byteCode,args);
+    }
+
+}
+
+
+FImageDesc ASTOpImageNormalComposite::GetImageDesc( bool returnBestOption, GetImageDescContext* context )
+{
+    FImageDesc res;
+
+    // Local context in case it is necessary
+    GetImageDescContext localContext;
+    if (!context)
+    {
+      context = &localContext;
+    }
+    else
+    {
+        // Cached result?
+		FImageDesc* PtrValue = context->m_results.Find(this);
+		if (PtrValue)
+		{
+			return *PtrValue;
+		}
+	}
+
+    if (Base)
+    {
+        res = Base->GetImageDesc( returnBestOption, context );
+    }
+
+    if (context)
+    {
+        context->m_results.Add(this, res);
+    }
+
+    return res;
+}
+
+void ASTOpImageNormalComposite::GetLayoutBlockSize(int* pBlockX, int* pBlockY)
+{
+    if (Base)
+    {
+        Base->GetLayoutBlockSize(pBlockX, pBlockY);
+    }
+}
+
+mu::Ptr<ImageSizeExpression> ASTOpImageNormalComposite::GetImageSizeExpression() const 
+{
+    if (Base)
+    {
+        return Base->GetImageSizeExpression();
+    }
+
+    return nullptr;
+}
diff --git a/Engine/Plugins/Experimental/Mutable/Source/MutableTools/Private/ASTOpImageNormalComposite.h b/Engine/Plugins/Experimental/Mutable/Source/MutableTools/Private/ASTOpImageNormalComposite.h
new file mode 100644
--- /dev/null
+++ b/Engine/Plugins/Experimental/Mutable/Source/MutableTools/Private/ASTOpImageNormalComposite.h
@@ -0,0 +1,37 @@
+// Copyright Epic Games, Inc. All Rights Reserved.
+
+#pragma once
+
+#include "AST.h"
+
+
+namespace mu
+{
+
+	class ASTOpImageNormalComposite : public ASTOp
+	{
+	public:
+		ASTChild Base;
+		ASTChild Normal;
+		ECompositeImageMode Mode;
+		float Power;
+
+	public:
+
+		ASTOpImageNormalComposite();
+		ASTOpImageNormalComposite(const ASTOpImageNormalComposite&) = delete;
+		~ASTOpImageNormalComposite();
+
+		OP_TYPE GetOpType() const override { return OP_TYPE::IM_NORMALCOMPOSITE; }
+		uint64 Hash() const override;
+		bool IsEqual(const ASTOp& otherUntyped) const override;
+		Ptr<ASTOp> Clone(MapChildFunc& mapChild) const override;
+		void ForEachChild(const std::function<void(ASTChild&)>&) override;
+		void Link(PROGRAM& program, const FLinkerOptions* Options) override;
+		FImageDesc GetImageDesc(bool returnBestOption, GetImageDescContext* context) override;
+		void GetLayoutBlockSize(int* pBlockX, int* pBlockY) override;
+		Ptr<ImageSizeExpression> GetImageSizeExpression() const override;
+	};
+
+}
+
diff --git a/Engine/Plugins/Experimental/Mutable/Source/MutableTools/Private/ASTOpImagePatch.cpp b/Engine/Plugins/Experimental/Mutable/Source/MutableTools/Private/ASTOpImagePatch.cpp
new file mode 100644
--- /dev/null
+++ b/Engine/Plugins/Experimental/Mutable/Source/MutableTools/Private/ASTOpImagePatch.cpp
@@ -0,0 +1,143 @@
+// Copyright Epic Games, Inc. All Rights Reserved.
+
+#include "ASTOpImagePatch.h"
+#include "ErrorLogPrivate.h"
+#include "CodeOptimiser.h"
+#include "StreamsPrivate.h"
+
+
+using namespace mu;
+
+
+//-------------------------------------------------------------------------------------------------
+//-------------------------------------------------------------------------------------------------
+//-------------------------------------------------------------------------------------------------
+ASTOpImagePatch::ASTOpImagePatch()
+    : base(this)
+    , patch(this)
+{
+}
+
+
+//-------------------------------------------------------------------------------------------------
+ASTOpImagePatch::~ASTOpImagePatch()
+{
+    // Explicit call needed to avoid recursive destruction
+    ASTOp::RemoveChildren();
+}
+
+
+//-------------------------------------------------------------------------------------------------
+bool ASTOpImagePatch::IsEqual(const ASTOp& otherUntyped) const
+{
+    if ( auto other = dynamic_cast<const ASTOpImagePatch*>(&otherUntyped) )
+    {
+        return base==other->base &&
+                patch==other->patch &&
+                location==other->location;
+    }
+    return false;
+}
+
+
+//-------------------------------------------------------------------------------------------------
+uint64 ASTOpImagePatch::Hash() const
+{
+	uint64 res = std::hash<OP_TYPE>()(OP_TYPE::IM_PATCH);
+    hash_combine( res, base.child().get() );
+    hash_combine( res, patch.child().get() );
+    return res;
+}
+
+
+//-------------------------------------------------------------------------------------------------
+mu::Ptr<ASTOp> ASTOpImagePatch::Clone(MapChildFunc& mapChild) const
+{
+    Ptr<ASTOpImagePatch> n = new ASTOpImagePatch();
+    n->base = mapChild(base.child());
+    n->patch = mapChild(patch.child());
+    n->location = location;
+    return n;
+}
+
+
+//-------------------------------------------------------------------------------------------------
+void ASTOpImagePatch::ForEachChild(const std::function<void(ASTChild&)>& f )
+{
+    f( base );
+    f( patch );
+}
+
+
+//-------------------------------------------------------------------------------------------------
+void ASTOpImagePatch::Link( PROGRAM& program, const FLinkerOptions*)
+{
+    // Already linked?
+    if (!linkedAddress)
+    {
+        OP::ImagePatchArgs args;
+        memset( &args,0, sizeof(args) );
+
+        if (base) args.base = base->linkedAddress;
+        if (patch) args.patch = patch->linkedAddress;
+        args.minX = location[0];
+        args.minY = location[1];
+
+        linkedAddress = (OP::ADDRESS)program.m_opAddress.size();
+        program.m_opAddress.push_back((uint32_t)program.m_byteCode.size());
+        AppendCode(program.m_byteCode,OP_TYPE::IM_PATCH);
+        AppendCode(program.m_byteCode,args);
+    }
+
+}
+
+
+//-------------------------------------------------------------------------------------------------
+FImageDesc ASTOpImagePatch::GetImageDesc( bool returnBestOption, GetImageDescContext* context )
+{
+    FImageDesc res;
+
+    // Local context in case it is necessary
+    GetImageDescContext localContext;
+    if (!context)
+    {
+      context = &localContext;
+    }
+    else
+    {
+        // Cached result?
+		FImageDesc* PtrValue = context->m_results.Find(this);
+		if (PtrValue)
+		{
+			return *PtrValue;
+		}
+	}
+
+    // Actual work
+    if (base)
+    {
+        res = base->GetImageDesc( returnBestOption, context );
+    }
+
+
+    // Cache the result
+    if (context)
+    {
+		context->m_results.Add(this, res);
+	}
+
+    return res;
+}
+
+
+//-------------------------------------------------------------------------------------------------
+mu::Ptr<ImageSizeExpression> ASTOpImagePatch::GetImageSizeExpression() const
+{
+    if ( base )
+    {
+        return base->GetImageSizeExpression();
+    }
+
+    return nullptr;
+}
+
diff --git a/Engine/Plugins/Experimental/Mutable/Source/MutableTools/Private/ASTOpImagePatch.h b/Engine/Plugins/Experimental/Mutable/Source/MutableTools/Private/ASTOpImagePatch.h
new file mode 100644
--- /dev/null
+++ b/Engine/Plugins/Experimental/Mutable/Source/MutableTools/Private/ASTOpImagePatch.h
@@ -0,0 +1,39 @@
+// Copyright Epic Games, Inc. All Rights Reserved.
+
+#pragma once
+
+#include "MutableTools/Private/AST.h"
+
+
+namespace mu
+{
+
+	class ASTOpImagePatch : public ASTOp
+	{
+	public:
+
+		ASTChild base;
+		ASTChild patch;
+		vec2<uint16_t> location;
+
+	public:
+
+		ASTOpImagePatch();
+		ASTOpImagePatch(const ASTOpImagePatch&) = delete;
+		~ASTOpImagePatch();
+
+		OP_TYPE GetOpType() const override { return OP_TYPE::IM_PATCH; }
+		uint64 Hash() const override;
+		bool IsEqual(const ASTOp& otherUntyped) const override;
+		Ptr<ASTOp> Clone(MapChildFunc& mapChild) const override;
+		void ForEachChild(const std::function<void(ASTChild&)>&) override;
+		void Link(PROGRAM& program, const FLinkerOptions* Options) override;
+		FImageDesc GetImageDesc(bool returnBestOption, GetImageDescContext* context) override;
+		Ptr<ImageSizeExpression> GetImageSizeExpression() const override;
+		//TODO: void GetLayoutBlockSize(int* pBlockX, int* pBlockY) override;
+		//TODO: bool IsImagePlainConstant(vec4<float>& colour) const override;
+	};
+
+
+}
+
diff --git a/Engine/Plugins/Experimental/Mutable/Source/MutableTools/Private/ASTOpImagePixelFormat.h b/Engine/Plugins/Experimental/Mutable/Source/MutableTools/Private/ASTOpImagePixelFormat.h
new file mode 100644
--- /dev/null
+++ b/Engine/Plugins/Experimental/Mutable/Source/MutableTools/Private/ASTOpImagePixelFormat.h
@@ -0,0 +1,40 @@
+// Copyright Epic Games, Inc. All Rights Reserved.
+
+#pragma once
+
+#include "MutableTools/Private/AST.h"
+
+
+namespace mu
+{
+
+	class ASTOpImagePixelFormat : public ASTOp
+	{
+	public:
+
+		ASTChild Source;
+		EImageFormat Format = EImageFormat::IF_NONE;
+		EImageFormat FormatIfAlpha = EImageFormat::IF_NONE;
+
+	public:
+
+		ASTOpImagePixelFormat();
+		ASTOpImagePixelFormat(const ASTOpImagePixelFormat&) = delete;
+		~ASTOpImagePixelFormat();
+
+		OP_TYPE GetOpType() const override { return OP_TYPE::IM_PIXELFORMAT; }
+		uint64 Hash() const override;
+		bool IsEqual(const ASTOp& otherUntyped) const override;
+		Ptr<ASTOp> Clone(MapChildFunc& mapChild) const override;
+		void ForEachChild(const std::function<void(ASTChild&)>&) override;
+		void Link(PROGRAM& program, const FLinkerOptions* Options) override;
+		Ptr<ASTOp> OptimiseSink(const MODEL_OPTIMIZATION_OPTIONS& options, OPTIMIZE_SINK_CONTEXT& context) const override;
+		FImageDesc GetImageDesc(bool returnBestOption, GetImageDescContext* context) override;
+		void GetLayoutBlockSize(int* pBlockX, int* pBlockY) override;
+		bool IsImagePlainConstant(vec4<float>& colour) const override;
+		Ptr<ImageSizeExpression> GetImageSizeExpression() const override;
+
+	};
+
+}
+
diff --git a/Engine/Plugins/Experimental/Mutable/Source/MutableTools/Private/ASTOpImageTransform.cpp b/Engine/Plugins/Experimental/Mutable/Source/MutableTools/Private/ASTOpImageTransform.cpp
new file mode 100644
--- /dev/null
+++ b/Engine/Plugins/Experimental/Mutable/Source/MutableTools/Private/ASTOpImageTransform.cpp
@@ -0,0 +1,170 @@
+// Copyright Epic Games, Inc. All Rights Reserved.
+
+#include "ASTOpImageTransform.h"
+#include "ErrorLogPrivate.h"
+#include "CodeOptimiser.h"
+#include "StreamsPrivate.h"
+
+
+using namespace mu;
+
+
+//-------------------------------------------------------------------------------------------------
+//-------------------------------------------------------------------------------------------------
+//-------------------------------------------------------------------------------------------------
+ASTOpImageTransform::ASTOpImageTransform()
+    : base(this)
+    , offsetX(this)
+    , offsetY(this)
+    , scaleX(this)
+    , scaleY(this)
+    , rotation(this)
+{
+}
+
+
+//-------------------------------------------------------------------------------------------------
+ASTOpImageTransform::~ASTOpImageTransform()
+{
+    // Explicit call needed to avoid recursive destruction
+    ASTOp::RemoveChildren();
+}
+
+
+//-------------------------------------------------------------------------------------------------
+bool ASTOpImageTransform::IsEqual(const ASTOp& otherUntyped) const
+{
+    if ( auto other = dynamic_cast<const ASTOpImageTransform*>(&otherUntyped) )
+    {
+        return base == other->base &&
+			offsetX == other->offsetX &&
+            offsetY == other->offsetY &&
+            scaleX == other->scaleX &&
+            scaleY == other->scaleY &&
+            rotation == other->rotation;
+    }
+    return false;
+}
+
+
+//-------------------------------------------------------------------------------------------------
+uint64 ASTOpImageTransform::Hash() const
+{
+	uint64 res = std::hash<OP_TYPE>()(OP_TYPE::IM_MULTILAYER);
+    hash_combine( res, base.child().get() );
+    hash_combine( res, offsetX.child().get() );
+    hash_combine( res, offsetY.child().get() );
+    hash_combine( res, scaleX.child().get() );
+    hash_combine( res, scaleY.child().get() );
+    hash_combine( res, rotation.child().get() );
+    return res;
+}
+
+
+//-------------------------------------------------------------------------------------------------
+mu::Ptr<ASTOp> ASTOpImageTransform::Clone(MapChildFunc& mapChild) const
+{
+    Ptr<ASTOpImageTransform> n = new ASTOpImageTransform();
+    n->base = mapChild(base.child());
+	n->offsetX = mapChild(offsetX.child());
+	n->offsetY = mapChild(offsetY.child());
+	n->scaleX = mapChild(scaleX.child());
+	n->scaleY = mapChild(scaleY.child());
+	n->rotation = mapChild(rotation.child());
+    return n;
+}
+
+
+//-------------------------------------------------------------------------------------------------
+void ASTOpImageTransform::ForEachChild( const std::function<void(ASTChild&)>& f )
+{
+    f( base );
+    f( offsetX );
+    f( offsetY );
+    f( scaleX );
+    f( scaleY );
+    f( rotation );
+}
+
+
+//-------------------------------------------------------------------------------------------------
+void ASTOpImageTransform::Link( PROGRAM& program, const FLinkerOptions*)
+{
+    // Already linked?
+    if (!linkedAddress)
+    {
+		OP::ImageTransformArgs Args;
+	
+		Args.base = base ? base->linkedAddress : 0;
+		Args.offsetX = offsetX ? offsetX->linkedAddress : 0;
+		Args.offsetY = offsetY ? offsetY->linkedAddress : 0;
+		Args.scaleX = scaleX ? scaleX->linkedAddress : 0;
+		Args.scaleY = scaleY ? scaleY->linkedAddress : 0;
+		Args.rotation = rotation ? rotation->linkedAddress : 0;
+
+        linkedAddress = (OP::ADDRESS)program.m_opAddress.size();
+        program.m_opAddress.push_back((uint32_t)program.m_byteCode.size());
+        AppendCode(program.m_byteCode, OP_TYPE::IM_TRANSFORM);
+		AppendCode(program.m_byteCode, Args);
+    }
+}
+
+//-------------------------------------------------------------------------------------------------
+FImageDesc ASTOpImageTransform::GetImageDesc( bool returnBestOption, GetImageDescContext* context )
+{
+    FImageDesc res;
+
+    // Local context in case it is necessary
+    GetImageDescContext localContext;
+    if (!context)
+    {
+      context = &localContext;
+    }
+    else
+    {
+        // Cached result?
+		FImageDesc* PtrValue = context->m_results.Find(this);
+		if (PtrValue)
+		{
+			return *PtrValue;
+		}
+	}
+
+    // Actual work
+    if (base)
+    {
+        res = base->GetImageDesc( returnBestOption, context );
+    }
+
+
+    // Cache the result
+    if (context)
+    {
+		context->m_results.Add(this, res);
+	}
+
+    return res;
+}
+
+
+//-------------------------------------------------------------------------------------------------
+void ASTOpImageTransform::GetLayoutBlockSize(int* pBlockX, int* pBlockY)
+{
+	if (base)
+	{
+		base->GetLayoutBlockSize(pBlockX, pBlockY);
+	}
+}
+
+
+//-------------------------------------------------------------------------------------------------
+mu::Ptr<ImageSizeExpression> ASTOpImageTransform::GetImageSizeExpression() const
+{
+    if ( base )
+    {
+        return base->GetImageSizeExpression();
+    }
+
+    return nullptr;
+}
+
diff --git a/Engine/Plugins/Experimental/Mutable/Source/MutableTools/Private/ASTOpImageTransform.h b/Engine/Plugins/Experimental/Mutable/Source/MutableTools/Private/ASTOpImageTransform.h
new file mode 100644
--- /dev/null
+++ b/Engine/Plugins/Experimental/Mutable/Source/MutableTools/Private/ASTOpImageTransform.h
@@ -0,0 +1,44 @@
+// Copyright Epic Games, Inc. All Rights Reserved.
+
+#pragma once
+
+#include "AST.h"
+
+
+namespace mu
+{
+
+
+	//---------------------------------------------------------------------------------------------
+	//!
+	//---------------------------------------------------------------------------------------------
+	class ASTOpImageTransform : public ASTOp
+	{
+	public:
+
+		ASTChild base;
+		ASTChild offsetX;
+		ASTChild offsetY;
+		ASTChild scaleX;
+		ASTChild scaleY;
+		ASTChild rotation;
+
+	public:
+
+		ASTOpImageTransform();
+		ASTOpImageTransform(const ASTOpImageTransform&) = delete;
+		~ASTOpImageTransform();
+
+		OP_TYPE GetOpType() const override { return OP_TYPE::IM_TRANSFORM; }
+		uint64 Hash() const override;
+		bool IsEqual(const ASTOp& otherUntyped) const override;
+		Ptr<ASTOp> Clone(MapChildFunc& mapChild) const override;
+		void ForEachChild(const std::function<void(ASTChild&)>&) override;
+		void Link(PROGRAM& program, const FLinkerOptions* Options) override;
+		FImageDesc GetImageDesc(bool returnBestOption, GetImageDescContext* context) override;
+		void GetLayoutBlockSize(int* pBlockX, int* pBlockY) override;
+		Ptr<ImageSizeExpression> GetImageSizeExpression() const override;
+	};
+
+}
+
diff --git a/Engine/Plugins/Experimental/Mutable/Source/MutableTools/Private/ASTOpInstanceAdd.cpp b/Engine/Plugins/Experimental/Mutable/Source/MutableTools/Private/ASTOpInstanceAdd.cpp
new file mode 100644
--- /dev/null
+++ b/Engine/Plugins/Experimental/Mutable/Source/MutableTools/Private/ASTOpInstanceAdd.cpp
@@ -0,0 +1,165 @@
+// Copyright Epic Games, Inc. All Rights Reserved.
+
+#include "ASTOpInstanceAdd.h"
+#include "ErrorLogPrivate.h"
+#include "CodeOptimiser.h"
+#include "StreamsPrivate.h"
+
+#include <queue>
+
+
+using namespace mu;
+
+
+//-------------------------------------------------------------------------------------------------
+//-------------------------------------------------------------------------------------------------
+//-------------------------------------------------------------------------------------------------
+ASTOpInstanceAdd::ASTOpInstanceAdd()
+    : instance(this)
+    , value(this)
+{
+}
+
+
+//-------------------------------------------------------------------------------------------------
+ASTOpInstanceAdd::~ASTOpInstanceAdd()
+{
+    // Explicit call needed to avoid recursive destruction
+    ASTOp::RemoveChildren();
+}
+
+
+//-------------------------------------------------------------------------------------------------
+bool ASTOpInstanceAdd::IsEqual(const ASTOp& otherUntyped) const
+{
+    if ( auto other = dynamic_cast<const ASTOpInstanceAdd*>(&otherUntyped) )
+    {
+        return type==other->type &&
+                instance==other->instance &&
+                value==other->value &&
+                id==other->id &&
+                externalId==other->externalId &&
+                name==other->name;
+    }
+    return false;
+}
+
+
+//-------------------------------------------------------------------------------------------------
+mu::Ptr<ASTOp> ASTOpInstanceAdd::Clone(MapChildFunc& mapChild) const
+{
+    Ptr<ASTOpInstanceAdd> n = new ASTOpInstanceAdd();
+    n->type = type;
+    n->instance = mapChild(instance.child());
+    n->value = mapChild(value.child());
+    n->id = id;
+    n->externalId = externalId;
+    n->name = name;
+    return n;
+}
+
+
+//-------------------------------------------------------------------------------------------------
+uint64 ASTOpInstanceAdd::Hash() const
+{
+	uint64 res = std::hash<size_t>()(size_t(type));
+    hash_combine( res, instance.child().get() );
+    hash_combine( res, value.child().get() );
+    return res;
+}
+
+
+//-------------------------------------------------------------------------------------------------
+void ASTOpInstanceAdd::Assert()
+{
+    switch(type)
+    {
+    case OP_TYPE::IN_ADDMESH:
+    case OP_TYPE::IN_ADDIMAGE:
+    case OP_TYPE::IN_ADDVECTOR:
+    case OP_TYPE::IN_ADDSCALAR:
+    case OP_TYPE::IN_ADDSTRING:
+    case OP_TYPE::IN_ADDSURFACE:
+    case OP_TYPE::IN_ADDCOMPONENT:
+    case OP_TYPE::IN_ADDLOD:
+        break;
+    default:
+        // Unexpected type
+        check(false);
+        break;
+    }
+
+    ASTOp::Assert();
+}
+
+
+//-------------------------------------------------------------------------------------------------
+void ASTOpInstanceAdd::ForEachChild(const std::function<void(ASTChild&)>& f )
+{
+    f( instance );
+    f( value );
+}
+
+
+//-------------------------------------------------------------------------------------------------
+void ASTOpInstanceAdd::Link( PROGRAM& program, const FLinkerOptions* )
+{
+    // Already linked?
+    if (!linkedAddress)
+    {
+        OP::InstanceAddArgs args;
+        memset( &args,0, sizeof(args) );
+        args.id = id;
+        args.externalId = externalId;
+
+        args.name = program.AddConstant(name);
+
+        if (instance) args.instance = instance->linkedAddress;
+        if (value) args.value = value->linkedAddress;
+
+        if ( type==OP_TYPE::IN_ADDIMAGE ||
+             type==OP_TYPE::IN_ADDMESH  )
+        {
+            // Find out relevant parameters. \todo: this may be optimised by reusing partial
+            // values in a LINK_CONTEXT or similar
+            SubtreeRelevantParametersVisitorAST visitor;
+            visitor.Run( value.child() );
+
+            vector<uint16_t> params;
+            for(const string& paramName: visitor.m_params)
+            {
+                for( size_t i=0; i<program.m_parameters.size(); ++i )
+                {
+                    const auto& param = program.m_parameters[i];
+                    if (param.m_name==paramName)
+                    {
+                        params.push_back( uint16_t(i) );
+                        break;
+                    }
+                }
+            }
+
+            std::sort(params.begin(),params.end());
+
+            auto it = std::find( program.m_parameterLists.begin(),
+                                 program.m_parameterLists.end(),
+                                 params );
+
+            if (it!=program.m_parameterLists.end())
+            {
+                args.relevantParametersListIndex = uint32_t(it - program.m_parameterLists.begin());
+            }
+            else
+            {
+                args.relevantParametersListIndex = uint32_t(program.m_parameterLists.size());
+                program.m_parameterLists.push_back( params );
+            }
+        }
+
+        linkedAddress = (OP::ADDRESS)program.m_opAddress.size();
+        program.m_opAddress.push_back((uint32_t)program.m_byteCode.size());
+        AppendCode(program.m_byteCode,type);
+        AppendCode(program.m_byteCode,args);
+    }
+
+}
diff --git a/Engine/Plugins/Experimental/Mutable/Source/MutableTools/Private/ASTOpInstanceAdd.h b/Engine/Plugins/Experimental/Mutable/Source/MutableTools/Private/ASTOpInstanceAdd.h
new file mode 100644
--- /dev/null
+++ b/Engine/Plugins/Experimental/Mutable/Source/MutableTools/Private/ASTOpInstanceAdd.h
@@ -0,0 +1,46 @@
+// Copyright Epic Games, Inc. All Rights Reserved.
+
+#pragma once
+
+#include "AST.h"
+
+
+namespace mu
+{
+
+
+	//---------------------------------------------------------------------------------------------
+	//! Operations to add elements to an instance
+	//---------------------------------------------------------------------------------------------
+	class ASTOpInstanceAdd : public ASTOp
+	{
+	public:
+
+		//! Type of switch
+		OP_TYPE type;
+
+		ASTChild instance;
+		ASTChild value;
+		uint32_t id = 0;
+		uint32_t externalId = 0;
+		string name;
+
+	public:
+
+		ASTOpInstanceAdd();
+		ASTOpInstanceAdd(const ASTOpInstanceAdd&) = delete;
+		~ASTOpInstanceAdd();
+
+		OP_TYPE GetOpType() const override { return type; }
+
+		void ForEachChild(const std::function<void(ASTChild&)>&) override;
+		bool IsEqual(const ASTOp& otherUntyped) const override;
+		Ptr<ASTOp> Clone(MapChildFunc& mapChild) const override;
+		uint64 Hash() const override;
+		void Assert() override;
+		void Link(PROGRAM& program, const FLinkerOptions* Options) override;
+	};
+
+
+}
+
diff --git a/Engine/Plugins/Experimental/Mutable/Source/MutableTools/Private/ASTOpMeshApplyPose.cpp b/Engine/Plugins/Experimental/Mutable/Source/MutableTools/Private/ASTOpMeshApplyPose.cpp
new file mode 100644
--- /dev/null
+++ b/Engine/Plugins/Experimental/Mutable/Source/MutableTools/Private/ASTOpMeshApplyPose.cpp
@@ -0,0 +1,84 @@
+// Copyright Epic Games, Inc. All Rights Reserved.
+
+#include "ASTOpMeshApplyPose.h"
+#include "ErrorLogPrivate.h"
+#include "CodeOptimiser.h"
+#include "StreamsPrivate.h"
+
+
+
+using namespace mu;
+
+
+//-------------------------------------------------------------------------------------------------
+//-------------------------------------------------------------------------------------------------
+//-------------------------------------------------------------------------------------------------
+ASTOpMeshApplyPose::ASTOpMeshApplyPose()
+    : base(this)
+    , pose(this)
+{
+}
+
+
+ASTOpMeshApplyPose::~ASTOpMeshApplyPose()
+{
+    // Explicit call needed to avoid recursive destruction
+    ASTOp::RemoveChildren();
+}
+
+
+bool ASTOpMeshApplyPose::IsEqual(const ASTOp& otherUntyped) const
+{
+    if ( auto other = dynamic_cast<const ASTOpMeshApplyPose*>(&otherUntyped) )
+    {
+        return base==other->base &&
+                pose==other->pose;
+    }
+    return false;
+}
+
+
+uint64 ASTOpMeshApplyPose::Hash() const
+{
+	uint64 res = std::hash<void*>()(base.child().get() );
+    hash_combine( res, pose.child().get() );
+    return res;
+}
+
+
+mu::Ptr<ASTOp> ASTOpMeshApplyPose::Clone(MapChildFunc& mapChild) const
+{
+    Ptr<ASTOpMeshApplyPose> n = new ASTOpMeshApplyPose();
+    n->base = mapChild(base.child());
+    n->pose = mapChild(pose.child());
+    return n;
+}
+
+
+void ASTOpMeshApplyPose::ForEachChild(const std::function<void(ASTChild&)>& f )
+{
+    f( base );
+    f( pose );
+}
+
+
+void ASTOpMeshApplyPose::Link( PROGRAM& program, const FLinkerOptions* )
+{
+    // Already linked?
+    if (!linkedAddress)
+    {
+        OP::MeshApplyPoseArgs args;
+        memset( &args,0, sizeof(args) );
+
+        if (base) args.base = base->linkedAddress;
+        if (pose) args.pose = pose->linkedAddress;
+
+        linkedAddress = (OP::ADDRESS)program.m_opAddress.size();
+        //program.m_code.push_back(op);
+        program.m_opAddress.push_back((uint32_t)program.m_byteCode.size());
+        AppendCode(program.m_byteCode,OP_TYPE::ME_APPLYPOSE);
+        AppendCode(program.m_byteCode,args);
+    }
+
+}
+
diff --git a/Engine/Plugins/Experimental/Mutable/Source/MutableTools/Private/ASTOpMeshApplyPose.h b/Engine/Plugins/Experimental/Mutable/Source/MutableTools/Private/ASTOpMeshApplyPose.h
new file mode 100644
--- /dev/null
+++ b/Engine/Plugins/Experimental/Mutable/Source/MutableTools/Private/ASTOpMeshApplyPose.h
@@ -0,0 +1,38 @@
+// Copyright Epic Games, Inc. All Rights Reserved.
+
+#pragma once
+
+#include "AST.h"
+
+
+namespace mu
+{
+
+
+	//---------------------------------------------------------------------------------------------
+	//!
+	//---------------------------------------------------------------------------------------------
+	class ASTOpMeshApplyPose : public ASTOp
+	{
+	public:
+
+		ASTChild base;
+		ASTChild pose;
+
+	public:
+
+		ASTOpMeshApplyPose();
+		ASTOpMeshApplyPose(const ASTOpMeshApplyPose&) = delete;
+		~ASTOpMeshApplyPose();
+
+		OP_TYPE GetOpType() const override { return OP_TYPE::ME_APPLYPOSE; }
+		uint64 Hash() const override;
+		bool IsEqual(const ASTOp& otherUntyped) const override;
+		Ptr<ASTOp> Clone(MapChildFunc& mapChild) const override;
+		void ForEachChild(const std::function<void(ASTChild&)>&) override;
+		void Link(PROGRAM& program, const FLinkerOptions* Options) override;
+	};
+
+
+}
+
diff --git a/Engine/Plugins/Experimental/Mutable/Source/MutableTools/Private/ASTOpMeshApplyShape.cpp b/Engine/Plugins/Experimental/Mutable/Source/MutableTools/Private/ASTOpMeshApplyShape.cpp
new file mode 100644
--- /dev/null
+++ b/Engine/Plugins/Experimental/Mutable/Source/MutableTools/Private/ASTOpMeshApplyShape.cpp
@@ -0,0 +1,104 @@
+// Copyright Epic Games, Inc. All Rights Reserved.
+
+#include "ASTOpMeshApplyShape.h"
+#include "ErrorLogPrivate.h"
+#include "CodeOptimiser.h"
+#include "StreamsPrivate.h"
+
+using namespace mu;
+
+
+ASTOpMeshApplyShape::ASTOpMeshApplyShape()
+    : Mesh(this)
+    , Shape(this)
+{
+}
+
+
+ASTOpMeshApplyShape::~ASTOpMeshApplyShape()
+{
+    // Explicit call needed to avoid recursive destruction
+    ASTOp::RemoveChildren();
+}
+
+
+bool ASTOpMeshApplyShape::IsEqual(const ASTOp& otherUntyped) const
+{
+    if ( auto other = dynamic_cast<const ASTOpMeshApplyShape*>(&otherUntyped) )
+    {
+    	const bool bSameFlags =
+    		m_reshapePhysicsVolumes == other->m_reshapePhysicsVolumes &&
+    		m_reshapeSkeleton == other->m_reshapeSkeleton &&
+			m_reshapeVertices == other->m_reshapeVertices;
+    	
+        return Mesh == other->Mesh && Shape == other->Shape && bSameFlags;
+    }
+    return false;
+}
+
+
+uint64 ASTOpMeshApplyShape::Hash() const
+{
+	uint64 res = std::hash<void*>()(Mesh.child().get() );
+	hash_combine(res, Shape.child().get());
+	hash_combine(res, m_reshapeSkeleton);
+	hash_combine(res, m_reshapePhysicsVolumes);
+	hash_combine(res, m_reshapeVertices);
+	return res;
+}
+
+
+mu::Ptr<ASTOp> ASTOpMeshApplyShape::Clone(MapChildFunc& mapChild) const
+{
+    Ptr<ASTOpMeshApplyShape> n = new ASTOpMeshApplyShape();
+	n->Mesh = mapChild(Mesh.child());
+	n->Shape = mapChild(Shape.child());
+	n->m_reshapeSkeleton = m_reshapeSkeleton;
+	n->m_reshapePhysicsVolumes = m_reshapePhysicsVolumes;
+	n->m_reshapeVertices = m_reshapeVertices;
+	return n;
+}
+
+
+void ASTOpMeshApplyShape::ForEachChild(const std::function<void(ASTChild&)>& f )
+{
+	f(Mesh);
+	f(Shape);
+}
+
+
+void ASTOpMeshApplyShape::Link( PROGRAM& program, const FLinkerOptions* )
+{
+    // Already linked?
+    if (!linkedAddress)
+    {
+        OP::MeshApplyShapeArgs args;
+        memset( &args,0, sizeof(args) );
+
+		if (m_reshapeSkeleton)
+		{
+			args.flags |= uint32(OP::EMeshBindShapeFlags::ReshapeSkeleton);
+		}
+
+    	if (m_reshapePhysicsVolumes)
+    	{
+			args.flags |= uint32(OP::EMeshBindShapeFlags::ReshapePhysicsVolumes);
+    	}
+
+		if (m_reshapeVertices)
+		{
+			args.flags |= uint32(OP::EMeshBindShapeFlags::ReshapeVertices);
+		}
+
+    	
+		if (Mesh) args.mesh = Mesh->linkedAddress;
+		if (Shape) args.shape = Shape->linkedAddress;
+
+        linkedAddress = (OP::ADDRESS)program.m_opAddress.size();
+        program.m_opAddress.push_back((uint32_t)program.m_byteCode.size());
+        AppendCode(program.m_byteCode,OP_TYPE::ME_APPLYSHAPE);
+        AppendCode(program.m_byteCode,args);
+    }
+
+}
+
diff --git a/Engine/Plugins/Experimental/Mutable/Source/MutableTools/Private/ASTOpMeshApplyShape.h b/Engine/Plugins/Experimental/Mutable/Source/MutableTools/Private/ASTOpMeshApplyShape.h
new file mode 100644
--- /dev/null
+++ b/Engine/Plugins/Experimental/Mutable/Source/MutableTools/Private/ASTOpMeshApplyShape.h
@@ -0,0 +1,39 @@
+// Copyright Epic Games, Inc. All Rights Reserved.
+
+#pragma once
+
+#include "AST.h"
+
+
+namespace mu
+{
+
+	//---------------------------------------------------------------------------------------------
+	//!
+	//---------------------------------------------------------------------------------------------
+	class ASTOpMeshApplyShape : public ASTOp
+	{
+	public:
+
+		ASTChild Mesh;
+		ASTChild Shape;
+		bool m_reshapeSkeleton = false;
+		bool m_reshapePhysicsVolumes = false;
+		bool m_reshapeVertices = true;
+
+	public:
+
+		ASTOpMeshApplyShape();
+		ASTOpMeshApplyShape(const ASTOpMeshApplyShape&) = delete;
+		~ASTOpMeshApplyShape();
+
+		OP_TYPE GetOpType() const override { return OP_TYPE::ME_APPLYSHAPE; }
+		uint64 Hash() const override;
+		bool IsEqual(const ASTOp& otherUntyped) const override;
+		Ptr<ASTOp> Clone(MapChildFunc& mapChild) const override;
+		void ForEachChild(const std::function<void(ASTChild&)>&) override;
+		void Link(PROGRAM& program, const FLinkerOptions* Options) override;
+	};
+
+}
+
diff --git a/Engine/Plugins/Experimental/Mutable/Source/MutableTools/Private/ASTOpMeshBindShape.cpp b/Engine/Plugins/Experimental/Mutable/Source/MutableTools/Private/ASTOpMeshBindShape.cpp
new file mode 100644
--- /dev/null
+++ b/Engine/Plugins/Experimental/Mutable/Source/MutableTools/Private/ASTOpMeshBindShape.cpp
@@ -0,0 +1,334 @@
+// Copyright Epic Games, Inc. All Rights Reserved.
+
+#include "ASTOpMeshBindShape.h"
+#include "ASTOpSwitch.h"
+#include "ErrorLogPrivate.h"
+#include "CodeOptimiser.h"
+#include "StreamsPrivate.h"
+
+using namespace mu;
+
+
+ASTOpMeshBindShape::ASTOpMeshBindShape()
+    : Mesh(this)
+    , Shape(this)
+{
+}
+
+
+ASTOpMeshBindShape::~ASTOpMeshBindShape()
+{
+    // Explicit call needed to avoid recursive destruction
+    ASTOp::RemoveChildren();
+}
+
+
+bool ASTOpMeshBindShape::IsEqual(const ASTOp& otherUntyped) const
+{
+	if ( auto other = dynamic_cast<const ASTOpMeshBindShape*>(&otherUntyped) )
+    {
+		const bool bSameFlags =
+			m_reshapeSkeleton == other->m_reshapeSkeleton &&
+			m_discardInvalidBindings == other->m_discardInvalidBindings &&
+			m_enableRigidParts == other->m_enableRigidParts &&
+			m_deformAllBones == other->m_deformAllBones &&
+			m_deformAllPhysics == other->m_deformAllPhysics &&
+    		m_reshapePhysicsVolumes == other->m_reshapePhysicsVolumes &&
+			m_reshapeVertices == other->m_reshapeVertices;
+    	
+		return bSameFlags && 
+				Mesh == other->Mesh && 
+				Shape == other->Shape && 
+				m_bonesToDeform == other->m_bonesToDeform &&
+				m_physicsToDeform == other->m_physicsToDeform && 
+				m_bindingMethod == other->m_bindingMethod;
+	}
+	
+    return false;
+}
+
+
+uint64 ASTOpMeshBindShape::Hash() const
+{
+	uint64 res = std::hash<void*>()(Mesh.child().get() );
+	hash_combine(res, Shape.child().get());
+	hash_combine(res, m_reshapeSkeleton);
+	hash_combine(res, m_discardInvalidBindings);
+	hash_combine(res, m_enableRigidParts);
+	hash_combine(res, m_deformAllBones);
+	hash_combine(res, m_deformAllPhysics);
+	hash_combine(res, m_reshapePhysicsVolumes);
+	hash_combine(res, m_reshapeVertices);
+	hash_combine(res, m_bindingMethod);
+
+	for ( const string& S : m_bonesToDeform)
+	{
+		hash_combine(res, S);
+	}
+
+	for ( const string& S : m_physicsToDeform)
+	{
+		hash_combine(res, S);
+	}
+
+	return res;
+}
+
+
+mu::Ptr<ASTOp> ASTOpMeshBindShape::Clone(MapChildFunc& mapChild) const
+{
+    Ptr<ASTOpMeshBindShape> n = new ASTOpMeshBindShape();
+	n->Mesh = mapChild(Mesh.child());
+	n->Shape = mapChild(Shape.child());
+	n->m_reshapeSkeleton = m_reshapeSkeleton;
+	n->m_discardInvalidBindings = m_discardInvalidBindings;
+	n->m_enableRigidParts = m_enableRigidParts;
+	n->m_deformAllBones = m_deformAllBones;
+	n->m_deformAllPhysics = m_deformAllPhysics;
+	n->m_reshapePhysicsVolumes = m_reshapePhysicsVolumes;
+	n->m_reshapeVertices = m_reshapeVertices;
+	n->m_bonesToDeform = m_bonesToDeform;
+	n->m_physicsToDeform = m_physicsToDeform;	
+	n->m_bindingMethod = m_bindingMethod;
+	return n;
+}
+
+
+void ASTOpMeshBindShape::ForEachChild(const std::function<void(ASTChild&)>& f )
+{
+	f(Mesh);
+	f(Shape);
+}
+
+
+void ASTOpMeshBindShape::Link( PROGRAM& program, const FLinkerOptions*)
+{
+    // Already linked?
+    if (!linkedAddress)
+    {
+        OP::MeshBindShapeArgs args;
+        FMemory::Memzero( &args, sizeof(args) );
+
+		if (m_reshapeSkeleton)
+		{
+			args.flags |= uint32(OP::EMeshBindShapeFlags::ReshapeSkeleton);
+		}
+
+    	if (m_discardInvalidBindings)
+    	{
+    		args.flags |= uint32(OP::EMeshBindShapeFlags::DiscardInvalidBindings);
+    	}
+    
+    	if (m_enableRigidParts)
+    	{
+    		args.flags |= uint32(OP::EMeshBindShapeFlags::EnableRigidParts);
+    	}
+		
+		if (m_deformAllBones)
+    	{
+    		args.flags |= uint32(OP::EMeshBindShapeFlags::DeformAllBones);
+    	}
+
+		if (m_deformAllPhysics)
+    	{
+    		args.flags |= uint32(OP::EMeshBindShapeFlags::DeformAllPhysics);
+    	}
+
+    	if (m_reshapePhysicsVolumes)
+    	{
+    		args.flags |= uint32(OP::EMeshBindShapeFlags::ReshapePhysicsVolumes);
+    	}
+  	
+		if (m_reshapeVertices)
+    	{
+    		args.flags |= uint32(OP::EMeshBindShapeFlags::ReshapeVertices);
+    	}  
+ 	
+		args.bindingMethod = m_bindingMethod;    
+		
+		if (Mesh) 
+		{
+			args.mesh = Mesh->linkedAddress;
+		}
+
+		if (Shape)
+		{
+			args.shape = Shape->linkedAddress;
+		}
+		
+		linkedAddress = (OP::ADDRESS)program.m_opAddress.size();
+		program.m_opAddress.push_back((uint32_t)program.m_byteCode.size());
+		AppendCode(program.m_byteCode, OP_TYPE::ME_BINDSHAPE);
+		AppendCode(program.m_byteCode, args);
+
+		AppendCode(program.m_byteCode, (int32)m_bonesToDeform.Num());
+		for (const string& S : m_bonesToDeform)
+		{
+			const uint32 index = program.AddConstant(S);
+			AppendCode(program.m_byteCode, index);
+		}
+
+		AppendCode(program.m_byteCode, (int32)m_physicsToDeform.Num());
+		for (const string& S : m_physicsToDeform)
+		{
+			const uint32 index = program.AddConstant(S);
+			AppendCode(program.m_byteCode, index);
+		}
+    }
+}
+
+
+//-------------------------------------------------------------------------------------------------
+mu::Ptr<ASTOp> ASTOpMeshBindShape::OptimiseSink(const MODEL_OPTIMIZATION_OPTIONS&, OPTIMIZE_SINK_CONTEXT&) const
+{
+	Ptr<ASTOp> NewOp;
+
+	Ptr<ASTOp> MeshAt = Mesh.child();
+	if (!MeshAt)
+	{ 
+		return nullptr;
+	}
+
+	Ptr<ASTOp> ShapeAt = Shape.child();
+	if (!ShapeAt)
+	{
+		return nullptr;
+	}
+
+	OP_TYPE MeshType = MeshAt->GetOpType();
+	OP_TYPE ShapeType = ShapeAt->GetOpType();
+
+	// See if both mesh and shape have an operation that can be optimized in a combined way
+	if (MeshType==ShapeType)
+	{
+		switch (MeshType)
+		{
+
+		case OP_TYPE::ME_SWITCH:
+		{
+			// Id the switch variable and structure is the same
+			const ASTOpSwitch* MeshSwitch = reinterpret_cast<const ASTOpSwitch*>(MeshAt.get());
+			const ASTOpSwitch* ShapeSwitch = reinterpret_cast<const ASTOpSwitch*>(ShapeAt.get());
+			bool bIsSimilarSwitch = MeshSwitch->IsCompatibleWith(ShapeSwitch);
+			if (!bIsSimilarSwitch)
+			{
+				break;
+			}
+
+			// Move the operation down all the paths
+			Ptr<ASTOpSwitch> NewSwitch = mu::Clone<ASTOpSwitch>(MeshAt);
+
+			if (NewSwitch->def)
+			{
+				Ptr<ASTOpMeshBindShape> NewBind = mu::Clone<ASTOpMeshBindShape>(this);
+				NewBind->Mesh = MeshSwitch->def.child();
+				NewBind->Shape = ShapeSwitch->def.child();
+				NewSwitch->def = NewBind;
+			}
+
+			// We need to copy the options because we change them
+			for (size_t v = 0; v < NewSwitch->cases.size(); ++v)
+			{
+				if (NewSwitch->cases[v].branch)
+				{
+					Ptr<ASTOpMeshBindShape> NewBind = mu::Clone<ASTOpMeshBindShape>(this);
+					NewBind->Mesh = MeshSwitch->cases[v].branch.child();
+					NewBind->Shape = ShapeSwitch->FindBranch(MeshSwitch->cases[v].condition);
+					NewSwitch->cases[v].branch = NewBind;
+				}
+			}
+
+			NewOp = NewSwitch;
+			break;
+		}
+
+		default:
+			break;
+
+		}
+	}
+
+
+	// If not already optimized
+	if (!NewOp)
+	{
+		// Optimize only the mesh parameter
+		switch (MeshType)
+		{
+
+		case OP_TYPE::ME_SWITCH:
+		{
+			// Move the operation down all the paths
+			Ptr<ASTOpSwitch> NewSwitch = mu::Clone<ASTOpSwitch>(MeshAt);
+
+			if (NewSwitch->def)
+			{
+				Ptr<ASTOpMeshBindShape> NewBind = mu::Clone<ASTOpMeshBindShape>(this);
+				NewBind->Mesh = NewSwitch->def.child();
+				NewSwitch->def = NewBind;
+			}
+
+			// We need to copy the options because we change them
+			for (size_t v = 0; v < NewSwitch->cases.size(); ++v)
+			{
+				if (NewSwitch->cases[v].branch)
+				{
+					Ptr<ASTOpMeshBindShape> NewBind = mu::Clone<ASTOpMeshBindShape>(this);
+					NewBind->Mesh = NewSwitch->cases[v].branch.child();
+					NewSwitch->cases[v].branch = NewBind;
+				}
+			}
+
+			NewOp = NewSwitch;
+			break;
+		}
+
+		default:
+			break;
+
+		}
+	}
+
+	// If not already optimized
+	if (!NewOp)
+	{
+		// Optimize only the shape parameter
+		switch (ShapeType)
+		{
+
+		case OP_TYPE::ME_SWITCH:
+		{
+			// Move the operation down all the paths
+			Ptr<ASTOpSwitch> NewSwitch = mu::Clone<ASTOpSwitch>(ShapeAt);
+
+			if (NewSwitch->def)
+			{
+				Ptr<ASTOpMeshBindShape> NewBind = mu::Clone<ASTOpMeshBindShape>(this);
+				NewBind->Shape = NewSwitch->def.child();
+				NewSwitch->def = NewBind;
+			}
+
+			// We need to copy the options because we change them
+			for (size_t v = 0; v < NewSwitch->cases.size(); ++v)
+			{
+				if (NewSwitch->cases[v].branch)
+				{
+					Ptr<ASTOpMeshBindShape> NewBind = mu::Clone<ASTOpMeshBindShape>(this);
+					NewBind->Shape = NewSwitch->cases[v].branch.child();
+					NewSwitch->cases[v].branch = NewBind;
+				}
+			}
+
+			NewOp = NewSwitch;
+			break;
+		}
+
+		default:
+			break;
+
+		}
+	}
+
+
+	return NewOp;
+}
diff --git a/Engine/Plugins/Experimental/Mutable/Source/MutableTools/Private/ASTOpMeshBindShape.h b/Engine/Plugins/Experimental/Mutable/Source/MutableTools/Private/ASTOpMeshBindShape.h
new file mode 100644
--- /dev/null
+++ b/Engine/Plugins/Experimental/Mutable/Source/MutableTools/Private/ASTOpMeshBindShape.h
@@ -0,0 +1,53 @@
+// Copyright Epic Games, Inc. All Rights Reserved.
+
+#pragma once
+
+#include "AST.h"
+
+
+namespace mu
+{
+
+
+	//---------------------------------------------------------------------------------------------
+	//!
+	//---------------------------------------------------------------------------------------------
+	class ASTOpMeshBindShape : public ASTOp
+	{
+	public:
+
+		ASTChild Mesh;
+		ASTChild Shape;
+
+		TArray<string> m_bonesToDeform;
+		TArray<string> m_physicsToDeform;
+
+		bool m_reshapeSkeleton = false;
+		bool m_discardInvalidBindings = true;
+		bool m_enableRigidParts = false;
+		bool m_deformAllBones = false;
+		bool m_deformAllPhysics = false;
+		bool m_reshapePhysicsVolumes = false;
+		bool m_reshapeVertices = true;
+
+		uint32 m_bindingMethod = 0;
+
+
+	public:
+
+		ASTOpMeshBindShape();
+		ASTOpMeshBindShape(const ASTOpMeshBindShape&) = delete;
+		~ASTOpMeshBindShape();
+
+		OP_TYPE GetOpType() const override { return OP_TYPE::ME_BINDSHAPE; }
+		uint64 Hash() const override;
+		bool IsEqual(const ASTOp& otherUntyped) const override;
+		Ptr<ASTOp> Clone(MapChildFunc& mapChild) const override;
+		void ForEachChild(const std::function<void(ASTChild&)>&) override;
+		void Link(PROGRAM& program, const FLinkerOptions* Options) override;
+		Ptr<ASTOp> OptimiseSink(const MODEL_OPTIMIZATION_OPTIONS&, OPTIMIZE_SINK_CONTEXT&) const override;
+
+	};
+
+}
+
diff --git a/Engine/Plugins/Experimental/Mutable/Source/MutableTools/Private/ASTOpMeshClipDeform.cpp b/Engine/Plugins/Experimental/Mutable/Source/MutableTools/Private/ASTOpMeshClipDeform.cpp
new file mode 100644
--- /dev/null
+++ b/Engine/Plugins/Experimental/Mutable/Source/MutableTools/Private/ASTOpMeshClipDeform.cpp
@@ -0,0 +1,85 @@
+// Copyright Epic Games, Inc. All Rights Reserved.
+
+#include "ASTOpMeshClipDeform.h"
+#include "ErrorLogPrivate.h"
+#include "CodeOptimiser.h"
+#include "StreamsPrivate.h"
+
+using namespace mu;
+
+
+ASTOpMeshClipDeform::ASTOpMeshClipDeform()
+    : Mesh(this)
+    , ClipShape(this)
+{
+}
+
+
+ASTOpMeshClipDeform::~ASTOpMeshClipDeform()
+{
+    // Explicit call needed to avoid recursive destruction
+    ASTOp::RemoveChildren();
+}
+
+
+bool ASTOpMeshClipDeform::IsEqual(const ASTOp& OtherUntyped) const
+{
+    if ( const ASTOpMeshClipDeform* Other = dynamic_cast<const ASTOpMeshClipDeform*>(&OtherUntyped) )
+    {
+        return Mesh == Other->Mesh && ClipShape == Other->ClipShape;
+    }
+
+    return false;
+}
+
+
+uint64 ASTOpMeshClipDeform::Hash() const
+{
+	uint64 res = std::hash<void*>()(Mesh.child().get());
+	hash_combine(res, ClipShape.child().get());
+	return res;
+}
+
+
+mu::Ptr<ASTOp> ASTOpMeshClipDeform::Clone(MapChildFunc& mapChild) const
+{
+    Ptr<ASTOpMeshClipDeform> n = new ASTOpMeshClipDeform();
+	n->Mesh = mapChild(Mesh.child());
+	n->ClipShape = mapChild(ClipShape.child());
+	return n;
+}
+
+
+void ASTOpMeshClipDeform::ForEachChild(const std::function<void(ASTChild&)>& f )
+{
+	f(Mesh);
+	f(ClipShape);
+}
+
+
+void ASTOpMeshClipDeform::Link( PROGRAM& program, const FLinkerOptions*)
+{
+    // Already linked?
+    if (!linkedAddress)
+    {
+        OP::MeshClipDeformArgs args;
+        memset( &args,0, sizeof(args) );	
+
+		if (Mesh) 
+		{
+			args.mesh = Mesh->linkedAddress;
+		}
+
+		if (ClipShape)
+		{
+			 args.clipShape = ClipShape->linkedAddress;
+		}	
+
+        linkedAddress = (OP::ADDRESS)program.m_opAddress.size();
+        program.m_opAddress.push_back((uint32_t)program.m_byteCode.size());
+        AppendCode(program.m_byteCode,OP_TYPE::ME_CLIPDEFORM);
+        AppendCode(program.m_byteCode,args);
+    }
+
+}
+
diff --git a/Engine/Plugins/Experimental/Mutable/Source/MutableTools/Private/ASTOpMeshClipDeform.h b/Engine/Plugins/Experimental/Mutable/Source/MutableTools/Private/ASTOpMeshClipDeform.h
new file mode 100644
--- /dev/null
+++ b/Engine/Plugins/Experimental/Mutable/Source/MutableTools/Private/ASTOpMeshClipDeform.h
@@ -0,0 +1,38 @@
+// Copyright Epic Games, Inc. All Rights Reserved.
+
+#pragma once
+
+#include "AST.h"
+
+
+namespace mu
+{
+
+
+	//---------------------------------------------------------------------------------------------
+	//!
+	//---------------------------------------------------------------------------------------------
+	class ASTOpMeshClipDeform : public ASTOp
+	{
+	public:
+
+		ASTChild Mesh;
+		ASTChild ClipShape;
+
+	public:
+
+		ASTOpMeshClipDeform();
+		ASTOpMeshClipDeform(const ASTOpMeshClipDeform&) = delete;
+		~ASTOpMeshClipDeform();
+
+		OP_TYPE GetOpType() const override { return OP_TYPE::ME_CLIPDEFORM; }
+		uint64 Hash() const override;
+		bool IsEqual(const ASTOp& otherUntyped) const override;
+		Ptr<ASTOp> Clone(MapChildFunc& mapChild) const override;
+		void ForEachChild(const std::function<void(ASTChild&)>&) override;
+		void Link(PROGRAM& program, const FLinkerOptions* Options) override;
+	};
+
+
+}
+
diff --git a/Engine/Plugins/Experimental/Mutable/Source/MutableTools/Private/ASTOpMeshClipMorphPlane.cpp b/Engine/Plugins/Experimental/Mutable/Source/MutableTools/Private/ASTOpMeshClipMorphPlane.cpp
new file mode 100644
--- /dev/null
+++ b/Engine/Plugins/Experimental/Mutable/Source/MutableTools/Private/ASTOpMeshClipMorphPlane.cpp
@@ -0,0 +1,108 @@
+// Copyright Epic Games, Inc. All Rights Reserved.
+
+#include "ASTOpMeshClipMorphPlane.h"
+#include "ErrorLogPrivate.h"
+#include "CodeOptimiser.h"
+#include "StreamsPrivate.h"
+
+
+using namespace mu;
+
+
+//-------------------------------------------------------------------------------------------------
+//-------------------------------------------------------------------------------------------------
+//-------------------------------------------------------------------------------------------------
+ASTOpMeshClipMorphPlane::ASTOpMeshClipMorphPlane()
+    : source(this)
+{
+}
+
+
+ASTOpMeshClipMorphPlane::~ASTOpMeshClipMorphPlane()
+{
+    // Explicit call needed to avoid recursive destruction
+    ASTOp::RemoveChildren();
+}
+
+
+bool ASTOpMeshClipMorphPlane::IsEqual(const ASTOp& otherUntyped) const
+{
+    if ( auto other = dynamic_cast<const ASTOpMeshClipMorphPlane*>(&otherUntyped) )
+    {
+        return source==other->source &&
+                morphShape==other->morphShape &&
+                selectionShape==other->selectionShape &&
+                vertexSelectionBone==other->vertexSelectionBone &&
+                vertexSelectionType==other->vertexSelectionType &&
+                dist==other->dist &&
+                factor==other->factor;
+    }
+    return false;
+}
+
+
+uint64 ASTOpMeshClipMorphPlane::Hash() const
+{
+	uint64 res = std::hash<void*>()(source.child().get());
+    hash_combine( res, factor );
+    hash_combine( res, dist );
+    return res;
+}
+
+
+mu::Ptr<ASTOp> ASTOpMeshClipMorphPlane::Clone(MapChildFunc& mapChild) const
+{
+    Ptr<ASTOpMeshClipMorphPlane> n = new ASTOpMeshClipMorphPlane();
+    n->source = mapChild(source.child());
+    n->morphShape = morphShape;
+    n->selectionShape = selectionShape;
+    n->vertexSelectionBone = vertexSelectionBone;
+    n->vertexSelectionType = vertexSelectionType;
+    n->dist = dist;
+    n->factor = factor;
+    return n;
+}
+
+
+void ASTOpMeshClipMorphPlane::ForEachChild(const std::function<void(ASTChild&)>& f )
+{
+    f( source );
+}
+
+
+void ASTOpMeshClipMorphPlane::Link( PROGRAM& program, const FLinkerOptions* )
+{
+    // Already linked?
+    if (!linkedAddress)
+    {
+        OP::MeshClipMorphPlaneArgs args;
+        memset( &args,0, sizeof(args) );
+
+        if (source) args.source = source->linkedAddress;
+
+        args.morphShape = (OP::ADDRESS)program.m_constantShapes.size();
+        program.m_constantShapes.push_back(morphShape);
+
+        args.vertexSelectionType = (uint8_t)vertexSelectionType;
+        if (vertexSelectionType==OP::MeshClipMorphPlaneArgs::VS_BONE_HIERARCHY)
+        {
+            args.vertexSelectionShapeOrBone = program.AddConstant(vertexSelectionBone);
+        }
+        else if (vertexSelectionType==OP::MeshClipMorphPlaneArgs::VS_SHAPE)
+        {
+            args.vertexSelectionShapeOrBone = program.AddConstant(selectionShape);
+        }
+
+        args.dist = dist;
+        args.factor = factor;
+		args.maxBoneRadius = vertexSelectionBoneMaxRadius;
+
+        linkedAddress = (OP::ADDRESS)program.m_opAddress.size();
+        //program.m_code.push_back(op);
+        program.m_opAddress.push_back((uint32_t)program.m_byteCode.size());
+        AppendCode(program.m_byteCode,OP_TYPE::ME_CLIPMORPHPLANE);
+        AppendCode(program.m_byteCode,args);
+    }
+
+}
+
diff --git a/Engine/Plugins/Experimental/Mutable/Source/MutableTools/Private/ASTOpMeshClipMorphPlane.h b/Engine/Plugins/Experimental/Mutable/Source/MutableTools/Private/ASTOpMeshClipMorphPlane.h
new file mode 100644
--- /dev/null
+++ b/Engine/Plugins/Experimental/Mutable/Source/MutableTools/Private/ASTOpMeshClipMorphPlane.h
@@ -0,0 +1,44 @@
+// Copyright Epic Games, Inc. All Rights Reserved.
+
+#pragma once
+
+#include "AST.h"
+
+
+namespace mu
+{
+
+	//---------------------------------------------------------------------------------------------
+	//!
+	//---------------------------------------------------------------------------------------------
+	class ASTOpMeshClipMorphPlane : public ASTOp
+	{
+	public:
+
+		ASTChild source;
+
+		SHAPE morphShape;
+		SHAPE selectionShape;
+		string vertexSelectionBone;
+
+		OP::MeshClipMorphPlaneArgs::VERTEX_SELECTION_TYPE vertexSelectionType = OP::MeshClipMorphPlaneArgs::VS_NONE;
+
+		float dist = 0.f, factor = 0.f, vertexSelectionBoneMaxRadius = -1.f;
+
+	public:
+
+		ASTOpMeshClipMorphPlane();
+		ASTOpMeshClipMorphPlane(const ASTOpMeshClipMorphPlane&) = delete;
+		~ASTOpMeshClipMorphPlane();
+
+		OP_TYPE GetOpType() const override { return OP_TYPE::ME_CLIPMORPHPLANE; }
+		uint64 Hash() const override;
+		bool IsEqual(const ASTOp& otherUntyped) const override;
+		Ptr<ASTOp> Clone(MapChildFunc& mapChild) const override;
+		void ForEachChild(const std::function<void(ASTChild&)>&) override;
+		void Link(PROGRAM& program, const FLinkerOptions* Options) override;
+	};
+
+
+}
+
diff --git a/Engine/Plugins/Experimental/Mutable/Source/MutableTools/Private/ASTOpMeshExtractLayoutBlocks.cpp b/Engine/Plugins/Experimental/Mutable/Source/MutableTools/Private/ASTOpMeshExtractLayoutBlocks.cpp
new file mode 100644
--- /dev/null
+++ b/Engine/Plugins/Experimental/Mutable/Source/MutableTools/Private/ASTOpMeshExtractLayoutBlocks.cpp
@@ -0,0 +1,89 @@
+// Copyright Epic Games, Inc. All Rights Reserved.
+
+#include "ASTOpMeshExtractLayoutBlocks.h"
+#include "ErrorLogPrivate.h"
+#include "CodeOptimiser.h"
+#include "StreamsPrivate.h"
+
+
+using namespace mu;
+
+
+//---------------------------------------------------------------------------------------------
+//---------------------------------------------------------------------------------------------
+//---------------------------------------------------------------------------------------------
+ASTOpMeshExtractLayoutBlocks::ASTOpMeshExtractLayoutBlocks()
+    : source(this)
+{
+}
+
+
+ASTOpMeshExtractLayoutBlocks::~ASTOpMeshExtractLayoutBlocks()
+{
+    // Explicit call needed to avoid recursive destruction
+    ASTOp::RemoveChildren();
+}
+
+
+bool ASTOpMeshExtractLayoutBlocks::IsEqual(const ASTOp& otherUntyped) const
+{
+    if ( auto other = dynamic_cast<const ASTOpMeshExtractLayoutBlocks*>(&otherUntyped) )
+    {
+        return source==other->source && layout==other->layout && blocks==other->blocks;
+    }
+    return false;
+}
+
+
+mu::Ptr<ASTOp> ASTOpMeshExtractLayoutBlocks::Clone(MapChildFunc& mapChild) const
+{
+    Ptr<ASTOpMeshExtractLayoutBlocks> n = new ASTOpMeshExtractLayoutBlocks();
+    n->source = mapChild(source.child());
+    n->layout = layout;
+    n->blocks = blocks;
+    return n;
+}
+
+
+void ASTOpMeshExtractLayoutBlocks::Assert()
+{
+    check( blocks.size()<std::numeric_limits<uint16_t>::max() );
+    ASTOp::Assert();
+}
+
+
+void ASTOpMeshExtractLayoutBlocks::ForEachChild(const std::function<void(ASTChild&)>& f )
+{
+    f(source);
+}
+
+
+uint64 ASTOpMeshExtractLayoutBlocks::Hash() const
+{
+	uint64 res = std::hash<size_t>()( size_t(source.child().get() ) );
+    return res;
+}
+
+
+void ASTOpMeshExtractLayoutBlocks::Link( PROGRAM& program, const FLinkerOptions* )
+{
+    // Already linked?
+    if (!linkedAddress)
+    {
+        linkedAddress = (OP::ADDRESS)program.m_opAddress.size();
+
+        program.m_opAddress.push_back((uint32_t)program.m_byteCode.size());
+        AppendCode(program.m_byteCode, OP_TYPE::ME_EXTRACTLAYOUTBLOCK);
+        OP::ADDRESS sourceAt = source ? source->linkedAddress : 0;
+        AppendCode(program.m_byteCode, sourceAt );
+        AppendCode(program.m_byteCode, (uint16_t)layout );
+        AppendCode(program.m_byteCode, (uint16_t)blocks.size() );
+
+        for ( auto b: blocks )
+        {
+            AppendCode(program.m_byteCode, (uint32_t)b );
+        }
+    }
+}
+
+
diff --git a/Engine/Plugins/Experimental/Mutable/Source/MutableTools/Private/ASTOpMeshExtractLayoutBlocks.h b/Engine/Plugins/Experimental/Mutable/Source/MutableTools/Private/ASTOpMeshExtractLayoutBlocks.h
new file mode 100644
--- /dev/null
+++ b/Engine/Plugins/Experimental/Mutable/Source/MutableTools/Private/ASTOpMeshExtractLayoutBlocks.h
@@ -0,0 +1,47 @@
+// Copyright Epic Games, Inc. All Rights Reserved.
+
+#pragma once
+
+#include "AST.h"
+
+
+namespace mu
+{
+
+
+	//---------------------------------------------------------------------------------------------
+	//! Variable sized mesh block extract operation.
+	//---------------------------------------------------------------------------------------------
+	class ASTOpMeshExtractLayoutBlocks : public ASTOp
+	{
+	public:
+
+		//! Source mesh to extract block from.
+		ASTChild source;
+
+		//! Layout to use to select the blocks.
+		uint16_t layout = 0;
+
+		//! Blocks
+		vector<uint32_t> blocks;
+
+	public:
+
+		ASTOpMeshExtractLayoutBlocks();
+		ASTOpMeshExtractLayoutBlocks(const ASTOpMeshExtractLayoutBlocks&) = delete;
+		~ASTOpMeshExtractLayoutBlocks() override;
+
+		// ASTOp interface
+		OP_TYPE GetOpType() const override { return OP_TYPE::ME_EXTRACTLAYOUTBLOCK; }
+		uint64 Hash() const override;
+		void ForEachChild(const std::function<void(ASTChild&)>&) override;
+		bool IsEqual(const ASTOp& otherUntyped) const override;
+		Ptr<ASTOp> Clone(MapChildFunc& mapChild) const override;
+		void Assert() override;
+		void Link(PROGRAM& program, const FLinkerOptions* Options) override;
+
+	};
+
+
+}
+
diff --git a/Engine/Plugins/Experimental/Mutable/Source/MutableTools/Private/ASTOpMeshFormat.cpp b/Engine/Plugins/Experimental/Mutable/Source/MutableTools/Private/ASTOpMeshFormat.cpp
new file mode 100644
--- /dev/null
+++ b/Engine/Plugins/Experimental/Mutable/Source/MutableTools/Private/ASTOpMeshFormat.cpp
@@ -0,0 +1,400 @@
+// Copyright Epic Games, Inc. All Rights Reserved.
+
+#include "MutableTools/Private/ASTOpMeshFormat.h"
+#include "MutableTools/Private/ASTOpConstantResource.h"
+#include "MutableTools/Private/ASTOpSwitch.h"
+#include "MutableTools/Private/ASTOpConditional.h"
+#include "MutableTools/Private/ASTOpMeshRemoveMask.h"
+#include "MutableTools/Private/ASTOpMeshClipMorphPlane.h"
+#include "ErrorLogPrivate.h"
+#include "CodeOptimiser.h"
+#include "StreamsPrivate.h"
+
+
+using namespace mu;
+
+
+//-------------------------------------------------------------------------------------------------
+//-------------------------------------------------------------------------------------------------
+//-------------------------------------------------------------------------------------------------
+ASTOpMeshFormat::ASTOpMeshFormat()
+    : Source(this)
+    , Format(this)
+{
+}
+
+
+ASTOpMeshFormat::~ASTOpMeshFormat()
+{
+    // Explicit call needed to avoid recursive destruction
+    ASTOp::RemoveChildren();
+}
+
+
+bool ASTOpMeshFormat::IsEqual(const ASTOp& otherUntyped) const
+{
+    if (const ASTOpMeshFormat* other = dynamic_cast<const ASTOpMeshFormat*>(&otherUntyped) )
+    {
+        return Source==other->Source && Format==other->Format && Buffers==other->Buffers;
+    }
+    return false;
+}
+
+
+uint64 ASTOpMeshFormat::Hash() const
+{
+	uint64 res = std::hash<void*>()(Source.child().get() );
+    hash_combine( res, Format.child().get() );
+    return res;
+}
+
+
+mu::Ptr<ASTOp> ASTOpMeshFormat::Clone(MapChildFunc& mapChild) const
+{
+    Ptr<ASTOpMeshFormat> n = new ASTOpMeshFormat();
+    n->Source = mapChild(Source.child());
+    n->Format = mapChild(Format.child());
+	n->Buffers = Buffers;
+    return n;
+}
+
+
+void ASTOpMeshFormat::ForEachChild(const std::function<void(ASTChild&)>& f )
+{
+    f( Source );
+    f( Format );
+}
+
+
+void ASTOpMeshFormat::Link( PROGRAM& program, const FLinkerOptions* )
+{
+    // Already linked?
+    if (!linkedAddress)
+    {
+        OP::MeshFormatArgs args;
+        FMemory::Memzero( &args, sizeof(args) );
+
+		args.buffers = Buffers;
+        if (Source) args.source = Source->linkedAddress;
+        if (Format) args.format = Format->linkedAddress;
+
+        linkedAddress = (OP::ADDRESS)program.m_opAddress.size();
+        //program.m_code.push_back(op);
+        program.m_opAddress.push_back((uint32_t)program.m_byteCode.size());
+        AppendCode(program.m_byteCode,OP_TYPE::ME_FORMAT);
+        AppendCode(program.m_byteCode,args);
+    }
+
+}
+
+
+Ptr<ASTOp> ASTOpMeshFormat::OptimiseSink(const MODEL_OPTIMIZATION_OPTIONS& options, OPTIMIZE_SINK_CONTEXT& context) const
+{
+	Ptr<ASTOp> at = context.MeshFormatSinker.Apply(this);
+	return at;
+}
+
+
+//---------------------------------------------------------------------------------------------
+//---------------------------------------------------------------------------------------------
+//---------------------------------------------------------------------------------------------
+mu::Ptr<ASTOp> Sink_MeshFormatAST::Apply(const ASTOp* root)
+{
+	m_root = dynamic_cast<const ASTOpMeshFormat*>(root);
+	check(m_root);
+
+	m_oldToNew.clear();
+
+	m_initialSource = m_root->Source.child();
+	mu::Ptr<ASTOp> newSource = Visit(m_initialSource, m_root);
+
+	// If there is any change, it is the new root.
+	if (newSource != m_initialSource)
+	{
+		return newSource;
+	}
+
+	return nullptr;
+}
+
+
+namespace
+{
+	Ptr<const Mesh> FindBaseMeshConstant(Ptr<ASTOp> at)
+	{
+		Ptr<const Mesh> res;
+
+		switch (at->GetOpType())
+		{
+		case OP_TYPE::ME_CONSTANT:
+		{
+			const ASTOpConstantResource* typed = dynamic_cast<const ASTOpConstantResource*>(at.get());
+			res = static_cast<const Mesh*>(typed->GetValue().get());
+			break;
+		}
+
+		default:
+			check(false);
+		}
+
+		check(res);
+
+		return res;
+	}
+
+
+	// Make a mesh format suitable to morph a particular other format.
+	MeshPtr MakeMorphTargetFormat(MeshPtrConst pTargetFormat)
+	{
+		MUTABLE_CPUPROFILER_SCOPE(MakeMorphTargetFormat);
+
+		// Make a morph format by adding all the vertex channels from the base into a single
+		// vertex buffer, adding the vertex index channel
+
+		int offset = 0;
+		int numChannels = 0;
+		vector<MESH_BUFFER_SEMANTIC> semantics;
+		vector<int> semanticIndices;
+		vector<MESH_BUFFER_FORMAT> formats;
+		vector<int> components;
+		vector<int> offsets;
+
+		// Add the vertex index channel
+		semantics.push_back(MBS_VERTEXINDEX);
+		semanticIndices.push_back(0);
+		formats.push_back(MBF_UINT32);
+		components.push_back(1);
+		offsets.push_back(offset);
+		offset += components[numChannels] * GetMeshFormatData(formats[numChannels]).m_size;
+		numChannels++;
+
+		// Add the vertex channels from the new format
+		for (int vb = 0; vb < pTargetFormat->GetVertexBuffers().GetBufferCount(); ++vb)
+		{
+			for (int c = 0; c < pTargetFormat->GetVertexBuffers().GetBufferChannelCount(vb); ++c)
+			{
+				// Channel info
+				MESH_BUFFER_SEMANTIC semantic;
+				int semanticIndex;
+				MESH_BUFFER_FORMAT format;
+				int component;
+				pTargetFormat->GetVertexBuffers().GetChannel
+				(vb, c, &semantic, &semanticIndex, &format, &component, nullptr);
+
+				// TODO: Filter useless semantics for morphing.
+				semantics.push_back(semantic);
+				semanticIndices.push_back(semanticIndex);
+				formats.push_back(format);
+				components.push_back(component);
+				offsets.push_back(offset);
+				offset += components[numChannels]
+					* GetMeshFormatData(formats[numChannels]).m_size;
+				numChannels++;
+			}
+		}
+
+
+		MeshPtr pTargetMorphFormat = new Mesh;
+		pTargetMorphFormat->GetVertexBuffers().SetBufferCount(1);
+		if (numChannels > 0)
+		{
+			pTargetMorphFormat->GetVertexBuffers().SetBuffer(0, offset,
+				numChannels,
+				&semantics[0],
+				&semanticIndices[0],
+				&formats[0],
+				&components[0],
+				&offsets[0]);
+		}
+
+		return pTargetMorphFormat;
+	}
+
+}
+
+
+mu::Ptr<ASTOp> Sink_MeshFormatAST::Visit(const mu::Ptr<ASTOp>& at, const ASTOpMeshFormat* currentFormatOp)
+{
+	if (!at) return nullptr;
+
+	// Newly created?
+	if (std::find(m_newOps.begin(), m_newOps.end(), at) != m_newOps.end())
+	{
+		return at;
+	}
+
+	// Already visited?
+	auto cacheIt = m_oldToNew.find(at);
+	if (cacheIt != m_oldToNew.end())
+	{
+		return cacheIt->second;
+	}
+
+	mu::Ptr<ASTOp> newAt = at;
+	switch (at->GetOpType())
+	{
+
+	case OP_TYPE::ME_APPLYLAYOUT:
+	{
+		auto newOp = mu::Clone<ASTOpFixed>(at);
+		newOp->SetChild(newOp->op.args.MeshApplyLayout.mesh, Visit(newOp->children[newOp->op.args.MeshApplyLayout.mesh].child(), currentFormatOp));
+		newAt = newOp;
+		break;
+	}
+
+	case OP_TYPE::ME_SETSKELETON:
+	{
+		auto newOp = mu::Clone<ASTOpFixed>(at);
+		newOp->SetChild(newOp->op.args.MeshSetSkeleton.source, Visit(newOp->children[newOp->op.args.MeshSetSkeleton.source].child(), currentFormatOp));
+		newAt = newOp;
+		break;
+	}
+
+	case OP_TYPE::ME_CLIPMORPHPLANE:
+	{
+		auto newOp = mu::Clone<ASTOpMeshClipMorphPlane>(at);
+		newOp->source = Visit(newOp->source.child(), currentFormatOp);
+		newAt = newOp;
+		break;
+	}
+
+	case OP_TYPE::ME_MORPH2:
+	{
+		// Move the format down the base of the morph
+		auto newOp = mu::Clone<ASTOpFixed>(at);
+		newOp->SetChild(newOp->op.args.MeshMorph2.base, Visit(newOp->children[newOp->op.args.MeshMorph2.base].child(), currentFormatOp));
+
+		// Reformat the morph targets to match the new format.
+		// \TODO: Cache pTargetMorphFormat? motaop?
+		MeshPtrConst pTargetFormat = FindBaseMeshConstant(currentFormatOp->Format.child());
+		MeshPtrConst pTargetMorphFormat = MakeMorphTargetFormat(pTargetFormat);
+
+		Ptr<ASTOpConstantResource> motaop = new ASTOpConstantResource();
+		motaop->type = OP_TYPE::ME_CONSTANT;
+		motaop->SetValue(pTargetMorphFormat, false /* useDiskCache */);
+		auto targetMorphFormatAt = motaop;
+
+		for (int t = 0; t < MUTABLE_OP_MAX_MORPH2_TARGETS; ++t)
+		{
+			if (newOp->children[newOp->op.args.MeshMorph2.targets[t]])
+			{
+				Ptr<ASTOpMeshFormat> newFormat = mu::Clone<ASTOpMeshFormat>(currentFormatOp);
+				newFormat->Buffers =
+					OP::MeshFormatArgs::BT_VERTEX
+					| OP::MeshFormatArgs::BT_IGNORE_MISSING;
+				newFormat->Format = targetMorphFormatAt;
+
+				newOp->SetChild(newOp->op.args.MeshMorph2.targets[t], Visit(newOp->children[newOp->op.args.MeshMorph2.targets[t]].child(), newFormat.get()));
+			}
+		}
+
+		newAt = newOp;
+		break;
+	}
+
+	case OP_TYPE::ME_MERGE:
+	{
+		auto newOp = mu::Clone<ASTOpFixed>(at);
+		newOp->SetChild(newOp->op.args.MeshMerge.base, Visit(newOp->children[newOp->op.args.MeshMerge.base].child(), currentFormatOp));
+		newOp->SetChild(newOp->op.args.MeshMerge.added, Visit(newOp->children[newOp->op.args.MeshMerge.added].child(), currentFormatOp));
+		newAt = newOp;
+		break;
+	}
+
+	case OP_TYPE::ME_INTERPOLATE:
+	{
+		// Move the format down the base of the morph
+		auto newOp = mu::Clone<ASTOpFixed>(at);
+		newOp->SetChild(newOp->op.args.MeshInterpolate.base, Visit(newOp->children[newOp->op.args.MeshInterpolate.base].child(), currentFormatOp));
+
+		// Reformat the morph targets to match the new format.
+		// \TODO: Cache pTargetMorphFormat? motaop?
+		MeshPtrConst pTargetFormat = FindBaseMeshConstant(currentFormatOp->Format.child());
+		MeshPtrConst pTargetMorphFormat = MakeMorphTargetFormat(pTargetFormat);
+
+		Ptr<ASTOpConstantResource> motaop = new ASTOpConstantResource();
+		motaop->type = OP_TYPE::ME_CONSTANT;
+		motaop->SetValue(pTargetMorphFormat, false /* useDiskCache */);
+		auto targetMorphFormatAt = motaop;
+
+		for (int t = 0; t < MUTABLE_OP_MAX_INTERPOLATE_COUNT - 1; ++t)
+		{
+			if (newOp->children[newOp->op.args.MeshInterpolate.targets[t]])
+			{
+				Ptr<ASTOpMeshFormat> newFormat = mu::Clone<ASTOpMeshFormat>(currentFormatOp);
+				newFormat->Buffers =
+					OP::MeshFormatArgs::BT_VERTEX
+					| OP::MeshFormatArgs::BT_IGNORE_MISSING;
+				newFormat->Format = targetMorphFormatAt;
+
+				newOp->SetChild(newOp->op.args.MeshInterpolate.targets[t], Visit(newOp->children[newOp->op.args.MeshInterpolate.targets[t]].child(), newFormat.get()));
+			}
+		}
+
+		newAt = newOp;
+		break;
+	}
+
+	case OP_TYPE::ME_REMOVEMASK:
+	{
+		auto newOp = mu::Clone<ASTOpMeshRemoveMask>(at);
+		newOp->source = Visit(newOp->source.child(), currentFormatOp);
+		newAt = newOp;
+		break;
+	}
+
+	case OP_TYPE::ME_CONDITIONAL:
+	{
+		auto newOp = mu::Clone<ASTOpConditional>(at);
+		newOp->yes = Visit(newOp->yes.child(), currentFormatOp);
+		newOp->no = Visit(newOp->no.child(), currentFormatOp);
+		newAt = newOp;
+		break;
+	}
+
+	case OP_TYPE::ME_SWITCH:
+	{
+		auto newOp = mu::Clone<ASTOpSwitch>(at);
+		newOp->def = Visit(newOp->def.child(), currentFormatOp);
+		for (auto& c : newOp->cases)
+		{
+			c.branch = Visit(c.branch.child(), currentFormatOp);
+		}
+		newAt = newOp;
+		break;
+	}
+
+	// This cannot be sunk since the result is different. Since the clipping is now correctly
+	// generated at the end of the chain when really necessary, this wrong optimisation is no 
+	// longer needed.
+	//case OP_TYPE::ME_CLIPMORPHPLANE:
+//         {
+//             // We move the mask creation down the source
+//             auto typedAt = dynamic_cast<const ASTOpMeshClipMorphPlane*>(at.get());
+//             newAt = Visit(typedAt->source.child());
+//             break;
+//         }
+
+	case OP_TYPE::ME_FORMAT:
+		// TODO: The child format can be removed. 
+		// Unless channels are removed and re-added, which would change their content?
+		break;
+
+	case OP_TYPE::ME_DIFFERENCE:
+	case OP_TYPE::ME_SUBTRACT:
+
+	default:
+		if (at != m_initialSource)
+		{
+			Ptr<ASTOpMeshFormat> newOp = mu::Clone<ASTOpMeshFormat>(currentFormatOp);
+			newOp->Source = at;
+			newAt = newOp;
+		}
+		break;
+
+	}
+
+	m_oldToNew[at] = newAt;
+
+	return newAt;
+}
diff --git a/Engine/Plugins/Experimental/Mutable/Source/MutableTools/Private/ASTOpMeshFormat.h b/Engine/Plugins/Experimental/Mutable/Source/MutableTools/Private/ASTOpMeshFormat.h
new file mode 100644
--- /dev/null
+++ b/Engine/Plugins/Experimental/Mutable/Source/MutableTools/Private/ASTOpMeshFormat.h
@@ -0,0 +1,36 @@
+// Copyright Epic Games, Inc. All Rights Reserved.
+
+#pragma once
+
+#include "MutableTools/Private/AST.h"
+
+
+namespace mu
+{
+	
+	class ASTOpMeshFormat : public ASTOp
+	{
+	public:
+
+		ASTChild Source;
+		ASTChild Format;
+		uint8_t Buffers = 0;
+
+	public:
+
+		ASTOpMeshFormat();
+		ASTOpMeshFormat(const ASTOpMeshFormat&) = delete;
+		~ASTOpMeshFormat();
+
+		OP_TYPE GetOpType() const override { return OP_TYPE::ME_FORMAT; }
+		uint64 Hash() const override;
+		bool IsEqual(const ASTOp& otherUntyped) const override;
+		Ptr<ASTOp> Clone(MapChildFunc& mapChild) const override;
+		void ForEachChild(const std::function<void(ASTChild&)>&) override;
+		void Link(PROGRAM& program, const FLinkerOptions* Options) override;
+		Ptr<ASTOp> OptimiseSink(const MODEL_OPTIMIZATION_OPTIONS& options, OPTIMIZE_SINK_CONTEXT& context) const override;
+
+	};
+
+}
+
diff --git a/Engine/Plugins/Experimental/Mutable/Source/MutableTools/Private/ASTOpMeshGeometryOperation.cpp b/Engine/Plugins/Experimental/Mutable/Source/MutableTools/Private/ASTOpMeshGeometryOperation.cpp
new file mode 100644
--- /dev/null
+++ b/Engine/Plugins/Experimental/Mutable/Source/MutableTools/Private/ASTOpMeshGeometryOperation.cpp
@@ -0,0 +1,89 @@
+// Copyright Epic Games, Inc. All Rights Reserved.
+
+#include "ASTOpMeshGeometryOperation.h"
+#include "ErrorLogPrivate.h"
+#include "CodeOptimiser.h"
+#include "StreamsPrivate.h"
+
+using namespace mu;
+
+
+ASTOpMeshGeometryOperation::ASTOpMeshGeometryOperation()
+    : meshA(this)
+    , meshB(this)
+	, scalarA(this)
+	, scalarB(this)
+{
+}
+
+
+ASTOpMeshGeometryOperation::~ASTOpMeshGeometryOperation()
+{
+    // Explicit call needed to avoid recursive destruction
+    ASTOp::RemoveChildren();
+}
+
+
+bool ASTOpMeshGeometryOperation::IsEqual(const ASTOp& otherUntyped) const
+{
+    if ( auto other = dynamic_cast<const ASTOpMeshGeometryOperation*>(&otherUntyped) )
+    {
+        return meshA ==other->meshA && meshB == other->meshB &&
+			scalarA == other->scalarA && scalarB == other->scalarB;
+    }
+    return false;
+}
+
+
+uint64 ASTOpMeshGeometryOperation::Hash() const
+{
+	uint64 res = std::hash<void*>()(meshA.child().get() );
+	hash_combine(res, meshB.child().get());
+	hash_combine(res, scalarA.child().get());
+	hash_combine(res, scalarB.child().get());
+	return res;
+}
+
+
+mu::Ptr<ASTOp> ASTOpMeshGeometryOperation::Clone(MapChildFunc& mapChild) const
+{
+    Ptr<ASTOpMeshGeometryOperation> n = new ASTOpMeshGeometryOperation();
+	n->meshA = mapChild(meshA.child());
+	n->meshB = mapChild(meshB.child());
+	n->scalarA = mapChild(scalarA.child());
+	n->scalarB = mapChild(scalarB.child());
+	return n;
+}
+
+
+void ASTOpMeshGeometryOperation::ForEachChild(const std::function<void(ASTChild&)>& f )
+{
+	f(meshA);
+	f(meshB);
+	f(scalarA);
+	f(scalarB);
+}
+
+
+void ASTOpMeshGeometryOperation::Link( PROGRAM& program, const FLinkerOptions* )
+{
+    // Already linked?
+    if (!linkedAddress)
+    {
+        OP::MeshGeometryOperationArgs args;
+        memset( &args,0, sizeof(args) );
+
+		if (meshA) args.meshA = meshA->linkedAddress;
+		if (meshB) args.meshB = meshB->linkedAddress;
+		if (scalarA) args.scalarA = scalarA->linkedAddress;
+		if (scalarB) args.scalarB = scalarB->linkedAddress;
+
+        linkedAddress = (OP::ADDRESS)program.m_opAddress.size();
+        //program.m_code.push_back(op);
+        program.m_opAddress.push_back((uint32_t)program.m_byteCode.size());
+        AppendCode(program.m_byteCode,OP_TYPE::ME_GEOMETRYOPERATION);
+        AppendCode(program.m_byteCode,args);
+    }
+
+}
+
diff --git a/Engine/Plugins/Experimental/Mutable/Source/MutableTools/Private/ASTOpMeshGeometryOperation.h b/Engine/Plugins/Experimental/Mutable/Source/MutableTools/Private/ASTOpMeshGeometryOperation.h
new file mode 100644
--- /dev/null
+++ b/Engine/Plugins/Experimental/Mutable/Source/MutableTools/Private/ASTOpMeshGeometryOperation.h
@@ -0,0 +1,40 @@
+// Copyright Epic Games, Inc. All Rights Reserved.
+
+#pragma once
+
+#include "AST.h"
+
+
+namespace mu
+{
+
+
+	//---------------------------------------------------------------------------------------------
+	//!
+	//---------------------------------------------------------------------------------------------
+	class ASTOpMeshGeometryOperation : public ASTOp
+	{
+	public:
+
+		ASTChild meshA;
+		ASTChild meshB;
+		ASTChild scalarA;
+		ASTChild scalarB;
+
+	public:
+
+		ASTOpMeshGeometryOperation();
+		ASTOpMeshGeometryOperation(const ASTOpMeshGeometryOperation&) = delete;
+		~ASTOpMeshGeometryOperation();
+
+		OP_TYPE GetOpType() const override { return OP_TYPE::ME_GEOMETRYOPERATION; }
+		uint64 Hash() const override;
+		bool IsEqual(const ASTOp& otherUntyped) const override;
+		Ptr<ASTOp> Clone(MapChildFunc& mapChild) const override;
+		void ForEachChild(const std::function<void(ASTChild&)>&) override;
+		void Link(PROGRAM& program, const FLinkerOptions* Options) override;
+	};
+
+
+}
+
diff --git a/Engine/Plugins/Experimental/Mutable/Source/MutableTools/Private/ASTOpMeshMaskClipMesh.cpp b/Engine/Plugins/Experimental/Mutable/Source/MutableTools/Private/ASTOpMeshMaskClipMesh.cpp
new file mode 100644
--- /dev/null
+++ b/Engine/Plugins/Experimental/Mutable/Source/MutableTools/Private/ASTOpMeshMaskClipMesh.cpp
@@ -0,0 +1,237 @@
+// Copyright Epic Games, Inc. All Rights Reserved.
+
+#include "MutableTools/Private/ASTOpMeshMaskClipMesh.h"
+#include "MutableTools/Private/ASTOpSwitch.h"
+#include "MutableTools/Private/ASTOpConditional.h"
+#include "MutableTools/Private/ASTOpMeshRemoveMask.h"
+#include "ErrorLogPrivate.h"
+#include "CodeOptimiser.h"
+#include "StreamsPrivate.h"
+
+
+using namespace mu;
+
+
+//-------------------------------------------------------------------------------------------------
+//-------------------------------------------------------------------------------------------------
+//-------------------------------------------------------------------------------------------------
+ASTOpMeshMaskClipMesh::ASTOpMeshMaskClipMesh()
+    : source(this)
+    , clip(this)
+{
+}
+
+
+//-------------------------------------------------------------------------------------------------
+ASTOpMeshMaskClipMesh::~ASTOpMeshMaskClipMesh()
+{
+    // Explicit call needed to avoid recursive destruction
+    ASTOp::RemoveChildren();
+}
+
+
+//-------------------------------------------------------------------------------------------------
+bool ASTOpMeshMaskClipMesh::IsEqual(const ASTOp& otherUntyped) const
+{
+    if ( auto other = dynamic_cast<const ASTOpMeshMaskClipMesh*>(&otherUntyped) )
+    {
+        return source==other->source && clip==other->clip;
+    }
+    return false;
+}
+
+
+//-------------------------------------------------------------------------------------------------
+uint64 ASTOpMeshMaskClipMesh::Hash() const
+{
+	uint64 res = std::hash<void*>()(source.child().get());
+    hash_combine( res, clip.child() );
+    return res;
+}
+
+
+//-------------------------------------------------------------------------------------------------
+mu::Ptr<ASTOp> ASTOpMeshMaskClipMesh::Clone(MapChildFunc& mapChild) const
+{
+    Ptr<ASTOpMeshMaskClipMesh> n = new ASTOpMeshMaskClipMesh();
+    n->source = mapChild(source.child());
+    n->clip = mapChild(clip.child());
+    return n;
+}
+
+
+//-------------------------------------------------------------------------------------------------
+void ASTOpMeshMaskClipMesh::ForEachChild(const std::function<void(ASTChild&)>& f )
+{
+    f( source );
+    f( clip );
+}
+
+
+//-------------------------------------------------------------------------------------------------
+void ASTOpMeshMaskClipMesh::Link( PROGRAM& program, const FLinkerOptions* )
+{
+    // Already linked?
+    if (!linkedAddress)
+    {
+        OP::MeshMaskClipMeshArgs args;
+        memset( &args,0, sizeof(args) );
+
+        if (source) args.source = source->linkedAddress;
+        if (clip) args.clip = clip->linkedAddress;
+
+        linkedAddress = (OP::ADDRESS)program.m_opAddress.size();
+        //program.m_code.push_back(op);
+        program.m_opAddress.push_back((uint32_t)program.m_byteCode.size());
+        AppendCode(program.m_byteCode,OP_TYPE::ME_MASKCLIPMESH);
+        AppendCode(program.m_byteCode,args);
+    }
+}
+
+
+//-------------------------------------------------------------------------------------------------
+//-------------------------------------------------------------------------------------------------
+//-------------------------------------------------------------------------------------------------
+namespace
+{
+    class Sink_MeshMaskClipMeshAST
+    {
+    public:
+
+        // \TODO This is recursive and may cause stack overflows in big models.
+
+        mu::Ptr<ASTOp> Apply( const ASTOpMeshMaskClipMesh* root )
+        {
+            m_root = root;
+            m_oldToNew.clear();
+
+            m_initialSource = root->source.child();
+			mu::Ptr<ASTOp> newSource = Visit( m_initialSource );
+
+            // If there is any change, it is the new root.
+            if (newSource!=m_initialSource)
+            {
+                return newSource;
+            }
+
+            return nullptr;
+        }
+
+    protected:
+
+        const ASTOpMeshMaskClipMesh* m_root;
+		mu::Ptr<ASTOp> m_initialSource;
+        std::unordered_map<mu::Ptr<ASTOp>, mu::Ptr<ASTOp>> m_oldToNew;
+        vector<mu::Ptr<ASTOp>> m_newOps;
+
+		mu::Ptr<ASTOp> Visit( const mu::Ptr<ASTOp>& at )
+        {
+            if (!at) return nullptr;
+
+            // Newly created?
+            if (std::find(m_newOps.begin(), m_newOps.end(), at )!=m_newOps.end())
+            {
+                return at;
+            }
+
+            // Already visited?
+            auto cacheIt = m_oldToNew.find(at);
+            if (cacheIt!=m_oldToNew.end())
+            {
+                return cacheIt->second;
+            }
+
+			mu::Ptr<ASTOp> newAt = at;
+            switch ( at->GetOpType() )
+            {
+
+			// This cannot be sunk since the result is different. Since the clipping is now correctly
+			// generated at the end of the chain when really necessary, this wrong optimisation is no 
+			// longer needed.
+			//case OP_TYPE::ME_MORPH2:
+   //         {
+   //             // Remove this op: Let's clip the unmorphed mesh only
+   //             auto typedAt = dynamic_cast<const ASTOpFixed*>(at.get());
+   //             newAt = Visit( typedAt->children[typedAt->op.args.MeshMorph2.base].child() );
+   //             break;
+   //         }
+
+            case OP_TYPE::ME_REMOVEMASK:
+            {
+                // Remove this op:
+                // This may lead to the mask being bigger than needed since it will include
+                // faces removed by the ignored removemask, but it is ok
+
+                // TODO: Swap instead of ignore, and implement removemask on a mask?
+                auto typedAt = dynamic_cast<const ASTOpMeshRemoveMask*>(at.get());
+                newAt = Visit( typedAt->source.child() );
+                break;
+            }
+
+            case OP_TYPE::ME_CONDITIONAL:
+            {
+                // We move the mask creation down the two paths
+                // It always needs to be a clone because otherwise we could be modifying an
+                // instruction that shouldn't if the parent was a ME_REMOVEMASK above and we
+                // skipped the cloning for the parent.
+                auto newOp = mu::Clone<ASTOpConditional>(at);
+                newOp->yes = Visit(newOp->yes.child());
+                newOp->no = Visit(newOp->no.child());
+                newAt = newOp;
+                break;
+            }
+
+            case OP_TYPE::ME_SWITCH:
+            {
+                // We move the mask creation down all the paths
+                auto newOp = mu::Clone<ASTOpSwitch>(at);
+                newOp->def = Visit(newOp->def.child());
+                for( auto& c:newOp->cases )
+                {
+                    c.branch = Visit(c.branch.child());
+                }
+                newAt = newOp;
+                break;
+            }
+
+			// This cannot be sunk since the result is different. Since the clipping is now correctly
+			// generated at the end of the chain when really necessary, this wrong optimisation is no 
+			// longer needed.
+			//case OP_TYPE::ME_CLIPMORPHPLANE:
+   //         {
+   //             // We move the mask creation down the source
+   //             auto typedAt = dynamic_cast<const ASTOpMeshClipMorphPlane*>(at.get());
+   //             newAt = Visit(typedAt->source.child());
+   //             break;
+   //         }
+
+            default:
+            {
+                //
+                if (at!=m_initialSource)
+                {
+                    auto newOp = mu::Clone<ASTOpMeshMaskClipMesh>(m_root);
+                    newOp->source = at;
+                    newAt = newOp;
+                }
+                break;
+            }
+
+            }
+
+            m_oldToNew[at] = newAt;
+
+            return newAt;
+        }
+    };
+}
+
+
+//-------------------------------------------------------------------------------------------------
+mu::Ptr<ASTOp> ASTOpMeshMaskClipMesh::OptimiseSink(const MODEL_OPTIMIZATION_OPTIONS&, OPTIMIZE_SINK_CONTEXT&) const
+{
+    Sink_MeshMaskClipMeshAST sinker;
+    auto at = sinker.Apply(this);
+
+    return at;
+}
diff --git a/Engine/Plugins/Experimental/Mutable/Source/MutableTools/Private/ASTOpMeshMaskClipMesh.h b/Engine/Plugins/Experimental/Mutable/Source/MutableTools/Private/ASTOpMeshMaskClipMesh.h
new file mode 100644
--- /dev/null
+++ b/Engine/Plugins/Experimental/Mutable/Source/MutableTools/Private/ASTOpMeshMaskClipMesh.h
@@ -0,0 +1,39 @@
+// Copyright Epic Games, Inc. All Rights Reserved.
+
+#pragma once
+
+#include "AST.h"
+
+
+namespace mu
+{
+
+
+	//---------------------------------------------------------------------------------------------
+	//!
+	//---------------------------------------------------------------------------------------------
+	class ASTOpMeshMaskClipMesh : public ASTOp
+	{
+	public:
+
+		ASTChild source;
+		ASTChild clip;
+
+	public:
+
+		ASTOpMeshMaskClipMesh();
+		ASTOpMeshMaskClipMesh(const ASTOpMeshMaskClipMesh&) = delete;
+		~ASTOpMeshMaskClipMesh();
+
+		OP_TYPE GetOpType() const override { return OP_TYPE::ME_MASKCLIPMESH; }
+		uint64 Hash() const override;
+		bool IsEqual(const ASTOp& otherUntyped) const override;
+		Ptr<ASTOp> Clone(MapChildFunc& mapChild) const override;
+		void ForEachChild(const std::function<void(ASTChild&)>&) override;
+		void Link(PROGRAM& program, const FLinkerOptions* Options) override;
+		Ptr<ASTOp> OptimiseSink(const MODEL_OPTIMIZATION_OPTIONS&, OPTIMIZE_SINK_CONTEXT&) const override;
+
+	};
+
+}
+
diff --git a/Engine/Plugins/Experimental/Mutable/Source/MutableTools/Private/ASTOpMeshMorphReshape.cpp b/Engine/Plugins/Experimental/Mutable/Source/MutableTools/Private/ASTOpMeshMorphReshape.cpp
new file mode 100644
--- /dev/null
+++ b/Engine/Plugins/Experimental/Mutable/Source/MutableTools/Private/ASTOpMeshMorphReshape.cpp
@@ -0,0 +1,85 @@
+// Copyright Epic Games, Inc. All Rights Reserved.
+
+#include "ASTOpMeshMorphReshape.h"
+#include "ErrorLogPrivate.h"
+#include "CodeOptimiser.h"
+#include "StreamsPrivate.h"
+
+using namespace mu;
+
+
+ASTOpMeshMorphReshape::ASTOpMeshMorphReshape()
+    : Morph(this)
+    , Reshape(this)
+{
+}
+
+
+ASTOpMeshMorphReshape::~ASTOpMeshMorphReshape()
+{
+    // Explicit call needed to avoid recursive destruction
+    ASTOp::RemoveChildren();
+}
+
+
+bool ASTOpMeshMorphReshape::IsEqual(const ASTOp& otherUntyped) const
+{
+	if ( auto other = dynamic_cast<const ASTOpMeshMorphReshape*>(&otherUntyped) )
+	{
+		return Morph == other->Morph && Reshape == other->Reshape;
+	}
+	
+    return false;
+}
+
+
+uint64 ASTOpMeshMorphReshape::Hash() const
+{
+	uint64 res = std::hash<void*>()(Morph.child().get() );
+	hash_combine(res, Reshape.child().get());
+
+	return res;
+}
+
+
+mu::Ptr<ASTOp> ASTOpMeshMorphReshape::Clone(MapChildFunc& mapChild) const
+{
+    Ptr<ASTOpMeshMorphReshape> n = new ASTOpMeshMorphReshape();
+	n->Morph = mapChild(Morph.child());
+	n->Reshape = mapChild(Reshape.child());
+	return n;
+}
+
+
+void ASTOpMeshMorphReshape::ForEachChild(const std::function<void(ASTChild&)>& f )
+{
+	f(Morph);
+	f(Reshape);
+}
+
+
+void ASTOpMeshMorphReshape::Link( PROGRAM& program, const FLinkerOptions*)
+{
+    // Already linked?
+    if (!linkedAddress)
+    {
+        OP::MeshMorphReshapeArgs Args;
+        memset( &Args, 0, sizeof(Args) );
+    	
+		if (Morph) 
+		{
+			Args.Morph = Morph->linkedAddress;
+		}
+
+		if (Reshape) 
+		{
+			Args.Reshape = Reshape->linkedAddress;
+		}
+
+        linkedAddress = (OP::ADDRESS)program.m_opAddress.size();
+        program.m_opAddress.push_back((uint32_t)program.m_byteCode.size());
+        AppendCode(program.m_byteCode, OP_TYPE::ME_MORPHRESHAPE);
+        AppendCode(program.m_byteCode, Args);
+    }
+
+}
diff --git a/Engine/Plugins/Experimental/Mutable/Source/MutableTools/Private/ASTOpMeshMorphReshape.h b/Engine/Plugins/Experimental/Mutable/Source/MutableTools/Private/ASTOpMeshMorphReshape.h
new file mode 100644
--- /dev/null
+++ b/Engine/Plugins/Experimental/Mutable/Source/MutableTools/Private/ASTOpMeshMorphReshape.h
@@ -0,0 +1,37 @@
+// Copyright Epic Games, Inc. All Rights Reserved.
+
+#pragma once
+
+#include "AST.h"
+
+
+namespace mu
+{
+
+	//---------------------------------------------------------------------------------------------
+	//!
+	//---------------------------------------------------------------------------------------------
+	class ASTOpMeshMorphReshape : public ASTOp
+	{
+	public:
+
+		ASTChild Morph;
+		ASTChild Reshape;
+
+	public:
+
+		ASTOpMeshMorphReshape();
+		ASTOpMeshMorphReshape(const ASTOpMeshMorphReshape&) = delete;
+		~ASTOpMeshMorphReshape();
+
+		OP_TYPE GetOpType() const override { return OP_TYPE::ME_MORPHRESHAPE; }
+		uint64 Hash() const override;
+		bool IsEqual(const ASTOp& otherUntyped) const override;
+		Ptr<ASTOp> Clone(MapChildFunc& mapChild) const override;
+		void ForEachChild(const std::function<void(ASTChild&)>&) override;
+		void Link(PROGRAM& program, const FLinkerOptions* Options) override;
+	};
+
+
+}
+
