[mutable] Move Mutable plugin from the restricted folder to the experimental plugins folder.

diff --git a/Engine/Plugins/Experimental/Mutable/Source/MutableTools/Private/ASTOpMeshRemapIndices.cpp b/Engine/Plugins/Experimental/Mutable/Source/MutableTools/Private/ASTOpMeshRemapIndices.cpp
new file mode 100644
--- /dev/null
+++ b/Engine/Plugins/Experimental/Mutable/Source/MutableTools/Private/ASTOpMeshRemapIndices.cpp
@@ -0,0 +1,85 @@
+// Copyright Epic Games, Inc. All Rights Reserved.
+
+#include "ASTOpMeshRemapIndices.h"
+#include "ErrorLogPrivate.h"
+#include "CodeOptimiser.h"
+#include "StreamsPrivate.h"
+
+using namespace mu;
+
+
+//-------------------------------------------------------------------------------------------------
+//-------------------------------------------------------------------------------------------------
+//-------------------------------------------------------------------------------------------------
+ASTOpMeshRemapIndices::ASTOpMeshRemapIndices()
+    : source(this)
+    , reference(this)
+{
+}
+
+
+//-------------------------------------------------------------------------------------------------
+ASTOpMeshRemapIndices::~ASTOpMeshRemapIndices()
+{
+    // Explicit call needed to avoid recursive destruction
+    ASTOp::RemoveChildren();
+}
+
+
+//-------------------------------------------------------------------------------------------------
+bool ASTOpMeshRemapIndices::IsEqual(const ASTOp& otherUntyped) const
+{
+    if ( auto other = dynamic_cast<const ASTOpMeshRemapIndices*>(&otherUntyped) )
+    {
+        return source==other->source && reference==other->reference;
+    }
+    return false;
+}
+
+
+//-------------------------------------------------------------------------------------------------
+uint64 ASTOpMeshRemapIndices::Hash() const
+{
+	uint64 res = std::hash<void*>()(source.child().get());
+    hash_combine( res, reference.child() );
+    return res;
+}
+
+
+//-------------------------------------------------------------------------------------------------
+mu::Ptr<ASTOp> ASTOpMeshRemapIndices::Clone(MapChildFunc& mapChild) const
+{
+    Ptr<ASTOpMeshRemapIndices> n = new ASTOpMeshRemapIndices();
+    n->source = mapChild(source.child());
+    n->reference = mapChild(reference.child());
+    return n;
+}
+
+
+//-------------------------------------------------------------------------------------------------
+void ASTOpMeshRemapIndices::ForEachChild(const std::function<void(ASTChild&)>& f )
+{
+    f( source );
+    f( reference );
+}
+
+
+//-------------------------------------------------------------------------------------------------
+void ASTOpMeshRemapIndices::Link( PROGRAM& program, const FLinkerOptions* )
+{
+    // Already linked?
+    if (!linkedAddress)
+    {
+        OP::MeshRemapIndicesArgs args;
+        memset( &args,0, sizeof(args) );
+
+        if (source) args.source = source->linkedAddress;
+        if (reference) args.reference = reference->linkedAddress;
+
+        linkedAddress = (OP::ADDRESS)program.m_opAddress.size();
+        program.m_opAddress.push_back((uint32_t)program.m_byteCode.size());
+        AppendCode(program.m_byteCode,OP_TYPE::ME_REMAPINDICES);
+        AppendCode(program.m_byteCode,args);
+    }
+}
+
diff --git a/Engine/Plugins/Experimental/Mutable/Source/MutableTools/Private/ASTOpMeshRemapIndices.h b/Engine/Plugins/Experimental/Mutable/Source/MutableTools/Private/ASTOpMeshRemapIndices.h
new file mode 100644
--- /dev/null
+++ b/Engine/Plugins/Experimental/Mutable/Source/MutableTools/Private/ASTOpMeshRemapIndices.h
@@ -0,0 +1,42 @@
+// Copyright Epic Games, Inc. All Rights Reserved.
+
+#pragma once
+
+#include "AST.h"
+
+
+namespace mu
+{
+
+
+	//---------------------------------------------------------------------------------------------
+	//! 
+	//---------------------------------------------------------------------------------------------
+	class ASTOpMeshRemapIndices : public ASTOp
+	{
+	public:
+
+		//! Mesh that will have the vertex indices remapped.
+		ASTChild source;
+
+		//! Mesh used to obtain the final vertex indices.
+		ASTChild reference;
+
+	public:
+
+		ASTOpMeshRemapIndices();
+		ASTOpMeshRemapIndices(const ASTOpMeshRemapIndices&) = delete;
+		~ASTOpMeshRemapIndices() override;
+
+		// ASTOp interface
+		OP_TYPE GetOpType() const override { return OP_TYPE::ME_REMAPINDICES; }
+		uint64 Hash() const override;
+		void ForEachChild(const std::function<void(ASTChild&)>&) override;
+		bool IsEqual(const ASTOp& otherUntyped) const override;
+		Ptr<ASTOp> Clone(MapChildFunc& mapChild) const override;
+		void Link(PROGRAM& program, const FLinkerOptions*) override;
+	};
+
+
+}
+
diff --git a/Engine/Plugins/Experimental/Mutable/Source/MutableTools/Private/ASTOpMeshRemoveMask.cpp b/Engine/Plugins/Experimental/Mutable/Source/MutableTools/Private/ASTOpMeshRemoveMask.cpp
new file mode 100644
--- /dev/null
+++ b/Engine/Plugins/Experimental/Mutable/Source/MutableTools/Private/ASTOpMeshRemoveMask.cpp
@@ -0,0 +1,236 @@
+// Copyright Epic Games, Inc. All Rights Reserved.
+
+#include "ASTOpMeshRemoveMask.h"
+#include "ErrorLogPrivate.h"
+#include "CodeOptimiser.h"
+#include "StreamsPrivate.h"
+
+
+using namespace mu;
+
+
+//---------------------------------------------------------------------------------------------
+//---------------------------------------------------------------------------------------------
+//---------------------------------------------------------------------------------------------
+ASTOpMeshRemoveMask::ASTOpMeshRemoveMask()
+    : source(this)
+{
+}
+
+
+//---------------------------------------------------------------------------------------------
+ASTOpMeshRemoveMask::~ASTOpMeshRemoveMask()
+{
+    // Explicit call needed to avoid recursive destruction
+    ASTOp::RemoveChildren();
+}
+
+
+//---------------------------------------------------------------------------------------------
+void ASTOpMeshRemoveMask::AddRemove( const Ptr<ASTOp>& condition, const Ptr<ASTOp>& mask  )
+{
+    removes.push_back( std::make_pair<>( ASTChild(this,condition),
+                                         ASTChild(this,mask) ) );
+}
+
+
+//---------------------------------------------------------------------------------------------
+bool ASTOpMeshRemoveMask::IsEqual(const ASTOp& otherUntyped) const
+{
+    if ( auto other = dynamic_cast<const ASTOpMeshRemoveMask*>(&otherUntyped) )
+    {
+        return source==other->source && removes==other->removes;
+    }
+    return false;
+}
+
+
+//---------------------------------------------------------------------------------------------
+mu::Ptr<ASTOp> ASTOpMeshRemoveMask::Clone(MapChildFunc& mapChild) const
+{
+    Ptr<ASTOpMeshRemoveMask> n = new ASTOpMeshRemoveMask();
+    n->source = mapChild(source.child());
+    for(const auto& r:removes)
+    {
+        n->removes.push_back( std::make_pair<>( ASTChild(n,mapChild(r.first.child())),
+                                                ASTChild(n,mapChild(r.second.child()))
+                                                ) );
+    }
+    return n;
+}
+
+
+//---------------------------------------------------------------------------------------------
+void ASTOpMeshRemoveMask::Assert()
+{
+    ASTOp::Assert();
+}
+
+
+//---------------------------------------------------------------------------------------------
+void ASTOpMeshRemoveMask::ForEachChild(const std::function<void(ASTChild&)>& f )
+{
+    f(source);
+    for(auto& r:removes)
+    {
+        f(r.first);
+        f(r.second);
+    }
+}
+
+
+//---------------------------------------------------------------------------------------------
+uint64 ASTOpMeshRemoveMask::Hash() const
+{
+	uint64 res = std::hash<ASTOp*>()( source.child().get() );
+    for(const auto& r:removes)
+    {
+        hash_combine( res, r.first.child().get() );
+        hash_combine( res, r.second.child().get() );
+    }
+    return res;
+}
+
+
+//---------------------------------------------------------------------------------------------
+void ASTOpMeshRemoveMask::Link( PROGRAM& program, const FLinkerOptions*)
+{
+    // Already linked?
+    if (!linkedAddress)
+    {
+        linkedAddress = (OP::ADDRESS)program.m_opAddress.size();
+
+        program.m_opAddress.push_back((uint32_t)program.m_byteCode.size());
+        AppendCode(program.m_byteCode, OP_TYPE::ME_REMOVEMASK);
+        OP::ADDRESS sourceAt = source ? source->linkedAddress : 0;
+        AppendCode(program.m_byteCode, sourceAt );
+        AppendCode(program.m_byteCode, (uint16_t)removes.size() );
+        for ( const auto& b: removes )
+        {
+            OP::ADDRESS condition = b.first ? b.first->linkedAddress : 0;
+            AppendCode(program.m_byteCode, condition );
+
+            OP::ADDRESS remove = b.second ? b.second->linkedAddress : 0;
+            AppendCode(program.m_byteCode, remove );
+        }
+    }
+}
+
+
+//-------------------------------------------------------------------------------------------------
+//-------------------------------------------------------------------------------------------------
+//-------------------------------------------------------------------------------------------------
+namespace
+{
+    class Sink_MeshRemoveMaskAST
+    {
+    public:
+
+        // \TODO This is recursive and may cause stack overflows in big models.
+
+		mu::Ptr<ASTOp> Apply( const ASTOpMeshRemoveMask* root )
+        {
+            m_root = root;
+            m_oldToNew.clear();
+
+            m_initialSource = root->source.child();
+			mu::Ptr<ASTOp> newSource = Visit( m_initialSource );
+
+            // If there is any change, it is the new root.
+            if (newSource!=m_initialSource)
+            {
+                return newSource;
+            }
+
+            return nullptr;
+        }
+
+    protected:
+
+        const ASTOpMeshRemoveMask* m_root;
+		mu::Ptr<ASTOp> m_initialSource;
+        std::unordered_map<mu::Ptr<ASTOp>, mu::Ptr<ASTOp>> m_oldToNew;
+        vector<mu::Ptr<ASTOp>> m_newOps;
+
+		mu::Ptr<ASTOp> Visit( const mu::Ptr<ASTOp>& at )
+        {
+            if (!at) return nullptr;
+
+            // Newly created?
+            if (std::find(m_newOps.begin(), m_newOps.end(), at )!=m_newOps.end())
+            {
+                return at;
+            }
+
+            // Already visited?
+            auto cacheIt = m_oldToNew.find(at);
+            if (cacheIt!=m_oldToNew.end())
+            {
+                return cacheIt->second;
+            }
+
+			mu::Ptr<ASTOp> newAt = at;
+            switch ( at->GetOpType() )
+            {
+
+            case OP_TYPE::ME_MORPH2:
+            {
+				mu::Ptr<ASTOpFixed> newOp = mu::Clone<ASTOpFixed>(at);
+                newOp->SetChild( newOp->op.args.MeshMorph2.base, Visit( newOp->children[newOp->op.args.MeshMorph2.base].child() ) );
+                newAt = newOp;
+                break;
+            }
+
+                // disabled to avoid code explosion (or bug?) TODO
+//            case OP_TYPE::ME_CONDITIONAL:
+//            {
+//                Ptr<ASTOpConditional> newOp = mu::Clone<ASTOpConditional>(at);
+//                newOp->yes = Visit(newOp->yes.child());
+//                newOp->no = Visit(newOp->no.child());
+//                newAt = newOp;
+//                break;
+//            }
+
+//            case OP_TYPE::ME_SWITCH:
+//            {
+//                auto newOp = mu::Clone<ASTOpSwitch>(at);
+//                newOp->def = Visit(newOp->def.child());
+//                for( auto& c:newOp->cases )
+//                {
+//                    c.branch = Visit(c.branch.child());
+//                }
+//                newAt = newOp;
+//                break;
+//            }
+
+            default:
+            {
+                //
+                if (at!=m_initialSource)
+                {
+                    auto newOp = mu::Clone<ASTOpMeshRemoveMask>(m_root);
+                    newOp->source = at;
+                    newAt = newOp;
+                }
+                break;
+            }
+
+            }
+
+            m_oldToNew[at] = newAt;
+
+            return newAt;
+        }
+    };
+}
+
+
+//-------------------------------------------------------------------------------------------------
+mu::Ptr<ASTOp> ASTOpMeshRemoveMask::OptimiseSink(const MODEL_OPTIMIZATION_OPTIONS&, OPTIMIZE_SINK_CONTEXT&) const
+{
+    Sink_MeshRemoveMaskAST sinker;
+    auto at = sinker.Apply(this);
+
+    return at;
+}
+
diff --git a/Engine/Plugins/Experimental/Mutable/Source/MutableTools/Private/ASTOpMeshRemoveMask.h b/Engine/Plugins/Experimental/Mutable/Source/MutableTools/Private/ASTOpMeshRemoveMask.h
new file mode 100644
--- /dev/null
+++ b/Engine/Plugins/Experimental/Mutable/Source/MutableTools/Private/ASTOpMeshRemoveMask.h
@@ -0,0 +1,46 @@
+// Copyright Epic Games, Inc. All Rights Reserved.
+
+#pragma once
+
+#include "AST.h"
+
+
+namespace mu
+{
+
+
+	//---------------------------------------------------------------------------------------------
+	//! From a source mesh, remove a list of fragments with a condition.
+	//---------------------------------------------------------------------------------------------
+	class ASTOpMeshRemoveMask : public ASTOp
+	{
+	public:
+
+		//! Source mesh to remove from.
+		ASTChild source;
+
+		//! Pairs of remove candidates: condition + mesh to remove
+		vector< std::pair<ASTChild, ASTChild> > removes;
+
+	public:
+
+		ASTOpMeshRemoveMask();
+		ASTOpMeshRemoveMask(const ASTOpMeshRemoveMask&) = delete;
+		~ASTOpMeshRemoveMask() override;
+
+		// ASTOp interface
+		OP_TYPE GetOpType() const override { return OP_TYPE::ME_REMOVEMASK; }
+		uint64 Hash() const override;
+		void ForEachChild(const std::function<void(ASTChild&)>&) override;
+		bool IsEqual(const ASTOp& otherUntyped) const override;
+		Ptr<ASTOp> Clone(MapChildFunc& mapChild) const override;
+		void Assert() override;
+		void Link(PROGRAM& program, const FLinkerOptions*) override;
+		Ptr<ASTOp> OptimiseSink(const MODEL_OPTIMIZATION_OPTIONS&, OPTIMIZE_SINK_CONTEXT&) const override;
+
+		// Own interface
+		void AddRemove(const Ptr<ASTOp>& condition, const Ptr<ASTOp>& mask);
+	};
+
+}
+
diff --git a/Engine/Plugins/Experimental/Mutable/Source/MutableTools/Private/ASTOpMeshTransform.cpp b/Engine/Plugins/Experimental/Mutable/Source/MutableTools/Private/ASTOpMeshTransform.cpp
new file mode 100644
--- /dev/null
+++ b/Engine/Plugins/Experimental/Mutable/Source/MutableTools/Private/ASTOpMeshTransform.cpp
@@ -0,0 +1,81 @@
+// Copyright Epic Games, Inc. All Rights Reserved.
+
+#include "ASTOpMeshTransform.h"
+#include "ErrorLogPrivate.h"
+#include "CodeOptimiser.h"
+#include "StreamsPrivate.h"
+
+
+using namespace mu;
+
+
+//-------------------------------------------------------------------------------------------------
+//-------------------------------------------------------------------------------------------------
+//-------------------------------------------------------------------------------------------------
+ASTOpMeshTransform::ASTOpMeshTransform()
+    : source(this)
+{
+}
+
+
+ASTOpMeshTransform::~ASTOpMeshTransform()
+{
+    // Explicit call needed to avoid recursive destruction
+    ASTOp::RemoveChildren();
+}
+
+
+bool ASTOpMeshTransform::IsEqual(const ASTOp& otherUntyped) const
+{
+    if ( auto other = dynamic_cast<const ASTOpMeshTransform*>(&otherUntyped) )
+    {
+        return source==other->source &&
+                matrix==other->matrix;
+    }
+    return false;
+}
+
+
+uint64 ASTOpMeshTransform::Hash() const
+{
+	uint64 res = std::hash<OP_TYPE>()(OP_TYPE::ME_TRANSFORM);
+    hash_combine( res, source.child().get() );
+    return res;
+}
+
+
+mu::Ptr<ASTOp> ASTOpMeshTransform::Clone(MapChildFunc& mapChild) const
+{
+    Ptr<ASTOpMeshTransform> n = new ASTOpMeshTransform();
+    n->matrix = matrix;
+    n->source = mapChild(source.child());
+    return n;
+}
+
+
+void ASTOpMeshTransform::ForEachChild(const std::function<void(ASTChild&)>& f )
+{
+    f( source );
+}
+
+
+void ASTOpMeshTransform::Link( PROGRAM& program, const FLinkerOptions*)
+{
+    // Already linked?
+    if (!linkedAddress)
+    {
+        OP::MeshTransformArgs args;
+        memset( &args,0, sizeof(args) );
+
+        if (source) args.source = source->linkedAddress;
+
+        args.matrix = program.AddConstant(matrix);
+
+        linkedAddress = (OP::ADDRESS)program.m_opAddress.size();
+        //program.m_code.push_back(op);
+        program.m_opAddress.push_back((uint32_t)program.m_byteCode.size());
+        AppendCode(program.m_byteCode,OP_TYPE::ME_TRANSFORM);
+        AppendCode(program.m_byteCode,args);
+    }
+
+}
diff --git a/Engine/Plugins/Experimental/Mutable/Source/MutableTools/Private/ASTOpMeshTransform.h b/Engine/Plugins/Experimental/Mutable/Source/MutableTools/Private/ASTOpMeshTransform.h
new file mode 100644
--- /dev/null
+++ b/Engine/Plugins/Experimental/Mutable/Source/MutableTools/Private/ASTOpMeshTransform.h
@@ -0,0 +1,40 @@
+// Copyright Epic Games, Inc. All Rights Reserved.
+
+#pragma once
+
+#include "AST.h"
+
+
+namespace mu
+{
+
+
+	//---------------------------------------------------------------------------------------------
+	//!
+	//---------------------------------------------------------------------------------------------
+	class ASTOpMeshTransform : public ASTOp
+	{
+	public:
+
+		ASTChild source;
+
+		mat4f matrix;
+
+	public:
+
+		ASTOpMeshTransform();
+		ASTOpMeshTransform(const ASTOpMeshTransform&) = delete;
+		~ASTOpMeshTransform();
+
+		OP_TYPE GetOpType() const override { return OP_TYPE::ME_TRANSFORM; }
+		uint64 Hash() const override;
+		void ForEachChild(const std::function<void(ASTChild&)>&) override;
+		bool IsEqual(const ASTOp& otherUntyped) const override;
+		Ptr<ASTOp> Clone(MapChildFunc& mapChild) const override;
+		void Link(PROGRAM& program, const FLinkerOptions* Options) override;
+
+	};
+
+
+}
+
diff --git a/Engine/Plugins/Experimental/Mutable/Source/MutableTools/Private/ASTOpParameter.cpp b/Engine/Plugins/Experimental/Mutable/Source/MutableTools/Private/ASTOpParameter.cpp
new file mode 100644
--- /dev/null
+++ b/Engine/Plugins/Experimental/Mutable/Source/MutableTools/Private/ASTOpParameter.cpp
@@ -0,0 +1,156 @@
+// Copyright Epic Games, Inc. All Rights Reserved.
+
+#include "ASTOpParameter.h"
+#include "ErrorLogPrivate.h"
+#include "CodeOptimiser.h"
+#include "StreamsPrivate.h"
+
+
+using namespace mu;
+
+
+//-------------------------------------------------------------------------------------------------
+//-------------------------------------------------------------------------------------------------
+//-------------------------------------------------------------------------------------------------
+ASTOpParameter::~ASTOpParameter()
+{
+    // Explicit call needed to avoid recursive destruction
+    ASTOp::RemoveChildren();
+}
+
+
+void ASTOpParameter::ForEachChild(const std::function<void(ASTChild&)>& f)
+{
+    for (auto &c: ranges)
+    {
+        f(c.rangeSize);
+    }
+    for (auto &c: additionalImages)
+    {
+        f(c);
+    }
+}
+
+
+bool ASTOpParameter::IsEqual(const ASTOp& otherUntyped) const
+{
+    if ( auto other = dynamic_cast<const ASTOpParameter*>(&otherUntyped) )
+    {
+        return type==other->type &&
+                parameter==other->parameter &&
+                ranges==other->ranges &&
+                additionalImages==other->additionalImages;
+    }
+    return false;
+}
+
+
+mu::Ptr<ASTOp> ASTOpParameter::Clone(MapChildFunc& mapChild) const
+{
+    Ptr<ASTOpParameter> n = new ASTOpParameter();
+    n->type = type;
+    n->parameter = parameter;
+    for(const auto& c:ranges)
+    {
+        n->ranges.emplace_back( n.get(), mapChild(c.rangeSize.child()), c.rangeName, c.rangeUID );
+    }
+    for(const auto& c:additionalImages)
+    {
+        n->additionalImages.emplace_back(n,mapChild(c.child()));
+    }
+    return n;
+}
+
+
+uint64 ASTOpParameter::Hash() const
+{
+	uint64 res = std::hash<uint64>()(uint64(type));
+    hash_combine( res, parameter.m_type );
+    hash_combine( res, parameter.m_name );
+    return res;
+}
+
+
+void ASTOpParameter::Assert()
+{
+    ASTOp::Assert();
+}
+
+
+void ASTOpParameter::Link( PROGRAM& program, const FLinkerOptions* )
+{
+    // Already linked?
+    if (!linkedAddress)
+    {
+        OP::ParameterArgs args;
+        memset( &args,0, sizeof(args) );
+
+        args.variable = (OP::ADDRESS)program.m_parameters.size();
+        program.m_parameters.push_back(parameter);
+
+        for( const auto& d: ranges)
+        {
+            OP::ADDRESS sizeAt = 0;
+            uint16_t rangeId = 0;
+            LinkRange( program, d, sizeAt, rangeId );
+            program.m_parameters.back().m_ranges.push_back(rangeId);
+        }
+
+        for( const auto& d: additionalImages)
+        {
+            OP::ADDRESS descAt = d ? d->linkedAddress : 0;
+            program.m_parameters.back().m_descImages.push_back(descAt);
+        }
+
+        linkedAddress = (OP::ADDRESS)program.m_opAddress.size();
+        //program.m_code.push_back(op);
+
+        program.m_opAddress.push_back((uint32_t)program.m_byteCode.size());
+        AppendCode(program.m_byteCode,type);
+        AppendCode(program.m_byteCode,args);
+    }
+}
+
+
+int ASTOpParameter::EvaluateInt( ASTOpList& facts, bool &unknown ) const
+{
+    unknown = true;
+
+    // Check the facts, in case we have the value for our parameter.
+    for( const auto& f: facts )
+    {
+        if (f->GetOpType()==OP_TYPE::BO_EQUAL_INT_CONST)
+        {
+            const auto typedFact = dynamic_cast<const ASTOpFixed*>(f.get());
+            auto value = typedFact->children[typedFact->op.args.BoolEqualScalarConst.value].child();
+            if (value.get()==this)
+            {
+                unknown = false;
+                return typedFact->op.args.BoolEqualScalarConst.constant;
+            }
+            else
+            {
+                // We could try something more if it was an expression and it had the parameter
+                // somewhere in it.
+            }
+        }
+    }
+
+    return 0;
+}
+
+
+ASTOp::BOOL_EVAL_RESULT ASTOpParameter::EvaluateBool( ASTOpList& /*facts*/,
+                                                      EVALUATE_BOOL_CACHE* ) const
+{
+    check(type==OP_TYPE::BO_PARAMETER);
+    return BET_UNKNOWN;
+}
+
+
+FImageDesc ASTOpParameter::GetImageDesc(bool, GetImageDescContext*)
+{
+	check(type == OP_TYPE::IM_PARAMETER);
+	return FImageDesc();
+}
+
diff --git a/Engine/Plugins/Experimental/Mutable/Source/MutableTools/Private/ASTOpParameter.h b/Engine/Plugins/Experimental/Mutable/Source/MutableTools/Private/ASTOpParameter.h
new file mode 100644
--- /dev/null
+++ b/Engine/Plugins/Experimental/Mutable/Source/MutableTools/Private/ASTOpParameter.h
@@ -0,0 +1,49 @@
+// Copyright Epic Games, Inc. All Rights Reserved.
+
+#pragma once
+
+#include "AST.h"
+
+
+namespace mu
+{
+
+
+	//---------------------------------------------------------------------------------------------
+	//! Parameter operation.
+	//---------------------------------------------------------------------------------------------
+	class ASTOpParameter : public ASTOp
+	{
+	public:
+
+		//! Type of parameter
+		OP_TYPE type;
+
+		//!
+		PARAMETER_DESC parameter;
+
+		//! Ranges adding dimensions to this parameter
+		vector<RANGE_DATA> ranges;
+
+		//! Additional images attached to the parameter
+		vector<ASTChild> additionalImages;
+
+	public:
+
+		~ASTOpParameter() override;
+
+		OP_TYPE GetOpType() const override { return type; }
+		uint64 Hash() const override;
+		void ForEachChild(const std::function<void(ASTChild&)>&) override;
+		bool IsEqual(const ASTOp& otherUntyped) const override;
+		Ptr<ASTOp> Clone(MapChildFunc& mapChild) const override;
+		void Assert() override;
+		void Link(PROGRAM& program, const FLinkerOptions*) override;
+		int EvaluateInt(ASTOpList& facts, bool& unknown) const override;
+		BOOL_EVAL_RESULT EvaluateBool(ASTOpList& /*facts*/, EVALUATE_BOOL_CACHE* = nullptr) const override;
+		FImageDesc GetImageDesc(bool, GetImageDescContext*) override;
+
+	};
+
+}
+
diff --git a/Engine/Plugins/Experimental/Mutable/Source/MutableTools/Private/ASTOpScalarCurve.cpp b/Engine/Plugins/Experimental/Mutable/Source/MutableTools/Private/ASTOpScalarCurve.cpp
new file mode 100644
--- /dev/null
+++ b/Engine/Plugins/Experimental/Mutable/Source/MutableTools/Private/ASTOpScalarCurve.cpp
@@ -0,0 +1,79 @@
+// Copyright Epic Games, Inc. All Rights Reserved.
+
+#include "ASTOpScalarCurve.h"
+#include "ErrorLogPrivate.h"
+#include "CodeOptimiser.h"
+#include "StreamsPrivate.h"
+
+
+using namespace mu;
+
+
+//-------------------------------------------------------------------------------------------------
+//-------------------------------------------------------------------------------------------------
+//-------------------------------------------------------------------------------------------------
+ASTOpScalarCurve::ASTOpScalarCurve()
+    : time(this)
+{
+}
+
+
+ASTOpScalarCurve::~ASTOpScalarCurve()
+{
+    // Explicit call needed to avoid recursive destruction
+    ASTOp::RemoveChildren();
+}
+
+
+void ASTOpScalarCurve::ForEachChild(const std::function<void(ASTChild&)>& f)
+{
+    f( time );
+}
+
+
+uint64 ASTOpScalarCurve::Hash() const
+{
+	uint64 res = std::hash<uint64>()(size_t(OP_TYPE::SC_CURVE));
+    hash_combine( res, curve.keyFrames.size() );
+    return res;
+}
+
+
+bool ASTOpScalarCurve::IsEqual(const ASTOp& otherUntyped) const
+{
+    if ( auto other = dynamic_cast<const ASTOpScalarCurve*>(&otherUntyped) )
+    {
+        return time==other->time && curve==other->curve;
+    }
+    return false;
+}
+
+
+mu::Ptr<ASTOp> ASTOpScalarCurve::Clone(MapChildFunc& mapChild) const
+{
+    Ptr<ASTOpScalarCurve> n = new ASTOpScalarCurve();
+    n->curve = curve;
+    n->time = mapChild(time.child());
+    return n;
+}
+
+
+void ASTOpScalarCurve::Link( PROGRAM& program, const FLinkerOptions* )
+{
+    if (!linkedAddress)
+    {
+        OP::ScalarCurveArgs args;
+        memset( &args,0, sizeof(args) );
+        args.time = time->linkedAddress;
+        args.curve = program.AddConstant(curve);
+
+        linkedAddress = (OP::ADDRESS)program.m_opAddress.size();
+        //program.m_code.push_back(op);
+        program.m_opAddress.push_back((uint32_t)program.m_byteCode.size());
+        AppendCode(program.m_byteCode,OP_TYPE::SC_CURVE);
+        AppendCode(program.m_byteCode,args);
+    }
+
+}
+
+
diff --git a/Engine/Plugins/Experimental/Mutable/Source/MutableTools/Private/ASTOpScalarCurve.h b/Engine/Plugins/Experimental/Mutable/Source/MutableTools/Private/ASTOpScalarCurve.h
new file mode 100644
--- /dev/null
+++ b/Engine/Plugins/Experimental/Mutable/Source/MutableTools/Private/ASTOpScalarCurve.h
@@ -0,0 +1,40 @@
+// Copyright Epic Games, Inc. All Rights Reserved.
+
+#pragma once
+
+#include "AST.h"
+
+
+namespace mu
+{
+
+
+	//---------------------------------------------------------------------------------------------
+	//!
+	//---------------------------------------------------------------------------------------------
+	class ASTOpScalarCurve : public ASTOp
+	{
+	public:
+
+		ASTChild time;
+
+		//!
+		Curve curve;
+
+	public:
+
+		ASTOpScalarCurve();
+		ASTOpScalarCurve(const ASTOpScalarCurve&) = delete;
+		~ASTOpScalarCurve() override;
+
+		OP_TYPE GetOpType() const override { return OP_TYPE::PR_CONSTANT; }
+		uint64 Hash() const override;
+		void ForEachChild(const std::function<void(ASTChild&)>&) override;
+		bool IsEqual(const ASTOp& otherUntyped) const override;
+		Ptr<ASTOp> Clone(MapChildFunc& mapChild) const override;
+		void Link(PROGRAM& program, const FLinkerOptions* Options) override;
+	};
+
+
+}
+
diff --git a/Engine/Plugins/Experimental/Mutable/Source/MutableTools/Private/ASTOpSwitch.cpp b/Engine/Plugins/Experimental/Mutable/Source/MutableTools/Private/ASTOpSwitch.cpp
new file mode 100644
--- /dev/null
+++ b/Engine/Plugins/Experimental/Mutable/Source/MutableTools/Private/ASTOpSwitch.cpp
@@ -0,0 +1,502 @@
+// Copyright Epic Games, Inc. All Rights Reserved.
+
+#include "ASTOpSwitch.h"
+#include "ErrorLogPrivate.h"
+#include "CodeOptimiser.h"
+#include "StreamsPrivate.h"
+
+
+using namespace mu;
+
+
+//---------------------------------------------------------------------------------------------
+//---------------------------------------------------------------------------------------------
+//---------------------------------------------------------------------------------------------
+ASTOpSwitch::ASTOpSwitch()
+    : variable(this)
+    , def(this)
+{
+}
+
+
+//-------------------------------------------------------------------------------------------------
+ASTOpSwitch::~ASTOpSwitch()
+{
+    // Explicit call needed to avoid recursive destruction
+    ASTOp::RemoveChildren();
+}
+
+
+//-------------------------------------------------------------------------------------------------
+bool ASTOpSwitch::IsEqual(const ASTOp& otherUntyped) const
+{
+    if ( auto other = dynamic_cast<const ASTOpSwitch*>(&otherUntyped) )
+    {
+        return type==other->type && variable==other->variable &&
+                cases==other->cases && def==other->def;
+    }
+    return false;
+}
+
+
+//-------------------------------------------------------------------------------------------------
+mu::Ptr<ASTOp> ASTOpSwitch::Clone(MapChildFunc& mapChild) const
+{
+    Ptr<ASTOpSwitch> n = new ASTOpSwitch();
+    n->type = type;
+    n->variable = mapChild(variable.child());
+    n->def = mapChild(def.child());
+    for(const auto& c:cases)
+    {
+        n->cases.emplace_back(c.condition, n, mapChild(c.branch.child()));
+    }
+    return n;
+}
+
+
+//-------------------------------------------------------------------------------------------------
+void ASTOpSwitch::Assert()
+{
+    switch(type)
+    {
+    case OP_TYPE::NU_SWITCH:
+    case OP_TYPE::SC_SWITCH:
+    case OP_TYPE::CO_SWITCH:
+    case OP_TYPE::IM_SWITCH:
+    case OP_TYPE::ME_SWITCH:
+    case OP_TYPE::LA_SWITCH:
+    case OP_TYPE::IN_SWITCH:
+        break;
+    default:
+        // Unexpected type
+        check(false);
+        break;
+    }
+
+    ASTOp::Assert();
+}
+
+
+//-------------------------------------------------------------------------------------------------
+uint64 ASTOpSwitch::Hash() const
+{
+	uint64 res = std::hash<uint64>()(uint64(type));
+    for (const auto& c: cases)
+    {
+        hash_combine( res, c.condition );
+        hash_combine( res, c.branch.child().get() );
+    }
+    return res;
+}
+
+
+//-------------------------------------------------------------------------------------------------
+mu::Ptr<ASTOp> ASTOpSwitch::GetFirstValidValue()
+{
+    for ( size_t i=0; i<cases.size(); ++i )
+    {
+        if ( cases[i].branch )
+        {
+            return cases[i].branch.child();
+        }
+    }
+    return nullptr;
+}
+
+
+//-------------------------------------------------------------------------------------------------
+bool ASTOpSwitch::IsCompatibleWith( const ASTOpSwitch* other ) const
+{
+    if (!other) return false;
+    if (variable.child()!=other->variable.child()) return false;
+    if (cases.size()!=other->cases.size()) return false;
+    for (const auto& c: cases)
+    {
+        bool found = false;
+        for (const auto& o: other->cases)
+        {
+            if (c.condition==o.condition)
+            {
+                found=true;
+                break;
+            }
+        }
+        if (!found)
+        {
+            return false;
+        }
+    }
+
+    return true;
+}
+
+
+//-------------------------------------------------------------------------------------------------
+mu::Ptr<ASTOp> ASTOpSwitch::FindBranch( int32_t condition ) const
+{
+    for (const auto& c: cases)
+    {
+        if (c.condition==condition)
+        {
+            return c.branch.child();
+        }
+    }
+
+    return def.child();
+}
+
+
+//-------------------------------------------------------------------------------------------------
+void ASTOpSwitch::ForEachChild(const std::function<void(ASTChild&)>& f )
+{
+    f(variable);
+    f(def);
+    for ( auto& cas: cases)
+    {
+        f(cas.branch);
+    }
+}
+
+
+//-------------------------------------------------------------------------------------------------
+void ASTOpSwitch::Link( PROGRAM& program, const FLinkerOptions* )
+{
+    // Already linked?
+    if (!linkedAddress)
+    {        
+        linkedAddress = (OP::ADDRESS)program.m_opAddress.size();
+        program.m_opAddress.push_back((uint32_t)program.m_byteCode.size());
+ 
+        OP::ADDRESS VarAddress = variable ? variable->linkedAddress : 0;
+        OP::ADDRESS DefAddress = def ? def->linkedAddress : 0;
+ 
+        AppendCode( program.m_byteCode, type );
+        AppendCode( program.m_byteCode, VarAddress );
+        AppendCode( program.m_byteCode, DefAddress );
+        AppendCode( program.m_byteCode, (uint32_t)cases.size() );
+
+        for ( const CASE& Case : cases )
+        {
+            OP::ADDRESS CaseBranchAddress = Case.branch ? Case.branch->linkedAddress : 0;  
+            AppendCode( program.m_byteCode, Case.condition );
+            AppendCode( program.m_byteCode, CaseBranchAddress );
+        }
+    }
+}
+
+//-------------------------------------------------------------------------------------------------
+FImageDesc ASTOpSwitch::GetImageDesc( bool returnBestOption, class GetImageDescContext* context )
+{
+    FImageDesc res;
+
+    // Local context in case it is necessary
+    GetImageDescContext localContext;
+    if (!context)
+    {
+        context = &localContext;
+    }
+    else
+    {
+        // Cached result?
+		FImageDesc* PtrValue = context->m_results.Find(this);
+		if (PtrValue)
+		{
+			return *PtrValue;
+		}
+	}
+
+    // In a switch we cannot guarantee the size and format.
+    // We check all the options, and if they are the same we return that.
+    // Otherwise, we return a descriptor with empty fields in the conflicting ones, size or format.
+    // In some places this will force re-formatting of the image.
+    // The code optimiser will take care then of moving the format operations down to each
+    // branch and remove the unnecessary ones.
+    FImageDesc candidate;
+    bool sameSize = true;
+    bool sameFormat = true;
+    bool sameLods = true;
+    bool first = true;
+
+    if ( def )
+    {
+        FImageDesc childDesc = def->GetImageDesc( returnBestOption, context );
+        candidate = childDesc;
+        first = false;
+    }
+
+    for ( int i=0; i<(int)cases.size(); ++i )
+    {
+        if ( cases[i].branch )
+        {
+            FImageDesc childDesc = cases[i].branch->GetImageDesc(  returnBestOption, context );
+            if ( first )
+            {
+                candidate = childDesc;
+                first = false;
+            }
+            else
+            {
+                sameSize = (candidate.m_size == childDesc.m_size);
+                sameFormat = (candidate.m_format == childDesc.m_format);
+                sameLods = (candidate.m_lods == childDesc.m_lods);
+
+                if (returnBestOption)
+                {
+                    candidate.m_format =
+                        GetMostGenericFormat( candidate.m_format, childDesc.m_format );
+                }
+            }
+          }
+      }
+
+
+    res = candidate;
+
+    if (!sameFormat && !returnBestOption)
+    {
+        res.m_format = EImageFormat::IF_NONE;
+    }
+
+    if (!sameSize && !returnBestOption)
+    {
+        res.m_size = FImageSize(0, 0);
+    }
+
+    if (!sameLods && !returnBestOption)
+    {
+        res.m_lods = 0;
+    }
+
+    // Cache the result
+    if (context)
+    {
+		context->m_results.Add(this, res);
+	}
+
+    return res;
+}
+
+
+//-------------------------------------------------------------------------------------------------
+void ASTOpSwitch::GetBlockLayoutSize( int blockIndex, int* pBlockX, int* pBlockY,
+                                      BLOCK_LAYOUT_SIZE_CACHE* cache )
+{
+    switch ( type )
+    {
+    case OP_TYPE::LA_SWITCH:
+    {
+        Ptr<ASTOp> child = GetFirstValidValue();
+        if (!child)
+        {
+            child = def.child();
+        }
+
+        if (child)
+        {
+            child->GetBlockLayoutSizeCached( blockIndex, pBlockX, pBlockY, cache );
+        }
+        else
+        {
+            *pBlockX = 0;
+            *pBlockY = 0;
+        }
+        break;
+    }
+
+    default:
+        check(false);
+        break;
+    }
+}
+
+
+//-------------------------------------------------------------------------------------------------
+void ASTOpSwitch::GetLayoutBlockSize( int* pBlockX, int* pBlockY )
+{
+    switch ( type )
+    {
+    case OP_TYPE::IM_SWITCH:
+    {
+        Ptr<ASTOp> child = GetFirstValidValue();
+        if (!child)
+        {
+            child = def.child();
+        }
+
+        if (child)
+        {
+            child->GetLayoutBlockSize(pBlockX, pBlockY);
+        }
+        else
+        {
+            checkf(false, TEXT("Image switch had no options.") );
+        }
+        break;
+    }
+
+    default:
+        checkf( false, TEXT("Instruction not supported") );
+    }
+}
+
+
+//-------------------------------------------------------------------------------------------------
+bool ASTOpSwitch::GetNonBlackRect( FImageRect& maskUsage ) const
+{
+    if (type == OP_TYPE::IM_SWITCH )
+    {
+        FImageRect local;
+        bool localValid = false;
+        if (def)
+        {
+            localValid = def->GetNonBlackRect(local);
+            if (!localValid)
+            {
+                return false;
+            }
+        }
+
+        for (const auto& c: cases)
+        {
+            if (c.branch)
+            {
+                FImageRect branchRect;
+                bool validBranch = c.branch->GetNonBlackRect(branchRect);
+                if (validBranch)
+                {
+                    if (localValid)
+                    {
+                        local.Bound(branchRect);
+                    }
+                    else
+                    {
+                        local = branchRect;
+                        localValid = true;
+                    }
+                }
+                else
+                {
+                    return false;
+                }
+            }
+        }
+
+        if (localValid)
+        {
+            maskUsage = local;
+            return true;
+        }
+    }
+
+    return false;
+}
+
+
+//-------------------------------------------------------------------------------------------------
+bool ASTOpSwitch::IsImagePlainConstant( vec4<float>& ) const
+{
+    // We could check if every option is plain and exactly the same colour, but probably it is
+    // not worth.
+    return false;
+}
+
+
+//-------------------------------------------------------------------------------------------------
+mu::Ptr<ASTOp> ASTOpSwitch::OptimiseSemantic( const MODEL_OPTIMIZATION_OPTIONS& ) const
+{
+    // Constant condition?
+    if  ( variable->GetOpType() == OP_TYPE::NU_CONSTANT )
+    {
+        Ptr<ASTOp> branch = def.child();
+
+        auto typedCondition = dynamic_cast<const ASTOpFixed*>(variable.child().get());
+        for ( size_t o=0; o<cases.size(); ++o )
+        {
+            if ( cases[o].branch &&
+                typedCondition->op.args.IntConstant.value
+                 ==
+                 (int)cases[o].condition)
+            {
+                branch = cases[o].branch.child();
+                break;
+            }
+        }
+
+        return branch;
+    }
+
+    //else
+    //{
+        // If all the possible branches are the same, remove the instruction
+        // TODO: If the variables is a parameter, checking that all the possible
+        // values of this parameter are handled.
+
+//            Disable this: I think it doesn't deal correctly with def. If all branches are the same but the
+//              value is none of the options, "def" should be selected.
+//            OP::ADDRESS switchAt = at;
+//            bool same = true;
+//            OP::ADDRESS value = 0;
+//            //vector< UINT16 > options;
+
+//            const auto& options = program.m_constantSwitches[program.m_code[switchAt].args.Switch.options].m_options;
+//            for ( size_t o=0; o<options.size(); ++o )
+//            {
+//                if ( options[o].at )
+//                {
+//                    if (!value)
+//                    {
+//                        value = options[o].at;
+//                    }
+//                    else
+//                    {
+//                        same &= ( value == options[o].at );
+//                    }
+
+//                    //options.push_back
+//                    //		( program.m_code[switchAt].args.Switch.conditions[o] );
+//                }
+//            }
+
+
+//            if (same)
+//            {
+//                m_modified = true;
+
+//                at = value;
+//            }
+
+    //}
+
+    return nullptr;
+}
+
+
+//-------------------------------------------------------------------------------------------------
+mu::Ptr<ImageSizeExpression> ASTOpSwitch::GetImageSizeExpression() const
+{
+    Ptr<ImageSizeExpression> pRes = new ImageSizeExpression;
+
+    bool first = true;
+    for (const auto& c : cases )
+    {
+        if (c.branch)
+        {
+            if (first)
+            {
+                pRes = c.branch->GetImageSizeExpression();
+            }
+            else
+            {
+                Ptr<ImageSizeExpression> pOther = c.branch->GetImageSizeExpression();
+                if ( !(*pOther==*pRes) )
+                {
+                    pRes->type = ImageSizeExpression::ISET_UNKNOWN;
+                    break;
+                }
+            }
+        }
+    }
+
+    return pRes;
+}
+
diff --git a/Engine/Plugins/Experimental/Mutable/Source/MutableTools/Private/ASTOpSwitch.h b/Engine/Plugins/Experimental/Mutable/Source/MutableTools/Private/ASTOpSwitch.h
new file mode 100644
--- /dev/null
+++ b/Engine/Plugins/Experimental/Mutable/Source/MutableTools/Private/ASTOpSwitch.h
@@ -0,0 +1,89 @@
+// Copyright Epic Games, Inc. All Rights Reserved.
+
+#pragma once
+
+#include "AST.h"
+
+
+namespace mu
+{
+
+
+	//---------------------------------------------------------------------------------------------
+	//! Variable sized switch operation.
+	//---------------------------------------------------------------------------------------------
+	class ASTOpSwitch : public ASTOp
+	{
+	public:
+
+		//! Type of switch
+		OP_TYPE type;
+
+		//! Variable whose value will be used to choose the switch branch.
+		ASTChild variable;
+
+		//! Default branch in case none matches the value in the variable.
+		ASTChild def;
+
+		struct CASE
+		{
+			CASE(int32_t cond, Ptr<ASTOp> parent, Ptr<ASTOp> b)
+				: condition(cond)
+				, branch(parent.get(), b)
+			{
+			}
+
+			int32_t condition;
+			ASTChild branch;
+
+			//!
+			bool operator==(const CASE& other) const
+			{
+				return condition == other.condition &&
+					branch == other.branch;
+			}
+		};
+
+		vector<CASE> cases;
+
+	public:
+
+		ASTOpSwitch();
+		ASTOpSwitch(const ASTOpSwitch&) = delete;
+		~ASTOpSwitch() override;
+
+		// ASTOp interface
+		OP_TYPE GetOpType() const override { return type; }
+		uint64 Hash() const override;
+		void ForEachChild(const std::function<void(ASTChild&)>& f) override;
+		bool IsEqual(const ASTOp& otherUntyped) const override;
+		Ptr<ASTOp> Clone(MapChildFunc& mapChild) const override;
+		void Assert() override;
+		void Link(PROGRAM& program, const FLinkerOptions* Options) override;
+		FImageDesc GetImageDesc(bool returnBestOption,
+			class GetImageDescContext* context) override;
+		void GetBlockLayoutSize(int blockIndex, int* pBlockX, int* pBlockY,
+			BLOCK_LAYOUT_SIZE_CACHE* cache) override;
+		void GetLayoutBlockSize(int* pBlockX, int* pBlockY) override;
+		Ptr<ASTOp> OptimiseSemantic(const MODEL_OPTIMIZATION_OPTIONS&) const override;
+		bool GetNonBlackRect(FImageRect& maskUsage) const override;
+		bool IsImagePlainConstant(vec4<float>& colour) const override;
+		Ptr<ImageSizeExpression> GetImageSizeExpression() const override;
+
+		// Own interface
+
+		//!
+		Ptr<ASTOp> GetFirstValidValue();
+
+		//! Return true if the two switches have the same condition, amd case value (but not
+		//! necessarily branches) or operation type.
+		bool IsCompatibleWith(const ASTOpSwitch* other) const;
+
+		//!
+		Ptr<ASTOp> FindBranch(int32_t condition) const;
+
+	};
+
+
+}
+
