[mutable] Move Mutable plugin from the restricted folder to the experimental plugins folder.

diff --git a/Engine/Plugins/Experimental/Mutable/Source/MutableTools/Private/CodeGenerator.cpp b/Engine/Plugins/Experimental/Mutable/Source/MutableTools/Private/CodeGenerator.cpp
new file mode 100644
--- /dev/null
+++ b/Engine/Plugins/Experimental/Mutable/Source/MutableTools/Private/CodeGenerator.cpp
@@ -0,0 +1,2433 @@
+// Copyright Epic Games, Inc. All Rights Reserved.
+
+#include "CodeGenerator.h"
+
+#include "CodeGenerator_SecondPass.h"
+#include "CodeOptimiser.h"
+#include "Platform.h"
+#include "CompilerPrivate.h"
+#include "ImageDescGenerator.h"
+#include "TablePrivate.h"
+#include "NodeLayoutPrivate.h"
+
+#include "NodeBoolPrivate.h"
+
+#include "NodeRangeFromScalarPrivate.h"
+
+#include "NodeProjectorPrivate.h"
+
+#include "NodeMeshTablePrivate.h"
+#include "NodeImageTablePrivate.h"
+
+#include "NodeSurfaceNewPrivate.h"
+#include "NodeSurfaceEditPrivate.h"
+#include "NodeSurfaceVariationPrivate.h"
+
+#include "NodeComponentPrivate.h"
+#include "NodeComponentNewPrivate.h"
+#include "NodeComponentEditPrivate.h"
+
+#include "NodeLODPrivate.h"
+#include "NodeObjectNewPrivate.h"
+#include "NodeObjectStatePrivate.h"
+#include "NodeObjectGroupPrivate.h"
+
+#include "NodeImageFormatPrivate.h"
+#include "NodeImageMipmapPrivate.h"
+
+#include "NodeMeshClipMorphPlane.h"
+#include "NodeMeshClipWithMesh.h"
+#include "NodeMeshClipDeform.h"
+#include "NodeMeshGeometryOperation.h"
+#include "NodeMeshReshape.h"
+
+#include "MutableTools/Private/NodeModifierPrivate.h"
+#include "MutableTools/Private/NodeModifierMeshClipMorphPlanePrivate.h"
+#include "MutableTools/Private/NodeModifierMeshClipWithMeshPrivate.h"
+#include "MutableTools/Private/NodeModifierMeshClipDeformPrivate.h"
+
+#include "MutableTools/Private/ASTOpImagePixelFormat.h"
+#include "MutableTools/Private/ASTOpImageMipmap.h"
+#include "MutableTools/Private/ASTOpImageCompose.h"
+#include "MutableTools/Private/ASTOpImagePatch.h"
+#include "MutableTools/Private/ASTOpConditional.h"
+#include "MutableTools/Private/ASTOpInstanceAdd.h"
+#include "MutableTools/Private/ASTOpConstantBool.h"
+#include "MutableTools/Private/ASTOpConstantResource.h"
+#include "MutableTools/Private/ASTOpMeshRemoveMask.h"
+#include "MutableTools/Private/ASTOpMeshClipDeform.h"
+#include "MutableTools/Private/ASTOpMeshBindShape.h"
+#include "MutableTools/Private/ASTOpMeshClipMorphPlane.h"
+#include "MutableTools/Private/ASTOpMeshMaskClipMesh.h"
+#include "MutableTools/Private/ASTOpAddLOD.h"
+
+#include "ModelPrivate.h"
+#include "Operations.h"
+#include "OpImageCrop.h"
+#include "SystemPrivate.h"
+
+#include "MutableTrace.h"
+
+
+namespace mu
+{
+
+    void PartialOptimise( Ptr<ASTOp>& lastCompOp, const CompilerOptions::Private* options )
+    {
+        if (!lastCompOp) return;
+
+        MUTABLE_CPUPROFILER_SCOPE(PartialOptimise);
+
+        //UE_LOG(LogMutableCore, Verbose, " Begin partial optimisation.");
+
+        ASTOpList roots;
+        roots.push_back( lastCompOp );
+
+        //size_t initialOps = ASTOp::CountNodes(roots);
+        //(void)initialOps;
+
+        //AXE_INT_VALUE("Mutable", Verbose, "ast size", (int64_t)ASTOp::CountNodes(roots));
+        //UE_LOG(LogMutableCore, Verbose, " - duplicated data remover");
+        DuplicatedDataRemoverAST( roots );
+        //AXE_INT_VALUE("Mutable", Verbose, "ast size", (int64_t)ASTOp::CountNodes(roots));
+
+        //UE_LOG(LogMutableCore, Verbose, " - duplicated code remover");
+        DuplicatedCodeRemoverAST( roots );
+        //AXE_INT_VALUE("Mutable", Verbose, "ast size", (int64_t)ASTOp::CountNodes(roots));
+
+        //UE_LOG(LogMutableCore, Verbose, " - size optimiser");
+        SizeOptimiserAST( roots );
+        //AXE_INT_VALUE("Mutable", Verbose, "ast size", (int64_t)ASTOp::CountNodes(roots));
+
+        //UE_LOG(LogMutableCore, Verbose, " - semantic optimiser");
+        SemanticOptimiserAST( roots, options->m_optimisationOptions );
+        //AXE_INT_VALUE("Mutable", Verbose, "ast size", (int64_t)ASTOp::CountNodes(roots));
+
+        //UE_LOG(LogMutableCore, Verbose, " - sink optimiser");
+        SinkOptimiserAST( roots, options->m_optimisationOptions );
+        //AXE_INT_VALUE("Mutable", Verbose, "ast size", (int64_t)ASTOp::CountNodes(roots));
+
+        // Image size operations are treated separately
+        //UE_LOG(LogMutableCore, Verbose, " - size optimiser");
+        SizeOptimiserAST( roots );
+        //AXE_INT_VALUE("Mutable", Verbose, "ast size", (int64_t)ASTOp::CountNodes(roots));
+
+        //UE_LOG(LogMutableCore, Verbose, " - semantic optimiser");
+        SemanticOptimiserAST( roots, options->m_optimisationOptions );
+        //AXE_INT_VALUE("Mutable", Verbose, "ast size", (int64_t)ASTOp::CountNodes(roots));
+
+        //UE_LOG(LogMutableCore, Verbose, " - duplicated data remover");
+        DuplicatedDataRemoverAST( roots );
+        //AXE_INT_VALUE("Mutable", Verbose, "ast size", (int64_t)ASTOp::CountNodes(roots));
+
+        //UE_LOG(LogMutableCore, Verbose, " - duplicated code remover");
+        DuplicatedCodeRemoverAST( roots );
+        //AXE_INT_VALUE("Mutable", Verbose, "ast size", (int64_t)ASTOp::CountNodes(roots));
+
+        // Constant subtree generation
+        ConstantGeneratorAST( options, *roots.begin(), nullptr );
+
+        //UE_LOG(LogMutableCore, Verbose, " - duplicated data remover");
+        DuplicatedDataRemoverAST( roots );
+        //AXE_INT_VALUE("Mutable", Verbose, "ast size", (int64_t)ASTOp::CountNodes(roots));
+
+        //UE_LOG(LogMutableCore, Verbose, " - duplicated code remover");
+        DuplicatedCodeRemoverAST( roots );
+        //AXE_INT_VALUE("Mutable", Verbose, "ast size", (int64_t)ASTOp::CountNodes(roots));
+
+        //UE_LOG(LogMutableCore, Verbose, " - logic optimiser");
+        LocalLogicOptimiserAST( roots );
+        //AXE_INT_VALUE("Mutable", Verbose, "ast size", (int64_t)ASTOp::CountNodes(roots));
+
+        //size_t finalOps = ASTOp::CountNodes(roots);
+        //(void)finalOps;
+
+        //lastCompOp = *roots.begin();
+
+        //UE_LOG(LogMutableCore, Verbose, " End partial optimisation.");
+
+        //UE_LOG(LogMutableCore, Verbose,
+        //        "Partial optimisation from %6d to %6d", int64_t(initialOps), int64_t(finalOps) );
+
+        lastCompOp = roots[0];
+    }
+
+
+    //---------------------------------------------------------------------------------------------
+    class OptimizeTask : public TaskManager::Task
+    {
+    public:
+        // input
+        Ptr<ASTOp> m_source;
+        CompilerOptions::Private* m_options;
+
+        std::shared_timed_mutex* m_codeAccessMutex = nullptr;
+
+    private:
+
+        // output
+        Ptr<ASTOp> m_result;
+
+    public:
+
+        // mu::Task interface
+        void Run() override
+        {
+            // This runs in a random thread
+
+            // Clone the tree
+            m_codeAccessMutex->lock_shared();
+            m_result = ASTOp::DeepClone( m_source );
+            m_codeAccessMutex->unlock_shared();
+
+            // Optimise
+            PartialOptimise( m_result, m_options );
+        }
+
+
+        // TaskManager::Task interface
+        void Complete() override
+        {
+            // This runs in the main thread
+            m_codeAccessMutex->lock();
+            ASTOp::Replace( m_source, m_result );
+            m_codeAccessMutex->unlock();
+        }
+
+    };
+
+
+    //---------------------------------------------------------------------------------------------
+    //---------------------------------------------------------------------------------------------
+    //---------------------------------------------------------------------------------------------
+    CodeGenerator::CodeGenerator( CompilerOptions::Private* options )
+    {
+        m_compilerOptions = options;
+
+        // Create the message log
+        m_pErrorLog = new ErrorLog;
+
+        // Add the parent at the top of the hierarchy
+        m_currentParents.push_back( PARENT_KEY() );
+    }
+
+    //---------------------------------------------------------------------------------------------
+    void CodeGenerator::GenerateRoot(const NodePtrConst pNode, TaskManager* taskManager )
+    {
+        MUTABLE_CPUPROFILER_SCOPE(Generate);
+
+        m_pTaskManager = taskManager;
+
+//        UE_LOG(LogMutableCore,Warning,"astop %9u#  %9u+  %9u- generate root start", size_t(ASTOp::s_count), size_t(ASTOp::s_created), size_t(ASTOp::s_destroyed) );
+
+		// First pass
+        m_firstPass.Generate( m_pErrorLog, pNode->GetBasePrivate(), m_compilerOptions->m_ignoreStates );
+
+//        UE_LOG(LogMutableCore,Warning,"astop %9u#  %9u+  %9u- first pass completed", size_t(ASTOp::s_count), size_t(ASTOp::s_created), size_t(ASTOp::s_destroyed) );
+
+		// Second pass
+		SecondPassGenerator secondPass( &m_firstPass, m_compilerOptions );
+		bool success = secondPass.Generate(m_pErrorLog, pNode->GetBasePrivate() );
+		if (!success)
+		{
+            return;
+		}
+
+//        UE_LOG(LogMutableCore,Warning,"astop %9u#  %9u+  %9u- second pass completed ", size_t(ASTOp::s_count), size_t(ASTOp::s_created), size_t(ASTOp::s_destroyed) );
+
+        // Main pass for each state
+        {
+            MUTABLE_CPUPROFILER_SCOPE(MainPass);
+
+            m_currentStateIndex = 0;
+            for( const auto& s: m_firstPass.m_states )
+            {
+                MUTABLE_CPUPROFILER_SCOPE(MainPassState);
+
+                Ptr<ASTOp> stateRoot = Generate(pNode);
+                m_states.emplace_back( s.first, stateRoot );
+
+                ++m_currentStateIndex;
+            }
+        }
+
+//        UE_LOG(LogMutableCore,Warning,"astop %9u#  %9u+  %9u- before cache clear", size_t(ASTOp::s_count), size_t(ASTOp::s_created), size_t(ASTOp::s_destroyed) );
+
+        // Free caches
+        m_compiled.Reset();
+        m_constantMeshes.clear();
+        m_addedLayouts.clear();
+        m_nodeVariables.clear();
+        m_generatedMeshes.Reset();
+        m_generatedProjectors.Reset();
+        m_generatedRanges.Reset();
+        m_generatedTables.clear();
+        m_firstPass = FirstPassGenerator();
+        m_currentBottomUpState = BOTTOM_UP_STATE();
+        m_overrideLayoutsStack.clear();
+        m_imageState.clear();
+        m_currentParents.clear();
+        m_currentObject.clear();
+        m_additionalComponents.clear();
+
+
+//        UE_LOG(LogMutableCore,Warning,"astop %9u#  %9u+  %9u- after cache clear", size_t(ASTOp::s_count), size_t(ASTOp::s_created), size_t(ASTOp::s_destroyed) );
+
+        // Complete pending tasks
+        {
+            if (m_pTaskManager)
+            {
+                UE_LOG(LogMutableCore,Verbose,TEXT("Waiting for remaining code generation tasks..."));
+                MUTABLE_CPUPROFILER_SCOPE(PendingTasks);
+                m_pTaskManager->CompleteTasks();
+                UE_LOG(LogMutableCore,Verbose, TEXT("End code generation tasks"));
+            }
+        }
+
+        m_pTaskManager = nullptr;
+
+//        UE_LOG(LogMutableCore,Warning,"astop %9u#  %9u+  %9u- exiting generate root", size_t(ASTOp::s_count), size_t(ASTOp::s_created), size_t(ASTOp::s_destroyed) );
+	}
+
+
+	//---------------------------------------------------------------------------------------------
+    Ptr<ASTOp> CodeGenerator::Generate(const NodePtrConst pNode)
+	{
+        if (!pNode)
+        {
+            return nullptr;
+        }
+
+		// Clear bottom-up state
+		m_currentBottomUpState.m_address = nullptr;
+
+        // Temp by-passes while we remove the visitor pattern
+		if (auto ScalarNode = dynamic_cast<const NodeScalar*>(pNode.get()))
+		{
+			SCALAR_GENERATION_RESULT ScalarResult;
+			GenerateScalar(ScalarResult, ScalarNode);
+			return ScalarResult.op;
+		}
+
+		if (auto ColorNode = dynamic_cast<const NodeColour*>(pNode.get()))
+		{
+			COLOR_GENERATION_RESULT Result;
+			GenerateColor(Result, ColorNode);
+			return Result.op;
+		}
+
+		if (auto ImageNode = dynamic_cast<const NodeImage*>(pNode.get()))
+		{
+			IMAGE_GENERATION_RESULT Result;
+			GenerateImage(Result, ImageNode);
+			return Result.op;
+		}
+
+		if (auto MeshNode = dynamic_cast<const NodeMesh*>(pNode.get()))
+		{
+			MESH_GENERATION_RESULT MeshResult;
+			GenerateMesh(MeshResult, MeshNode);
+			return MeshResult.meshOp;
+		}
+
+		if ( auto projNode = dynamic_cast<const NodeProjector*>(pNode.get()) )
+        {
+            PROJECTOR_GENERATION_RESULT ProjResult;
+            GenerateProjector( ProjResult, projNode );
+            return ProjResult.op;
+        }
+
+        if ( auto surfNode = dynamic_cast<const NodeSurfaceNew*>(pNode.get()) )
+        {
+            // This happens only if we generate a node graph that has a NodeSurfaceNew at the root.
+            SURFACE_GENERATION_RESULT surfResult;
+            const vector<FirstPassGenerator::SURFACE::EDIT> edits;
+            GenerateSurface( surfResult, surfNode, edits );
+            return surfResult.surfaceOp;
+        }
+
+        else if ( dynamic_cast<const NodeSurfaceVariation*>(pNode.get()) )
+        {
+            // This happens only if we generate a node graph that has a NodeSurfaceVariation at the root.
+            return nullptr;
+        }
+
+        else if ( dynamic_cast<const NodeSurfaceEdit*>(pNode.get()) )
+        {
+            // This happens only if we generate a node graph that has a NodeSurfaceEdit at the root.
+            return nullptr;
+        }
+
+        else if (  dynamic_cast<const NodeModifier*>(pNode.get()) )
+        {
+            // This happens only if we generate a node graph that has a modifier at the root.
+            return nullptr;
+        }
+
+
+        Ptr<ASTOp> result;
+
+        // See if it was already generated
+		VISITED_MAP_KEY key = GetCurrentCacheKey(pNode);
+        VisitedMap::ValueType* it = m_compiled.Find( key );
+        if ( it )
+        {
+            m_currentBottomUpState = *it;
+            result = m_currentBottomUpState.m_address;
+        }
+        else
+        {
+			result = pNode->GetBasePrivate()->Accept(*this);
+			m_currentBottomUpState.m_address = result;
+			m_compiled.Add(key, m_currentBottomUpState);
+		}
+
+        // debug: expensive check of all code generation
+//        if (result)
+//        {
+//            ASTOpList roots;
+//            roots.push_back(result);
+//            ASTOp::FullAssert(roots);
+//        }
+
+        return result;
+    }
+
+
+    //---------------------------------------------------------------------------------------------
+    void CodeGenerator::GenerateRange( RANGE_GENERATION_RESULT& result,
+                                      NodeRangePtrConst untyped )
+    {
+        if (!untyped)
+        {
+            result = RANGE_GENERATION_RESULT();
+            return;
+        }
+
+        // See if it was already generated
+		VISITED_MAP_KEY key = GetCurrentCacheKey(untyped);
+		GeneratedRangeMap::ValueType* it = m_generatedRanges.Find( key );
+        if ( it  )
+        {
+            result = *it;
+            return;
+        }
+
+
+        // Generate for each different type of node
+        if ( auto fromScalar = dynamic_cast<const NodeRangeFromScalar*>(untyped.get()) )
+        {
+            result = RANGE_GENERATION_RESULT();
+            result.rangeName = fromScalar->GetName();
+            result.sizeOp = Generate( fromScalar->GetSize() );
+        }
+        else
+        {
+            check(false);
+            mu::Halt();
+        }
+
+
+        // Cache the result
+        m_generatedRanges.Add( key, result);
+    }
+
+	
+    //---------------------------------------------------------------------------------------------
+    Ptr<ASTOp> CodeGenerator::GenerateTableVariable( TablePtr pTable, const string& strName )
+    {
+        Ptr<ASTOp> result;
+
+        PARAMETER_DESC param;
+        param.m_name = strName;
+        if ( param.m_name.size()==0 )
+        {
+            param.m_name = pTable->GetName();
+        }
+        param.m_type = PARAMETER_TYPE::T_INT;
+        param.m_defaultValue.m_int = 0;
+
+        // Add the possible values
+        {
+            // See if there is a string column. If there is one, we will use it as names for the
+            // options. Only the first string column will be used.
+            int nameCol = -1;
+            std::size_t cols = pTable->GetPrivate()->m_columns.size();
+            for ( std::size_t c=0; c<cols && nameCol<0; ++c )
+            {
+                if ( pTable->GetPrivate()->m_columns[c].m_type==TCT_STRING )
+                {
+                    nameCol = (int)c;
+                }
+            }
+
+			if (pTable->GetPrivate()->m_NoneOption)
+			{
+				PARAMETER_DESC::INT_VALUE_DESC nullValue;
+				nullValue.m_value = -1;
+				nullValue.m_name = "None";
+				param.m_possibleValues.push_back(nullValue);
+				param.m_defaultValue.m_int = nullValue.m_value;
+			}
+
+            // Add every row
+            std::size_t rows = pTable->GetPrivate()->m_rows.size();
+			for (size_t i = 0; i<rows; ++i)
+            {
+                PARAMETER_DESC::INT_VALUE_DESC value;
+                value.m_value = (int16_t)pTable->GetPrivate()->m_rows[i].m_id;
+
+                if (nameCol>-1)
+                {
+                    value.m_name = pTable->GetPrivate()->m_rows[i].m_values[nameCol].m_string;
+                }
+
+                param.m_possibleValues.push_back( value );
+
+                // The first row is the default one
+                if (i==0)
+                {
+                    param.m_defaultValue.m_int = value.m_value;
+                }
+            }
+        }
+
+        Ptr<ASTOpParameter> op = new ASTOpParameter();
+        op->type = OP_TYPE::NU_PARAMETER;
+        op->parameter = param;
+
+        return op;
+    }
+
+
+    //---------------------------------------------------------------------------------------------
+	Ptr<const Layout> CodeGenerator::AddLayout(Ptr<const Layout> pLayout )
+    {
+        auto it = m_addedLayouts.find( pLayout.get() );
+
+		Ptr<const Layout> pResult;
+        if ( it != m_addedLayouts.end() )
+        {
+            pResult = it->second;
+        }
+        else
+        {
+			Ptr<Layout> pCloned = pLayout->Clone();
+            for (int32 b=0;b< pCloned->m_blocks.Num();++b)
+            {
+				pCloned->m_blocks[b].m_id = m_absoluteLayoutIndex++;
+            }
+            m_addedLayouts[ pLayout.get() ] = pCloned;
+			pResult = pCloned;
+        }
+
+        return pResult;
+    }
+
+
+    //---------------------------------------------------------------------------------------------
+    Ptr<ASTOp> CodeGenerator::GenerateImageBlockPatch( Ptr<ASTOp> blockAd,
+                                                        const NodePatchImage* pPatch,
+                                                        Ptr<ASTOp> conditionAd )
+    {
+        // Blend operation
+         Ptr<ASTOp> blendAd;
+         {
+             MUTABLE_CPUPROFILER_SCOPE(PatchBlend);
+
+             Ptr<ASTOpFixed> op = new ASTOpFixed();
+			 op->op.type = OP_TYPE::IM_LAYER;
+			 op->op.args.ImageLayer.blendType = uint8_t(pPatch->GetPrivate()->m_blendType);
+             op->SetChild(op->op.args.ImageLayer.base, blockAd );
+
+             // When we patch from edit nodes, we want to apply it to all the channels.
+             // \todo: since we can choose the patch function, maybe we want to be able to
+             // select this as well.
+             op->op.args.ImageLayer.flags = pPatch->GetPrivate()->m_applyToAlpha
+                     ? OP::ImageLayerArgs::F_APPLY_TO_ALPHA
+                     : 0;
+
+             NodeImage* pImage = pPatch->GetPrivate()->m_pImage.get();
+             Ptr<ASTOp> blend;
+             if ( pImage )
+             {
+                 blend = Generate( pImage );
+             }
+             else
+             {
+                 blend = GenerateMissingImageCode( "Blend top image", EImageFormat::IF_RGB_UBYTE,
+                                                   pPatch->GetPrivate()->m_errorContext );
+             }
+             blend = GenerateImageFormat( blend, blockAd->GetImageDesc().m_format );
+             blend = GenerateImageSize
+                 (
+                     blend,
+                     FImageSize( (uint16_t)m_imageState.back().m_imageRect.size[0],
+                                 (uint16_t)m_imageState.back().m_imageRect.size[1] )
+                 );
+             op->SetChild(op->op.args.ImageLayer.blended, blend );
+
+             NodeImage* pMask = pPatch->GetPrivate()->m_pMask.get();
+             Ptr<ASTOp> mask;
+             if ( pMask )
+             {
+                 mask = Generate( pMask );
+             }
+             else
+             {
+                 // Set the argument default value: affect all pixels.
+                 // TODO: Special operation code without mask
+                 mask = GeneratePlainImageCode( vec3<float>( 1,1,1 ) );
+             }
+             mask = GenerateImageFormat( mask, EImageFormat::IF_L_UBYTE );
+             mask = GenerateImageSize
+                 (
+                     mask,
+                     FImageSize( (uint16_t)m_imageState.back().m_imageRect.size[0],
+                                 (uint16_t)m_imageState.back().m_imageRect.size[1] )
+                 );
+             op->SetChild(op->op.args.ImageLayer.mask, mask );
+
+             blendAd = op;
+         }
+
+         // Condition to enable this patch
+         if (conditionAd)
+         {
+             Ptr<ASTOp> conditionalAd;
+             {
+                 Ptr<ASTOpConditional> op = new ASTOpConditional();
+                 op->type = OP_TYPE::IM_CONDITIONAL;
+                 op->no = blockAd;
+                 op->yes = blendAd;
+                 op->condition = conditionAd;
+                 conditionalAd = op;
+             }
+
+             blockAd = conditionalAd;
+         }
+         else
+         {
+             blockAd = blendAd;
+         }
+
+         return blockAd;
+    }
+
+
+    //---------------------------------------------------------------------------------------------
+    const NodeMeshConstant* FindSourceMesh( const Node* pNode )
+    {
+        const NodeMeshConstant* pResult = nullptr;
+
+        if (!pNode)
+        {
+            return pResult;
+        }
+
+        if ( auto pTypedSN = dynamic_cast<const NodeSurfaceNew*>( pNode ) )
+        {
+            // TODO: 0...
+            pResult = FindSourceMesh( pTypedSN->GetMesh(0).get() );
+        }
+        else if ( auto pTypedSE = dynamic_cast<const NodeSurfaceEdit*>(pNode) )
+		{
+            NodePatchMesh* pPatch = pTypedSE->GetMesh();
+			if (pPatch)
+			{
+				pResult = FindSourceMesh(pPatch->GetAdd());
+			}
+		}
+        else if ( auto pTypedSV = dynamic_cast<const NodeSurfaceVariation*>(pNode) )
+		{
+            if (pTypedSV->GetPrivate()->m_defaultSurfaces.size())
+			{
+                pResult = FindSourceMesh(pTypedSV->GetPrivate()->m_defaultSurfaces[0].get());
+			}
+		}
+        else if ( auto pTypedMI = dynamic_cast<const NodeMeshInterpolate*>( pNode ) )
+        {
+            if ( pTypedMI->GetTargetCount()>0 )
+            {
+                pResult = FindSourceMesh( pTypedMI->GetTarget(0).get() );
+            }
+        }
+        else if ( auto pTypedMM = dynamic_cast<const NodeMeshMorph*>( pNode ) )
+        {
+            pResult = FindSourceMesh( pTypedMM->GetBase().get() );
+        }
+		else if (auto pTypedGO = dynamic_cast<const NodeMeshGeometryOperation*>(pNode))
+		{
+			pResult = FindSourceMesh(pTypedGO->GetMeshA().get());
+		}
+		else if (auto pTypedR = dynamic_cast<const NodeMeshReshape*>(pNode))
+		{
+			pResult = FindSourceMesh(pTypedR->GetBaseMesh().get());
+		}
+		else if ( auto pTypedMF = dynamic_cast<const NodeMeshFormat*>( pNode ) )
+        {
+            if ( pTypedMF->GetSource() )
+            {
+                pResult = FindSourceMesh( pTypedMF->GetSource().get() );
+            }
+        }
+        else if ( auto pTypedMFrag = dynamic_cast<const NodeMeshFragment*>( pNode ) )
+        {
+            if ( pTypedMFrag->GetMesh() )
+            {
+                pResult = FindSourceMesh( pTypedMFrag->GetMesh().get() );
+            }
+        }
+        else if ( auto pCasted = dynamic_cast<const NodeMeshConstant*>( pNode ) )
+        {
+            pResult = pCasted;
+        }
+        else if ( auto pTypedCMP = dynamic_cast<const NodeMeshClipMorphPlane*>(pNode) )
+		{
+            pResult = FindSourceMesh(pTypedCMP->GetSource().get());
+		}
+        else if ( auto pTypedCWM = dynamic_cast<const NodeMeshClipWithMesh*>(pNode) )
+		{
+            pResult = FindSourceMesh(pTypedCWM->GetSource().get());
+		}
+		else
+        {
+            check( false );
+        }
+
+        return pResult;
+    }
+
+
+    //---------------------------------------------------------------------------------------------
+    Ptr<ASTOp> CodeGenerator::Visit( const NodePatchImage::Private& node )
+    {
+        // Get the parent component layout
+        const NodeObjectNew::Private* pParent = nullptr;
+        if (m_currentParents.size()>2)
+        {
+            pParent = (m_currentParents.end()-2)->m_pObject;
+        }
+
+        const NodeLayout* pNodeLayout = nullptr;
+        if (pParent)
+        {
+            pNodeLayout = pParent->GetLayout
+                (
+                    m_currentParents.back().m_lod,
+                    m_currentParents.back().m_component,
+                    m_currentParents.back().m_surface,
+                    m_currentParents.back().m_texture
+                ).get();
+        }
+
+        if (!pNodeLayout)
+        {
+            char buf[256];
+            mutable_snprintf
+                (
+                    buf, 256,
+                    "In object [%s] NodePatchImage couldn't find the layout in parent.",
+                    m_currentParents.rbegin()->m_pObject->m_name.c_str()
+                );
+
+            m_pErrorLog->GetPrivate()->Add( buf, ELMT_ERROR, node.m_errorContext );
+        }
+
+        return 0;
+    }
+
+
+    //---------------------------------------------------------------------------------------------
+    Ptr<ASTOp> CodeGenerator::Visit( const NodeLOD::Private& node )
+    {
+        // Build a series of operations to assemble the component
+        Ptr<ASTOp> lastCompOp;
+
+        // Create the expression for each component in this object
+        // TODO: More components per operation
+        for ( std::size_t t=0; t<node.m_components.size(); ++t )
+        {
+            if ( const NodeComponent* pComponentNode = node.m_components[t].get() )
+            {
+                m_currentParents.back().m_component = (int)t;
+
+                Ptr<ASTOp> componentOp = Generate( pComponentNode );
+
+                if (componentOp)
+                {
+                    check( componentOp->GetOpType() == OP_TYPE::IN_ADDCOMPONENT );
+
+                    auto typedOp = dynamic_cast<ASTOpInstanceAdd*>( componentOp.get() );
+
+                    // Complete the instruction adding the base
+                    typedOp->instance = lastCompOp;
+
+                    lastCompOp = componentOp;
+                }
+            }
+        }
+
+        // Add components from child objects
+        ADDITIONAL_COMPONENT_KEY thisKey;
+        thisKey.m_lod = m_currentParents.back().m_lod;
+        thisKey.m_pObject = m_currentParents.back().m_pObject;
+        auto addIt = m_additionalComponents.find( thisKey );
+        if (addIt!=m_additionalComponents.end())
+        {
+            for( const auto& cop: addIt->second )
+            {
+                // Add the additional components after the main ones, this means higher up in the
+                // op tree.
+
+                // First find the last op in the chain of IN_ADDCOMPONENT operations
+                check( cop->GetOpType() == OP_TYPE::IN_ADDCOMPONENT );
+                auto typedOp = dynamic_cast<ASTOpInstanceAdd*>( cop.get() );
+                auto bottomOp = typedOp;
+                while( bottomOp->instance )
+                {
+                    // Step down
+                    check( bottomOp->instance->GetOpType() == OP_TYPE::IN_ADDCOMPONENT );
+                    bottomOp = dynamic_cast<ASTOpInstanceAdd*>( bottomOp->instance.child().get() );
+                }
+
+                // Chain
+                bottomOp->instance = lastCompOp;
+                lastCompOp = typedOp;
+            }
+        }
+
+        // Store for possible parent objects if necessary
+        // 2 is because there must be a parent and there is always a null element as well.
+        if (lastCompOp && m_currentParents.size()>2)
+        {
+            const auto& parentObjectKey = m_currentParents[ m_currentParents.size()-2 ];
+            ADDITIONAL_COMPONENT_KEY parentKey;
+            parentKey.m_lod = m_currentParents.back().m_lod;
+            parentKey.m_pObject = parentObjectKey.m_pObject;
+            m_additionalComponents[parentKey].push_back( lastCompOp );
+        }
+
+        return lastCompOp;
+    }
+
+
+
+    //---------------------------------------------------------------------------------------------
+    void CodeGenerator::GenerateSurface( SURFACE_GENERATION_RESULT& result,
+                                         NodeSurfaceNewPtrConst surfaceNode,
+                                         const vector<FirstPassGenerator::SURFACE::EDIT>& edits )
+    {
+        //MUTABLE_CPUPROFILER_SCOPE(GenerateSurface);
+
+        const auto& node = *surfaceNode->GetPrivate();
+
+        // Clear the surface generation state
+//        m_addedLayouts.clear();
+//        m_absoluteLayoutIndex = 0;
+//        m_currentLayoutMesh = nullptr;
+//        m_currentLayoutChannel = 0;
+//        m_generatedMeshes.clear();
+
+        // Build a series of operations to assemble the surface
+        Ptr<ASTOp> lastSurfOp;
+
+        // Generate the mesh
+        //------------------------------------------------------------------------
+        NodeMeshPtr pMesh;
+        if ( node.m_meshes.size()>0 )
+        {
+            pMesh = node.m_meshes[0].m_pMesh;
+
+            // TODO: Support more meshes
+            check( node.m_meshes.size()==1 );
+        }
+
+
+        MESH_GENERATION_RESULT meshResults;
+
+
+		// We don't add the mesh here, since it will be added directly at the top of the
+		// component expression in the NodeComponentNew visitor with the right merges
+		// and conditions.
+		// But we store it to be used then.
+
+		// Do we need to generate the mesh? Or was it already generated for state conditions 
+		// accepting the current state?
+		FirstPassGenerator::SURFACE* targetSurface = nullptr;
+		for (vector<FirstPassGenerator::SURFACE>::iterator its = m_firstPass.surfaces.begin()
+			; its != m_firstPass.surfaces.end()
+			; ++its)
+		{
+            if (its->node->GetPrivate() == &node)
+            {
+                // Check state conditions
+                bool surfaceValidForThisState =
+                    m_currentStateIndex >= int(its->stateCondition.size())
+                    ||
+                    its->stateCondition[m_currentStateIndex];
+
+                if (surfaceValidForThisState)
+                {
+                    if (its->resultSurfaceOp)
+                    {
+                        // Reuse the entire surface
+                        result.surfaceOp = its->resultSurfaceOp;
+                        return;
+                    }
+                    else
+                    {
+                        // Not already generated, we will generate this
+                        targetSurface = &(*its);
+                    }
+                }
+            }
+		}
+
+        if (!targetSurface)
+        {
+            return;
+        }
+
+        if ( pMesh )
+        {
+            MUTABLE_CPUPROFILER_SCOPE(SurfaceMesh);
+
+            Ptr<ASTOp> lastMeshOp;
+
+			// Store the data necessary to apply modifiers for the pre-normal operations stage.
+			m_activeTags.push_back(node.m_tags);
+
+            // Generate the mesh
+            GenerateMesh( meshResults, pMesh );
+            lastMeshOp = meshResults.meshOp;
+            meshResults.extraMeshLayouts.resize( edits.size() );
+
+			m_activeTags.pop_back();
+
+            // Let's remember the base mesh of each added mesh, so that we can use them for the
+            // "remove mesh" operations that apply to them instead of the base.
+            std::map<const Node::Private*,Ptr<ASTOp>> baseMeshesForEachAddedMesh;
+
+
+            // Apply mesh merges from child objects "edit surface" nodes
+            for ( size_t editIndex=0; editIndex<edits.size(); ++editIndex )
+            {
+                const auto& e = edits[editIndex];
+
+                if ( e.node->m_pMesh )
+                {
+                    if ( NodeMeshPtr pAdd = e.node->m_pMesh->GetAdd() )
+                    {
+						// Store the data necessary to apply modifiers for the pre-normal operations stage.
+						m_activeTags.push_back(e.node->m_tags);
+						
+						MESH_GENERATION_RESULT addResults;
+                        GenerateMesh( addResults, pAdd );
+
+						m_activeTags.pop_back();
+
+                        baseMeshesForEachAddedMesh[e.node] = addResults.baseMeshOp;
+
+						// Apply the modifier for the post-normal operations stage to the added mesh
+						bool bModifiersForBeforeOperations = false;
+						lastMeshOp = ApplyMeshModifiers(lastMeshOp, e.node->m_tags,
+							bModifiersForBeforeOperations, node.m_errorContext);
+
+                        MESH_GENERATION_RESULT::EXTRA_LAYOUTS data;
+                        data.layouts = addResults.layouts;
+                        data.condition = e.condition;
+                        data.meshFragment = addResults.meshOp;
+                        meshResults.extraMeshLayouts[editIndex] = data;
+
+                        Ptr<ASTOpFixed> mop = new ASTOpFixed();
+                        mop->op.type = OP_TYPE::ME_MERGE;
+                        mop->SetChild(mop->op.args.MeshMerge.base, lastMeshOp );
+                        mop->SetChild(mop->op.args.MeshMerge.added, addResults.meshOp );
+                        // will merge the meshes under the same surface
+                        mop->op.args.MeshMerge.newSurfaceID = 0;
+
+                        // Condition to apply
+                        if (e.condition)
+                        {
+                            Ptr<ASTOpConditional> conditionalAd = new ASTOpConditional();
+                            conditionalAd->type = OP_TYPE::ME_CONDITIONAL;
+                            conditionalAd->no = lastMeshOp;
+                            conditionalAd->yes = mop;
+                            conditionalAd->condition = e.condition;
+                            lastMeshOp = conditionalAd;
+                        }
+                        else
+                        {
+                            lastMeshOp = mop;
+                        }
+                    }
+                }
+            }
+
+
+            // Apply mesh removes from child objects "edit surface" nodes.
+            // "Removes" need to come after "Adds" because some removes may refer to added meshes,
+            // and not the base.
+            Ptr<ASTOpMeshRemoveMask> rop;
+            for ( const auto& e: edits )
+            {
+                if ( e.node->m_pMesh )
+                {
+                    if ( NodeMeshPtr pRemove = e.node->m_pMesh->GetRemove() )
+                    {
+                        MESH_GENERATION_RESULT removeResults;
+                        GenerateMesh( removeResults, pRemove );
+
+                        Ptr<ASTOpFixed> maskOp = new ASTOpFixed();
+                        maskOp->op.type = OP_TYPE::ME_MASKDIFF;
+
+                        // By default, remove from the base
+                        auto removeFrom = meshResults.baseMeshOp;
+                        // See if we want to remove from an added mesh instead.
+                        if ( e.node->m_pParent )
+                        {
+                            auto addedBaseMeshIt = baseMeshesForEachAddedMesh.find(e.node->m_pParent->GetBasePrivate());
+                            if (addedBaseMeshIt!=baseMeshesForEachAddedMesh.end())
+                            {
+                                removeFrom = addedBaseMeshIt->second;
+                            }
+                        }
+
+                        maskOp->SetChild(maskOp->op.args.MeshMaskDiff.source, removeFrom );
+                        maskOp->SetChild(maskOp->op.args.MeshMaskDiff.fragment, removeResults.meshOp );
+
+                        if (!rop)
+                        {
+                            rop = new ASTOpMeshRemoveMask();
+                            rop->source = lastMeshOp;
+                        }
+
+                        rop->AddRemove( e.condition, maskOp );
+                    }
+                }
+            }
+
+            if (rop)
+            {
+                lastMeshOp = rop;
+            }
+
+            // Apply mesh morphs from child objects "edit surface" nodes
+            for ( const auto& e: edits )
+            {
+                if ( NodeMeshPtr pMorph = e.node->m_pMorph )
+                {
+					// Base mesh
+					MESH_GENERATION_RESULT baseMesh;
+					GenerateMesh(baseMesh, pMesh);
+
+                    // Target mesh
+                    MESH_GENERATION_RESULT morphResult;
+                    GenerateMesh( morphResult, pMorph );
+
+					// BaseMorph generation through mesh diff
+					Ptr<ASTOpFixed> diffBase;
+					{
+						Ptr<ASTOpFixed> op = new ASTOpFixed();
+						op->op.type = OP_TYPE::ME_DIFFERENCE;
+						op->SetChild(op->op.args.MeshDifference.base, baseMesh.meshOp);
+						op->SetChild(op->op.args.MeshDifference.target, baseMesh.meshOp);
+
+						// Morphing tex coords here is not supported:
+						// Generating the homogoneous UVs is difficult since we don't have the base
+						// layout yet.                       
+						op->op.args.MeshDifference.ignoreTextureCoords = 1;
+						diffBase = op;
+					}
+
+					// Morph generation through mesh diff
+					Ptr<ASTOp> targetAd = morphResult.meshOp;
+                    Ptr<ASTOpFixed> diffAd;
+                    {
+                        Ptr<ASTOpFixed> op = new ASTOpFixed();
+                        op->op.type = OP_TYPE::ME_DIFFERENCE;
+                        op->SetChild(op->op.args.MeshDifference.base, meshResults.baseMeshOp );
+                        op->SetChild(op->op.args.MeshDifference.target, targetAd );
+					
+                        // Morphing tex coords here is not supported:
+                        // Generating the homogoneous UVs is difficult since we don't have the base
+                        // layout yet.                       
+                        op->op.args.MeshDifference.ignoreTextureCoords = 1;
+                        diffAd = op;
+                    }
+
+                    // Morph operation
+                    Ptr<ASTOp> morphAd;
+                    {
+                        Ptr<ASTOpFixed> op = new ASTOpFixed();
+						op->op.type = OP_TYPE::ME_MORPH2;
+
+						// Factor
+						if (e.node->m_pFactor)
+						{
+							op->SetChild(op->op.args.MeshMorph2.factor, Generate(e.node->m_pFactor));
+						}
+						else
+						{
+							NodeScalarConstantPtr auxNode = new NodeScalarConstant();
+							auxNode->SetValue(1.0f);
+							Ptr<ASTOp> resultNode = Generate(auxNode);
+
+							op->SetChild(op->op.args.MeshMorph2.factor, resultNode);
+						}
+
+						// Base		
+						MESH_GENERATION_RESULT test;
+						GenerateMesh(test, pMesh);
+						op->SetChild(op->op.args.MeshMorph2.base, lastMeshOp);
+
+						// Targets
+						op->SetChild(op->op.args.MeshMorph2.targets[0], diffBase);
+						op->SetChild(op->op.args.MeshMorph2.targets[1], diffAd);
+                        morphAd = op;
+                    }
+
+                    // Condition to apply the morph
+                    if (e.condition)
+                    {
+                        Ptr<ASTOpConditional> conditionalAd = new ASTOpConditional();
+                        conditionalAd->type = OP_TYPE::ME_CONDITIONAL;
+                        conditionalAd->no = lastMeshOp;
+                        conditionalAd->yes = morphAd;
+                        conditionalAd->condition = e.condition;
+                        lastMeshOp = conditionalAd;
+                    }
+                    else
+                    {
+                        lastMeshOp = morphAd;
+                    }
+                }
+            }
+
+			// Apply the modifier for the post-normal operations stage.
+			bool bModifiersForBeforeOperations = false;
+			lastMeshOp = ApplyMeshModifiers( lastMeshOp, node.m_tags, 
+				bModifiersForBeforeOperations, node.m_errorContext);
+
+            // Layouts
+            for ( int l=0; l<(int)meshResults.layouts.size(); ++l )
+            {
+                Ptr<ASTOp> layoutOp = nullptr;
+
+				Ptr<const Layout> pLayout = meshResults.layouts[l].get();
+                if ( pLayout )
+                {
+                    // Create the layout expression
+
+                    // Constant layout to start with
+                    {
+                        Ptr<ASTOpConstantResource> op = new ASTOpConstantResource();
+                        op->type = OP_TYPE::LA_CONSTANT;
+
+                        // Clone and set absolute block ids
+						Ptr<const Layout> pCloned = AddLayout( pLayout );
+
+                        op->SetValue( pCloned, m_compilerOptions->m_optimisationOptions.m_useDiskCache );
+                        layoutOp = op;
+                    }
+
+                    // Add children merged meshes layouts
+                    for ( const auto& data : meshResults.extraMeshLayouts )
+                    {
+                        if ( !data.meshFragment )
+                        {
+                            // No mesh to add, we assume there are no layouts to add either.
+                            check(data.layouts.empty());
+                            continue;
+                        }
+
+                        if ( data.layouts.size() != meshResults.layouts.size() )
+                        {
+                            m_pErrorLog->GetPrivate()->Add
+                                ( "Merged layout has been ignored because the number of layouts"
+                                    " is different.",
+                                  ELMT_ERROR, node.m_errorContext );
+                        }
+                        else
+                        {
+                            // Constant layout to start with
+                            Ptr<ASTOp> layoutFragmentAd;
+                            {
+                                Ptr<ASTOpConstantResource> op = new ASTOpConstantResource();
+                                op->type = OP_TYPE::LA_CONSTANT;
+
+								Ptr<const Layout> pCloned = AddLayout( data.layouts[l] );
+                                op->SetValue( pCloned, m_compilerOptions->m_optimisationOptions.m_useDiskCache );
+
+                                layoutFragmentAd = op;
+                            }
+
+                            // Merge operation
+                            Ptr<ASTOpFixed> mergeAd = new ASTOpFixed();
+                            {
+                                mergeAd->op.type = OP_TYPE::LA_MERGE;
+                                // A normal mesh that we merge
+                                mergeAd->SetChild(mergeAd->op.args.LayoutMerge.base, layoutOp );
+                                mergeAd->SetChild(mergeAd->op.args.LayoutMerge.added, layoutFragmentAd );
+                            }
+
+                            // Condition to apply
+                            if (data.condition)
+                            {
+                                Ptr<ASTOpConditional> conditionalAd = new ASTOpConditional();
+                                conditionalAd->type = OP_TYPE::LA_CONDITIONAL;
+                                conditionalAd->no = layoutOp;
+                                conditionalAd->yes = mergeAd;
+                                conditionalAd->condition = data.condition;
+                                layoutOp = conditionalAd;
+                            }
+                            else
+                            {
+                                layoutOp = mergeAd;
+                            }
+                        }
+                    }
+
+
+                    if ( m_compilerOptions->m_textureLayoutStrategy
+                         ==
+                         CompilerOptions::TextureLayoutStrategy::Pack )
+                    {
+                        // Add layout packing instructions
+                        {
+                            // Make sure we removed unnecessary blocks
+                            Ptr<ASTOpFixed> op = new ASTOpFixed();
+                            op->op.type = OP_TYPE::LA_REMOVEBLOCKS;
+                            op->SetChild(op->op.args.LayoutRemoveBlocks.source, layoutOp );
+                            op->SetChild(op->op.args.LayoutRemoveBlocks.mesh, lastMeshOp );
+                            op->op.args.LayoutRemoveBlocks.meshLayoutIndex = uint8_t( l );
+                            layoutOp = op;
+
+                            // Pack uv blocks
+                            op = new ASTOpFixed();
+                            op->op.type = OP_TYPE::LA_PACK;
+                            op->SetChild(op->op.args.LayoutPack.layout, layoutOp );
+                            layoutOp = op;
+                        }
+
+                        // Create the expression to apply the layout to the mesh
+                        {
+                            Ptr<ASTOpFixed> op = new ASTOpFixed();
+                            op->op.type = OP_TYPE::ME_APPLYLAYOUT;
+                            op->SetChild(op->op.args.MeshApplyLayout.mesh, lastMeshOp );
+                            op->SetChild(op->op.args.MeshApplyLayout.layout, layoutOp );
+                            op->op.args.MeshApplyLayout.channel = (uint16_t)l;
+                            lastMeshOp = op;
+                        }
+                    }
+
+                }
+
+                meshResults.layoutOps.push_back( layoutOp );
+            }
+
+            // Store in the surface for later use.
+            targetSurface->resultMeshOp = lastMeshOp;
+        }
+
+
+        // Create the expression for each texture
+        //------------------------------------------------------------------------
+        for ( std::size_t t=0; t<node.m_images.size(); ++t )
+        {
+            //MUTABLE_CPUPROFILER_SCOPE(SurfaceTexture);
+
+            if ( NodeImagePtr pImageNode = node.m_images[t].m_pImage )
+            {
+                // Any image-specific format or mipmapping needs to be applied at the end
+                NodeImageMipmapPtr mipmapNode;
+                NodeImageFormatPtr formatNode;
+                bool found = false;
+                while ( !found)
+                {
+                    if (NodeImageMipmap* tm = dynamic_cast<NodeImageMipmap*>(pImageNode.get()))
+                    {
+                        if (!mipmapNode) mipmapNode = tm;
+                        pImageNode = tm->GetSource();
+                    }
+                    else if (NodeImageFormat* tf = dynamic_cast<NodeImageFormat*>(pImageNode.get()))
+                    {
+                        if (!formatNode) formatNode = tf;
+                        pImageNode = tf->GetSource();
+                    }
+                    else
+                    {
+                        found = true;
+                    }
+                }
+
+                // Find out the size of the image
+                // TODO: What if the image is empty and everything is added?
+                //       Look in the extending images.
+                FImageDesc desc;
+                if (pImageNode)
+                {
+                    MUTABLE_CPUPROFILER_SCOPE(CalculateImageDesc);
+                    desc = CalculateImageDesc( *pImageNode->GetBasePrivate() );
+                }
+
+				// If the image format doesn't come bottom-up, it may come top-down
+				if (desc.m_format == EImageFormat::IF_NONE && formatNode)
+				{
+					desc.m_format = formatNode->GetFormat();
+				}
+
+				const int layout = node.m_images[t].m_layoutIndex;
+
+				if (desc.m_size[0] == 0 || desc.m_size[1] == 0)
+				{
+					int currentLOD = m_currentParents.back().m_lod;
+					char buf[256];
+					mutable_snprintf(buf, 256, "An image block for [%s] [%s] [%s] at lod [%d] has zero size and will not be generated. ", 
+						node.m_images[t].m_name.c_str(),
+						node.m_images[t].m_materialName.c_str(),
+						node.m_images[t].m_materialParameterName.c_str(),
+						currentLOD
+					);
+					m_pErrorLog->GetPrivate()->Add(buf, ELMT_INFO, node.m_errorContext);
+				}
+
+                else if ( desc.m_format == EImageFormat::IF_NONE )
+                {
+					char buf[256];
+					mutable_snprintf(buf, 256, "An image [%s] has an unidentified pixel format. ", node.m_images[t].m_name.c_str() );
+                    m_pErrorLog->GetPrivate()->Add(buf, ELMT_ERROR, node.m_errorContext);
+                }
+
+                else if ( 
+					// If the layout index has been set to negative, it means we should ignore the layout node.
+					(layout<0)
+					||
+					(m_compilerOptions->m_textureLayoutStrategy == CompilerOptions::TextureLayoutStrategy::None) 
+					)
+                {
+					check( desc.m_format!= EImageFormat::IF_NONE );
+
+                    // Generate the image
+                    IMAGE_STATE newState;
+                    newState.m_imageSize = desc.m_size;
+                    newState.m_imageRect = box<vec2<int>>::FromMinSize( vec2<int>(0,0), desc.m_size );
+                    newState.m_layoutBlock = -1;
+                    newState.m_pLayout = nullptr;
+                    m_imageState.push_back( newState );
+
+                    Ptr<ASTOp> imageAd = Generate( pImageNode );
+
+                    // Look for patches to this block
+                    for( size_t editIndex=0; editIndex<edits.size(); ++editIndex )
+                    {
+                        const auto& e = edits[editIndex];
+                        if ( t<e.node->m_textures.size() )
+                        {
+                            if ( const NodePatchImage* pPatch = e.node->m_textures[t].m_pPatch.get() )
+                            {
+                                imageAd = GenerateImageBlockPatch(imageAd,pPatch,e.condition);
+                            }
+                        }
+
+                    }
+
+                    check(imageAd);
+
+                    m_imageState.pop_back();
+
+                    if (mipmapNode)
+                    {
+                        Ptr<ASTOpImageMipmap> op = new ASTOpImageMipmap();
+                        op->Levels = 0;
+                        op->Source = imageAd;
+                        op->BlockLevels = 0;
+
+						op->AddressMode = mipmapNode->GetPrivate()->m_settings.m_addressMode;
+						op->FilterType = mipmapNode->GetPrivate()->m_settings.m_filterType;
+						op->SharpenFactor = mipmapNode->GetPrivate()->m_settings.m_sharpenFactor;
+						op->DitherMipmapAlpha = mipmapNode->GetPrivate()->m_settings.m_ditherMipmapAlpha;
+                        imageAd = op;
+                    }
+
+                    if (formatNode)
+                    {
+                        Ptr<ASTOpImagePixelFormat> fop = new ASTOpImagePixelFormat();
+                        fop->Format = formatNode->GetPrivate()->m_format;
+                        fop->FormatIfAlpha = formatNode->GetPrivate()->m_formatIfAlpha;
+                        fop->Source = imageAd;
+                        check(fop->Format != EImageFormat::IF_NONE);
+                        imageAd = fop;
+                    }
+
+                    // if (imageAd && m_compilerOptions->m_enablePartialOptimise)
+                    // {
+                    //     if (m_pTaskManager->HasRunner())
+                    //     {
+                    //         OptimizeTask* ot = new OptimizeTask;
+                    //         ot->m_optimizationOptions = &m_compilerOptions->m_optimisationOptions;
+                    //         ot->m_source = imageAd;
+                    //         ot->m_codeAccessMutex = &m_codeAccessMutex;
+                    //         m_pTaskManager->AddTask(ot);
+                    //     }
+                    //     else
+                    //     {
+                    //         PartialOptimise(imageAd, m_compilerOptions->m_optimisationOptions);
+                    //     }
+                    // }
+
+                    Ptr<ASTOpInstanceAdd> op = new ASTOpInstanceAdd();
+                    op->type = OP_TYPE::IN_ADDIMAGE;
+                    op->instance = lastSurfOp;
+                    op->value = imageAd;
+                    op->name = node.m_images[t].m_name;
+
+                    lastSurfOp = op;
+                }
+
+                else if ( m_compilerOptions->m_textureLayoutStrategy
+                          ==
+                          CompilerOptions::TextureLayoutStrategy::Pack )
+                {
+                    if( layout >= (int)meshResults.layouts.size() ||
+                        layout >= (int)meshResults.layoutOps.size() )
+                    {
+                        m_pErrorLog->GetPrivate()->Add("Missing layout in object, or its parent.",
+                                                       ELMT_ERROR, node.m_errorContext);
+                    }
+                    else
+                    {
+                        const Layout* pLayout = m_addedLayouts[ meshResults.layouts[ layout ].get() ].get();
+                        check(pLayout);
+
+                        Ptr<ASTOpInstanceAdd> op = new ASTOpInstanceAdd();
+                        op->type = OP_TYPE::IN_ADDIMAGE;
+                        op->instance = lastSurfOp;
+
+                        // Image
+                        //-------------------------------------
+
+                        // Size of a layout block in pixels
+                        FIntPoint grid = pLayout->GetGridSize();
+
+                        check( desc.m_format!= EImageFormat::IF_NONE );
+
+                        // If the image is too small or not a multiple of the the layout size, 
+						// resize it, but raise a warning.
+						if ( (desc.m_size[0] % grid[0] != 0) || (desc.m_size[1] % grid[1] != 0) )
+						{
+							FImageSize oldSize = desc.m_size;
+							desc.m_size[0] = grid[0] * FMath::Max(1, desc.m_size[0] / grid[0]);
+							desc.m_size[1] = grid[1] * FMath::Max(1, desc.m_size[1] / grid[1]);
+
+							// It is already visible in the error log
+							//UE_LOG( LogMutableCore, Warning, TEXT("A texture [%s] has been resized from [%d x %d] to [%d x %d] because it didn't fit the layout [%d x %d]. "),
+							//	ANSI_TO_TCHAR(node.m_images[t].m_name.c_str()), oldSize[0], oldSize[1], desc.m_size[0], desc.m_size[1], grid[0], grid[1] );
+
+							int currentLOD = m_currentParents.back().m_lod;
+							char buf[256];
+							mutable_snprintf(buf, 256, "A texture [%s] for material [%s] parameter [%s] in LOD [%d] has been resized from [%d x %d] to [%d x %d] because it didn't fit the layout [%d x %d]. ",
+								node.m_images[t].m_name.c_str(), 
+								node.m_images[t].m_materialName.c_str(),
+								node.m_images[t].m_materialParameterName.c_str(),
+								currentLOD,
+								oldSize[0], oldSize[1], desc.m_size[0], desc.m_size[1], grid[0], grid[1] );
+							m_pErrorLog->GetPrivate()->Add(buf, ELMT_INFO, node.m_errorContext);
+						}
+
+                        int blockSizeX = std::max( 1, desc.m_size[0] / grid[0]);
+                        int blockSizeY = std::max( 1, desc.m_size[1] / grid[1]);
+
+						bool bBlocksHaveMips = desc.m_lods > 1;
+
+                        // Start with a blank image
+                        Ptr<ASTOp> imageAd;
+                        {
+                            Ptr<ASTOpFixed> bop = new ASTOpFixed();
+                            bop->op.type = OP_TYPE::IM_BLANKLAYOUT;
+                            bop->SetChild(bop->op.args.ImageBlankLayout.layout, meshResults.layoutOps[ layout ] );
+                            bop->op.args.ImageBlankLayout.blockSize[0] = uint16_t(blockSizeX);
+                            bop->op.args.ImageBlankLayout.blockSize[1] = uint16_t(blockSizeY);
+                            // We support block compression directly here, but not non-block compression
+                            if ( GetImageFormatData(desc.m_format).m_pixelsPerBlockX == 0 )
+                            {
+                                // It's something like RLE
+                                bop->op.args.ImageBlankLayout.format = GetUncompressedFormat(desc.m_format);
+                            }
+                            else
+                            {
+                                // Directly supported
+                                bop->op.args.ImageBlankLayout.format = desc.m_format;
+                            }
+                            bop->op.args.ImageBlankLayout.generateMipmaps = bBlocksHaveMips;
+                            bop->op.args.ImageBlankLayout.mipmapCount = 0;
+                            imageAd = bop;
+                        }
+
+                        for( int b=0; b<pLayout->GetBlockCount(); ++b )
+                        {
+                            // Block in layout grid units
+                            box< vec2<int> > rectInCells;
+                            pLayout->GetBlock
+                                (
+                                    b,
+                                    &rectInCells.min[0], &rectInCells.min[1],
+                                    &rectInCells.size[0], &rectInCells.size[1]
+                                );
+
+                            // Transform to pixels
+                            box< vec2<int> > rect = rectInCells;
+                            rect.min[0] *= blockSizeX;
+                            rect.min[1] *= blockSizeY;
+                            rect.size[0] *= blockSizeX;
+                            rect.size[1] *= blockSizeY;
+
+                            // Generate the image
+                            IMAGE_STATE newState;
+                            newState.m_imageSize = desc.m_size;
+                            newState.m_imageRect = rect;
+                            newState.m_layoutBlock = pLayout->m_blocks[b].m_id;
+                            newState.m_pLayout = pLayout;
+                            m_imageState.push_back( newState );
+
+                            Ptr<ASTOp> blockAd = Generate( pImageNode );
+
+                            // Look for patches to this block
+                            for( size_t editIndex=0; editIndex<edits.size(); ++editIndex )
+                            {
+                                const auto& e = edits[editIndex];
+                                if ( t<e.node->m_textures.size() )
+                                {
+                                    if ( const NodePatchImage* pPatch = e.node->m_textures[t].m_pPatch.get() )
+                                    {
+                                        // Is the current block to be patched?
+                                        if ( std::find( pPatch->GetPrivate()->m_blocks.begin(),
+                                                        pPatch->GetPrivate()->m_blocks.end(),
+                                                        b )
+                                             != pPatch->GetPrivate()->m_blocks.end() )
+                                        {
+                                            blockAd = GenerateImageBlockPatch(blockAd,pPatch,e.condition);
+                                        }
+                                    }
+                                }
+                            }
+
+                            // Compose layout operation
+                            Ptr<ASTOpImageCompose> composeOp = new ASTOpImageCompose();
+                            composeOp->Layout = meshResults.layoutOps[ layout ];
+                            // Direct compression support
+                            //composeOp.args.ImageCompose.base = GenerateImageUncompressed( imageAd );
+                            composeOp->Base = imageAd;
+							EImageFormat baseFormat = imageAd->GetImageDesc().m_format;
+                            composeOp->BlockImage = GenerateImageFormat( blockAd, baseFormat );
+
+                            // Set the absolute block index.
+                            check( pLayout->m_blocks[b].m_id>=0 );
+                            composeOp->BlockIndex = pLayout->m_blocks[b].m_id;
+
+                            imageAd = composeOp;
+
+                            m_imageState.pop_back();
+                        }
+                        check(imageAd);
+
+                        // Apply composition of blocks coming from child objects
+                        for( size_t editIndex=0; editIndex<edits.size(); ++editIndex )
+                        {
+                            const auto& e = edits[editIndex];
+                            if ( t<e.node->m_textures.size() )
+                            {
+                                auto pExtend = e.node->m_textures[t].m_pExtend;
+                                if (pExtend)
+                                {
+                                    if ( layout>=int(meshResults.extraMeshLayouts[editIndex].layouts.size())
+                                         ||
+                                         !meshResults.extraMeshLayouts[editIndex].layouts[layout] )
+                                    {
+                                        char buf[256];
+                                        mutable_snprintf
+                                            (
+                                                buf, 256,
+                                                "Trying to extend a layout that doesn't exist in object [%s].",
+                                                m_currentParents.rbegin()->m_pObject->m_name.c_str()
+                                            );
+
+                                        m_pErrorLog->GetPrivate()->Add( buf, ELMT_ERROR, node.m_errorContext );
+                                    }
+                                    else
+                                    {
+                                        auto pExtendLayoutSource = meshResults.extraMeshLayouts[editIndex].layouts[layout];
+                                        auto pExtendLayout = m_addedLayouts[ pExtendLayoutSource.get() ];
+
+                                        // Find out the size of the image
+                                        // TODO: What if the image is empty and everything is added?
+                                        //       Look in the extending images.
+                                        FImageDesc extendDesc = CalculateImageDesc( *pExtend->GetBasePrivate() );
+
+                                        // Size of a layout block in pixels
+                                        FIntPoint extlayout = pExtendLayout->GetGridSize( );
+
+                                        Ptr<ASTOp> lastBase = imageAd;
+
+                                        for (int b=0; b<pExtendLayout->GetBlockCount(); ++b)
+                                        {
+                                            // Block in layout grid units
+                                            box< vec2<int> > blockRect;
+                                            pExtendLayout->GetBlock
+                                                (
+                                                    b,
+                                                    &blockRect.min[0], &blockRect.min[1],
+                                                    &blockRect.size[0], &blockRect.size[1]
+                                                );
+
+                                            // Transform to pixels
+                                            box< vec2<int> > rect = blockRect;
+                                            rect.min[0] = (rect.min[0] * extendDesc.m_size[0]) / extlayout[0];
+                                            rect.min[1] = (rect.min[1] * extendDesc.m_size[1]) / extlayout[1];
+                                            rect.size[0] = (rect.size[0] * extendDesc.m_size[0]) / extlayout[0];
+                                            rect.size[1] = (rect.size[1] * extendDesc.m_size[1]) / extlayout[1];
+
+                                            // Generate the image block
+                                            IMAGE_STATE newState;
+                                            newState.m_imageSize = extendDesc.m_size;
+                                            newState.m_imageRect = rect;
+                                            newState.m_layoutBlock = pExtendLayout->m_blocks[b].m_id;
+                                            newState.m_pLayout = pExtendLayout;
+                                            m_imageState.push_back( newState );
+
+                                            Ptr<ASTOp> fragmentAd = Generate( pExtend );
+
+                                            m_imageState.pop_back();
+
+                                            // Adjust the format and size of the block to be added
+                                            Ptr<ASTOp> formatted = GenerateImageFormat( fragmentAd, GetUncompressedFormat(desc.m_format) );
+                                            FImageSize expectedSize;
+                                            expectedSize[0] = (uint16_t)(blockSizeX * blockRect.size[0]);
+                                            expectedSize[1] = (uint16_t)(blockSizeY * blockRect.size[1]);
+                                            formatted = GenerateImageSize( formatted, expectedSize );
+                                            // Direct compression support
+                                            // formatted = GenerateImageUncompressed( formatted );
+                                            //formatted = formatted;
+                                            fragmentAd = formatted;
+
+
+                                            // Compose operation
+                                            Ptr<ASTOpImageCompose> composeOp = new ASTOpImageCompose();
+                                            composeOp->Layout = meshResults.layoutOps[layout];
+                                            composeOp->Base = lastBase;
+                                            // Direct compression support
+                                            // composeOp.args.ImageCompose.blockImage = GenerateImageUncompressed( fragmentAd );
+                                            composeOp->BlockImage = fragmentAd;
+
+                                            // Set the absolute block index.
+                                            check( pExtendLayout->m_blocks[b].m_id>=0 );
+                                            composeOp->BlockIndex = pExtendLayout->m_blocks[b].m_id;
+
+                                            lastBase = composeOp;
+                                        }
+
+                                        // Condition to enable this image extension
+                                        if (e.condition)
+                                        {
+                                            Ptr<ASTOp> conditionalAd;
+                                            Ptr<ASTOpConditional> cop = new ASTOpConditional();
+                                            cop->type = OP_TYPE::IM_CONDITIONAL;
+                                            cop->no = imageAd;
+                                            cop->yes = lastBase;
+                                            cop->condition = e.condition;
+                                            conditionalAd = cop;
+                                            imageAd = conditionalAd;
+                                        }
+                                        else
+                                        {
+                                            imageAd = lastBase;
+                                        }
+
+                                    }
+                                }
+                            }
+                        }
+
+
+                        // Apply mipmap and format if necessary
+                        if (mipmapNode)
+                        {
+                            Ptr<ASTOpImageMipmap> mop = new ASTOpImageMipmap();
+
+                            // At the end of the day, we want all the mipmaps. Maybe the code
+                            // optimiser will split the process later.
+                            mop->Levels = 0;
+							mop->bOnlyTail = false;
+							mop->Source = imageAd;
+
+                            // We have to avoid mips smaller than the image format block size, so
+                            // we will devide the layout block by the format block
+                            const FImageFormatData& finfo = GetImageFormatData( desc.m_format );
+
+                            int mipsX = (int)ceilf( logf( (float)blockSizeX/finfo.m_pixelsPerBlockX )/logf(2.0f) );
+                            int mipsY = (int)ceilf( logf( (float)blockSizeY/finfo.m_pixelsPerBlockY )/logf(2.0f) );
+                            mop->BlockLevels = (uint8_t)std::max( mipsX, mipsY );
+
+							mop->AddressMode	   = mipmapNode->GetPrivate()->m_settings.m_addressMode;
+							mop->FilterType		   = mipmapNode->GetPrivate()->m_settings.m_filterType;
+							mop->SharpenFactor	   = mipmapNode->GetPrivate()->m_settings.m_sharpenFactor;
+							mop->DitherMipmapAlpha = mipmapNode->GetPrivate()->m_settings.m_ditherMipmapAlpha;
+
+                            imageAd = mop;
+                        }
+
+						else if (bBlocksHaveMips)
+						{
+							// If the blocks had mipmaps, we still need to generate them after the compose, to have the full chain
+							Ptr<ASTOpImageMipmap> mop = new ASTOpImageMipmap();
+
+							// At the end of the day, we want all the mipmaps. Maybe the code
+							// optimiser will split the process later.
+							mop->Levels = 0;
+							mop->bOnlyTail = false;
+							mop->Source = imageAd;
+
+							// We have to avoid mips smaller than the image format block size, so
+							// we will devide the layout block by the format block
+							const FImageFormatData& finfo = GetImageFormatData(desc.m_format);
+
+							int mipsX = (int)ceilf(logf((float)blockSizeX / finfo.m_pixelsPerBlockX) / logf(2.0f));
+							int mipsY = (int)ceilf(logf((float)blockSizeY / finfo.m_pixelsPerBlockY) / logf(2.0f));
+							mop->BlockLevels = (uint8_t)std::max(mipsX, mipsY);
+
+							// Not important for the end of the mip tail?
+							mop->AddressMode = EAddressMode::AM_CLAMP;
+							mop->FilterType = EMipmapFilterType::MFT_SimpleAverage;
+							mop->SharpenFactor = 0;
+							mop->DitherMipmapAlpha = false;
+
+							imageAd = mop;
+						}
+
+                        if (formatNode)
+                        {
+                            Ptr<ASTOpImagePixelFormat> fop = new ASTOpImagePixelFormat();
+                            fop->Format = formatNode->GetPrivate()->m_format;
+                            fop->FormatIfAlpha = formatNode->GetPrivate()->m_formatIfAlpha;
+                            fop->Source = imageAd;
+                            check(fop->Format != EImageFormat::IF_NONE);
+                            imageAd = fop;
+                        }
+
+                        // if (imageAd && m_compilerOptions->m_enablePartialOptimise)
+                        // {
+                        //     if (m_pTaskManager->HasRunner())
+                        //     {                            
+                        //         OptimizeTask* ot = new OptimizeTask;
+                        //         ot->m_optimizationOptions = &m_compilerOptions->m_optimisationOptions;
+                        //         ot->m_source = imageAd;
+                        //         ot->m_codeAccessMutex = &m_codeAccessMutex;
+                        //         m_pTaskManager->AddTask(ot);
+                        //     }
+                        //     else
+                        //     {
+                        //         PartialOptimise(imageAd, m_compilerOptions->m_optimisationOptions);
+                        //     }
+                        // }
+
+                        op->value = imageAd;
+
+                        // Name
+                        op->name = node.m_images[t].m_name.c_str();
+
+                        lastSurfOp = op;
+                    }
+                }
+
+                else
+                {
+                    // Unimplemented texture layout strategy
+                    check(false);
+                }
+            }
+        }
+
+        // Create the expression for each vector
+        //------------------------------------------------------------------------
+        for (std::size_t t = 0; t<node.m_vectors.size(); ++t)
+        {
+            //MUTABLE_CPUPROFILER_SCOPE(SurfaceVector);
+
+            if (NodeColourPtr pVectorNode = node.m_vectors[t].m_pVector)
+            {
+                Ptr<ASTOpInstanceAdd> op = new ASTOpInstanceAdd();
+                op->type = OP_TYPE::IN_ADDVECTOR;
+                op->instance = lastSurfOp;
+
+                // Vector
+                Ptr<ASTOp> vectorAd = Generate(pVectorNode);
+                op->value = vectorAd;
+
+                // Name
+                op->name = node.m_vectors[t].m_name;
+
+                lastSurfOp = op;
+            }
+        }
+
+        // Create the expression for each scalar
+        //------------------------------------------------------------------------
+        for ( std::size_t t = 0; t < node.m_scalars.size(); ++t )
+        {
+            // MUTABLE_CPUPROFILER_SCOPE(SurfaceScalar);
+
+            if ( NodeScalarPtr pScalarNode = node.m_scalars[t].m_pScalar )
+            {
+                Ptr<ASTOpInstanceAdd> op = new ASTOpInstanceAdd();
+                op->type = OP_TYPE::IN_ADDSCALAR;
+                op->instance = lastSurfOp;
+
+                // Scalar
+                Ptr<ASTOp> scalarAd = Generate( pScalarNode );
+                op->value = scalarAd;
+
+                // Name
+                op->name = node.m_scalars[t].m_name;
+
+                lastSurfOp = op;
+            }
+        }
+
+        // Create the expression for each string
+        //------------------------------------------------------------------------
+        for ( std::size_t t = 0; t < node.m_strings.size(); ++t )
+        {
+            if ( NodeStringPtr pStringNode = node.m_strings[t].m_pString )
+            {
+                Ptr<ASTOpInstanceAdd> op = new ASTOpInstanceAdd();
+                op->type = OP_TYPE::IN_ADDSTRING;
+                op->instance = lastSurfOp;
+
+                Ptr<ASTOp> stringAd = Generate( pStringNode );
+                op->value = stringAd;
+
+                // Name
+                op->name = node.m_strings[t].m_name;
+
+                lastSurfOp = op;
+            }
+        }
+
+
+        // Partial optimisation of the code generated so far
+        if (lastSurfOp && m_compilerOptions->m_enablePartialOptimise)
+        {
+            if (m_pTaskManager->IsConcurrencyEnabled())
+            {
+                OptimizeTask* t = new OptimizeTask;
+                t->m_options = m_compilerOptions;
+                t->m_source = lastSurfOp;
+                t->m_codeAccessMutex = &m_codeAccessMutex;
+                m_pTaskManager->AddTask(t);
+            }
+            else
+            {
+                PartialOptimise(lastSurfOp, m_compilerOptions);
+            }
+        }
+
+
+        result.surfaceOp = lastSurfOp;
+        targetSurface->resultSurfaceOp = lastSurfOp;
+    }
+
+
+    //---------------------------------------------------------------------------------------------
+    Ptr<ASTOp> CodeGenerator::Visit( const NodeComponentNew::Private& node )
+    {
+        MUTABLE_CPUPROFILER_SCOPE(NodeComponentNew);
+
+		// Build a series of operations to assemble the component
+        Ptr<ASTOp> lastCompOp;
+        Ptr<ASTOp> lastMeshOp;
+        string lastMeshName;
+
+        // This generates a different ID for each surface. It can be used to match it to the
+        // mesh surface, or for debugging. It cannot be 0 because it is a special case for the
+        // merge operation.
+        int surfaceID=1;
+
+        // Look for all surfaces that belong to this component
+		for (vector<FirstPassGenerator::SURFACE>::const_iterator its = m_firstPass.surfaces.begin()
+			; its != m_firstPass.surfaces.end()
+            ; ++its, ++surfaceID)
+		{
+			if (its->component==&node)
+			{
+                // Apply state conditions: only generate it if it enabled in this state
+                {
+                    bool enabledInThisState = true;
+                    if (its->stateCondition.size() && m_currentStateIndex >= 0)
+                    {
+                        enabledInThisState =
+                                ( m_currentStateIndex < int(its->stateCondition.size()) )
+                                &&
+                                ( its->stateCondition[m_currentStateIndex] );
+                    }
+                    if (!enabledInThisState)
+                    {
+                        continue;
+                    }
+                }
+
+                Ptr<ASTOpInstanceAdd> sop = new ASTOpInstanceAdd();
+                sop->type = OP_TYPE::IN_ADDSURFACE;
+                sop->name = its->node->GetPrivate()->m_name;
+                sop->instance = lastCompOp;
+
+				SURFACE_GENERATION_RESULT surfaceGenerationResult;
+                GenerateSurface( surfaceGenerationResult, its->node, its->edits );
+                sop->value = surfaceGenerationResult.surfaceOp;
+
+                sop->id = surfaceID;
+                sop->externalId = its->node->GetPrivate()->m_customID;
+                Ptr<ASTOp> surfaceAt = sop;
+
+				// TODO: This could be done earlier?
+                Ptr<ASTOpFixed> surfCondOp = new ASTOpFixed();
+                surfCondOp->op.type = OP_TYPE::BO_AND;
+                surfCondOp->SetChild(surfCondOp->op.args.BoolBinary.a, its->objectCondition );
+                surfCondOp->SetChild(surfCondOp->op.args.BoolBinary.b, its->surfaceCondition );
+                Ptr<ASTOp> surfaceCondition = surfCondOp;
+
+                if (surfaceCondition)
+                {
+                    Ptr<ASTOpConditional> op = new ASTOpConditional();
+                    op->type = OP_TYPE::IN_CONDITIONAL;
+                    op->no = lastCompOp;
+                    op->yes = surfaceAt;
+                    op->condition = surfaceCondition;
+                    lastCompOp = op;
+                }
+                else
+                {
+                    lastCompOp = surfaceAt;
+                }
+
+                // Add the mesh with its condition
+                Ptr<ASTOp> mergeAd;
+                if (!lastMeshOp && its->node->GetPrivate()->m_meshes.size())
+                {
+                    // \todo: mesh index in case of multiple meshes.
+                    const int meshIndex = 0;
+                    lastMeshName = its->node->GetPrivate()->m_meshes[ meshIndex ].m_name;
+                }
+
+                // We add the merge op even for the first mesh, so that we set the surface id.
+                {
+                    Ptr<ASTOp> added = its->resultMeshOp;
+                    // if (added && m_compilerOptions->m_enablePartialOptimise)
+                    // {
+                    //     if (m_pTaskManager->HasRunner())
+                    //     {
+                    //         OptimizeTask* t = new OptimizeTask;
+                    //         t->m_optimizationOptions = &m_compilerOptions->m_optimisationOptions;
+                    //         t->m_source = added;
+                    //         t->m_codeAccessMutex = &m_codeAccessMutex;
+                    //         m_pTaskManager->AddTask(t);
+                    //     }
+                    //     else
+                    //     {
+                    //         PartialOptimise( added, m_compilerOptions->m_optimisationOptions );
+                    //     }
+                    // }
+
+                    Ptr<ASTOpFixed> mop = new ASTOpFixed();
+                    mop->op.type = OP_TYPE::ME_MERGE;
+                    mop->SetChild(mop->op.args.MeshMerge.base, lastMeshOp );
+                    mop->SetChild(mop->op.args.MeshMerge.added, added );
+                    mop->op.args.MeshMerge.newSurfaceID = surfaceID;
+                    mergeAd = mop;
+                }
+
+                if (surfaceCondition)
+                {
+                    Ptr<ASTOpConditional> op = new ASTOpConditional();
+                    op->type = OP_TYPE::ME_CONDITIONAL;
+                    op->no = lastMeshOp;
+                    op->yes = mergeAd;
+                    op->condition = surfaceCondition;
+                    lastMeshOp = op;
+                }
+                else
+                {
+                    lastMeshOp = mergeAd;
+                }
+			}
+		}
+
+        // Add the component mesh
+        {
+            Ptr<ASTOpInstanceAdd> iop = new ASTOpInstanceAdd();
+            iop->type = OP_TYPE::IN_ADDMESH;
+            iop->instance = lastCompOp;
+            iop->value = lastMeshOp;
+
+            // Name
+            // \todo: do mesh names make any sense?
+            iop->name = lastMeshName;
+
+            lastCompOp = iop;
+        }
+
+
+        Ptr<ASTOpInstanceAdd> iop = new ASTOpInstanceAdd();
+        iop->type = OP_TYPE::IN_ADDCOMPONENT;
+        iop->value = lastCompOp;
+        iop->name = node.m_name;
+    	iop->id = node.m_id;
+        lastCompOp = iop;
+
+        return lastCompOp;
+    }
+
+
+    //---------------------------------------------------------------------------------------------
+    Ptr<ASTOp> CodeGenerator::Visit( const NodeComponentEdit::Private& )
+    {
+        MUTABLE_CPUPROFILER_SCOPE(NodeComponentEdit);
+
+        // Nothing to do. Surface information will be already collected in the suitable
+        // parent components furing the first and second passes.
+
+        return nullptr;
+    }
+
+
+    //---------------------------------------------------------------------------------------------
+    Ptr<ASTOp> CodeGenerator::Visit( const NodeObjectState::Private& node )
+    {
+        // Generate the source object
+        Ptr<ASTOp> source;
+
+        if (node.m_pSource)
+        {
+            // Generate the base object where we will add the state.
+            source = Generate( node.m_pSource );
+        }
+
+        return source;
+    }
+
+
+    //---------------------------------------------------------------------------------------------
+    Ptr<ASTOp> CodeGenerator::Visit( const NodeObjectNew::Private& node )
+    {
+        MUTABLE_CPUPROFILER_SCOPE(NodeObjectNew);
+
+        m_currentParents.push_back( PARENT_KEY() );
+        m_currentParents.back().m_pObject = &node;
+
+        // Parse the child objects first, which will accumulate operations in the patching lists
+        for ( std::size_t t=0; t<node.m_children.size(); ++t )
+        {
+            if ( const NodeObject* pChildNode = node.m_children[t].get() )
+            {
+                Ptr<ASTOp> paramOp;
+
+                // If there are parent objects, the condition of this object depends on the
+                // condition of the parent object
+                if ( m_currentObject.size() )
+                {
+                    paramOp = m_currentObject.back().m_condition;
+                }
+                else
+                {
+                    // In case there is no group node, we generate a constant true condition
+                    // This condition will be overwritten by the group nodes.
+                    Ptr<ASTOpConstantBool> op = new ASTOpConstantBool();
+                    op->value = true;
+                    paramOp = op;
+                }
+
+                OBJECT_GENERATION_DATA data;
+                data.m_condition = paramOp;
+                m_currentObject.push_back( data );
+
+                // This op is ignored: everything is stored as patches to apply to the parent when
+                // it is compiled.
+                Generate( pChildNode );
+
+                m_currentObject.pop_back();
+            }
+        }
+
+        // Create the expression for each lod
+        Ptr<ASTOpAddLOD> lodsOp = new ASTOpAddLOD();
+        for ( std::size_t t=0; t<node.m_lods.size(); ++t )
+        {
+            if ( const NodeLOD* pLODNode = node.m_lods[t].get() )
+            {
+                m_currentParents.back().m_lod = (int)t;
+
+                Ptr<ASTOp> lodOp = Generate( pLODNode );
+
+                lodsOp->lods.emplace_back( lodsOp, lodOp );
+            }
+        }
+        Ptr<ASTOp> rootOp = lodsOp;
+
+        m_currentParents.pop_back();
+
+        return rootOp;
+    }
+
+
+    //---------------------------------------------------------------------------------------------
+    Ptr<ASTOp> CodeGenerator::Visit( const NodeObjectGroup::Private& node )
+    {
+		vector<string> usedNames;
+
+        // Parse the child objects first, which will accumulate operations in the patching lists
+        for ( std::size_t t=0; t<node.m_children.size(); ++t )
+        {
+            if ( const NodeObject* pChildNode = node.m_children[t].get() )
+            {
+				// Look for the child condition in the first pass
+                Ptr<ASTOp> conditionOp;
+				bool found = false;
+                for( auto it = m_firstPass.objects.begin(); !found && it != m_firstPass.objects.end(); it++ )
+				{
+					if (it->node == pChildNode->GetBasePrivate())
+					{
+						found = true;
+						conditionOp = it->condition;
+					}
+				}
+
+                // \todo
+                // TrunkStraight_02_BranchTop crash
+                // It may happen with partial compilations?
+                // check(found);
+
+                OBJECT_GENERATION_DATA data;
+                data.m_condition = conditionOp;
+                m_currentObject.push_back( data );
+
+                // This op is ignored: everything is stored as patches to apply to the parent when
+                // it is compiled.
+                Generate( pChildNode );
+
+                m_currentObject.pop_back();
+
+				// Check for duplicated child names
+				const char* strChildName = pChildNode->GetName();
+				if (std::find(usedNames.begin(), usedNames.end(), strChildName)
+					!=
+					usedNames.end() )
+				{
+					char buf[256];
+					mutable_snprintf
+					(
+						buf, 256,
+						"Object group has more than one children with the same name [%s].",
+						strChildName
+					);
+					m_pErrorLog->GetPrivate()->Add(buf, ELMT_WARNING, node.m_errorContext);
+				}
+				else
+				{
+					usedNames.push_back(strChildName);
+				}
+            }
+        }
+
+        return 0;
+    }
+
+    //---------------------------------------------------------------------------------------------
+    Ptr<ASTOp> CodeGenerator::GenerateMissingBoolCode( const char* strWhere,
+                                                      bool value,
+                                                      const void* errorContext )
+    {
+        // Log a warning
+        char buf[256];
+        mutable_snprintf
+            (
+                buf, 256,
+                "Required connection not found: %s",
+                strWhere
+            );
+        m_pErrorLog->GetPrivate()->Add( buf, ELMT_ERROR, errorContext );
+
+        // Create a constant node
+        NodeBoolConstantPtr pNode = new NodeBoolConstant();
+        pNode->SetValue( value );
+
+        Ptr<ASTOp> result = Generate( pNode );
+
+        return result;
+    }
+
+
+	//---------------------------------------------------------------------------------------------
+	void CodeGenerator::GetModifiersFor(
+		const vector<string>& tags,
+		int LOD, bool bModifiersForBeforeOperations,
+		vector<FirstPassGenerator::MODIFIER>& modifiers)
+	{
+        MUTABLE_CPUPROFILER_SCOPE(GetModifiersFor);
+
+		if (tags.size())
+		{
+			for (const auto& m: m_firstPass.modifiers)
+			{
+				// Correct LOD?
+				if (m.lod != LOD)
+				{
+					continue;
+				}
+
+				// Correct stage
+				if (m.node->m_applyBeforeNormalOperations != bModifiersForBeforeOperations)
+				{
+					continue;
+				}
+
+				// Already there?
+				bool alreadyAdded = std::find_if(modifiers.begin(), modifiers.end(),
+					[&m](const FirstPassGenerator::MODIFIER& c) {return c.node == m.node; })
+					!= modifiers.end();
+
+				if (alreadyAdded)
+				{
+					continue;
+				}
+
+				// Matching tags?
+				bool found = false;
+				for (auto it = m.node->m_tags.begin(); !found && it != m.node->m_tags.end(); ++it)
+				{
+					found = std::find(tags.begin(), tags.end(), *it) != tags.end();
+				}
+
+				if (found)
+				{
+					modifiers.push_back(m);
+				}
+			}
+		}
+	}
+
+	//---------------------------------------------------------------------------------------------
+	Ptr<ASTOp> CodeGenerator::ApplyMeshModifiers(
+		const Ptr<ASTOp>& sourceOp,
+		const vector<string>& tags,
+		bool bModifiersForBeforeOperations,
+		const void* errorContext )
+	{
+		Ptr<ASTOp> lastMeshOp = sourceOp;
+
+		// Apply mesh modifiers
+		vector<FirstPassGenerator::MODIFIER> modifiers;
+
+		int currentLOD = m_currentParents.back().m_lod;
+		GetModifiersFor(tags, currentLOD, bModifiersForBeforeOperations, modifiers);
+
+		Ptr<ASTOp> preModifiersMesh = lastMeshOp;
+
+		m_activeTags.push_back({});
+
+		// Process clip-with-mesh modifiers
+		Ptr<ASTOpMeshRemoveMask> removeOp;
+		for (const FirstPassGenerator::MODIFIER& m : modifiers)
+		{
+			if (auto TypedClipNode = dynamic_cast<const NodeModifierMeshClipWithMesh::Private*>(m.node))
+			{
+				Ptr<ASTOpMeshMaskClipMesh> op = new ASTOpMeshMaskClipMesh();
+				op->source = preModifiersMesh;
+
+				// Parameters
+				MESH_GENERATION_RESULT clipResult;
+				GenerateMesh(clipResult, TypedClipNode->m_clipMesh);
+				op->clip = clipResult.meshOp;
+
+				if (!op->clip)
+				{
+					m_pErrorLog->GetPrivate()->Add
+					("Clip mesh has not been generated", ELMT_ERROR, errorContext);
+					continue;
+				}
+
+				Ptr<ASTOp> maskAt = op;
+
+				if (!removeOp)
+				{
+					removeOp = new ASTOpMeshRemoveMask();
+					removeOp->source = lastMeshOp;
+					lastMeshOp = removeOp;
+				}
+
+				Ptr<ASTOpFixed> surfCondOp = new ASTOpFixed();
+				surfCondOp->op.type = OP_TYPE::BO_AND;
+				surfCondOp->SetChild(surfCondOp->op.args.BoolBinary.a, m.objectCondition);
+				surfCondOp->SetChild(surfCondOp->op.args.BoolBinary.b, m.surfaceCondition);
+				Ptr<ASTOp> fullCondition = surfCondOp;
+
+				removeOp->AddRemove(fullCondition, maskAt);
+			}
+		}
+
+
+		// Process clip-morph-plane modifiers
+		for (const FirstPassGenerator::MODIFIER& m : modifiers)
+		{
+			Ptr<ASTOp> modifiedMeshOp;
+
+			if (auto TypedNode = dynamic_cast<const NodeModifierMeshClipMorphPlane::Private*>(m.node))
+			{
+				Ptr<ASTOpMeshClipMorphPlane> op = new ASTOpMeshClipMorphPlane();
+				op->source = lastMeshOp;
+
+				// Morph to an ellipse
+				{
+					SHAPE morphShape;
+					morphShape.type = (uint8_t)SHAPE::Type::Ellipse;
+					morphShape.position = TypedNode->m_origin;
+					morphShape.up = TypedNode->m_normal;
+					// TODO: Move rotation to ellipse rotation reference base instead of passing it directly
+					morphShape.size = vec3f(TypedNode->m_radius1, TypedNode->m_radius2, TypedNode->m_rotation);
+
+					// Generate a "side" vector.
+					// \todo: make generic and move to the vector class
+					{
+						// Generate vector perpendicular to normal for ellipse rotation reference base
+						vec3f aux_base(0.f, 1.f, 0.f);
+
+						if (fabs(dot(TypedNode->m_normal, aux_base)) > 0.95f)
+						{
+							aux_base = vec3f(0.f, 0.f, 1.f);
+						}
+
+						morphShape.side = cross(TypedNode->m_normal, aux_base);
+					}
+					op->morphShape = morphShape;
+				}
+
+				// Selection box
+				if (TypedNode->m_vertexSelectionType == NodeModifierMeshClipMorphPlane::Private::VS_SHAPE)
+				{
+					op->vertexSelectionType = OP::MeshClipMorphPlaneArgs::VS_SHAPE;
+					SHAPE selectionShape;
+					selectionShape.type = (uint8_t)SHAPE::Type::AABox;
+					selectionShape.position = TypedNode->m_selectionBoxOrigin;
+					selectionShape.size = TypedNode->m_selectionBoxRadius;
+					op->selectionShape = selectionShape;
+				}
+				else if (TypedNode->m_vertexSelectionType == NodeModifierMeshClipMorphPlane::Private::VS_BONE_HIERARCHY)
+				{
+					op->vertexSelectionType = OP::MeshClipMorphPlaneArgs::VS_BONE_HIERARCHY;
+					op->vertexSelectionBone = TypedNode->m_vertexSelectionBone;
+					op->vertexSelectionBoneMaxRadius = TypedNode->m_maxEffectRadius;
+				}
+				else
+				{
+					op->vertexSelectionType = OP::MeshClipMorphPlaneArgs::VS_NONE;
+				}
+
+				// Parameters
+				op->dist = TypedNode->m_dist;
+				op->factor = TypedNode->m_factor;
+
+				modifiedMeshOp = op;
+			}
+
+			if (modifiedMeshOp)
+			{
+				Ptr<ASTOpFixed> surfCondOp = new ASTOpFixed();
+				surfCondOp->op.type = OP_TYPE::BO_AND;
+				surfCondOp->SetChild(surfCondOp->op.args.BoolBinary.a, m.objectCondition);
+				surfCondOp->SetChild(surfCondOp->op.args.BoolBinary.b, m.surfaceCondition);
+				Ptr<ASTOp> fullCondition = surfCondOp;
+
+				Ptr<ASTOpConditional> op = new ASTOpConditional();
+				op->type = OP_TYPE::ME_CONDITIONAL;
+				op->no = lastMeshOp;
+				op->yes = modifiedMeshOp;
+				op->condition = fullCondition;
+				lastMeshOp = op;
+			}
+		}
+
+    	// Process clip deform modifiers.
+		for (const FirstPassGenerator::MODIFIER& M : modifiers)
+		{
+			Ptr<ASTOp> ModifiedMeshOp;
+
+			using ModifierType = NodeModifierMeshClipDeform::Private;
+			if ( const ModifierType* TypedClipNode = dynamic_cast<const ModifierType*>(M.node))
+			{
+				Ptr<ASTOpMeshBindShape>  BindOp = new ASTOpMeshBindShape();
+				Ptr<ASTOpMeshClipDeform> ClipOp = new ASTOpMeshClipDeform();
+
+				MESH_GENERATION_RESULT ClipShapeResult;
+				GenerateMesh(ClipShapeResult, TypedClipNode->ClipMesh);
+				ClipOp->ClipShape = ClipShapeResult.meshOp;
+				
+				BindOp->Mesh = lastMeshOp;
+				BindOp->Shape = ClipShapeResult.meshOp; 
+				BindOp->m_discardInvalidBindings = false;
+				BindOp->m_bindingMethod = static_cast<uint32>(TypedClipNode->BindingMethod);
+	
+				ClipOp->Mesh = BindOp;
+
+				if (!ClipOp->ClipShape)
+				{
+					m_pErrorLog->GetPrivate()->Add
+					("Clip shape mesh has not been generated", ELMT_ERROR, errorContext);
+				}
+				else
+				{
+					ModifiedMeshOp = ClipOp;
+				}
+			}
+			
+			if (ModifiedMeshOp)
+			{
+				Ptr<ASTOpFixed> SurfCondOp = new ASTOpFixed();
+				SurfCondOp->op.type = OP_TYPE::BO_AND;
+				SurfCondOp->SetChild(SurfCondOp->op.args.BoolBinary.a, M.objectCondition);
+				SurfCondOp->SetChild(SurfCondOp->op.args.BoolBinary.b, M.surfaceCondition);
+				Ptr<ASTOp> FullCondition = SurfCondOp;
+
+				Ptr<ASTOpConditional> Op = new ASTOpConditional();
+				Op->type = OP_TYPE::ME_CONDITIONAL;
+				Op->no = lastMeshOp;
+				Op->yes = ModifiedMeshOp;
+				Op->condition = FullCondition;
+				lastMeshOp = Op;
+			}
+		}
+			
+		m_activeTags.pop_back();
+
+		return lastMeshOp;
+	}
+	//---------------------------------------------------------------------------------------------
+//	void CodeGenerator::GetSurfacesWithTag(
+//		const string& tag,
+//		vector<FirstPassGenerator::SURFACE>& surfaces)
+//	{
+//		MUTABLE_CPUPROFILER_SCOPE(GetSurfacesWithTag);
+
+//		for (const auto& m : m_firstPass.surfaces)
+//		{
+//			// Matching tags?
+//            bool found = std::find(m.node->GetPrivate()->m_tags.begin(), m.node->GetPrivate()->m_tags.end(),tag)
+//				!=
+//                m.node->GetPrivate()->m_tags.end();
+
+//			if (found)
+//			{
+//				surfaces.push_back(m);
+//			}
+//		}
+//	}
+
+}
