[mutable] Move Mutable plugin from the restricted folder to the experimental plugins folder.

diff --git a/Engine/Plugins/Experimental/Mutable/Source/MutableTools/Private/CodeGenerator.h b/Engine/Plugins/Experimental/Mutable/Source/MutableTools/Private/CodeGenerator.h
new file mode 100644
--- /dev/null
+++ b/Engine/Plugins/Experimental/Mutable/Source/MutableTools/Private/CodeGenerator.h
@@ -0,0 +1,697 @@
+// Copyright Epic Games, Inc. All Rights Reserved.
+
+#pragma once
+
+#include "MutableTools/Private/Visitor.h"
+
+#include "MutableTools/Private/CodeGenerator_FirstPass.h"
+#include "MutableTools/Private/CompilerPrivate.h"
+#include "MutableTools/Private/NodeLayoutPrivate.h"
+#include "MutableTools/Private/ErrorLogPrivate.h"
+#include "MutableTools/Private/NodeObjectNewPrivate.h"
+#include "MutableTools/Private/TaskManager.h"
+
+#include "MutableTools/Private/TablePrivate.h"
+
+#include "MutableTools/Public/NodeBool.h"
+
+#include "MutableTools/Public/NodeScalar.h"
+#include "MutableTools/Public/NodeScalarConstant.h"
+#include "MutableTools/Public/NodeScalarParameter.h"
+#include "MutableTools/Public/NodeScalarSwitch.h"
+#include "MutableTools/Public/NodeScalarEnumParameter.h"
+#include "MutableTools/Public/NodeScalarCurve.h"
+#include "MutableTools/Public/NodeScalarArithmeticOperation.h"
+#include "MutableTools/Public/NodeScalarVariation.h"
+#include "MutableTools/Public/NodeRangeFromScalar.h"
+#include "MutableTools/Public/NodeScalarTable.h"
+
+#include "MutableTools/Public/NodeString.h"
+#include "MutableTools/Public/NodeStringConstant.h"
+#include "MutableTools/Public/NodeStringParameter.h"
+
+#include "MutableTools/Public/NodeColour.h"
+#include "MutableTools/Public/NodeColourConstant.h"
+#include "MutableTools/Public/NodeColourParameter.h"
+#include "MutableTools/Public/NodeColourSwitch.h"
+#include "MutableTools/Public/NodeColourSampleImage.h"
+#include "MutableTools/Public/NodeColourFromScalars.h"
+#include "MutableTools/Public/NodeColourTable.h"
+#include "MutableTools/Public/NodeColourArithmeticOperation.h"
+#include "MutableTools/Public/NodeColourVariation.h"
+
+#include "MutableTools/Public/NodeProjector.h"
+
+#include "MutableTools/Public/NodeImage.h"
+#include "MutableTools/Public/NodeImageConstant.h"
+#include "MutableTools/Public/NodeImageParameter.h"
+#include "MutableTools/Public/NodeImageDifference.h"
+#include "MutableTools/Public/NodeImageInterpolate.h"
+#include "MutableTools/Public/NodeImageInterpolate3.h"
+#include "MutableTools/Public/NodeImageSaturate.h"
+#include "MutableTools/Public/NodeImageLuminance.h"
+#include "MutableTools/Public/NodeImageLayer.h"
+#include "MutableTools/Public/NodeImageLayerColour.h"
+#include "MutableTools/Public/NodeImageMultiLayer.h"
+#include "MutableTools/Public/NodeImageTable.h"
+#include "MutableTools/Public/NodeImageSwizzle.h"
+#include "MutableTools/Public/NodeImageFormat.h"
+#include "MutableTools/Public/NodeImageSelectColour.h"
+#include "MutableTools/Public/NodeImageColourMap.h"
+#include "MutableTools/Public/NodeImageGradient.h"
+#include "MutableTools/Public/NodeImageBinarise.h"
+#include "MutableTools/Public/NodeImageResize.h"
+#include "MutableTools/Public/NodeImagePlainColour.h"
+#include "MutableTools/Public/NodeImageSwitch.h"
+#include "MutableTools/Public/NodeImageProject.h"
+#include "MutableTools/Public/NodeImageMipmap.h"
+#include "MutableTools/Public/NodeImageConditional.h"
+#include "MutableTools/Public/NodeImageInvert.h"
+#include "MutableTools/Public/NodeImageVariation.h"
+#include "MutableTools/Public/NodeImageNormalComposite.h"
+#include "MutableTools/Public/NodeImageTransform.h"
+
+#include "MutableTools/Public/NodeMesh.h"
+#include "MutableTools/Public/NodeMeshConstant.h"
+#include "MutableTools/Public/NodeMeshSwitch.h"
+#include "MutableTools/Public/NodeMeshTable.h"
+#include "MutableTools/Public/NodeMeshInterpolate.h"
+#include "MutableTools/Public/NodeMeshMorph.h"
+#include "MutableTools/Public/NodeMeshMakeMorph.h"
+#include "MutableTools/Public/NodePatchMesh.h"
+#include "MutableTools/Public/NodeMeshSubtract.h"
+#include "MutableTools/Public/NodeMeshFormat.h"
+#include "MutableTools/Public/NodeMeshFragment.h"
+#include "MutableTools/Public/NodeMeshTransform.h"
+#include "MutableTools/Public/NodeMeshClipMorphPlane.h"
+#include "MutableTools/Public/NodeMeshClipWithMesh.h"
+#include "MutableTools/Public/NodeMeshApplyPose.h"
+#include "MutableTools/Public/NodeMeshVariation.h"
+#include "MutableTools/Public/NodeMeshGeometryOperation.h"
+#include "MutableTools/Public/NodeMeshReshape.h"
+#include "MutableTools/Public/NodeMeshClipDeform.h"
+
+#include "MutableTools/Public/NodeComponent.h"
+#include "MutableTools/Public/NodeComponentNew.h"
+#include "MutableTools/Public/NodeComponentEdit.h"
+
+#include "MutableTools/Public/NodeModifier.h"
+#include "MutableTools/Public/NodeModifierMeshClipMorphPlane.h"
+#include "MutableTools/Public/NodeModifierMeshClipWithMesh.h"
+
+#include "MutableTools/Public/NodeSurface.h"
+#include "MutableTools/Public/NodeSurfaceNew.h"
+#include "MutableTools/Public/NodeSurfaceEdit.h"
+
+#include "MutableTools/Public/NodePatchImage.h"
+#include "MutableTools/Public/NodeLOD.h"
+#include "MutableTools/Public/NodeObjectNew.h"
+#include "MutableTools/Public/NodeObjectState.h"
+#include "MutableTools/Public/NodeObjectGroup.h"
+
+#include "MutableTools/Private/AST.h"
+#include "MutableTools/Private/ASTOpParameter.h"
+#include "MutableTools/Private/ASTOpSwitch.h"
+
+#include "Operations.h"
+#include "ModelPrivate.h"
+#include "ImagePrivate.h"
+
+#include "Containers/Map.h"
+
+#include <shared_mutex>
+
+
+namespace mu
+{
+
+    //---------------------------------------------------------------------------------------------
+    //! Code generator
+    //---------------------------------------------------------------------------------------------
+    class CodeGenerator : public Base,
+                          public BaseVisitor,
+
+                          public Visitor<NodeComponentNew::Private, Ptr<ASTOp>, true>,
+                          public Visitor<NodeComponentEdit::Private, Ptr<ASTOp>, true>,
+                          public Visitor<NodeLOD::Private, Ptr<ASTOp>, true>,
+                          public Visitor<NodeObjectNew::Private, Ptr<ASTOp>, true>,
+                          public Visitor<NodeObjectState::Private, Ptr<ASTOp>, true>,
+                          public Visitor<NodeObjectGroup::Private, Ptr<ASTOp>, true>,
+                          public Visitor<NodePatchImage::Private, Ptr<ASTOp>, true>
+    {
+    public:
+
+        CodeGenerator( CompilerOptions::Private* options );
+
+        //! Data will be stored in m_states
+        void GenerateRoot( const NodePtrConst pNode, TaskManager* );
+
+	protected:
+
+        Ptr<ASTOp> Generate(const NodePtrConst pNode);
+
+	public:
+
+        Ptr<ASTOp> Visit( const NodeComponentNew::Private& ) override;
+        Ptr<ASTOp> Visit( const NodeComponentEdit::Private& ) override;
+        Ptr<ASTOp> Visit( const NodeLOD::Private& ) override;
+        Ptr<ASTOp> Visit( const NodeObjectNew::Private& ) override;
+        Ptr<ASTOp> Visit( const NodeObjectState::Private& ) override;
+        Ptr<ASTOp> Visit( const NodeObjectGroup::Private& ) override;
+        Ptr<ASTOp> Visit( const NodePatchImage::Private& ) override;
+
+    public:
+
+        //! Settings
+        CompilerOptions::Private* m_compilerOptions = nullptr;
+
+		//!
+		FirstPassGenerator m_firstPass;
+
+
+        struct VISITED_MAP_KEY
+        {
+            VISITED_MAP_KEY()
+            {
+            }
+
+			friend FORCEINLINE uint32 GetTypeHash(const VISITED_MAP_KEY& InKey)
+			{
+				uint32 KeyHash = 0;
+				KeyHash = HashCombine(KeyHash, ::GetTypeHash(InKey.pNode.get()));
+				KeyHash = HashCombine(KeyHash, ::GetTypeHash(InKey.imageSize[0]));
+				KeyHash = HashCombine(KeyHash, ::GetTypeHash(InKey.imageSize[1]));
+				KeyHash = HashCombine(KeyHash, ::GetTypeHash(InKey.imageRect.min[0]));
+				KeyHash = HashCombine(KeyHash, ::GetTypeHash(InKey.imageRect.min[1]));
+				KeyHash = HashCombine(KeyHash, ::GetTypeHash(InKey.imageRect.size[0]));
+				KeyHash = HashCombine(KeyHash, ::GetTypeHash(InKey.imageRect.size[1]));
+				KeyHash = HashCombine(KeyHash, ::GetTypeHash(InKey.state));
+				KeyHash = HashCombine(KeyHash, ::GetTypeHash((uint64)InKey.activeTags.size()));
+				return KeyHash;
+			}
+
+			bool operator==(const VISITED_MAP_KEY& InKey) const
+			{
+				if (pNode != InKey.pNode) return false;
+				if (state != InKey.state) return false;
+				if (imageSize != InKey.imageSize) return false;
+				if (imageRect.min != InKey.imageRect.min) return false;
+				if (imageRect.size != InKey.imageRect.size) return false;
+				if (activeTags != InKey.activeTags) return false;
+				if (overrideLayouts != InKey.overrideLayouts) return false;
+				return true;
+			}
+
+            // This reference has to be the smart pointer to avoid memory aliasing, keeping
+            // processed nodes alive.
+            NodePtrConst pNode;
+            vec2<int> imageSize;
+            box< vec2<int> > imageRect;
+            int state = -1;
+			vector<mu::string> activeTags;
+			vector<LayoutPtrConst> overrideLayouts;
+        };
+
+        //! This struct contains additional state propagated from bottom to top of the object node graph.
+        //! It is stored for every visited node, and restored when the cache is used.
+        struct BOTTOM_UP_STATE
+        {
+            //! Generated root address for the node.
+            Ptr<ASTOp> m_address;
+        };
+        BOTTOM_UP_STATE m_currentBottomUpState;
+
+        typedef TMap<VISITED_MAP_KEY,BOTTOM_UP_STATE> VisitedMap;
+        VisitedMap m_compiled;
+
+        //!
+        ErrorLogPtr m_pErrorLog;
+
+        //! While generating code, this contains the index of the state being generated. This
+        //! can only be used with the state data in m_firstPass.
+        int m_currentStateIndex = -1;
+
+        //! After the entire code generation this contains the information about all the states
+        typedef vector< std::pair<OBJECT_STATE, Ptr<ASTOp>> > StateList;
+        StateList m_states;
+
+    private:
+
+        //! List of meshes generated to be able to reuse them
+        vector<MeshPtr> m_constantMeshes;
+
+        //! List of image resources for every image formata that have been generated so far as
+        //! palceholders for missing images.
+        ImagePtr m_missingImage[size_t(EImageFormat::IF_COUNT)];
+
+        //! If this has something, while generating meshes, the layouts will be ignored, because
+        //! they are supposed to match some other set of layouts. If the vector is empty, layouts
+        //! are generated normally.
+        vector< vector<Ptr<const Layout>> > m_overrideLayoutsStack;
+
+        //! Map of layouts found in the code already generated. The map is from the source layout
+        //! pointer to the cloned layout. The cloned layout will have absolute block ids assigned.
+        map<Ptr<const Layout>,Ptr<const Layout>> m_addedLayouts;
+
+        //! First free index for a layout block
+        int32_t m_absoluteLayoutIndex = 0;
+
+        //! First free index to be used to identify mesh vertices.
+        uint32_t m_freeVertexIndex = 0;
+
+        //! First free index to be used to identify mesh masks.
+        uint32_t m_freeMaskIndex = 0;
+
+        //! When generating images, here we have the entire source image size and the rect of the
+        //! image that we are generating.
+        struct IMAGE_STATE
+        {
+            vec2<int> m_imageSize;
+            box< vec2<int> > m_imageRect;
+            int32_t m_layoutBlock;
+            LayoutPtrConst m_pLayout;
+        };
+        vector<IMAGE_STATE> m_imageState;
+
+		// (top-down) Tags that are active when generating nodes.
+		vector< vector<mu::string> > m_activeTags;
+
+        struct PARENT_KEY
+        {
+            PARENT_KEY()
+            {
+                m_pObject = nullptr;
+                m_state = -1;
+                m_lod = -1;
+                m_component = -1;
+                m_surface = -1;
+                m_texture = -1;
+                m_block = -1;
+            }
+
+            const NodeObjectNew::Private* m_pObject;
+            int m_state;
+            int m_lod;
+            int m_component;
+            int m_surface;
+            int m_texture;
+            int m_block;
+        };
+
+        vector< PARENT_KEY > m_currentParents;
+
+        // List of additional components to add to an object that come from child objects.
+        // The index is the object and lod that should receive the components.
+        struct ADDITIONAL_COMPONENT_KEY
+        {
+            ADDITIONAL_COMPONENT_KEY()
+            {
+                m_pObject = nullptr;
+                m_lod = -1;
+            }
+
+            const NodeObjectNew::Private* m_pObject;
+            int m_lod;
+
+            inline bool operator<(const ADDITIONAL_COMPONENT_KEY& o) const
+            {
+                if (m_pObject < o.m_pObject) return true;
+                if (m_pObject > o.m_pObject) return false;
+                return m_lod < o.m_lod;
+            }
+        };
+        std::map< ADDITIONAL_COMPONENT_KEY, vector<Ptr<ASTOp>> > m_additionalComponents;
+
+
+        struct OBJECT_GENERATION_DATA
+        {
+            // Condition that enables a specific object
+            Ptr<ASTOp> m_condition;
+        };
+        vector< OBJECT_GENERATION_DATA > m_currentObject;
+
+        map< std::pair<TablePtr,string>, std::pair<TablePtr,Ptr<ASTOp>> > m_generatedTables;
+
+        //! Variables added for every node
+        map< Ptr<const Node>, Ptr<ASTOpParameter> > m_nodeVariables;
+
+        //! This is a list of nodes that are set to replace the identity nodes found while
+        //! generating code. This is useful when applying object transforms.
+        typedef map< NodePtr, Ptr<ASTOp> > IdentityMap;
+        vector<IdentityMap> m_identityReplacement;
+
+		//-----------------------------------------------------------------------------------------
+
+		// Get the modifiers that have to be applied to elements with a specific tag.
+		void GetModifiersFor(const vector<string>& tags, int LOD,
+			bool bModifiersForBeforeOperations, vector<FirstPassGenerator::MODIFIER>& modifiers);
+
+		// Apply the required mesh modifiers to the given operation.
+		Ptr<ASTOp> ApplyMeshModifiers( const Ptr<ASTOp>& sourceOp, const vector<string>& tags,
+			bool bModifiersForBeforeOperations, const void* errorContext);
+
+		// Get the modifiers that have to be applied to elements with a specific tag.
+        //void GetSurfacesWithTag(const string& tag, vector<FirstPassGenerator::SURFACE>& surfaces);
+
+        //-----------------------------------------------------------------------------------------
+        void PrepareForLayout( LayoutPtrConst pSourceLayout,
+                                  MeshPtr currentLayoutMesh,
+                                  size_t currentLayoutChannel,
+                                  const void* errorContext );
+
+        //-----------------------------------------------------------------------------------------
+        //!
+        Ptr<ASTOp> GenerateTableVariable( TablePtr pTable, const string& strName );
+
+        //!
+        Ptr<ASTOp> GenerateMissingBoolCode( const char* strWhere, bool value,
+                                         const void* errorContext );
+
+        //!
+		template<class NODE_TABLE_PRIVATE, TABLE_COLUMN_TYPE TYPE, OP_TYPE OPTYPE, typename F>
+		Ptr<ASTOp> GenerateTableSwitch( const NODE_TABLE_PRIVATE& node, F&& GenerateOption );
+
+        //!
+        Ptr<ASTOp> GenerateImageBlockPatch( Ptr<ASTOp> blockAd,
+                                             const NodePatchImage* pPatch,
+                                             Ptr<ASTOp> conditionAd );
+
+
+    private:
+
+        std::shared_timed_mutex m_codeAccessMutex;
+
+		//! Generate the key with all the relevant state that is used in generation of operations for a node.
+		VISITED_MAP_KEY GetCurrentCacheKey(const NodePtrConst& InNode) const
+		{
+			VISITED_MAP_KEY key;
+			key.pNode = InNode;
+			key.state = m_currentStateIndex;
+			if (!m_imageState.empty())
+			{
+				key.imageSize = m_imageState.back().m_imageSize;
+				key.imageRect = m_imageState.back().m_imageRect;
+			}
+			if (!m_activeTags.empty())
+			{
+				key.activeTags = m_activeTags.back();
+				if (!m_overrideLayoutsStack.empty())
+				{
+					key.overrideLayouts = m_overrideLayoutsStack.back();
+				}
+			}
+			return key;
+		}
+
+
+		//-----------------------------------------------------------------------------------------
+		// Images
+		struct IMAGE_GENERATION_RESULT
+		{
+			Ptr<ASTOp> op;
+		};
+
+		typedef TMap<VISITED_MAP_KEY, IMAGE_GENERATION_RESULT> GeneratedImagesMap;
+		GeneratedImagesMap m_generatedImages;
+
+		void GenerateImage(IMAGE_GENERATION_RESULT& result, const NodeImagePtrConst& node);
+		void GenerateImage_Constant(IMAGE_GENERATION_RESULT&, const NodeImageConstant*);
+		void GenerateImage_Difference(IMAGE_GENERATION_RESULT&, const NodeImageDifference*);
+		void GenerateImage_Interpolate(IMAGE_GENERATION_RESULT&, const NodeImageInterpolate*);
+		void GenerateImage_Saturate(IMAGE_GENERATION_RESULT&, const NodeImageSaturate*);
+		void GenerateImage_Table(IMAGE_GENERATION_RESULT&, const NodeImageTable*);
+		void GenerateImage_Swizzle(IMAGE_GENERATION_RESULT&, const NodeImageSwizzle*);
+		void GenerateImage_SelectColour(IMAGE_GENERATION_RESULT&, const NodeImageSelectColour*);
+		void GenerateImage_ColourMap(IMAGE_GENERATION_RESULT&, const NodeImageColourMap*);
+		void GenerateImage_Gradient(IMAGE_GENERATION_RESULT&, const NodeImageGradient*);
+		void GenerateImage_Binarise(IMAGE_GENERATION_RESULT&, const NodeImageBinarise*);
+		void GenerateImage_Luminance(IMAGE_GENERATION_RESULT&, const NodeImageLuminance*);
+		void GenerateImage_Layer(IMAGE_GENERATION_RESULT&, const NodeImageLayer*);
+		void GenerateImage_LayerColour(IMAGE_GENERATION_RESULT&, const NodeImageLayerColour*);
+		void GenerateImage_Resize(IMAGE_GENERATION_RESULT&, const NodeImageResize*);
+		void GenerateImage_PlainColour(IMAGE_GENERATION_RESULT&, const NodeImagePlainColour*);
+		void GenerateImage_Interpolate3(IMAGE_GENERATION_RESULT&, const NodeImageInterpolate3*);
+		void GenerateImage_Project(IMAGE_GENERATION_RESULT&, const NodeImageProject*);
+		void GenerateImage_Mipmap(IMAGE_GENERATION_RESULT&, const NodeImageMipmap*);
+		void GenerateImage_Switch(IMAGE_GENERATION_RESULT&, const NodeImageSwitch*);
+		void GenerateImage_Conditional(IMAGE_GENERATION_RESULT&, const NodeImageConditional*);
+		void GenerateImage_Format(IMAGE_GENERATION_RESULT&, const NodeImageFormat*);
+		void GenerateImage_Parameter(IMAGE_GENERATION_RESULT&, const NodeImageParameter*);
+		void GenerateImage_MultiLayer(IMAGE_GENERATION_RESULT&, const NodeImageMultiLayer*);
+		void GenerateImage_Invert(IMAGE_GENERATION_RESULT&, const NodeImageInvert*);
+		void GenerateImage_Variation(IMAGE_GENERATION_RESULT&, const NodeImageVariation*);
+		void GenerateImage_NormalComposite(IMAGE_GENERATION_RESULT&, const NodeImageNormalComposite*);
+		void GenerateImage_Transform(IMAGE_GENERATION_RESULT&, const NodeImageTransform*);
+
+		//!
+		ImagePtr GenerateMissingImage(EImageFormat);
+
+		//!
+		Ptr<ASTOp> GenerateMissingImageCode(const char* strWhere, EImageFormat, const void* errorContext);
+
+		//!
+		Ptr<ASTOp> GeneratePlainImageCode(const vec3<float>& colour);
+
+		//!
+		Ptr<ASTOp> GenerateImageFormat(Ptr<ASTOp>, EImageFormat);
+
+		//!
+		Ptr<ASTOp> GenerateImageUncompressed(Ptr<ASTOp>);
+
+		//!
+		Ptr<ASTOp> GenerateImageSize(Ptr<ASTOp>, FImageSize);
+
+		//!
+		FImageDesc CalculateImageDesc(const Node::Private&);
+
+
+        //-----------------------------------------------------------------------------------------
+        // Meshes
+        typedef TMap<VISITED_MAP_KEY,MESH_GENERATION_RESULT> GeneratedMeshMap;
+        GeneratedMeshMap m_generatedMeshes;
+
+        void GenerateMesh( MESH_GENERATION_RESULT& result, const NodeMeshPtrConst& node);
+        void GenerateMesh_Constant( MESH_GENERATION_RESULT&, const NodeMeshConstant* );
+        void GenerateMesh_Format( MESH_GENERATION_RESULT&, const NodeMeshFormat* );
+        void GenerateMesh_Morph( MESH_GENERATION_RESULT&, const NodeMeshMorph* );
+        void GenerateMesh_MakeMorph( MESH_GENERATION_RESULT&, const NodeMeshMakeMorph* );
+        void GenerateMesh_Fragment( MESH_GENERATION_RESULT&, const NodeMeshFragment* );
+        void GenerateMesh_Interpolate( MESH_GENERATION_RESULT&, const NodeMeshInterpolate* );
+        void GenerateMesh_Switch( MESH_GENERATION_RESULT&, const NodeMeshSwitch* );
+        void GenerateMesh_Subtract( MESH_GENERATION_RESULT&, const NodeMeshSubtract* );
+        void GenerateMesh_Transform( MESH_GENERATION_RESULT&, const NodeMeshTransform* );
+        void GenerateMesh_ClipMorphPlane( MESH_GENERATION_RESULT&, const NodeMeshClipMorphPlane* );
+        void GenerateMesh_ClipWithMesh( MESH_GENERATION_RESULT&, const NodeMeshClipWithMesh* );
+        void GenerateMesh_ApplyPose( MESH_GENERATION_RESULT&, const NodeMeshApplyPose* );
+        void GenerateMesh_Variation( MESH_GENERATION_RESULT&, const NodeMeshVariation* );
+		void GenerateMesh_Table(MESH_GENERATION_RESULT&, const NodeMeshTable*);
+		void GenerateMesh_GeometryOperation(MESH_GENERATION_RESULT&, const NodeMeshGeometryOperation*);
+		void GenerateMesh_Reshape(MESH_GENERATION_RESULT&, const NodeMeshReshape*);
+		void GenerateMesh_ClipDeform(MESH_GENERATION_RESULT&, const NodeMeshClipDeform*);
+
+		void GenerateLayout(MESH_GENERATION_RESULT& result, const NodeLayoutBlocksPtrConst& node, uint32 currentLayoutChannel, const MeshPtr currentLayoutMesh);
+
+		//!
+		Ptr<const Layout> AddLayout(Ptr<const Layout> pLayout);
+
+        //-----------------------------------------------------------------------------------------
+        // Projectors
+        struct PROJECTOR_GENERATION_RESULT
+        {
+            Ptr<ASTOp> op;
+            PROJECTOR_TYPE type;
+        };
+
+        typedef TMap<VISITED_MAP_KEY,PROJECTOR_GENERATION_RESULT> GeneratedProjectorsMap;
+        GeneratedProjectorsMap m_generatedProjectors;
+
+        void GenerateProjector( PROJECTOR_GENERATION_RESULT&, const NodeProjectorPtrConst& );
+        void GenerateProjector_Constant( PROJECTOR_GENERATION_RESULT&, const Ptr<const NodeProjectorConstant>& );
+        void GenerateProjector_Parameter( PROJECTOR_GENERATION_RESULT&, const Ptr<const NodeProjectorParameter>& );
+        void GenerateMissingProjectorCode( PROJECTOR_GENERATION_RESULT&, const void* errorContext );
+
+		//-----------------------------------------------------------------------------------------
+		// Bools
+		struct BOOL_GENERATION_RESULT
+		{
+			Ptr<ASTOp> op;
+		};
+
+		typedef TMap<VISITED_MAP_KEY, BOOL_GENERATION_RESULT> GeneratedBoolsMap;
+		GeneratedBoolsMap m_generatedBools;
+
+		void GenerateBool(BOOL_GENERATION_RESULT&, const NodeBoolPtrConst&);
+		void GenerateBool_Constant(BOOL_GENERATION_RESULT&, const Ptr<const NodeBoolConstant>&);
+		void GenerateBool_Parameter(BOOL_GENERATION_RESULT&, const Ptr<const NodeBoolParameter>&);
+		void GenerateBool_IsNull(BOOL_GENERATION_RESULT&, const Ptr<const NodeBoolIsNull>&);
+		void GenerateBool_Not(BOOL_GENERATION_RESULT&, const Ptr<const NodeBoolNot>&);
+		void GenerateBool_And(BOOL_GENERATION_RESULT&, const Ptr<const NodeBoolAnd>&);
+
+		//-----------------------------------------------------------------------------------------
+		// Scalars
+		struct SCALAR_GENERATION_RESULT
+		{
+			Ptr<ASTOp> op;
+		};
+
+		typedef TMap<VISITED_MAP_KEY, SCALAR_GENERATION_RESULT> GeneratedScalarsMap;
+		GeneratedScalarsMap m_generatedScalars;
+
+		void GenerateScalar(SCALAR_GENERATION_RESULT&, const NodeScalarPtrConst&);
+		void GenerateScalar_Constant(SCALAR_GENERATION_RESULT&, const Ptr<const NodeScalarConstant>&);
+		void GenerateScalar_Parameter(SCALAR_GENERATION_RESULT&, const Ptr<const NodeScalarParameter>&);
+		void GenerateScalar_Switch(SCALAR_GENERATION_RESULT&, const Ptr<const NodeScalarSwitch>&);
+		void GenerateScalar_EnumParameter(SCALAR_GENERATION_RESULT&, const Ptr<const NodeScalarEnumParameter>&);
+		void GenerateScalar_Curve(SCALAR_GENERATION_RESULT&, const Ptr<const NodeScalarCurve>&);
+		void GenerateScalar_Arithmetic(SCALAR_GENERATION_RESULT&, const Ptr<const NodeScalarArithmeticOperation>&);
+		void GenerateScalar_Variation(SCALAR_GENERATION_RESULT&, const Ptr<const NodeScalarVariation>&);
+		void GenerateScalar_Table(SCALAR_GENERATION_RESULT&, const Ptr<const NodeScalarTable>&);
+		Ptr<ASTOp> GenerateMissingScalarCode(const char* strWhere, float value, const void* errorContext);
+
+		//-----------------------------------------------------------------------------------------
+		// Colors
+		struct COLOR_GENERATION_RESULT
+		{
+			Ptr<ASTOp> op;
+		};
+
+		typedef TMap<VISITED_MAP_KEY, COLOR_GENERATION_RESULT> GeneratedColorsMap;
+		GeneratedColorsMap m_generatedColors;
+
+		void GenerateColor(COLOR_GENERATION_RESULT&, const NodeColourPtrConst&);
+		void GenerateColor_Constant(COLOR_GENERATION_RESULT&, const Ptr<const NodeColourConstant>&);
+		void GenerateColor_Parameter(COLOR_GENERATION_RESULT&, const Ptr<const NodeColourParameter>&);
+		void GenerateColor_Switch(COLOR_GENERATION_RESULT&, const Ptr<const NodeColourSwitch>&);
+		void GenerateColor_SampleImage(COLOR_GENERATION_RESULT&, const Ptr<const NodeColourSampleImage>&);
+		void GenerateColor_FromScalars(COLOR_GENERATION_RESULT&, const Ptr<const NodeColourFromScalars>&);
+		void GenerateColor_Arithmetic(COLOR_GENERATION_RESULT&, const Ptr<const NodeColourArithmeticOperation>&);
+		void GenerateColor_Variation(COLOR_GENERATION_RESULT&, const Ptr<const NodeColourVariation>&);
+		void GenerateColor_Table(COLOR_GENERATION_RESULT&, const Ptr<const NodeColourTable>&);
+		Ptr<ASTOp> GenerateMissingColourCode(const char* strWhere, const void* errorContext);
+
+		//-----------------------------------------------------------------------------------------
+		// Strings
+		struct STRING_GENERATION_RESULT
+		{
+			Ptr<ASTOp> op;
+		};
+
+		typedef TMap<VISITED_MAP_KEY, STRING_GENERATION_RESULT> GeneratedStringsMap;
+		GeneratedStringsMap m_generatedStrings;
+
+		void GenerateString(STRING_GENERATION_RESULT&, const NodeStringPtrConst&);
+		void GenerateString_Constant(STRING_GENERATION_RESULT&, const Ptr<const NodeStringConstant>&);
+		void GenerateString_Parameter(STRING_GENERATION_RESULT&, const Ptr<const NodeStringParameter>&);
+
+        //-----------------------------------------------------------------------------------------
+        // Ranges
+        struct RANGE_GENERATION_RESULT
+        {
+            //
+            Ptr<ASTOp> sizeOp;
+
+            //
+            string rangeName;
+
+            //
+            string rangeUID;
+        };
+
+        typedef TMap<VISITED_MAP_KEY,RANGE_GENERATION_RESULT> GeneratedRangeMap;
+        GeneratedRangeMap m_generatedRanges;
+
+        void GenerateRange( RANGE_GENERATION_RESULT& result, Ptr<const NodeRange> node);
+
+
+        //-----------------------------------------------------------------------------------------
+        struct SURFACE_GENERATION_RESULT
+        {
+            Ptr<ASTOp> surfaceOp;
+        };
+
+        void GenerateSurface( SURFACE_GENERATION_RESULT& result,
+                              NodeSurfaceNewPtrConst node,
+                              const vector<FirstPassGenerator::SURFACE::EDIT>& edits );
+
+        TaskManager* m_pTaskManager = nullptr;
+    };
+
+
+    //---------------------------------------------------------------------------------------------
+    //! Analyse the code trying to guess the descriptor of the image genereated by the instruction
+    //! address.
+    //! \param returnBestOption If true, try to resolve ambiguities returning some value.
+    //---------------------------------------------------------------------------------------------
+    extern FImageDesc GetImageDesc( const PROGRAM& program, OP::ADDRESS at,
+                                    bool returnBestOption = false,
+                                    class GetImageDescContext* context=nullptr );
+
+    //!
+    extern void PartialOptimise( Ptr<ASTOp>& op, const CompilerOptions* options );
+    
+
+	
+    //---------------------------------------------------------------------------------------------
+    template<class NODE_TABLE_PRIVATE, TABLE_COLUMN_TYPE TYPE, OP_TYPE OPTYPE, typename F>
+    Ptr<ASTOp> CodeGenerator::GenerateTableSwitch
+        (
+            const NODE_TABLE_PRIVATE& node, F&& GenerateOption
+        )
+    {
+        TablePtr pTable;
+        Ptr<ASTOp> variable;
+
+        map< std::pair<TablePtr,string>, std::pair<TablePtr,Ptr<ASTOp>> >::iterator it
+                = m_generatedTables.find
+                ( std::pair<TablePtr,string>(node.m_pTable,node.m_parameterName) );
+        if ( it!=m_generatedTables.end() )
+        {
+            pTable = it->second.first;
+            variable = it->second.second;
+        }
+
+        if ( !pTable )
+        {
+            // Create the table variable expression
+            pTable = node.m_pTable;
+            variable = GenerateTableVariable( pTable, node.m_parameterName );
+
+            m_generatedTables[ std::pair<TablePtr,string>(node.m_pTable,node.m_parameterName) ] =
+                    std::pair<TablePtr,Ptr<ASTOp>>( pTable, variable );
+        }
+
+        // Verify that the table column is the right type
+        int colIndex = pTable->FindColumn( node.m_columnName.c_str() );
+        if ( colIndex<0 )
+        {
+            m_pErrorLog->GetPrivate()->Add("Table column not found.", ELMT_ERROR, node.m_errorContext);
+            return nullptr;
+        }
+
+        if ( pTable->GetPrivate()->m_columns[ colIndex ].m_type != TYPE )
+        {
+            m_pErrorLog->GetPrivate()->Add("Table column type is not the right type.",
+                                           ELMT_ERROR, node.m_errorContext);
+            return nullptr;
+        }
+
+        // Create the switch to cover all the options
+        Ptr<ASTOp> lastSwitch;
+        std::size_t rows = pTable->GetPrivate()->m_rows.size();
+
+        Ptr<ASTOpSwitch> SwitchOp = new ASTOpSwitch();
+		SwitchOp->type = OPTYPE;
+		SwitchOp->variable = variable;
+		SwitchOp->def = nullptr;
+
+		for (size_t i = 0; i < rows; ++i)
+        {
+            check( pTable->GetPrivate()->m_rows[i].m_id <= 0xFFFF);
+            auto condition = (uint16_t)pTable->GetPrivate()->m_rows[i].m_id;
+            Ptr<ASTOp> Branch = GenerateOption( node, colIndex, (int)i, m_pErrorLog.get() );
+			SwitchOp->cases.push_back(ASTOpSwitch::CASE(condition, SwitchOp, Branch ));
+        }
+
+        return SwitchOp;
+    }
+}
diff --git a/Engine/Plugins/Experimental/Mutable/Source/MutableTools/Private/CodeGenerator_Bool.cpp b/Engine/Plugins/Experimental/Mutable/Source/MutableTools/Private/CodeGenerator_Bool.cpp
new file mode 100644
--- /dev/null
+++ b/Engine/Plugins/Experimental/Mutable/Source/MutableTools/Private/CodeGenerator_Bool.cpp
@@ -0,0 +1,156 @@
+// Copyright Epic Games, Inc. All Rights Reserved.
+
+#include "CodeGenerator.h"
+
+#include "Platform.h"
+#include "CompilerPrivate.h"
+#include "NodeBoolPrivate.h"
+#include "MutableTools/Private/ASTOpConstantBool.h"
+
+#include "ModelPrivate.h"
+#include "Operations.h"
+
+
+using namespace mu;
+
+
+//-------------------------------------------------------------------------------------------------
+void CodeGenerator::GenerateBool( BOOL_GENERATION_RESULT& result, const NodeBoolPtrConst& untyped )
+{
+    if (!untyped)
+    {
+        result = BOOL_GENERATION_RESULT();
+        return;
+    }
+
+    // See if it was already generated
+	VISITED_MAP_KEY key = GetCurrentCacheKey(untyped);
+	GeneratedBoolsMap::ValueType* it = m_generatedBools.Find( key );
+    if ( it )
+    {
+        result = *it;
+        return;
+    }
+
+    // Generate for each different type of node
+    if ( auto Constant = dynamic_cast<const NodeBoolConstant*>(untyped.get()) )
+    {
+        GenerateBool_Constant( result, Constant );
+    }
+    else if ( auto Param = dynamic_cast<const NodeBoolParameter*>(untyped.get()) )
+    {
+        GenerateBool_Parameter( result, Param );
+    }
+	else if (auto Switch = dynamic_cast<const NodeBoolIsNull*>(untyped.get()))
+	{
+		GenerateBool_IsNull(result, Switch);
+	}
+	else if (auto Sample = dynamic_cast<const NodeBoolNot*>(untyped.get()))
+	{
+		GenerateBool_Not(result, Sample);
+	}
+	else if (auto From = dynamic_cast<const NodeBoolAnd*>(untyped.get()))
+	{
+		GenerateBool_And(result, From);
+	}	
+	else
+    {
+        check(false);
+    }
+
+    // Cache the result
+    m_generatedBools.Add( key, result );
+}
+
+
+//-------------------------------------------------------------------------------------------------
+void CodeGenerator::GenerateBool_Constant(BOOL_GENERATION_RESULT& result, const Ptr<const NodeBoolConstant>& Typed )
+{
+    const NodeBoolConstant::Private& node = *Typed->GetPrivate();
+
+	Ptr<ASTOpConstantBool> op = new ASTOpConstantBool();
+	op->value = node.m_value;
+
+    result.op = op;
+}
+
+
+//-------------------------------------------------------------------------------------------------
+void CodeGenerator::GenerateBool_Parameter( BOOL_GENERATION_RESULT& result, const Ptr<const NodeBoolParameter>& Typed)
+{
+	const NodeBoolParameter::Private& node = *Typed->GetPrivate();
+
+	Ptr<ASTOpParameter> op;
+
+	auto it = m_nodeVariables.find(node.m_pNode);
+	if (it == m_nodeVariables.end())
+	{
+		PARAMETER_DESC param;
+		param.m_name = node.m_name;
+		param.m_uid = node.m_uid;
+		param.m_type = PARAMETER_TYPE::T_BOOL;
+		param.m_defaultValue.m_bool = node.m_defaultValue;
+
+		op = new ASTOpParameter();
+		op->type = OP_TYPE::BO_PARAMETER;
+		op->parameter = param;
+
+		// Generate the code for the ranges
+		for (std::size_t a = 0; a < node.m_ranges.size(); ++a)
+		{
+			RANGE_GENERATION_RESULT rangeResult;
+			GenerateRange(rangeResult, node.m_ranges[a]);
+			op->ranges.emplace_back
+			(op.get(), rangeResult.sizeOp, rangeResult.rangeName, rangeResult.rangeUID);
+		}
+
+		m_nodeVariables[node.m_pNode] = op;
+	}
+	else
+	{
+		op = it->second;
+	}
+
+    result.op = op;
+}
+
+
+//-------------------------------------------------------------------------------------------------
+void CodeGenerator::GenerateBool_IsNull(BOOL_GENERATION_RESULT& result, const Ptr<const NodeBoolIsNull>& Typed)
+{
+	const NodeBoolIsNull::Private& node = *Typed->GetPrivate();
+
+	Ptr<ASTOpConstantBool> op = new ASTOpConstantBool();
+	Ptr<ASTOp> source = Generate(node.m_pSource);
+	op->value = !source;
+
+	result.op = op;
+}
+
+
+//-------------------------------------------------------------------------------------------------
+void CodeGenerator::GenerateBool_Not(BOOL_GENERATION_RESULT& result, const Ptr<const NodeBoolNot>& Typed)
+{
+	const NodeBoolNot::Private& node = *Typed->GetPrivate();
+
+	Ptr<ASTOpFixed> op = new ASTOpFixed();
+	op->op.type = OP_TYPE::BO_NOT;
+	op->SetChild(op->op.args.BoolNot.source, Generate(node.m_pSource));
+
+	result.op = op;
+}
+
+
+//-------------------------------------------------------------------------------------------------
+void CodeGenerator::GenerateBool_And(BOOL_GENERATION_RESULT& result, const Ptr<const NodeBoolAnd>& Typed)
+{
+	const NodeBoolAnd::Private& node = *Typed->GetPrivate();
+
+	Ptr<ASTOpFixed> op = new ASTOpFixed();
+	op->op.type = OP_TYPE::BO_AND;
+	op->SetChild(op->op.args.BoolBinary.a, Generate(node.m_pA));
+	op->SetChild(op->op.args.BoolBinary.b, Generate(node.m_pB));
+
+	result.op = op;
+}
+
diff --git a/Engine/Plugins/Experimental/Mutable/Source/MutableTools/Private/CodeGenerator_Color.cpp b/Engine/Plugins/Experimental/Mutable/Source/MutableTools/Private/CodeGenerator_Color.cpp
new file mode 100644
--- /dev/null
+++ b/Engine/Plugins/Experimental/Mutable/Source/MutableTools/Private/CodeGenerator_Color.cpp
@@ -0,0 +1,483 @@
+// Copyright Epic Games, Inc. All Rights Reserved.
+
+#include "CodeGenerator.h"
+
+#include "Platform.h"
+#include "CompilerPrivate.h"
+
+#include "NodeColourPrivate.h"
+#include "NodeColourConstantPrivate.h"
+#include "NodeColourParameterPrivate.h"
+#include "NodeColourSwitchPrivate.h"
+#include "NodeColourSampleImagePrivate.h"
+#include "NodeColourFromScalarsPrivate.h"
+#include "NodeColourTablePrivate.h"
+#include "NodeColourArithmeticOperationPrivate.h"
+#include "NodeColourVariationPrivate.h"
+
+#include "MutableTools/Private/ASTOpConditional.h"
+
+#include "ModelPrivate.h"
+#include "Operations.h"
+
+
+using namespace mu;
+
+
+//-------------------------------------------------------------------------------------------------
+void CodeGenerator::GenerateColor( COLOR_GENERATION_RESULT& result, const NodeColourPtrConst& untyped )
+{
+    if (!untyped)
+    {
+        result = COLOR_GENERATION_RESULT();
+        return;
+    }
+
+    // See if it was already generated
+	VISITED_MAP_KEY key = GetCurrentCacheKey(untyped);
+	GeneratedColorsMap::ValueType* it = m_generatedColors.Find( key );
+    if ( it )
+    {
+        result = *it;
+        return;
+    }
+
+    // Generate for each different type of node
+    if ( auto Constant = dynamic_cast<const NodeColourConstant*>(untyped.get()) )
+    {
+        GenerateColor_Constant( result, Constant );
+    }
+    else if ( auto Param = dynamic_cast<const NodeColourParameter*>(untyped.get()) )
+    {
+        GenerateColor_Parameter( result, Param );
+    }
+	else if (auto Switch = dynamic_cast<const NodeColourSwitch*>(untyped.get()))
+	{
+		GenerateColor_Switch(result, Switch);
+	}
+	else if (auto Sample = dynamic_cast<const NodeColourSampleImage*>(untyped.get()))
+	{
+		GenerateColor_SampleImage(result, Sample);
+	}
+	else if (auto From = dynamic_cast<const NodeColourFromScalars*>(untyped.get()))
+	{
+		GenerateColor_FromScalars(result, From);
+	}
+	else if (auto Arithmetic = dynamic_cast<const NodeColourArithmeticOperation*>(untyped.get()))
+	{
+		GenerateColor_Arithmetic(result, Arithmetic);
+	}
+	else if (auto Variation = dynamic_cast<const NodeColourVariation*>(untyped.get()))
+	{
+		GenerateColor_Variation(result, Variation);
+	}
+	else if (auto Table = dynamic_cast<const NodeColourTable*>(untyped.get()))
+	{
+		GenerateColor_Table(result, Table);
+	}
+	else
+    {
+        check(false);
+    }
+
+    // Cache the result
+    m_generatedColors.Add( key, result );
+}
+
+
+//-------------------------------------------------------------------------------------------------
+void CodeGenerator::GenerateColor_Constant(COLOR_GENERATION_RESULT& result, const Ptr<const NodeColourConstant>& Typed )
+{
+    const NodeColourConstant::Private& node = *Typed->GetPrivate();
+
+	Ptr<ASTOpFixed> op = new ASTOpFixed();
+	op->op.type = OP_TYPE::CO_CONSTANT;
+	op->op.args.ColourConstant.value[0] = node.m_value[0];
+	op->op.args.ColourConstant.value[1] = node.m_value[1];
+	op->op.args.ColourConstant.value[2] = node.m_value[2];
+
+    result.op = op;
+}
+
+
+//-------------------------------------------------------------------------------------------------
+void CodeGenerator::GenerateColor_Parameter( COLOR_GENERATION_RESULT& result, const Ptr<const NodeColourParameter>& Typed)
+{
+    const NodeColourParameter::Private& node = *Typed->GetPrivate();
+
+    Ptr<ASTOpParameter> op;
+
+	auto it = m_nodeVariables.find(node.m_pNode);
+
+	if (it == m_nodeVariables.end())
+	{
+		PARAMETER_DESC param;
+		param.m_name = node.m_name;
+		param.m_uid = node.m_uid;
+		param.m_type = PARAMETER_TYPE::T_COLOUR;
+		param.m_defaultValue.m_colour[0] = node.m_defaultValue[0];
+		param.m_defaultValue.m_colour[1] = node.m_defaultValue[1];
+		param.m_defaultValue.m_colour[2] = node.m_defaultValue[2];
+
+		op = new ASTOpParameter();
+		op->type = OP_TYPE::CO_PARAMETER;
+		op->parameter = param;
+
+		// Generate the code for the ranges
+		for (std::size_t a = 0; a < node.m_ranges.size(); ++a)
+		{
+			RANGE_GENERATION_RESULT rangeResult;
+			GenerateRange(rangeResult, node.m_ranges[a]);
+			op->ranges.emplace_back
+			(op.get(), rangeResult.sizeOp, rangeResult.rangeName, rangeResult.rangeUID);
+		}
+
+		m_nodeVariables[node.m_pNode] = op;
+	}
+	else
+	{
+		op = it->second;
+	}
+
+    result.op = op;
+}
+
+
+//-------------------------------------------------------------------------------------------------
+void CodeGenerator::GenerateColor_Switch(COLOR_GENERATION_RESULT& result, const Ptr<const NodeColourSwitch>& Typed)
+{
+	const NodeColourSwitch::Private& node = *Typed->GetPrivate();
+
+	MUTABLE_CPUPROFILER_SCOPE(NodeColourSwitch);
+
+	if (node.m_options.size() == 0)
+	{
+		// No options in the switch!
+		Ptr<ASTOp> missingOp = GenerateMissingColourCode("Switch option",
+			node.m_errorContext);
+		result.op = missingOp;
+		return;
+	}
+
+	Ptr<ASTOpSwitch> op = new ASTOpSwitch();
+	op->type = OP_TYPE::CO_SWITCH;
+
+	// Variable value
+	if (node.m_pParameter)
+	{
+		op->variable = Generate(node.m_pParameter.get());
+	}
+	else
+	{
+		// This argument is required
+		op->variable = GenerateMissingScalarCode("Switch variable", 0.0f, node.m_errorContext);
+	}
+
+	// Options
+	for (std::size_t t = 0; t < node.m_options.size(); ++t)
+	{
+		Ptr<ASTOp> branch;
+		if (node.m_options[t])
+		{
+			branch = Generate(node.m_options[t].get());
+		}
+		else
+		{
+			// This argument is required
+			branch = GenerateMissingColourCode("Switch option",
+				node.m_errorContext);
+		}
+		op->cases.push_back(ASTOpSwitch::CASE((int16_t)t, op, branch));
+	}
+
+	result.op = op;
+}
+
+
+//-------------------------------------------------------------------------------------------------
+void CodeGenerator::GenerateColor_Variation(COLOR_GENERATION_RESULT& result, const Ptr<const NodeColourVariation>& Typed)
+{
+	const NodeColourVariation::Private& node = *Typed->GetPrivate();
+
+	Ptr<ASTOp> currentOp;
+
+	// Default case
+	if (node.m_defaultColour)
+	{
+		MESH_GENERATION_RESULT branchResults;
+		currentOp = Generate(node.m_defaultColour);
+	}
+	else
+	{
+		// This argument is required
+		currentOp = GenerateMissingColourCode("Variation default", node.m_errorContext);
+	}
+
+	// Process variations in reverse order, since conditionals are built bottom-up.
+	for (int t = int(node.m_variations.size()) - 1; t >= 0; --t)
+	{
+		int tagIndex = -1;
+		const string& tag = node.m_variations[t].m_tag;
+		for (int i = 0; i < int(m_firstPass.m_tags.size()); ++i)
+		{
+			if (m_firstPass.m_tags[i].tag == tag)
+			{
+				tagIndex = i;
+			}
+		}
+
+		if (tagIndex < 0)
+		{
+			char buf[256];
+			mutable_snprintf(buf, 256, "Unknown tag found in image variation [%s].",
+				tag.c_str());
+
+			m_pErrorLog->GetPrivate()->Add(buf, ELMT_WARNING, node.m_errorContext);
+			continue;
+		}
+
+		Ptr<ASTOp> variationOp;
+		if (node.m_variations[t].m_colour)
+		{
+			variationOp = Generate(node.m_variations[t].m_colour);
+		}
+		else
+		{
+			// This argument is required
+			variationOp = GenerateMissingColourCode("Variation option", node.m_errorContext);
+		}
+
+
+		Ptr<ASTOpConditional> conditional = new ASTOpConditional;
+		conditional->type = OP_TYPE::CO_CONDITIONAL;
+		conditional->no = currentOp;
+		conditional->yes = variationOp;
+		conditional->condition = m_firstPass.m_tags[tagIndex].genericCondition;
+
+		currentOp = conditional;
+	}
+
+	result.op = currentOp;
+}
+
+
+//-------------------------------------------------------------------------------------------------
+void CodeGenerator::GenerateColor_SampleImage(COLOR_GENERATION_RESULT& result, const Ptr<const NodeColourSampleImage>& Typed)
+{
+	const NodeColourSampleImage::Private& node = *Typed->GetPrivate();
+
+	// Generate the code
+	Ptr<ASTOpFixed> op = new ASTOpFixed();
+	op->op.type = OP_TYPE::CO_SAMPLEIMAGE;
+
+
+	// Source image
+	Ptr<ASTOp> base;
+	if (Node* pSource = node.m_pImage.get())
+	{
+		// We take whatever size will be produced
+		FImageDesc desc = CalculateImageDesc(*pSource->GetBasePrivate());
+		IMAGE_STATE newState;
+		newState.m_imageSize = desc.m_size;
+		newState.m_imageRect.min[0] = 0;
+		newState.m_imageRect.min[1] = 0;
+		newState.m_imageRect.size = desc.m_size;
+		m_imageState.push_back(newState);
+
+		// Generate
+		base = Generate(pSource);
+
+		// Restore rect
+		m_imageState.pop_back();
+	}
+	else
+	{
+		// This argument is required
+		base = GenerateMissingImageCode("Sample image", EImageFormat::IF_RGB_UBYTE, node.m_errorContext);
+	}
+	base = GenerateImageFormat(base, EImageFormat::IF_RGB_UBYTE);
+	op->SetChild(op->op.args.ColourSampleImage.image, base);
+
+
+	// X
+	if (Node* pX = node.m_pX.get())
+	{
+		op->SetChild(op->op.args.ColourSampleImage.x, Generate(pX));
+	}
+	else
+	{
+		// Set a constant 0.5 value
+		NodeScalarConstantPtr pNode = new NodeScalarConstant();
+		pNode->SetValue(0.5f);
+		op->SetChild(op->op.args.ColourSampleImage.x, Generate(pNode));
+	}
+
+
+	// Y
+	if (Node* pY = node.m_pY.get())
+	{
+		op->SetChild(op->op.args.ColourSampleImage.y, Generate(pY));
+	}
+	else
+	{
+		// Set a constant 0.5 value
+		NodeScalarConstantPtr pNode = new NodeScalarConstant();
+		pNode->SetValue(0.5f);
+		op->SetChild(op->op.args.ColourSampleImage.y, Generate(pNode));
+	}
+
+	// TODO
+	op->op.args.ColourSampleImage.filter = 0;
+
+	result.op = op;
+}
+
+
+//-------------------------------------------------------------------------------------------------
+void CodeGenerator::GenerateColor_FromScalars(COLOR_GENERATION_RESULT& result, const Ptr<const NodeColourFromScalars>& Typed)
+{
+	const NodeColourFromScalars::Private& node = *Typed->GetPrivate();
+
+	Ptr<ASTOpFixed> op = new ASTOpFixed();
+	op->op.type = OP_TYPE::CO_FROMSCALARS;
+
+	// X
+	if (Node* pX = node.m_pX.get())
+	{
+		op->SetChild(op->op.args.ColourFromScalars.x, Generate(pX));
+	}
+	else
+	{
+		NodeScalarConstantPtr pNode = new NodeScalarConstant();
+		pNode->SetValue(1.0f);
+		op->SetChild(op->op.args.ColourFromScalars.x, Generate(pNode));
+	}
+
+	// Y
+	if (Node* pY = node.m_pY.get())
+	{
+		op->SetChild(op->op.args.ColourFromScalars.y, Generate(pY));
+	}
+	else
+	{
+		NodeScalarConstantPtr pNode = new NodeScalarConstant();
+		pNode->SetValue(1.0f);
+		op->SetChild(op->op.args.ColourFromScalars.y, Generate(pNode));
+	}
+
+	// Z
+	if (Node* pZ = node.m_pZ.get())
+	{
+		op->SetChild(op->op.args.ColourFromScalars.z, Generate(pZ));
+	}
+	else
+	{
+		NodeScalarConstantPtr pNode = new NodeScalarConstant();
+		pNode->SetValue(1.0f);
+		op->SetChild(op->op.args.ColourFromScalars.z, Generate(pNode));
+	}
+
+	// W
+	if (Node* pW = node.m_pW.get())
+	{
+		op->SetChild(op->op.args.ColourFromScalars.w, Generate(pW));
+	}
+	else
+	{
+		NodeScalarConstantPtr pNode = new NodeScalarConstant();
+		pNode->SetValue(1.0f);
+		op->SetChild(op->op.args.ColourFromScalars.w, Generate(pNode));
+	}
+
+	result.op = op;
+}
+
+
+//-------------------------------------------------------------------------------------------------
+void CodeGenerator::GenerateColor_Arithmetic(COLOR_GENERATION_RESULT& result, const Ptr<const NodeColourArithmeticOperation>& Typed)
+{
+	const NodeColourArithmeticOperation::Private& node = *Typed->GetPrivate();
+
+	Ptr<ASTOpFixed> op = new ASTOpFixed();
+	op->op.type = OP_TYPE::CO_ARITHMETIC;
+
+	switch (node.m_operation)
+	{
+	case NodeColourArithmeticOperation::AO_ADD: op->op.args.ColourArithmetic.operation = OP::ArithmeticArgs::ADD; break;
+	case NodeColourArithmeticOperation::AO_SUBTRACT: op->op.args.ColourArithmetic.operation = OP::ArithmeticArgs::SUBTRACT; break;
+	case NodeColourArithmeticOperation::AO_MULTIPLY: op->op.args.ColourArithmetic.operation = OP::ArithmeticArgs::MULTIPLY; break;
+	case NodeColourArithmeticOperation::AO_DIVIDE: op->op.args.ColourArithmetic.operation = OP::ArithmeticArgs::DIVIDE; break;
+	default:
+		checkf(false, TEXT("Unknown arithmetic operation."));
+		op->op.args.ColourArithmetic.operation = OP::ArithmeticArgs::NONE;
+		break;
+	}
+
+	// A
+	if (Node* pA = node.m_pA.get())
+	{
+		op->SetChild(op->op.args.ColourArithmetic.a, Generate(pA));
+	}
+	else
+	{
+		op->SetChild(op->op.args.ColourArithmetic.a,
+			CodeGenerator::GenerateMissingColourCode
+			(
+				"ColourArithmetic A",
+				node.m_errorContext
+			));
+	}
+
+	// B
+	if (Node* pB = node.m_pB.get())
+	{
+		op->SetChild(op->op.args.ColourArithmetic.b, Generate(pB));
+	}
+	else
+	{
+		op->SetChild(op->op.args.ColourArithmetic.b,
+			CodeGenerator::GenerateMissingColourCode
+			(
+				"ColourArithmetic B",
+				node.m_errorContext
+			));
+	}
+
+	result.op = op;
+}
+
+
+//-------------------------------------------------------------------------------------------------
+void CodeGenerator::GenerateColor_Table(COLOR_GENERATION_RESULT& result, const Ptr<const NodeColourTable>& Typed)
+{
+	const NodeColourTable::Private& node = *Typed->GetPrivate();
+
+	result.op = GenerateTableSwitch<NodeColourTable::Private, TCT_COLOUR, OP_TYPE::CO_SWITCH>(node,
+		[this](const NodeColourTable::Private& node, int colIndex, int row, ErrorLog* pErrorLog)
+		{
+			NodeColourConstantPtr pCell = new NodeColourConstant();
+			vec3<float> colour = node.m_pTable->GetPrivate()->m_rows[row].m_values[colIndex].m_colour;
+			pCell->SetValue(colour[0], colour[1], colour[2]);
+			return Generate(pCell);
+		});
+}
+
+
+//---------------------------------------------------------------------------------------------
+mu::Ptr<ASTOp> CodeGenerator::GenerateMissingColourCode(const char* strWhere, const void* errorContext)
+{
+	// Log a warning
+	char buf[256];
+	mutable_snprintf( buf, 256, "Required connection not found: %s", strWhere );
+	m_pErrorLog->GetPrivate()->Add(buf, ELMT_ERROR, errorContext);
+
+	// Create a constant colour node
+	NodeColourConstantPtr pNode = new NodeColourConstant();
+	pNode->SetValue(1, 1, 0);
+
+	COLOR_GENERATION_RESULT Result;
+	GenerateColor(Result,pNode);
+
+	return Result.op;
+}
+
diff --git a/Engine/Plugins/Experimental/Mutable/Source/MutableTools/Private/CodeGenerator_FirstPass.cpp b/Engine/Plugins/Experimental/Mutable/Source/MutableTools/Private/CodeGenerator_FirstPass.cpp
new file mode 100644
--- /dev/null
+++ b/Engine/Plugins/Experimental/Mutable/Source/MutableTools/Private/CodeGenerator_FirstPass.cpp
@@ -0,0 +1,675 @@
+// Copyright Epic Games, Inc. All Rights Reserved.
+
+#include "CodeGenerator_FirstPass.h"
+
+#include "Platform.h"
+#include "CompilerPrivate.h"
+
+#include "NodeSurfaceNewPrivate.h"
+#include "NodeSurfaceEditPrivate.h"
+#include "NodeSurfaceVariationPrivate.h"
+
+#include "NodeComponentPrivate.h"
+#include "NodeComponentNewPrivate.h"
+#include "NodeComponentEditPrivate.h"
+
+#include "NodeLODPrivate.h"
+#include "NodeObjectNewPrivate.h"
+#include "NodeObjectStatePrivate.h"
+#include "NodeObjectGroupPrivate.h"
+
+#include "NodeMeshClipMorphPlane.h"
+#include "NodeMeshClipWithMesh.h"
+
+#include "NodeModifierPrivate.h"
+#include "NodeModifierMeshClipMorphPlanePrivate.h"
+#include "NodeModifierMeshClipWithMeshPrivate.h"
+#include "NodeModifierMeshClipDeformPrivate.h"
+
+#include "AST.h"
+#include "ASTOpConstantBool.h"
+#include "ASTOpParameter.h"
+
+#include "ModelPrivate.h"
+#include "MutableTrace.h"
+#include "Operations.h"
+#include "SystemPrivate.h"
+
+
+namespace mu
+{
+
+	//---------------------------------------------------------------------------------------------
+	FirstPassGenerator::FirstPassGenerator()
+	{
+		// Default conditions when there is no restriction accumulated.
+		CONDITION_CONTEXT noCondition;
+        m_currentCondition.push_back(noCondition);
+        m_currentStateCondition.push_back(StateCondition());
+	}
+
+
+	//---------------------------------------------------------------------------------------------
+    void FirstPassGenerator::Generate( ErrorLogPtr pErrorLog,
+                                       const Node::Private* root,
+                                       bool ignoreStates )
+	{
+		MUTABLE_CPUPROFILER_SCOPE(FirstPassGenerate);
+
+		m_pErrorLog = pErrorLog;
+        m_ignoreStates = ignoreStates;
+
+		// Step 1: collect all objects, surfaces and object conditions
+        if ( root )
+		{
+            bool isFirstPassNode =
+                    dynamic_cast<const NodeSurface*>(root->m_pNode)
+                    ||
+                    dynamic_cast<const NodeComponent*>(root->m_pNode)
+                    ||
+                    dynamic_cast<const NodeObject*>(root->m_pNode)
+                    ||
+                    dynamic_cast<const NodeLOD*>(root->m_pNode)
+                    ||
+                    dynamic_cast<const NodePatchMesh*>(root->m_pNode)
+                    ||
+                    dynamic_cast<const NodeModifier*>(root->m_pNode);
+
+            if (isFirstPassNode)
+            {
+                root->Accept(*this);
+            }
+		}
+
+		// Step 2: Collect all tags and a list of the surfaces that activate them
+		for (std::size_t s=0; s<surfaces.size(); ++s)
+		{
+            // \todo: edit surfaces should also be able to activate tags. T1245
+            for (std::size_t t=0; t<surfaces[s].node->GetPrivate()->m_tags.size(); ++t)
+			{
+				int tag = -1;
+                auto tagStr = surfaces[s].node->GetPrivate()->m_tags[t];
+                for (std::size_t i = 0; i<m_tags.size() && tag<0; ++i)
+				{
+                    if (m_tags[i].tag == tagStr)
+					{
+						tag = (int)i;
+					}
+				}
+
+				// New tag?
+				if (tag < 0)
+				{
+                    tag = (int)m_tags.size();
+					TAG newTag;
+                    newTag.tag = tagStr;
+                    m_tags.push_back(newTag);
+				}
+
+                if (std::find(m_tags[tag].surfaces.begin(), m_tags[tag].surfaces.end(),s)
+					== 
+                    m_tags[tag].surfaces.end())
+				{
+                    m_tags[tag].surfaces.push_back((int)s);
+				}
+			}
+
+            // Collect the tags in edit surfaces
+            for (std::size_t e=0; e<surfaces[s].edits.size(); ++e)
+            {
+                const auto& edit = surfaces[s].edits[e];
+                for (std::size_t t=0; t<edit.node->m_tags.size(); ++t)
+                {
+                    int tag = -1;
+                    auto tagStr = edit.node->m_tags[t];
+
+                    for (std::size_t i = 0; i<m_tags.size() && tag<0; ++i)
+                    {
+                        if (m_tags[i].tag == tagStr)
+                        {
+                            tag = (int)i;
+                        }
+                    }
+
+                    // New tag?
+                    if (tag < 0)
+                    {
+                        tag = (int)m_tags.size();
+                        TAG newTag;
+                        newTag.tag = tagStr;
+                        m_tags.push_back(newTag);
+                    }
+
+                    auto key = std::make_pair(int(s),int(e));
+                    if (std::find(m_tags[tag].edits.begin(), m_tags[tag].edits.end(),key)
+                        ==
+                        m_tags[tag].edits.end())
+                    {
+                        m_tags[tag].edits.push_back(key);
+                    }
+                }
+            }
+
+		}
+
+
+        // Step 3: Create default state if necessary
+        if ( ignoreStates )
+        {
+            m_states.clear();
+        }
+
+        if ( m_states.empty() )
+        {
+            OBJECT_STATE data;
+            data.m_name = "Default";
+            m_states.emplace_back( data, root );
+        }
+	}
+
+
+	//---------------------------------------------------------------------------------------------
+    Ptr<ASTOp> FirstPassGenerator::Visit(const NodeModifierMeshClipMorphPlane::Private& node)
+	{
+        // Add the data about this modifier
+        MODIFIER thisData;
+        thisData.node = &node;
+        thisData.objectCondition = m_currentCondition.back().objectCondition;
+        thisData.stateCondition = m_currentStateCondition.back();
+        thisData.lod = m_currentLOD;
+        thisData.positiveTags = m_currentPositiveTags;
+        thisData.negativeTags = m_currentNegativeTags;
+        modifiers.push_back(thisData);
+
+        return nullptr;
+    }
+
+
+	//---------------------------------------------------------------------------------------------
+    Ptr<ASTOp> FirstPassGenerator::Visit(const NodeModifierMeshClipWithMesh::Private& node)
+	{
+        // Add the data about this modifier
+		MODIFIER thisData;
+		thisData.node = &node;
+        thisData.objectCondition = m_currentCondition.back().objectCondition;
+        thisData.stateCondition = m_currentStateCondition.back();
+        thisData.lod = m_currentLOD;
+        thisData.positiveTags = m_currentPositiveTags;
+        thisData.negativeTags = m_currentNegativeTags;
+        modifiers.push_back(thisData);
+
+        return nullptr;
+	}
+
+	
+	//---------------------------------------------------------------------------------------------
+    Ptr<ASTOp> FirstPassGenerator::Visit(const NodeModifierMeshClipDeform::Private& node)
+	{
+        // Add the data about this modifier
+		MODIFIER thisData;
+		thisData.node = &node;
+        thisData.objectCondition = m_currentCondition.back().objectCondition;
+        thisData.stateCondition = m_currentStateCondition.back();
+        thisData.lod = m_currentLOD;
+        thisData.positiveTags = m_currentPositiveTags;
+        thisData.negativeTags = m_currentNegativeTags;
+        modifiers.push_back(thisData);
+
+        return nullptr;
+	}
+
+	//---------------------------------------------------------------------------------------------
+    Ptr<ASTOp> FirstPassGenerator::Visit(const NodeSurfaceNew::Private& node)
+	{
+		// Add the data about this surface
+		SURFACE thisData;
+		thisData.node = dynamic_cast<const NodeSurfaceNew*>(node.m_pNode);
+		thisData.component = m_currentComponent;
+        thisData.objectCondition = m_currentCondition.back().objectCondition;
+        thisData.stateCondition = m_currentStateCondition.back();
+        thisData.positiveTags = m_currentPositiveTags;
+        thisData.negativeTags = m_currentNegativeTags;
+        surfaces.push_back(thisData);
+
+        return nullptr;
+	}
+
+
+	//---------------------------------------------------------------------------------------------
+    Ptr<ASTOp> FirstPassGenerator::Visit(const NodeSurfaceEdit::Private& node)
+	{
+		// Store a reference to this node in the surface data for the surface that this node is
+		// editing.
+		auto its = std::find_if(surfaces.begin(), surfaces.end(), 
+            [&node](const FirstPassGenerator::SURFACE& s)
+        {
+            // Are we editing the main surface node of this surface?
+            if (s.node.get() == node.m_pParent.get()) return true;
+
+            // Are we editing an edit node modifying this surface?
+            for (const auto& e: s.edits)
+            {
+                if (node.m_pParent && e.node==node.m_pParent->GetBasePrivate())
+                {
+                    return true;
+                }
+            }
+
+            return false;
+        });
+		
+		// The surface could be missing if the parent is not in the hierarchy. This could happen
+		// with wrong input or in case of partial models for preview.
+		if (its != surfaces.end())
+		{
+			SURFACE& surface = *its;
+
+            SURFACE::EDIT edit;
+            edit.node = &node;
+            edit.condition = m_currentCondition.back().objectCondition;
+            surface.edits.push_back(edit);
+		}
+		else
+		{
+			m_pErrorLog->GetPrivate()->Add("Missing parent object for edit node.",
+				ELMT_WARNING, node.m_errorContext);
+		}
+
+        return nullptr;
+	}
+
+
+	//---------------------------------------------------------------------------------------------
+    Ptr<ASTOp> FirstPassGenerator::Visit(const NodeSurfaceVariation::Private& node)
+	{
+        switch(node.m_type)
+        {
+
+        case NodeSurfaceVariation::VariationType::Tag:
+        {
+            // Any of the tags in the variations would prevent the default surface
+            auto oldNegativeTags = m_currentNegativeTags;
+            for (size_t v=0; v<node.m_variations.size(); ++v)
+            {
+                m_currentNegativeTags.push_back(node.m_variations[v].m_tag);
+            }
+
+            for(const auto& n:node.m_defaultSurfaces)
+            {
+                n->GetBasePrivate()->Accept(*this);
+            }
+            for(const auto& n:node.m_defaultModifiers)
+            {
+                n->GetBasePrivate()->Accept(*this);
+            }
+
+            m_currentNegativeTags = oldNegativeTags;
+
+            for (size_t v=0; v<node.m_variations.size(); ++v)
+            {
+                m_currentPositiveTags.push_back(node.m_variations[v].m_tag);
+                for (const auto& s : node.m_variations[v].m_surfaces)
+                {
+                    s->GetBasePrivate()->Accept(*this);
+                }
+
+                for (const auto& s : node.m_variations[v].m_modifiers)
+                {
+                    s->GetBasePrivate()->Accept(*this);
+                }
+
+                m_currentPositiveTags.pop_back();
+
+                // Tags have an order in a variation node: the current tag should prevent any following
+                // variation surface
+                m_currentNegativeTags.push_back(node.m_variations[v].m_tag);
+            }
+
+            m_currentNegativeTags = oldNegativeTags;
+
+            break;
+        }
+
+
+        case NodeSurfaceVariation::VariationType::State:
+        {
+            size_t stateCount = m_states.size();
+
+            // Default
+            {
+                // Store the states for the default branch here
+                StateCondition defaultStates;
+                {
+                    defaultStates = m_currentStateCondition.back().empty()
+                            ? StateCondition( stateCount, true )
+                            : m_currentStateCondition.back();
+
+                    for (const auto& v:node.m_variations)
+                    {
+                        for( size_t s=0; s<stateCount; ++s )
+                        {
+                            if (m_states[s].first.m_name==v.m_tag)
+                            {
+                                // Remove this state from the default options, since it has its own variation
+                                defaultStates[s] = false;
+                            }
+                        }
+                    }
+                }
+
+                m_currentStateCondition.push_back(defaultStates);
+
+                for (const auto& n : node.m_defaultSurfaces)
+                {
+                    n->GetBasePrivate()->Accept(*this);
+                }
+                for (const auto& n : node.m_defaultModifiers)
+                {
+                    n->GetBasePrivate()->Accept(*this);
+                }
+
+                m_currentStateCondition.pop_back();
+            }
+
+
+            // Variation branches
+            for (const auto& v:node.m_variations)
+            {
+                // Store the states for this variation here
+                StateCondition variationStates( stateCount, false );
+
+                for( size_t s=0; s<stateCount; ++s )
+                {
+                    if (m_states[s].first.m_name==v.m_tag)
+                    {
+                        variationStates[s] = true;
+                    }
+                }
+
+                m_currentStateCondition.push_back(variationStates);
+
+                for (const auto& n : v.m_surfaces)
+                {
+                    n->GetBasePrivate()->Accept(*this);
+                }
+                for (const auto& n : v.m_modifiers)
+                {
+                    n->GetBasePrivate()->Accept(*this);
+                }
+
+                m_currentStateCondition.pop_back();
+            }
+
+            break;
+        }
+
+
+        default:
+            // Case not implemented.
+            check(false);
+            break;
+        }
+
+        return nullptr;
+	}
+
+
+	//---------------------------------------------------------------------------------------------
+    Ptr<ASTOp> FirstPassGenerator::Visit(const NodeComponentNew::Private& node)
+	{
+        m_currentComponent = node.GetParentComponentNew();
+
+		for (const auto& c : node.m_surfaces)
+		{
+			if (c)
+			{
+				c->GetBasePrivate()->Accept(*this);
+			}
+		}
+
+		m_currentComponent = nullptr;
+
+        return nullptr;
+	}
+
+
+	//---------------------------------------------------------------------------------------------
+    Ptr<ASTOp> FirstPassGenerator::Visit(const NodeComponentEdit::Private& node)
+	{
+		m_currentComponent = node.GetParentComponentNew();
+
+		for (const auto& c : node.m_surfaces)
+		{
+			if (c)
+			{
+                c->GetBasePrivate()->Accept(*this);
+			}
+		}
+
+		m_currentComponent = nullptr;
+
+        return nullptr;
+	}
+
+
+	//---------------------------------------------------------------------------------------------
+    Ptr<ASTOp> FirstPassGenerator::Visit(const NodeLOD::Private& node)
+	{
+		for (const auto& c : node.m_components)
+		{
+			if (c)
+			{
+				c->GetBasePrivate()->Accept(*this);
+			}
+		}
+		for (const auto& c : node.m_modifiers)
+		{
+			if (c)
+			{
+				c->GetBasePrivate()->Accept(*this);
+			}
+		}
+
+        return nullptr;
+	}
+
+
+	//---------------------------------------------------------------------------------------------
+    Ptr<ASTOp> FirstPassGenerator::Visit(const NodeObjectNew::Private& node)
+	{
+		// Add the data about this object
+		OBJECT thisData;
+		thisData.node = &node;
+        thisData.condition = m_currentCondition.back().objectCondition;
+		objects.push_back(thisData);
+
+        // Accumulate the model states
+        for ( const auto& s: node.m_states )
+        {
+            m_states.emplace_back( s, &node );
+
+            // \todo move this 32 constant to a macro or constexpr with a meaningful name and location
+            if ( s.m_runtimeParams.size() > 32 )
+            {
+                char temp[256];
+                mutable_snprintf( temp, 256,
+                                  "State [%s] has more than 32 runtime parameters. Their update may fail.",
+                                  s.m_name.c_str() );
+                m_pErrorLog->GetPrivate()->Add( temp, ELMT_ERROR, node.m_errorContext );
+            }
+        }
+
+		// Process the lods
+		int i = 0;
+		for (const auto& l : node.m_lods)
+		{
+			if (l)
+			{
+                m_currentLOD = i++;
+				l->GetBasePrivate()->Accept(*this);
+			}
+		}
+
+		m_currentLOD = -1;
+
+		// Process the children
+		for (const auto& c : node.m_children)
+		{
+			if (c)
+			{
+				c->GetBasePrivate()->Accept(*this);
+			}
+		}
+
+        return nullptr;
+	}
+
+
+	//---------------------------------------------------------------------------------------------
+    Ptr<ASTOp> FirstPassGenerator::Visit(const NodeObjectGroup::Private& node)
+	{
+       // Prepare the enumeration parameter if necessary
+        Ptr<ASTOpParameter> enumOp;
+        if ( node.m_type==NodeObjectGroup::CS_ALWAYS_ONE ||
+             node.m_type==NodeObjectGroup::CS_ONE_OR_NONE )
+        {
+            Ptr<ASTOpParameter> op = new ASTOpParameter();
+            op->type = OP_TYPE::NU_PARAMETER;
+
+            op->parameter.m_name = node.m_name;
+            op->parameter.m_uid = node.m_uid;
+            op->parameter.m_type = PARAMETER_TYPE::T_INT;
+            op->parameter.m_defaultValue.m_int = -1;
+
+            if ( node.m_type==NodeObjectGroup::CS_ONE_OR_NONE )
+            {
+                PARAMETER_DESC::INT_VALUE_DESC nullValue;
+                nullValue.m_value = -1;
+                nullValue.m_name = "None";
+                op->parameter.m_possibleValues.push_back( nullValue );
+                op->parameter.m_defaultValue.m_int = nullValue.m_value;
+            }
+
+            enumOp = op;
+        }
+
+
+        // Parse the child objects
+		for ( std::size_t t=0; t<node.m_children.size(); ++t )
+        {
+            if ( const NodeObject* pChildNode = node.m_children[t].get() )
+            {
+                // Overwrite the implicit condition
+                Ptr<ASTOp> paramOp = 0;
+                switch ( node.m_type )
+                {
+                    case NodeObjectGroup::CS_TOGGLE_EACH:
+                    {
+                        // Create a new boolean parameter
+                        Ptr<ASTOpParameter> op = new ASTOpParameter();
+                        op->type = OP_TYPE::BO_PARAMETER;
+
+                        op->parameter.m_name = pChildNode->GetName();
+                        op->parameter.m_uid = pChildNode->GetUid();
+                        op->parameter.m_type = PARAMETER_TYPE::T_BOOL;
+                        op->parameter.m_defaultValue.m_bool = false;
+
+                        paramOp = op;
+                        break;
+                    }
+
+                    case NodeObjectGroup::CS_ALWAYS_ALL:
+                    {
+                        // Create a constant true boolean that the optimiser will remove later.
+                        Ptr<ASTOpConstantBool> op = new ASTOpConstantBool();
+                        op->value = true;
+
+                        paramOp = op;
+                        break;
+                    }
+
+                    case NodeObjectGroup::CS_ONE_OR_NONE:
+                    case NodeObjectGroup::CS_ALWAYS_ONE:
+                    {
+                        // Add the option to the enumeration parameter
+                        PARAMETER_DESC::INT_VALUE_DESC value;
+                        value.m_value = (int16_t)t;
+                        value.m_name = pChildNode->GetName();
+                        enumOp->parameter.m_possibleValues.push_back( value );
+
+                        // Set as default if none has been set.
+                        if ( enumOp->parameter.m_defaultValue.m_int == -1 )
+                        {
+                            enumOp->parameter.m_defaultValue.m_int = (int)t;
+                        }
+
+                        check(enumOp);
+
+                        // Create a comparison operation as the boolean parameter for the child
+                        Ptr<ASTOpFixed> op = new ASTOpFixed();
+                        op->op.type = OP_TYPE::BO_EQUAL_INT_CONST;
+                        op->SetChild( op->op.args.BoolEqualScalarConst.value, enumOp);
+                        op->op.args.BoolEqualScalarConst.constant = (int16_t)t;
+
+                        paramOp = op;
+                        break;
+                    }
+
+                    default:
+                        check( false );
+                }
+
+                // Combine the new condition with previous conditions coming from parent objects
+                if (m_currentCondition.back().objectCondition)
+                {
+                    Ptr<ASTOpFixed> op = new ASTOpFixed();
+                    op->op.type = OP_TYPE::BO_AND;
+                    op->SetChild( op->op.args.BoolBinary.a,m_currentCondition.back().objectCondition);
+                    op->SetChild( op->op.args.BoolBinary.b,paramOp);
+                    paramOp = op;
+                }
+
+				CONDITION_CONTEXT data;
+                data.objectCondition = paramOp;
+                m_currentCondition.push_back( data );
+
+				pChildNode->GetBasePrivate()->Accept(*this);
+
+                m_currentCondition.pop_back();
+            }
+        }
+
+        return nullptr;
+ 	}
+
+
+    //---------------------------------------------------------------------------------------------
+    Ptr<ASTOp> FirstPassGenerator::Visit( const NodeObjectState::Private& node )
+    {
+        // Generate the source object
+        Ptr<ASTOp> source;
+
+        if (node.m_pSource)
+        {
+            // Process the source
+            node.m_pSource->GetBasePrivate()->Accept(*this);
+
+            if (node.m_pRoot)
+            {
+                // Remember the new state with the new root.
+                m_states.emplace_back( node.m_state, node.m_pRoot->GetBasePrivate() );
+            }
+        }
+
+        return source;
+    }
+
+
+	//---------------------------------------------------------------------------------------------
+    Ptr<ASTOp> FirstPassGenerator::Visit(const NodePatchMesh::Private&)
+	{
+        return nullptr;
+	}
+
+}
+
diff --git a/Engine/Plugins/Experimental/Mutable/Source/MutableTools/Private/CodeGenerator_FirstPass.h b/Engine/Plugins/Experimental/Mutable/Source/MutableTools/Private/CodeGenerator_FirstPass.h
new file mode 100644
--- /dev/null
+++ b/Engine/Plugins/Experimental/Mutable/Source/MutableTools/Private/CodeGenerator_FirstPass.h
@@ -0,0 +1,269 @@
+// Copyright Epic Games, Inc. All Rights Reserved.
+
+#pragma once
+
+#include "Visitor.h"
+
+#include "AST.h"
+#include "CompilerPrivate.h"
+#include "NodeLayoutPrivate.h"
+#include "ErrorLogPrivate.h"
+#include "NodeObjectNewPrivate.h"
+
+#include "Table.h"
+
+#include "NodeBool.h"
+
+#include "NodeComponent.h"
+#include "NodeComponentNew.h"
+#include "NodeComponentEdit.h"
+
+#include "NodeModifier.h"
+#include "NodeModifierMeshClipMorphPlane.h"
+#include "NodeModifierMeshClipWithMesh.h"
+#include "NodeModifierMeshClipDeform.h"
+
+#include "NodeSurface.h"
+#include "NodeSurfaceNew.h"
+#include "NodeSurfaceEdit.h"
+#include "NodeSurfaceVariation.h"
+
+#include "NodePatchImage.h"
+#include "NodeLOD.h"
+#include "NodeObjectNew.h"
+#include "NodeObjectState.h"
+#include "NodeObjectGroup.h"
+#include "NodePatchMesh.h"
+
+#include "Operations.h"
+#include "ModelPrivate.h"
+#include "ImagePrivate.h"
+
+#include <memory>
+
+namespace mu
+{
+
+
+    //! Store the results of the code generation of a mesh.
+    struct MESH_GENERATION_RESULT
+    {
+        //! Mesh after all code tree is applied
+        Ptr<ASTOp> meshOp;
+
+        //! Original base mesh before removes, morphs, etc.
+        Ptr<ASTOp> baseMeshOp;
+
+        vector<Ptr<const Layout>> layouts;
+
+        vector<Ptr<ASTOp>> layoutOps;
+
+
+        struct EXTRA_LAYOUTS
+        {
+            vector<Ptr<const Layout>> layouts;
+            Ptr<ASTOp> condition;
+            Ptr<ASTOp> meshFragment;
+        };
+        vector< EXTRA_LAYOUTS > extraMeshLayouts;
+
+    };
+
+
+	//---------------------------------------------------------------------------------------------
+	//! First pass of the code generation process.
+	//! It collects data about the object hierarchy, the conditions for each object and the global
+	//! modifiers.
+	//---------------------------------------------------------------------------------------------
+	class FirstPassGenerator :
+		public Base,
+		public BaseVisitor,
+
+        public Visitor<NodeSurfaceNew::Private, Ptr<ASTOp>, true>,
+        public Visitor<NodeSurfaceEdit::Private, Ptr<ASTOp>, true>,
+        public Visitor<NodeSurfaceVariation::Private, Ptr<ASTOp>, true>,
+        public Visitor<NodeComponentNew::Private, Ptr<ASTOp>, true>,
+        public Visitor<NodeComponentEdit::Private, Ptr<ASTOp>, true>,
+        public Visitor<NodeLOD::Private, Ptr<ASTOp>, true>,
+        public Visitor<NodeObjectNew::Private, Ptr<ASTOp>, true>,
+        public Visitor<NodeObjectGroup::Private, Ptr<ASTOp>, true>,
+        public Visitor<NodeObjectState::Private, Ptr<ASTOp>, true>,
+        public Visitor<NodePatchMesh::Private, Ptr<ASTOp>, true>,
+
+        public Visitor<NodeModifierMeshClipMorphPlane::Private, Ptr<ASTOp>, true>,
+        public Visitor<NodeModifierMeshClipWithMesh::Private, Ptr<ASTOp>, true>,
+        public Visitor<NodeModifierMeshClipDeform::Private, Ptr<ASTOp>, true>
+	{
+	public:
+
+		FirstPassGenerator();
+
+        void Generate(ErrorLogPtr pErrorLog,
+                      const Node::Private* root,
+                      bool ignoreStates);
+
+	public:
+
+		// Results
+		//-------------------------
+
+		//! Store the conditions that will enable or disable every object
+		struct OBJECT
+		{
+			const NodeObjectNew::Private* node;
+            Ptr<ASTOp> condition;
+		};
+		vector<OBJECT> objects;
+
+        //! Type used to represent the activation conditions regarding states
+        //! This is the state mask for the states in which this surface must be added. If it
+        //! is empty it means the surface is valid for all states. Otherwise it is only valid
+        //! for the states whose index is true.
+        using StateCondition = std::vector<uint8_t>;
+
+		//! Store information about every surface including
+		//! - the component it may be added to
+		//! - the conditions that will enable or disable it
+		//! - all edit operators
+        //! A surface may have different versions depending on the different parents and conditions
+        //! it is reached with.
+		struct SURFACE
+		{
+            NodeSurfaceNewPtrConst node;
+
+			// Parent component where this surface will be added. It may be different from the 
+			// component that defined it (if it was an edit component).
+            const NodeComponentNew::Private* component = nullptr;
+
+            // List of tags that are required for the presence of this surface
+            vector<string> positiveTags;
+
+            // List of tags that block the presence of this surface
+            vector<string> negativeTags;
+
+			// This conditions is the condition of the object defining this surface which may not
+			// be the parent object where this surface will be added.
+            Ptr<ASTOp> objectCondition;
+
+            // This is filled in the first pass.
+            StateCondition stateCondition;
+
+            // Condition for this surface to be enabled when all the object conditions are met.
+            // This is filled in CodeGenerator_SecondPass.
+            Ptr<ASTOp> surfaceCondition;
+
+			// All surface editing nodes that edit this surface
+            struct EDIT
+            {
+                //! Condition that enables the effects of this edit node on the surface
+                Ptr<ASTOp> condition;
+
+                //! Weak reference to the edit node, used during compilation.
+                const NodeSurfaceEdit::Private* node = nullptr;
+            };
+            vector<EDIT> edits;
+
+            // This is filled in the final code generation pass
+            Ptr<ASTOp> resultSurfaceOp;
+            Ptr<ASTOp> resultMeshOp;
+        };
+		vector<SURFACE> surfaces;
+
+		//! Store the conditions that enable every modifier.
+		struct MODIFIER
+		{
+            const NodeModifier::Private* node = nullptr;
+
+            // List of tags that are required for the presence of this surface
+            vector<string> positiveTags;
+
+            // List of tags that block the presence of this surface
+            vector<string> negativeTags;
+
+            // This conditions is the condition of the object defining this modifier which may not
+            // be the parent object where this surface will be added.
+            Ptr<ASTOp> objectCondition;
+
+            // This conditions is the condition for this modifier to be enabled when all the object
+            // conditions are met.
+            // This is filled in CodeGenerator_SecondPass.
+            Ptr<ASTOp> surfaceCondition;
+
+            // This is filled in CodeGenerator_SecondPass.
+            StateCondition stateCondition;
+
+            //
+            int lod = 0;
+        };
+		vector<MODIFIER> modifiers;
+
+		//! Info about all found tags.
+		struct TAG
+		{
+			string tag;
+
+            // Surfaces that activate the tag. These are indices to the FirstPassGenerator::surfaces
+            // vector.
+            vector<int> surfaces;
+
+            // Edit Surfaces that activate the tag. These first element of the pair are indices to
+            // the FirstPassGenerator::surfaces vector. The second element are indices to the
+            // "edits" in the specific surface.
+            vector<std::pair<int,int>> edits;
+
+            // This conditions is the condition for this tag to be enabled considering no other
+            // condition. This is filled in CodeGenerator_SecondPass.
+            Ptr<ASTOp> genericCondition;
+        };
+        vector<TAG> m_tags;
+
+        //! Accumulate the model states found while generating code, with their generated root
+        //! nodes.
+        typedef vector< std::pair<OBJECT_STATE, const Node::Private*> > StateList;
+        StateList m_states;
+
+	public:
+
+        Ptr<ASTOp> Visit(const NodeSurfaceNew::Private&) override;
+        Ptr<ASTOp> Visit(const NodeSurfaceEdit::Private&) override;
+        Ptr<ASTOp> Visit(const NodeSurfaceVariation::Private&) override;
+        Ptr<ASTOp> Visit(const NodeComponentNew::Private&) override;
+        Ptr<ASTOp> Visit(const NodeComponentEdit::Private&) override;
+        Ptr<ASTOp> Visit(const NodeLOD::Private&) override;
+        Ptr<ASTOp> Visit(const NodeObjectNew::Private& node) override;
+        Ptr<ASTOp> Visit(const NodeObjectGroup::Private&) override;
+        Ptr<ASTOp> Visit(const NodeObjectState::Private&) override;
+        Ptr<ASTOp> Visit(const NodeModifierMeshClipMorphPlane::Private&) override;
+        Ptr<ASTOp> Visit(const NodeModifierMeshClipWithMesh::Private&) override;
+        Ptr<ASTOp> Visit(const NodeModifierMeshClipDeform::Private&) override;
+        Ptr<ASTOp> Visit(const NodePatchMesh::Private&) override;
+
+	private:
+
+        struct CONDITION_CONTEXT
+        {
+            Ptr<ASTOp> objectCondition;
+        };
+        vector< CONDITION_CONTEXT > m_currentCondition;
+
+        //!
+        vector< StateCondition > m_currentStateCondition;
+
+		//! When processing surfaces, this is the parent component the surfaces may be added to
+        const NodeComponentNew::Private* m_currentComponent = nullptr;
+
+        //! Current relevant tags so far. Used during traversal.
+        vector<string> m_currentPositiveTags;
+        vector<string> m_currentNegativeTags;
+
+		//! Index of the LOD we are processing
+        int m_currentLOD = -1;
+
+        //!
+        ErrorLogPtr m_pErrorLog;
+
+        //!
+        bool m_ignoreStates = false;
+	};
+
+}
