[mutable] Move Mutable plugin from the restricted folder to the experimental plugins folder.

diff --git a/Engine/Plugins/Experimental/Mutable/Source/MutableTools/Private/CodeGenerator_Image.cpp b/Engine/Plugins/Experimental/Mutable/Source/MutableTools/Private/CodeGenerator_Image.cpp
new file mode 100644
--- /dev/null
+++ b/Engine/Plugins/Experimental/Mutable/Source/MutableTools/Private/CodeGenerator_Image.cpp
@@ -0,0 +1,1994 @@
+// Copyright Epic Games, Inc. All Rights Reserved.
+
+#include "CodeGenerator.h"
+
+#include "Platform.h"
+#include "CompilerPrivate.h"
+#include "ImageDescGenerator.h"
+#include "TablePrivate.h"
+#include "NodeLayoutPrivate.h"
+
+#include "NodeImagePrivate.h"
+#include "NodeImageConstantPrivate.h"
+#include "NodeImageDifferencePrivate.h"
+#include "NodeImageInterpolatePrivate.h"
+#include "NodeImageInterpolate3Private.h"
+#include "NodeImageSaturatePrivate.h"
+#include "NodeImageLuminancePrivate.h"
+#include "NodeImageLayerPrivate.h"
+#include "NodeImageLayerColourPrivate.h"
+#include "NodeImageMultiLayerPrivate.h"
+#include "NodeImageTablePrivate.h"
+#include "NodeImageSwizzlePrivate.h"
+#include "NodeImageFormatPrivate.h"
+#include "NodeImageSelectColourPrivate.h"
+#include "NodeImageColourMapPrivate.h"
+#include "NodeImageGradientPrivate.h"
+#include "NodeImageBinarisePrivate.h"
+#include "NodeImageResizePrivate.h"
+#include "NodeImagePlainColourPrivate.h"
+#include "NodeImageSwitchPrivate.h"
+#include "NodeImageProjectPrivate.h"
+#include "NodeImageMipmapPrivate.h"
+#include "NodeImageConditionalPrivate.h"
+#include "NodeImageParameterPrivate.h"
+#include "NodeImageInvertPrivate.h"
+#include "NodeImageVariationPrivate.h"
+#include "NodeImageNormalCompositePrivate.h"
+#include "NodeImageTransformPrivate.h"
+
+#include "MutableTools/Private/ASTOpMeshFormat.h"
+#include "MutableTools/Private/ASTOpImagePixelFormat.h"
+#include "MutableTools/Private/ASTOpImageMipmap.h"
+#include "MutableTools/Private/ASTOpConstantResource.h"
+#include "MutableTools/Private/ASTOpConditional.h"
+#include "MutableTools/Private/ASTOpImageMultiLayer.h"
+#include "MutableTools/Private/ASTOpImageTransform.h"
+#include "MutableTools/Private/ASTOpImageNormalComposite.h"
+#include "MutableTools/Private/ASTOpMeshExtractLayoutBlocks.h"
+
+#include "ModelPrivate.h"
+#include "MutableTrace.h"
+#include "Operations.h"
+#include "OpImageCrop.h"
+#include "OpImageProject.h"
+
+
+namespace mu
+{
+
+
+	//-------------------------------------------------------------------------------------------------
+	void CodeGenerator::GenerateImage(IMAGE_GENERATION_RESULT& result, const NodeImagePtrConst& Untyped)
+	{
+		if (!Untyped)
+		{
+			result = IMAGE_GENERATION_RESULT();
+			return;
+		}
+
+		// Generate the block size in case we are generating an expression whose root is an image
+		bool addedImageState = false;
+		if ( m_imageState.empty() )
+		{
+
+			FImageDesc desc = CalculateImageDesc(*Untyped->GetBasePrivate());
+			IMAGE_STATE newState;
+			newState.m_imageSize[0] = desc.m_size[0] ? desc.m_size[0] : 256;
+			newState.m_imageSize[1] = desc.m_size[1] ? desc.m_size[1] : 256;
+			newState.m_imageRect.size = desc.m_size;
+			newState.m_imageRect.min[0] = 0;
+			newState.m_imageRect.min[1] = 0;
+			newState.m_layoutBlock = -1;
+			m_imageState.push_back(newState);
+			addedImageState = true;
+		}
+
+		// See if it was already generated
+		VISITED_MAP_KEY key = GetCurrentCacheKey(Untyped);
+		GeneratedImagesMap::ValueType* it = m_generatedImages.Find(key);
+		if (it)
+		{
+			result = *it;
+		}
+		else
+		{ 
+			const NodeImage* Node = Untyped.get();
+
+			// Generate for each different type of node
+			switch (Untyped->GetImageNodeType())
+			{
+			case NodeImage::EType::Constant: GenerateImage_Constant(result, static_cast<const NodeImageConstant*>(Node)); break;
+			case NodeImage::EType::Difference: GenerateImage_Difference(result, static_cast<const NodeImageDifference*>(Node)); break;
+			case NodeImage::EType::Interpolate: GenerateImage_Interpolate(result, static_cast<const NodeImageInterpolate*>(Node)); break;
+			case NodeImage::EType::Saturate: GenerateImage_Saturate(result, static_cast<const NodeImageSaturate*>(Node)); break;
+			case NodeImage::EType::Table: GenerateImage_Table(result, static_cast<const NodeImageTable*>(Node)); break;
+			case NodeImage::EType::Swizzle: GenerateImage_Swizzle(result, static_cast<const NodeImageSwizzle*>(Node)); break;
+			case NodeImage::EType::SelectColour: GenerateImage_SelectColour(result, static_cast<const NodeImageSelectColour*>(Node)); break;
+			case NodeImage::EType::ColourMap: GenerateImage_ColourMap(result, static_cast<const NodeImageColourMap*>(Node)); break;
+			case NodeImage::EType::Gradient: GenerateImage_Gradient(result, static_cast<const NodeImageGradient*>(Node)); break;
+			case NodeImage::EType::Binarise: GenerateImage_Binarise(result, static_cast<const NodeImageBinarise*>(Node)); break;
+			case NodeImage::EType::Luminance: GenerateImage_Luminance(result, static_cast<const NodeImageLuminance*>(Node)); break;
+			case NodeImage::EType::Layer: GenerateImage_Layer(result, static_cast<const NodeImageLayer*>(Node)); break;
+			case NodeImage::EType::LayerColour: GenerateImage_LayerColour(result, static_cast<const NodeImageLayerColour*>(Node)); break;
+			case NodeImage::EType::Resize: GenerateImage_Resize(result, static_cast<const NodeImageResize*>(Node)); break;
+			case NodeImage::EType::PlainColour: GenerateImage_PlainColour(result, static_cast<const NodeImagePlainColour*>(Node)); break;
+			case NodeImage::EType::Interpolate3: GenerateImage_Interpolate3(result, static_cast<const NodeImageInterpolate3*>(Node)); break;
+			case NodeImage::EType::Project: GenerateImage_Project(result, static_cast<const NodeImageProject*>(Node)); break;
+			case NodeImage::EType::Mipmap: GenerateImage_Mipmap(result, static_cast<const NodeImageMipmap*>(Node)); break;
+			case NodeImage::EType::Switch: GenerateImage_Switch(result, static_cast<const NodeImageSwitch*>(Node)); break;
+			case NodeImage::EType::Conditional: GenerateImage_Conditional(result, static_cast<const NodeImageConditional*>(Node)); break;
+			case NodeImage::EType::Format: GenerateImage_Format(result, static_cast<const NodeImageFormat*>(Node)); break;
+			case NodeImage::EType::Parameter: GenerateImage_Parameter(result, static_cast<const NodeImageParameter*>(Node)); break;
+			case NodeImage::EType::MultiLayer: GenerateImage_MultiLayer(result, static_cast<const NodeImageMultiLayer*>(Node)); break;
+			case NodeImage::EType::Invert: GenerateImage_Invert(result, static_cast<const NodeImageInvert*>(Node)); break;
+			case NodeImage::EType::Variation: GenerateImage_Variation(result, static_cast<const NodeImageVariation*>(Node)); break;
+			case NodeImage::EType::NormalComposite: GenerateImage_NormalComposite(result, static_cast<const NodeImageNormalComposite*>(Node)); break;
+			case NodeImage::EType::Transform: GenerateImage_Transform(result, static_cast<const NodeImageTransform*>(Node)); break;
+			case NodeImage::EType::None: check(false);
+			}
+
+			// Cache the result
+			m_generatedImages.Add(key, result);
+		}
+
+		// Restore the modified image state
+		if (addedImageState)
+		{
+			m_imageState.pop_back();
+		}
+
+	}
+
+
+    //---------------------------------------------------------------------------------------------
+    void CodeGenerator::GenerateImage_Constant(IMAGE_GENERATION_RESULT& result, const NodeImageConstant* InNode)
+    {
+		const NodeImageConstant::Private& node = *InNode->GetPrivate();
+		
+		//MUTABLE_CPUPROFILER_SCOPE(NodeImageConstant);
+
+        Ptr<ASTOpConstantResource> op = new ASTOpConstantResource();
+        op->type = OP_TYPE::IM_CONSTANT;
+
+        // TODO: check duplicates
+        Ptr<const Image> pImage;
+		if (node.m_pProxy)
+		{
+			pImage = node.m_pProxy->Get();
+		}
+
+        if (!pImage)
+        {
+            // This data is required
+            pImage = GenerateMissingImage(EImageFormat::IF_RGB_UBYTE );
+
+            // Log an error message
+            m_pErrorLog->GetPrivate()->Add( "Constant image not set.", ELMT_WARNING, node.m_errorContext );
+        }
+
+        vec2<int> imageSize( pImage->GetSizeX(), pImage->GetSizeY() );
+
+        // The constant image size may be different than the parent rect we are generating
+        // In that case we need to crop the proportional part and the code generator will
+        // add scaling operations later
+        box< vec2< int > > cropRect;
+
+        // Order of the operations is important: multiply first to avoid losing precision.
+        // It will not overflow since image sizes are limited to 16 bit
+        cropRect.min = ( m_imageState.back().m_imageRect.min * imageSize )
+                / m_imageState.back().m_imageSize;
+        cropRect.size = ( m_imageState.back().m_imageRect.size * imageSize )
+                / m_imageState.back().m_imageSize;
+
+        //check( cropRect.size[0]>0 && cropRect.size[1]>0 );
+        cropRect.size[0] = std::max( cropRect.size[0], 1 );
+        cropRect.size[1] = std::max( cropRect.size[1], 1 );
+
+        //
+        if ( pImage->GetSizeX()!=cropRect.size[0] || pImage->GetSizeY()!=cropRect.size[1] )
+        {
+            ImagePtrConst pCropped =
+                ImageCrop( m_compilerOptions->m_imageCompressionQuality,
+                           pImage.get(), cropRect );
+            op->SetValue( pCropped, m_compilerOptions->m_optimisationOptions.m_useDiskCache );
+        }
+        else
+        {
+            check( cropRect.min[0]==0 && cropRect.min[1]==0 );
+            op->SetValue( pImage, m_compilerOptions->m_optimisationOptions.m_useDiskCache );
+        }
+
+		result.op = op;
+    }
+
+
+    //---------------------------------------------------------------------------------------------
+	void CodeGenerator::GenerateImage_Parameter(IMAGE_GENERATION_RESULT& result, const NodeImageParameter* InNode )
+    {
+		const NodeImageParameter::Private& node = *InNode->GetPrivate();
+
+        Ptr<ASTOpParameter> op;
+
+        auto it = m_nodeVariables.find( node.m_pNode );
+        if ( it == m_nodeVariables.end() )
+        {
+            op = new ASTOpParameter();
+            op->type = OP_TYPE::IM_PARAMETER;
+
+			op->parameter.m_name = node.m_name;
+			op->parameter.m_uid = node.m_uid;
+            op->parameter.m_type = PARAMETER_TYPE::T_IMAGE;
+
+            m_nodeVariables[node.m_pNode] = op;
+        }
+        else
+        {
+            op = it->second;
+        }
+
+		result.op = op;
+    }
+
+
+    //---------------------------------------------------------------------------------------------
+	void CodeGenerator::GenerateImage_Layer(IMAGE_GENERATION_RESULT& result, const NodeImageLayer* InNode)
+	{
+		const NodeImageLayer::Private& node = *InNode->GetPrivate();
+		
+		MUTABLE_CPUPROFILER_SCOPE(NodeImageLayer);
+
+        Ptr<ASTOpFixed> op = new ASTOpFixed();
+
+		op->op.type = OP_TYPE::IM_LAYER;
+        op->op.args.ImageLayer.blendType = uint8(node.m_type);
+
+        // Base image
+        Ptr<ASTOp> base;
+        if ( Node* pBase = node.m_pBase.get() )
+        {
+            base = Generate( pBase );
+        }
+        else
+        {
+            // This argument is required
+            base = GenerateMissingImageCode( "Image Layer base", EImageFormat::IF_RGB_UBYTE,
+                                             node.m_errorContext );
+        }
+
+		EImageFormat baseFormat = base->GetImageDesc().m_format;
+        //base = GenerateImageFormat( base, EImageFormat::IF_RGB_UBYTE );
+        base = GenerateImageSize
+                ( base, FImageSize((uint16_t)m_imageState.back().m_imageRect.size[0],
+                                   (uint16_t)m_imageState.back().m_imageRect.size[1]) );
+        op->SetChild( op->op.args.ImageLayer.base, base);
+
+        // Mask of the effect
+        Ptr<ASTOp> mask = 0;
+        if ( Node* pMask = node.m_pMask.get() )
+        {
+            mask = Generate( pMask );
+            mask = GenerateImageFormat( mask, EImageFormat::IF_L_UBYTE );
+            mask = GenerateImageSize
+                    ( mask, FImageSize((uint16_t)m_imageState.back().m_imageRect.size[0],
+                                       (uint16_t)m_imageState.back().m_imageRect.size[1]) );
+        }
+        op->SetChild( op->op.args.ImageLayer.mask, mask);
+
+        // Image to apply
+        Ptr<ASTOp> blended = 0;
+        if ( Node* pBlended = node.m_pBlended.get() )
+        {
+            blended = Generate( pBlended );
+        }
+        else
+        {
+            // This argument is required
+            blended = GeneratePlainImageCode( vec3<float>( 1,1,0 ) );
+        }
+        //blended = GenerateImageFormat( blended, EImageFormat::IF_RGB_UBYTE );
+        blended = GenerateImageFormat( blended, baseFormat );
+        blended = GenerateImageSize
+                ( blended, FImageSize((uint16_t)m_imageState.back().m_imageRect.size[0],
+                                      (uint16_t)m_imageState.back().m_imageRect.size[1]) );
+        op->SetChild( op->op.args.ImageLayer.blended, blended);
+
+        result.op = op;
+    }
+
+
+    //---------------------------------------------------------------------------------------------
+	void CodeGenerator::GenerateImage_LayerColour(IMAGE_GENERATION_RESULT& result, const NodeImageLayerColour* InNode)
+	{
+		const NodeImageLayerColour::Private& node = *InNode->GetPrivate();
+		
+		MUTABLE_CPUPROFILER_SCOPE(NodeImageLayerColour);
+
+        Ptr<ASTOpFixed> op = new ASTOpFixed();
+		op->op.type = OP_TYPE::IM_LAYERCOLOUR;
+		op->op.args.ImageLayerColour.blendType = uint8(node.m_type);
+
+        // Base image
+        Ptr<ASTOp> base;
+        if ( Node* pBase = node.m_pBase.get() )
+        {
+            base = Generate( pBase );
+        }
+        else
+        {
+            // This argument is required
+            base = GenerateMissingImageCode( "Layer base image", EImageFormat::IF_RGB_UBYTE, node.m_errorContext );
+        }
+        base = GenerateImageFormat( base, EImageFormat::IF_RGB_UBYTE );
+        base = GenerateImageSize
+                ( base, FImageSize((uint16_t)m_imageState.back().m_imageRect.size[0],(uint16_t)m_imageState.back().m_imageRect.size[1]) );
+        op->SetChild( op->op.args.ImageLayerColour.base, base);
+
+        // Mask of the effect
+        Ptr<ASTOp> mask = 0;
+        if ( Node* pMask = node.m_pMask.get() )
+        {
+            mask = Generate( pMask );
+            mask = GenerateImageFormat( mask, EImageFormat::IF_L_UBYTE );
+            mask = GenerateImageSize
+                    ( mask, FImageSize((uint16_t)m_imageState.back().m_imageRect.size[0],(uint16_t)m_imageState.back().m_imageRect.size[1]) );
+        }
+        op->SetChild( op->op.args.ImageLayerColour.mask, mask);
+
+        // Colour to apply
+        Ptr<ASTOp> colour = 0;
+        if ( Node* pColour = node.m_pColour.get() )
+        {
+            colour = Generate( pColour );
+        }
+        else
+        {
+            // This argument is required
+            colour = GenerateMissingColourCode( "Layer colour", node.m_errorContext );
+        }
+        op->SetChild( op->op.args.ImageLayerColour.colour, colour);
+
+        result.op = op;
+    }
+
+
+    //---------------------------------------------------------------------------------------------
+	void CodeGenerator::GenerateImage_MultiLayer(IMAGE_GENERATION_RESULT& result, const NodeImageMultiLayer* InNode)
+	{
+		const NodeImageMultiLayer::Private& node = *InNode->GetPrivate();
+
+		MUTABLE_CPUPROFILER_SCOPE(NodeImageMultiLayer);
+
+        Ptr<ASTOpImageMultiLayer> op = new ASTOpImageMultiLayer();
+
+		op->blendType = node.m_type;
+
+        // Base image
+        Ptr<ASTOp> base;
+        if ( Node* pBase = node.m_pBase.get() )
+        {
+            base = Generate( pBase );
+        }
+        else
+        {
+            // This argument is required
+            base = GenerateMissingImageCode( "Image MultiLayer base", EImageFormat::IF_RGB_UBYTE,
+                                             node.m_errorContext );
+        }
+
+		EImageFormat baseFormat = base->GetImageDesc().m_format;
+        base = GenerateImageSize
+                ( base, FImageSize((uint16_t)m_imageState.back().m_imageRect.size[0],
+                                   (uint16_t)m_imageState.back().m_imageRect.size[1]) );
+        op->base = base;
+
+        // Mask of the effect
+        Ptr<ASTOp> mask;
+        if ( Node* pMask = node.m_pMask.get() )
+        {
+            mask = Generate( pMask );
+            mask = GenerateImageFormat( mask, EImageFormat::IF_L_UBYTE );
+            mask = GenerateImageSize
+                    ( mask, FImageSize((uint16_t)m_imageState.back().m_imageRect.size[0],
+                                       (uint16_t)m_imageState.back().m_imageRect.size[1]) );
+        }
+        op->mask = mask;
+
+        // Image to apply
+        Ptr<ASTOp> blended;
+        if ( Node* pBlended = node.m_pBlended.get() )
+        {
+            blended = Generate( pBlended );
+        }
+        else
+        {
+            // This argument is required
+            blended = GeneratePlainImageCode( vec3<float>( 1,1,0 ) );
+        }
+        //blended = GenerateImageFormat( blended, EImageFormat::IF_RGB_UBYTE );
+        blended = GenerateImageFormat( blended, baseFormat );
+        blended = GenerateImageSize
+                ( blended, FImageSize((uint16_t)m_imageState.back().m_imageRect.size[0],
+                                      (uint16_t)m_imageState.back().m_imageRect.size[1]) );
+        op->blend = blended;
+
+        // Range of iteration
+        if ( node.m_pRange )
+        {
+            //op->range = Generate( pRange );
+            RANGE_GENERATION_RESULT rangeResult;
+            GenerateRange( rangeResult, node.m_pRange );
+
+            op->range.rangeSize = rangeResult.sizeOp;
+            op->range.rangeName = rangeResult.rangeName;
+            op->range.rangeUID = rangeResult.rangeUID;
+        }
+
+        result.op = op;
+    }
+
+
+	//---------------------------------------------------------------------------------------------
+	void CodeGenerator::GenerateImage_NormalComposite(IMAGE_GENERATION_RESULT& result, const NodeImageNormalComposite* InNode)
+	{
+		const NodeImageNormalComposite::Private& node = *InNode->GetPrivate();
+
+		MUTABLE_CPUPROFILER_SCOPE(NodeImageNormalComposite);
+
+        Ptr<ASTOpImageNormalComposite> op = new ASTOpImageNormalComposite();
+
+		op->Mode = node.m_mode; 
+		op->Power = node.m_power;
+
+        // Base image
+        Ptr<ASTOp> base;
+        if ( Node* pBase = node.m_pBase.get() )
+        {
+            base = Generate( pBase );
+        }
+        else
+        {
+            // This argument is required
+            base = GenerateMissingImageCode( "Image Composite Base", EImageFormat::IF_RGB_UBYTE,
+                                             node.m_errorContext );
+        }
+
+		EImageFormat baseFormat = base->GetImageDesc().m_format;
+        base = GenerateImageSize
+                ( base, FImageSize((uint16_t)m_imageState.back().m_imageRect.size[0],
+                                   (uint16_t)m_imageState.back().m_imageRect.size[1]) );
+        op->Base = base;
+
+        Ptr<ASTOp> normal;
+        if ( Node* pNormal = node.m_pNormal.get() )
+        {
+            normal = Generate( pNormal );
+            normal = GenerateImageFormat( normal, EImageFormat::IF_RGB_UBYTE );
+            normal = GenerateImageSize
+                    ( normal, FImageSize((uint16_t)m_imageState.back().m_imageRect.size[0],
+                                         (uint16_t)m_imageState.back().m_imageRect.size[1]) );
+        }
+		else
+		{
+            // This argument is required
+            normal = GenerateMissingImageCode( "Image Composite Normal", EImageFormat::IF_RGB_UBYTE,
+                                               node.m_errorContext );
+		}
+
+        op->Normal = normal;
+        
+		result.op = op;
+    }
+
+	void CodeGenerator::GenerateImage_Transform(IMAGE_GENERATION_RESULT& result, const NodeImageTransform* InNode)
+    {
+		const NodeImageTransform::Private& node = *InNode->GetPrivate();
+
+        MUTABLE_CPUPROFILER_SCOPE(NodeImageTransform);
+
+        Ptr<ASTOpImageTransform> op = new ASTOpImageTransform();
+
+		Ptr<ASTOp> OffsetX;
+		if (node.m_pOffsetX)
+		{
+			OffsetX = Generate(node.m_pOffsetX);
+		}
+
+		Ptr<ASTOp> OffsetY;
+		if (node.m_pOffsetY)
+		{
+			OffsetY = Generate(node.m_pOffsetY);
+		}
+	
+		Ptr<ASTOp> ScaleX;
+		if (node.m_pScaleX)
+		{
+			ScaleX = Generate(node.m_pScaleX);
+		}
+	
+		Ptr<ASTOp> ScaleY;
+		if (node.m_pScaleY)
+		{
+			ScaleY = Generate(node.m_pScaleY);
+		}
+
+		Ptr<ASTOp> Rotation;
+		if (node.m_pRotation)
+		{
+			Rotation = Generate(node.m_pRotation);
+		}
+
+		// If one of the inputs (offset or scale) is missig assume unifrom translation/scaling 
+		op->offsetX = OffsetX ? OffsetX : OffsetY;
+		op->offsetY = OffsetY ? OffsetY : OffsetX;
+ 		op->scaleX = ScaleX ? ScaleX : ScaleY;
+		op->scaleY = ScaleY ? ScaleY : ScaleX;
+		op->rotation = Rotation; 
+
+        // Base image
+        Ptr<ASTOp> base;
+        if ( Node* pBase = node.m_pBase.get() )
+        {
+            base = Generate( pBase );
+        }
+        else
+        {
+            // This argument is required
+            base = GenerateMissingImageCode( "Image Transform Base", EImageFormat::IF_RGB_UBYTE,
+                                             node.m_errorContext );
+        }
+
+		EImageFormat baseFormat = base->GetImageDesc().m_format;
+        base = GenerateImageSize
+                ( base, FImageSize((uint16_t)m_imageState.back().m_imageRect.size[0],
+									  (uint16_t)m_imageState.back().m_imageRect.size[1]) );
+        op->base = base;
+
+        result.op = op; 
+    }
+
+    //---------------------------------------------------------------------------------------------
+	void CodeGenerator::GenerateImage_Interpolate(IMAGE_GENERATION_RESULT& result, const NodeImageInterpolate* InNode)
+	{
+		const NodeImageInterpolate::Private& node = *InNode->GetPrivate();
+
+		MUTABLE_CPUPROFILER_SCOPE(NodeImageInterpolate);
+
+        Ptr<ASTOpFixed> op = new ASTOpFixed();
+        op->op.type = OP_TYPE::IM_INTERPOLATE;
+
+        // Factor
+        if ( Node* pFactor = node.m_pFactor.get() )
+        {
+            op->SetChild( op->op.args.ImageInterpolate.factor, Generate( pFactor ));
+        }
+        else
+        {
+            // This argument is required
+            op->SetChild( op->op.args.ImageInterpolate.factor,
+                          GenerateMissingScalarCode( "Interpolation factor", 0.5f, node.m_errorContext ));
+        }
+
+        // Target images
+        int numTargets = 0;
+
+        for ( std::size_t t=0
+            ; t< node.m_targets.size() && numTargets<MUTABLE_OP_MAX_INTERPOLATE_COUNT
+            ; ++t )
+        {
+            if ( Node* pA = node.m_targets[t].get() )
+            {
+                Ptr<ASTOp> target = Generate( pA );
+
+                // TODO: Support other formats
+                target = GenerateImageFormat( target, EImageFormat::IF_RGB_UBYTE );
+                target = GenerateImageSize
+                        ( target, FImageSize((uint16_t)m_imageState.back().m_imageRect.size[0],(uint16_t)m_imageState.back().m_imageRect.size[1]) );
+
+                op->SetChild( op->op.args.ImageInterpolate.targets[numTargets], target);
+                numTargets++;
+            }
+        }
+
+        // At least one target is required
+        if (!op->op.args.ImageInterpolate.targets[0])
+        {
+            Ptr<ASTOp> target = GenerateMissingImageCode( "First interpolation image", EImageFormat::IF_RGB_UBYTE,
+                                                       node.m_errorContext );
+            target = GenerateImageSize
+                    ( target, FImageSize((uint16_t)m_imageState.back().m_imageRect.size[0],(uint16_t)m_imageState.back().m_imageRect.size[1]) );
+            op->SetChild( op->op.args.ImageInterpolate.targets[0], target);
+        }
+
+        result.op = op;
+    }
+
+
+    //---------------------------------------------------------------------------------------------
+	void CodeGenerator::GenerateImage_Interpolate3(IMAGE_GENERATION_RESULT& result, const NodeImageInterpolate3* InNode)
+	{
+		const NodeImageInterpolate3::Private& node = *InNode->GetPrivate();
+
+		MUTABLE_CPUPROFILER_SCOPE(NodeImageInterpolate3);
+
+        Ptr<ASTOpFixed> op = new ASTOpFixed();
+        op->op.type = OP_TYPE::IM_INTERPOLATE3;
+
+        // Factors
+        if ( node.m_pFactor1.get() )
+        {
+            op->SetChild( op->op.args.ImageInterpolate3.factor1, Generate( node.m_pFactor1 ));
+        }
+        else
+        {
+            // This argument is required
+            op->SetChild( op->op.args.ImageInterpolate3.factor1,
+                          GenerateMissingScalarCode( "Interpolation factor 1", 0.3f, node.m_errorContext ));
+        }
+
+        if ( node.m_pFactor2.get() )
+        {
+            op->SetChild( op->op.args.ImageInterpolate3.factor2, Generate( node.m_pFactor2 ));
+        }
+        else
+        {
+            // This argument is required
+            op->SetChild( op->op.args.ImageInterpolate3.factor2,
+                          GenerateMissingScalarCode( "Interpolation factor 2", 0.3f, node.m_errorContext ));
+        }
+
+        // Target 0
+        Ptr<ASTOp> target;
+        if ( node.m_pTarget0.get() )
+        {
+            target = Generate( node.m_pTarget0 );
+        }
+        else
+        {
+            // This argument is required
+            target = GenerateMissingImageCode( "Interpolation target 0", EImageFormat::IF_RGB_UBYTE,
+                                               node.m_errorContext );
+        }
+        target = GenerateImageFormat( target, EImageFormat::IF_RGB_UBYTE );
+        target = GenerateImageSize
+                ( target, FImageSize((uint16_t)m_imageState.back().m_imageRect.size[0],(uint16_t)m_imageState.back().m_imageRect.size[1]) );
+        op->SetChild( op->op.args.ImageInterpolate3.target0, target);
+
+        // Target 1
+        if ( node.m_pTarget1.get() )
+        {
+            target = Generate( node.m_pTarget1 );
+        }
+        else
+        {
+            // This argument is required
+            target = GenerateMissingImageCode( "Interpolation target 1", EImageFormat::IF_RGB_UBYTE,
+                                               node.m_errorContext );
+        }
+        target = GenerateImageFormat( target, EImageFormat::IF_RGB_UBYTE );
+        target = GenerateImageSize
+                ( target, FImageSize( (uint16_t)m_imageState.back().m_imageRect.size[0],
+                                      (uint16_t)m_imageState.back().m_imageRect.size[1]) );
+        op->SetChild( op->op.args.ImageInterpolate3.target1, target);
+
+        // Target 2
+        if ( node.m_pTarget2.get() )
+        {
+            target = Generate( node.m_pTarget2 );
+        }
+        else
+        {
+            // This argument is required
+            target = GenerateMissingImageCode( "Interpolation target 2", EImageFormat::IF_RGB_UBYTE,
+                                               node.m_errorContext );
+        }
+        target = GenerateImageFormat( target, EImageFormat::IF_RGB_UBYTE );
+        target = GenerateImageSize
+                ( target, FImageSize( (uint16_t)m_imageState.back().m_imageRect.size[0],
+                                      (uint16_t)m_imageState.back().m_imageRect.size[1] ) );
+        op->SetChild( op->op.args.ImageInterpolate3.target2, target);
+
+
+        result.op = op;
+    }
+
+
+	//---------------------------------------------------------------------------------------------
+	void CodeGenerator::GenerateImage_Swizzle(IMAGE_GENERATION_RESULT& result, const NodeImageSwizzle* InNode)
+	{
+		const NodeImageSwizzle::Private& node = *InNode->GetPrivate();
+
+		//MUTABLE_CPUPROFILER_SCOPE(NodeImageSwizzle);
+
+        // This node always produces a swizzle operation and sometimes it may produce a pixelformat
+		// operation to compress the result
+        Ptr<ASTOpFixed> op = new ASTOpFixed();
+        op->op.type = OP_TYPE::IM_SWIZZLE;
+
+		// Format
+		EImageFormat compressedFormat = EImageFormat::IF_NONE;
+
+		switch (node.m_format)
+		{
+        case EImageFormat::IF_BC1:
+        case EImageFormat::IF_ASTC_4x4_RGB_LDR:
+            compressedFormat = node.m_format;
+            op->op.args.ImageSwizzle.format = node.m_sources[3] ? EImageFormat::IF_RGBA_UBYTE : EImageFormat::IF_RGB_UBYTE;
+			break;
+
+		case EImageFormat::IF_BC2:
+		case EImageFormat::IF_BC3:
+		case EImageFormat::IF_BC6:
+        case EImageFormat::IF_BC7:
+        case EImageFormat::IF_ASTC_4x4_RGBA_LDR:
+            compressedFormat = node.m_format;
+             op->op.args.ImageSwizzle.format = EImageFormat::IF_RGBA_UBYTE;
+			break;
+
+		case EImageFormat::IF_BC4:
+			compressedFormat = node.m_format;
+            op->op.args.ImageSwizzle.format = EImageFormat::IF_L_UBYTE;
+			break;
+
+		case EImageFormat::IF_BC5:
+        case EImageFormat::IF_ASTC_4x4_RG_LDR:
+            compressedFormat = node.m_format;
+			// TODO: Should be RG
+            op->op.args.ImageSwizzle.format = EImageFormat::IF_RGB_UBYTE;
+			break;
+
+		default:
+            op->op.args.ImageSwizzle.format = node.m_format;
+			break;
+
+		}
+
+		check(node.m_format != EImageFormat::IF_NONE);
+
+		// Source images and channels
+		check(node.m_sources.size() == node.m_sourceChannels.size());
+
+		// First source, for reference in the size
+        Ptr<ASTOp> first = 0;
+		for (std::size_t t = 0; t< node.m_sources.size(); ++t)
+		{
+			if (Node* pChannel = node.m_sources[t].get())
+			{
+                Ptr<ASTOp> source = Generate(pChannel);
+				source = GenerateImageUncompressed(source);
+
+				if (!source)
+				{
+					// TODO: Warn?
+					source = GenerateMissingImageCode("Swizzle channel", EImageFormat::IF_L_UBYTE,
+						node.m_errorContext);
+				}
+
+                Ptr<ASTOp> sizedSource;
+				if (first)
+				{
+                    Ptr<ASTOpFixed> sop = new ASTOpFixed();
+                    sop->op.type = OP_TYPE::IM_RESIZELIKE;
+                    sop->SetChild( sop->op.args.ImageResizeLike.source, source);
+                    sop->SetChild( sop->op.args.ImageResizeLike.sizeSource, first);
+                    sizedSource = sop;
+				}
+				else
+				{
+					first = source;
+					sizedSource = source;
+				}
+
+                op->SetChild( op->op.args.ImageSwizzle.sources[t], sizedSource);
+                op->op.args.ImageSwizzle.sourceChannels[t] = (uint8_t)node.m_sourceChannels[t];
+			}
+		}
+
+		// At least one source is required
+        if (!op->op.args.ImageSwizzle.sources[0])
+		{
+            Ptr<ASTOp> source = GenerateMissingImageCode("First swizzle image", EImageFormat::IF_RGBA_UBYTE,
+				node.m_errorContext);
+            op->SetChild( op->op.args.ImageSwizzle.sources[0], source);
+		}
+
+        Ptr<ASTOp> resultOp = op;
+
+		if (compressedFormat != EImageFormat::IF_NONE)
+		{
+            Ptr<ASTOpImagePixelFormat> fop = new ASTOpImagePixelFormat();
+            fop->Source = resultOp;
+            fop->Format = compressedFormat;
+			resultOp = fop;
+		}
+
+        result.op = resultOp;
+	}
+
+
+	//---------------------------------------------------------------------------------------------
+	void CodeGenerator::GenerateImage_Format(IMAGE_GENERATION_RESULT& result, const NodeImageFormat* InNode)
+	{
+		const NodeImageFormat::Private& node = *InNode->GetPrivate();
+
+		MUTABLE_CPUPROFILER_SCOPE(NodeImageFormat);
+
+        check(node.m_format != EImageFormat::IF_NONE);
+
+        Ptr<ASTOpImagePixelFormat> fop = new ASTOpImagePixelFormat();
+        fop->Format = node.m_format;
+        fop->FormatIfAlpha = node.m_formatIfAlpha;
+
+		// Source is required
+		if (!node.m_source)
+		{
+            fop->Source = GenerateMissingImageCode("Source image for format.", EImageFormat::IF_RGBA_UBYTE, node.m_errorContext);
+		}
+		else
+		{
+            fop->Source = Generate(node.m_source.get());
+		}
+
+        result.op = fop;
+	}
+
+
+    //---------------------------------------------------------------------------------------------
+	void CodeGenerator::GenerateImage_Saturate(IMAGE_GENERATION_RESULT& result, const NodeImageSaturate* InNode)
+	{
+		const NodeImageSaturate::Private& node = *InNode->GetPrivate();
+
+		Ptr<ASTOpFixed> op = new ASTOpFixed();
+        op->op.type = OP_TYPE::IM_SATURATE;
+
+
+        // Source image
+        Ptr<ASTOp> base = 0;
+        if ( Node* pSource = node.m_pSource.get() )
+        {
+            base = Generate( pSource );
+        }
+        else
+        {
+            // This argument is required
+            base = GenerateMissingImageCode( "Saturate image", EImageFormat::IF_RGB_UBYTE, node.m_errorContext );
+        }
+        base = GenerateImageFormat( base, EImageFormat::IF_RGB_UBYTE );
+        base = GenerateImageSize
+                ( base, FImageSize((uint16_t)m_imageState.back().m_imageRect.size[0],(uint16_t)m_imageState.back().m_imageRect.size[1]) );
+        op->SetChild( op->op.args.ImageSaturate.base, base);
+
+
+        // Factor
+        if ( Node* pFactor = node.m_pFactor.get() )
+        {
+            op->SetChild( op->op.args.ImageSaturate.factor, Generate( pFactor ));
+        }
+        else
+        {
+            // This argument is required
+            op->SetChild( op->op.args.ImageSaturate.factor, GenerateMissingScalarCode( "Saturation factor", 0.5f,
+                                                                      node.m_errorContext ) );
+        }
+
+        result.op = op;
+    }
+
+
+    //---------------------------------------------------------------------------------------------
+	void CodeGenerator::GenerateImage_Luminance(IMAGE_GENERATION_RESULT& result, const NodeImageLuminance* InNode)
+	{
+		const NodeImageLuminance::Private& node = *InNode->GetPrivate();
+
+		Ptr<ASTOpFixed> op = new ASTOpFixed();
+        op->op.type = OP_TYPE::IM_LUMINANCE;
+
+        // Source image
+        Ptr<ASTOp> base = 0;
+        if ( Node* pSource = node.m_pSource.get() )
+        {
+            base = Generate( pSource );
+        }
+        else
+        {
+            // This argument is required
+            base = GenerateMissingImageCode( "Image luminance", EImageFormat::IF_RGB_UBYTE,
+                                             node.m_errorContext );
+        }
+        base = GenerateImageFormat( base, EImageFormat::IF_RGB_UBYTE );
+        op->SetChild( op->op.args.ImageLuminance.base, base);
+
+        result.op = op;
+    }
+
+
+    //---------------------------------------------------------------------------------------------
+	void CodeGenerator::GenerateImage_SelectColour(IMAGE_GENERATION_RESULT& result, const NodeImageSelectColour* InNode)
+	{
+		const NodeImageSelectColour::Private& node = *InNode->GetPrivate();
+
+		Ptr<ASTOpFixed> op = new ASTOpFixed();
+        op->op.type = OP_TYPE::IM_SELECTCOLOUR;
+
+
+        // Source image
+        Ptr<ASTOp> base = 0;
+        if ( Node* pSource = node.m_pSource.get() )
+        {
+            base = Generate( pSource );
+        }
+        else
+        {
+            // This argument is required
+            base = GenerateMissingImageCode( "Image select colour base", EImageFormat::IF_RGB_UBYTE,
+                                             node.m_errorContext );
+        }
+        base = GenerateImageSize
+                ( base, FImageSize((uint16_t)m_imageState.back().m_imageRect.size[0],(uint16_t)m_imageState.back().m_imageRect.size[1]) );
+        op->SetChild( op->op.args.ImageSelectColour.base, base);
+
+
+        // Factor
+        if ( Node* pColour = node.m_pColour.get() )
+        {
+            op->SetChild( op->op.args.ImageSelectColour.colour, Generate( pColour ));
+        }
+        else
+        {
+            // This argument is required
+            op->SetChild( op->op.args.ImageSelectColour.colour, GenerateMissingColourCode( "Selected colour",
+                                                                          node.m_errorContext ));
+        }
+
+        result.op = op;
+    }
+
+
+    //---------------------------------------------------------------------------------------------
+	void CodeGenerator::GenerateImage_Difference(IMAGE_GENERATION_RESULT& result, const NodeImageDifference* InNode)
+	{
+		const NodeImageDifference::Private& node = *InNode->GetPrivate();
+
+		Ptr<ASTOpFixed> op = new ASTOpFixed();
+        op->op.type = OP_TYPE::IM_DIFFERENCE;
+
+        // A image
+        Ptr<ASTOp> a = 0;
+        if ( Node* pA = node.m_pA.get() )
+        {
+            a = Generate( pA );
+        }
+        else
+        {
+            // This argument is required
+            a = GenerateMissingImageCode( "Image Difference A", EImageFormat::IF_RGB_UBYTE, node.m_errorContext );
+        }
+        a = GenerateImageFormat( a, EImageFormat::IF_RGB_UBYTE );
+        a = GenerateImageSize( a,
+                               FImageSize( (uint16_t)m_imageState.back().m_imageRect.size[0],
+                                           (uint16_t)m_imageState.back().m_imageRect.size[1]) );
+        op->SetChild( op->op.args.ImageDifference.a, a);
+
+        // B image
+        Ptr<ASTOp> b = 0;
+        if ( Node* pB = node.m_pB.get() )
+        {
+            b = Generate( pB );
+        }
+        else
+        {
+            // This argument is required
+            b = GenerateMissingImageCode( "Image Difference B", EImageFormat::IF_RGB_UBYTE, node.m_errorContext );
+        }
+        b = GenerateImageFormat( b, EImageFormat::IF_RGB_UBYTE );
+        b = GenerateImageSize
+                ( b, FImageSize((uint16_t)m_imageState.back().m_imageRect.size[0],(uint16_t)m_imageState.back().m_imageRect.size[1]) );
+        op->SetChild( op->op.args.ImageDifference.b, b);
+
+        result.op = op;
+    }
+
+
+    //---------------------------------------------------------------------------------------------
+	void CodeGenerator::GenerateImage_ColourMap(IMAGE_GENERATION_RESULT& result, const NodeImageColourMap* InNode)
+	{
+		const NodeImageColourMap::Private& node = *InNode->GetPrivate();
+
+		Ptr<ASTOpFixed> op = new ASTOpFixed();
+        op->op.type = OP_TYPE::IM_COLOURMAP;
+
+        // Base image
+        Ptr<ASTOp> base = 0;
+        if ( Node* pBase = node.m_pBase.get() )
+        {
+            base = Generate( pBase );
+        }
+        else
+        {
+            // This argument is required
+            base = GenerateMissingImageCode( "Colourmap base image", EImageFormat::IF_RGB_UBYTE,
+                                             node.m_errorContext );
+        }
+        base = GenerateImageSize
+                ( base, FImageSize( (uint16_t)m_imageState.back().m_imageRect.size[0],
+                                    (uint16_t)m_imageState.back().m_imageRect.size[1]) );
+        op->SetChild( op->op.args.ImageColourMap.base, base);
+
+        // Mask of the effect
+        Ptr<ASTOp> mask = 0;
+        if ( Node* pMask = node.m_pMask.get() )
+        {
+            mask = Generate( pMask );
+        }
+        else
+        {
+            // Set the argument default value: affect all pixels.
+            // TODO: Special operation code without mask
+            mask = GeneratePlainImageCode( vec3<float>( 1,1,1 ) );
+        }
+        mask = GenerateImageFormat( mask, EImageFormat::IF_L_UBYTE );
+        mask = GenerateImageSize
+                ( mask, FImageSize( (uint16_t)m_imageState.back().m_imageRect.size[0],
+                                    (uint16_t)m_imageState.back().m_imageRect.size[1]) );
+        op->SetChild( op->op.args.ImageColourMap.mask, mask);
+
+        // Map image
+        Ptr<ASTOp> mapImage = 0;
+        if ( Node* pMap = node.m_pMap.get() )
+        {
+            mapImage = Generate( pMap );
+        }
+        else
+        {
+            // This argument is required
+            mapImage = GenerateMissingImageCode( "Map image", EImageFormat::IF_RGB_UBYTE, node.m_errorContext );
+        }
+        mapImage = GenerateImageFormat( mapImage, EImageFormat::IF_RGB_UBYTE );
+        mapImage = GenerateImageSize
+                ( mapImage, FImageSize( (uint16_t)m_imageState.back().m_imageRect.size[0],
+                                        (uint16_t)m_imageState.back().m_imageRect.size[1]) );
+        op->SetChild( op->op.args.ImageColourMap.map, mapImage);
+
+        result.op = op;
+    }
+
+
+    //---------------------------------------------------------------------------------------------
+	void CodeGenerator::GenerateImage_Gradient(IMAGE_GENERATION_RESULT& result, const NodeImageGradient* InNode)
+	{
+		const NodeImageGradient::Private& node = *InNode->GetPrivate();
+
+		Ptr<ASTOpFixed> op = new ASTOpFixed();
+        op->op.type = OP_TYPE::IM_GRADIENT;
+
+        // First colour
+        Ptr<ASTOp> colour0 = 0;
+        if ( Node* pColour0 = node.m_pColour0.get() )
+        {
+            colour0 = Generate( pColour0 );
+        }
+        else
+        {
+            // This argument is required
+            colour0 = GenerateMissingColourCode( "Gradient colour 0", node.m_errorContext );
+        }
+        op->SetChild( op->op.args.ImageGradient.colour0, colour0);
+
+        // Second colour
+        Ptr<ASTOp> colour1 = 0;
+        if ( Node* pColour1 = node.m_pColour1.get() )
+        {
+            colour1 = Generate( pColour1 );
+        }
+        else
+        {
+            // This argument is required
+            colour1 = GenerateMissingColourCode( "Gradient colour 1", node.m_errorContext );
+        }
+        op->SetChild( op->op.args.ImageGradient.colour1, colour1);
+
+        op->op.args.ImageGradient.size[0] = (uint16_t)std::max( 2, std::min( 1024, node.m_size[0] ) );
+        op->op.args.ImageGradient.size[1] = (uint16_t)std::max( 1, std::min( 1024, node.m_size[1] ) );
+
+        result.op = op;
+    }
+
+
+    //---------------------------------------------------------------------------------------------
+	void CodeGenerator::GenerateImage_Binarise(IMAGE_GENERATION_RESULT& result, const NodeImageBinarise* InNode)
+	{
+		const NodeImageBinarise::Private& node = *InNode->GetPrivate();
+
+		Ptr<ASTOpFixed> op = new ASTOpFixed();
+        op->op.type = OP_TYPE::IM_BINARISE;
+
+        // A image
+        Ptr<ASTOp> a;
+        if ( Node* pBase = node.m_pBase.get() )
+        {
+            a = Generate( pBase );
+        }
+        else
+        {
+            // This argument is required
+            a = GenerateMissingImageCode( "Image Binarise Base", EImageFormat::IF_RGB_UBYTE,
+                                          node.m_errorContext );
+        }
+        a = GenerateImageFormat( a, EImageFormat::IF_RGB_UBYTE );
+        a = GenerateImageSize
+                ( a, FImageSize( (uint16_t)m_imageState.back().m_imageRect.size[0],
+                                 (uint16_t)m_imageState.back().m_imageRect.size[1]) );
+        op->SetChild( op->op.args.ImageBinarise.base, a );
+
+        // Threshold
+        Ptr<ASTOp> b = 0;
+        if ( Node* pScalar = node.m_pThreshold.get() )
+        {
+            b = Generate( pScalar );
+        }
+        else
+        {
+            // This argument is required
+            b = GenerateMissingScalarCode( "Image Binarise Threshold", 0.5f, node.m_errorContext );
+        }
+        op->SetChild( op->op.args.ImageBinarise.threshold, b );
+
+        result.op = op;
+    }
+
+
+    //---------------------------------------------------------------------------------------------
+	void CodeGenerator::GenerateImage_Resize(IMAGE_GENERATION_RESULT& result, const NodeImageResize* InNode)
+	{
+		const NodeImageResize::Private& node = *InNode->GetPrivate();
+
+		MUTABLE_CPUPROFILER_SCOPE(NodeImageResize);
+
+        Ptr<ASTOp> at = 0;
+
+        // Source image
+        Ptr<ASTOp> base = 0;
+        if ( Node* pSource = node.m_pBase.get() )
+        {
+            IMAGE_STATE newState = m_imageState.back();
+            if ( node.m_relative )
+            {
+                newState.m_imageSize[0]=(int)std::roundf(newState.m_imageSize[0]/node.m_sizeX);
+                newState.m_imageSize[1]=(int)std::roundf(newState.m_imageSize[1]/node.m_sizeY);
+                newState.m_imageRect.min[0]=(int)std::roundf(newState.m_imageRect.min[0]/node.m_sizeX);
+                newState.m_imageRect.min[1]=(int)std::roundf(newState.m_imageRect.min[1]/node.m_sizeY);
+                newState.m_imageRect.size[0]=(int)std::roundf(newState.m_imageRect.size[0]/node.m_sizeX);
+                newState.m_imageRect.size[1]=(int)std::roundf(newState.m_imageRect.size[1]/node.m_sizeY);
+            }
+
+            m_imageState.push_back(newState);
+
+            base = Generate( pSource );
+
+            m_imageState.pop_back();
+        }
+        else
+        {
+            // This argument is required
+            base = GenerateMissingImageCode( "Image resize base", EImageFormat::IF_RGB_UBYTE, node.m_errorContext );
+        }
+
+        // Size
+        if ( node.m_relative )
+        {
+            Ptr<ASTOpFixed> op = new ASTOpFixed();
+            op->op.type = OP_TYPE::IM_RESIZEREL;
+            op->op.args.ImageResizeRel.factor[0] = node.m_sizeX;
+            op->op.args.ImageResizeRel.factor[1] = node.m_sizeY;
+            op->SetChild( op->op.args.ImageResizeRel.source, base);
+            at = op;
+        }
+        else
+        {
+            Ptr<ASTOpFixed> op = new ASTOpFixed();
+            op->op.type = OP_TYPE::IM_RESIZE;
+            op->op.args.ImageResize.size[0] = (uint16_t)node.m_sizeX;
+            op->op.args.ImageResize.size[1] = (uint16_t)node.m_sizeY;
+            op->SetChild( op->op.args.ImageResize.source, base);
+            at = op;
+        }
+
+        result.op = at;
+    }
+
+
+    //---------------------------------------------------------------------------------------------
+	void CodeGenerator::GenerateImage_PlainColour(IMAGE_GENERATION_RESULT& result, const NodeImagePlainColour* InNode)
+	{
+		const NodeImagePlainColour::Private& node = *InNode->GetPrivate();
+
+		// Source colour
+        Ptr<ASTOp> base = 0;
+        if ( node.m_pColour )
+        {
+            base = Generate( node.m_pColour.get() );
+        }
+        else
+        {
+            // This argument is required
+            base = GenerateMissingColourCode( "Image plain colour base", node.m_errorContext );
+        }
+
+        Ptr<ASTOpFixed> op = new ASTOpFixed();
+        op->op.type = OP_TYPE::IM_PLAINCOLOUR;
+        op->SetChild( op->op.args.ImagePlainColour.colour, base);
+        op->op.args.ImagePlainColour.format = EImageFormat::IF_RGB_UBYTE;
+        op->op.args.ImagePlainColour.size[0] = uint16_t(node.m_sizeX);
+        op->op.args.ImagePlainColour.size[1] = uint16_t(node.m_sizeY);
+
+        Ptr<ASTOpFixed> opSize = new ASTOpFixed();
+        opSize->op.type = OP_TYPE::IM_RESIZE;
+        if (m_imageState.size())
+        {
+            opSize->op.args.ImageResize.size[0] = (uint16_t)m_imageState.back().m_imageRect.size[0];
+            opSize->op.args.ImageResize.size[1] = (uint16_t)m_imageState.back().m_imageRect.size[1];
+        }
+        else
+        {
+            opSize->op.args.ImageResize.size[0] = (uint16_t)node.m_sizeX;
+            opSize->op.args.ImageResize.size[1] = (uint16_t)node.m_sizeY;
+        }
+        opSize->SetChild( opSize->op.args.ImageResize.source, op);
+
+        result.op = opSize;
+    }
+
+
+    //---------------------------------------------------------------------------------------------
+	void CodeGenerator::GenerateImage_Switch(IMAGE_GENERATION_RESULT& result, const NodeImageSwitch* InNode)
+	{
+		const NodeImageSwitch::Private& node = *InNode->GetPrivate();
+
+		MUTABLE_CPUPROFILER_SCOPE(NodeImageSwitch);
+
+        if (node.m_options.size() == 0)
+		{
+			// No options in the switch!
+            Ptr<ASTOp> missingOp = GenerateMissingImageCode("Switch option",
+				EImageFormat::IF_RGBA_UBYTE,
+				node.m_errorContext);
+			result.op = missingOp;
+			return;
+		}
+
+        Ptr<ASTOpSwitch> op = new ASTOpSwitch();
+        op->type = OP_TYPE::IM_SWITCH;
+
+		// Variable value
+		if ( node.m_pParameter )
+		{
+            op->variable = Generate( node.m_pParameter.get() );
+		}
+		else
+		{
+			// This argument is required
+            op->variable = GenerateMissingScalarCode( "Switch variable", 0.0f, node.m_errorContext );
+		}
+
+		// Options
+        for ( std::size_t t=0; t< node.m_options.size(); ++t )
+        {
+            Ptr<ASTOp> branch;
+
+            if (node.m_options[t])
+            {
+                branch = Generate( node.m_options[t].get() );
+            }
+            else
+            {
+                // This argument is required
+                branch = GenerateMissingImageCode("Switch option",
+                                                    EImageFormat::IF_RGBA_UBYTE,
+                                                    node.m_errorContext);
+            }
+
+            op->cases.emplace_back((int16_t)t,op,branch);
+        }
+
+        Ptr<ASTOp> switchAt = op;
+
+        // Make sure all options are the same format and size
+		// Disabled: This is not always desirable. For example if the image is going to be used in a 
+		// projector, the size doesn't need to be constrained.
+        //auto desc = switchAt->GetImageDesc( true );
+        //if ( desc.m_format == EImageFormat::IF_NONE )
+        //{
+        //    // TODO: Look for the most generic of the options?
+        //    // For now force a decently generic one
+        //    desc.m_format = EImageFormat::IF_RGBA_UBYTE;
+        //}
+
+        //if (desc.m_size[0]!=0 && desc.m_size[1]!=0)
+        //{
+        //    Ptr<ASTOpFixed> sop = new ASTOpFixed();
+        //    sop->op.type = OP_TYPE::IM_RESIZE;
+        //    sop->op.args.ImageResize.size[0] = desc.m_size[0];
+        //    sop->op.args.ImageResize.size[1] = desc.m_size[1];
+        //    sop->SetChild( sop->op.args.ImageResize.source, switchAt );
+        //    switchAt = sop;
+        //}
+
+        //{
+        //    Ptr<ASTOpFixed> fop = new ASTOpFixed();
+        //    fop->op.type = OP_TYPE::IM_PIXELFORMAT;
+        //    fop->op.args.ImagePixelFormat.format = desc.m_format;
+        //    fop->SetChild( fop->op.args.ImagePixelFormat.source, switchAt );
+        //    switchAt = fop;
+        //}
+
+        result.op = switchAt;
+    }
+
+
+    //---------------------------------------------------------------------------------------------
+	void CodeGenerator::GenerateImage_Conditional(IMAGE_GENERATION_RESULT& result, const NodeImageConditional* InNode)
+	{
+		const NodeImageConditional::Private& node = *InNode->GetPrivate();
+
+		Ptr<ASTOpConditional> op = new ASTOpConditional();
+        op->type = OP_TYPE::IM_CONDITIONAL;
+
+        // Condition
+        if ( node.m_parameter )
+        {
+            op->condition = Generate( node.m_parameter.get() );
+        }
+        else
+        {
+            // This argument is required
+            op->condition = GenerateMissingBoolCode( "Conditional condition", true, node.m_errorContext );
+        }
+
+        // Options
+        op->yes = Generate( node.m_true.get() );
+        op->no = Generate( node.m_false.get() );
+
+        result.op = op;
+    }
+
+
+    //---------------------------------------------------------------------------------------------
+	void CodeGenerator::GenerateImage_Project(IMAGE_GENERATION_RESULT& result, const NodeImageProject* InNode)
+	{
+		const NodeImageProject::Private& node = *InNode->GetPrivate();
+
+		MUTABLE_CPUPROFILER_SCOPE(NodeImageProject);
+
+        // Mesh project operation
+        //------------------------------
+        Ptr<ASTOpFixed> pop = new ASTOpFixed();
+        pop->op.type = OP_TYPE::ME_PROJECT;
+
+        Ptr<ASTOp> lastMeshOp = pop;
+
+        // Projector
+        PROJECTOR_GENERATION_RESULT projectorResult;
+        if ( node.m_pProjector )
+        {
+            GenerateProjector( projectorResult, node.m_pProjector );
+            //projectorAt = Generate( node.m_pProjector.get() );
+        }
+        else
+        {
+            // This argument is required
+            GenerateMissingProjectorCode( projectorResult, node.m_errorContext );
+        }
+
+        pop->SetChild( pop->op.args.MeshProject.projector, projectorResult.op );
+
+        // Mesh
+        if ( node.m_pMesh )
+        {
+            MESH_GENERATION_RESULT MeshResult;
+			check(m_overrideLayoutsStack.empty());
+            GenerateMesh(MeshResult, node.m_pMesh );
+
+            pop->SetChild( pop->op.args.MeshProject.mesh, MeshResult.meshOp );
+
+            if (projectorResult.type == PROJECTOR_TYPE::WRAPPING)
+            {
+                // For wrapping projector we need the entire mesh. The actual project operation
+                // will remove the faces that are not in the layout block we are generating.
+                Ptr<ASTOpConstantResource> cop = new ASTOpConstantResource();
+                cop->type = OP_TYPE::ME_CONSTANT;
+				MeshPtr pFormatMesh = CreateMeshOptimisedForWrappingProjection(node.m_layout);
+                cop->SetValue( pFormatMesh, m_compilerOptions->m_optimisationOptions.m_useDiskCache );
+
+                Ptr<ASTOpMeshFormat> fop = new ASTOpMeshFormat();
+                fop->Buffers = OP::MeshFormatArgs::BT_VERTEX
+                        | OP::MeshFormatArgs::BT_INDEX
+                        | OP::MeshFormatArgs::BT_FACE
+                        | OP::MeshFormatArgs::BT_RESETBUFFERINDICES;
+                fop->Format = cop;
+                fop->Source = pop->children[pop->op.args.MeshProject.mesh].child();
+                pop->SetChild( pop->op.args.MeshProject.mesh, fop );
+            }
+            else
+            {
+                // Extract the mesh layout block
+                if ( m_imageState.size() && m_imageState.back().m_layoutBlock>=0 )
+                {
+                    Ptr<ASTOpMeshExtractLayoutBlocks> eop = new ASTOpMeshExtractLayoutBlocks();
+                    eop->source = pop->children[pop->op.args.MeshProject.mesh].child();
+                    eop->layout = node.m_layout;
+                    eop->blocks.push_back( m_imageState.back().m_layoutBlock );
+
+                    pop->SetChild( pop->op.args.MeshProject.mesh, eop );
+                }
+
+                // Reformat the mesh to a more efficient format for this operation
+                Ptr<ASTOpConstantResource> cop = new ASTOpConstantResource();
+                cop->type = OP_TYPE::ME_CONSTANT;
+                MeshPtr pFormatMesh = CreateMeshOptimisedForProjection(node.m_layout);
+                cop->SetValue( pFormatMesh, m_compilerOptions->m_optimisationOptions.m_useDiskCache );
+
+                Ptr<ASTOpMeshFormat> fop = new ASTOpMeshFormat();
+                fop->Buffers = OP::MeshFormatArgs::BT_VERTEX
+					| OP::MeshFormatArgs::BT_INDEX
+					| OP::MeshFormatArgs::BT_FACE
+					| OP::MeshFormatArgs::BT_RESETBUFFERINDICES;
+				fop->Format = cop;
+                fop->Source = pop->children[pop->op.args.MeshProject.mesh].child();
+                pop->SetChild( pop->op.args.MeshProject.mesh, fop );
+            }
+        }
+        else
+        {
+            // This argument is required
+            MeshPtrConst pMesh = new Mesh();
+            Ptr<ASTOpConstantResource> cop = new ASTOpConstantResource();
+            cop->type = OP_TYPE::ME_CONSTANT;
+            cop->SetValue( pMesh, m_compilerOptions->m_optimisationOptions.m_useDiskCache );
+            pop->SetChild( pop->op.args.MeshProject.mesh, cop );
+            m_pErrorLog->GetPrivate()->Add( "Projector mesh not set.",
+                                            ELMT_ERROR, node.m_errorContext );
+        }
+
+
+        // Image raster operation
+        //------------------------------
+        Ptr<ASTOpFixed> op = new ASTOpFixed();
+        op->op.type = OP_TYPE::IM_RASTERMESH;
+        op->SetChild( op->op.args.ImageRasterMesh.mesh, lastMeshOp);
+        op->SetChild( op->op.args.ImageRasterMesh.projector, projectorResult.op);
+
+        // Image
+        if ( node.m_pImage )
+        {
+            // Remeber previous rect values
+            IMAGE_STATE newState;
+
+            // We take whatever size will be produced
+            FImageDesc desc = CalculateImageDesc( *node.m_pImage->GetBasePrivate() );
+            newState.m_imageSize = desc.m_size;
+            newState.m_imageRect.min[0] = 0;
+            newState.m_imageRect.min[1] = 0;
+            newState.m_imageRect.size = desc.m_size;
+            newState.m_layoutBlock = -1;
+            m_imageState.push_back( newState );
+
+            // Generate
+            op->SetChild( op->op.args.ImageRasterMesh.image, Generate( node.m_pImage.get() ) );
+
+            // Restore rect
+            m_imageState.pop_back();
+
+        }
+        else
+        {
+            // This argument is required
+            op->SetChild( op->op.args.ImageRasterMesh.image,
+                    GenerateMissingImageCode( "Projector image", EImageFormat::IF_RGB_UBYTE, node.m_errorContext ) );
+        }
+
+        // Image size, from the current block being generated
+        op->op.args.ImageRasterMesh.sizeX = (uint16_t)m_imageState.back().m_imageRect.size[0];
+        op->op.args.ImageRasterMesh.sizeY = (uint16_t)m_imageState.back().m_imageRect.size[1];
+        op->op.args.ImageRasterMesh.blockIndex = m_imageState.back().m_layoutBlock;
+
+        // Fading properties are optional, and stored in a colour
+        if (node.m_pAngleFadeStart||node.m_pAngleFadeEnd)
+        {
+            NodeScalarConstantPtr pDefaultFade = new NodeScalarConstant();
+            pDefaultFade->SetValue( 180.0f );
+
+            NodeColourFromScalarsPtr pPropsNode = new NodeColourFromScalars();
+
+            if (node.m_pAngleFadeStart) pPropsNode->SetX(node.m_pAngleFadeStart);
+            else pPropsNode->SetX(pDefaultFade);
+
+            if (node.m_pAngleFadeEnd) pPropsNode->SetY(node.m_pAngleFadeEnd);
+            else pPropsNode->SetY(pDefaultFade);
+
+            op->SetChild( op->op.args.ImageRasterMesh.angleFadeProperties, Generate( pPropsNode ) );
+        }
+
+        // Target mask
+        if ( node.m_pMask )
+        {
+            auto mask = Generate( node.m_pMask.get() );
+            mask = GenerateImageFormat( mask, EImageFormat::IF_L_UBYTE );
+            op->SetChild( op->op.args.ImageRasterMesh.mask, GenerateImageSize
+                    ( mask,
+                      FImageSize( (uint16_t)m_imageState.back().m_imageRect.size[0],
+                                  (uint16_t)m_imageState.back().m_imageRect.size[1]) )
+                              );
+        }
+
+        // Seam correction operations
+        //------------------------------
+        Ptr<ASTOpFixed> rasterop = new ASTOpFixed();
+        rasterop->op.type = OP_TYPE::IM_RASTERMESH;
+        rasterop->SetChild( rasterop->op.args.ImageRasterMesh.mesh, op->children[op->op.args.ImageRasterMesh.mesh].child() );
+        rasterop->op.args.ImageRasterMesh.image = 0;
+        rasterop->op.args.ImageRasterMesh.mask = 0;
+        rasterop->op.args.ImageRasterMesh.blockIndex = op->op.args.ImageRasterMesh.blockIndex;
+        rasterop->op.args.ImageRasterMesh.sizeX = op->op.args.ImageRasterMesh.sizeX;
+        rasterop->op.args.ImageRasterMesh.sizeY = op->op.args.ImageRasterMesh.sizeY;
+//        rasterop->op.args.ImageRasterMesh.growBorder = 0;
+
+        Ptr<ASTOpFixed> mapop = new ASTOpFixed();
+        mapop->op.type = OP_TYPE::IM_MAKEGROWMAP;
+        mapop->SetChild( mapop->op.args.ImageMakeGrowMap.mask, rasterop );
+        mapop->op.args.ImageMakeGrowMap.border = 2;
+
+        Ptr<ASTOpFixed> disop = new ASTOpFixed();
+        disop->op.type = OP_TYPE::IM_DISPLACE;
+        disop->SetChild( disop->op.args.ImageDisplace.displacementMap, mapop );
+        disop->SetChild( disop->op.args.ImageDisplace.source, op );
+
+        result.op = disop;
+    }
+
+
+    //---------------------------------------------------------------------------------------------
+	void CodeGenerator::GenerateImage_Mipmap(IMAGE_GENERATION_RESULT& result, const NodeImageMipmap* InNode)
+	{
+		const NodeImageMipmap::Private& node = *InNode->GetPrivate();
+
+		MUTABLE_CPUPROFILER_SCOPE(NodeImageMipmap);
+
+        Ptr<ASTOp> res = 0;
+
+        Ptr<ASTOpImageMipmap> op = new ASTOpImageMipmap();
+
+        // At the end of the day, we want all the mipmaps. Maybe the code optimiser will split the
+        // process later.
+        op->Levels = 0;
+
+        // Source image
+        Ptr<ASTOp> base = 0;
+        if ( Node* pSource = node.m_pSource.get() )
+        {
+            MUTABLE_CPUPROFILER_SCOPE(Base);
+            base = Generate( pSource );
+        }
+        else
+        {
+            // This argument is required
+            base = GenerateMissingImageCode( "Mipmap image", EImageFormat::IF_RGB_UBYTE, node.m_errorContext );
+        }
+
+        op->Source = base;
+
+        // The number of tail mipmaps depends on the cell size. We need to know it for some code
+        // optimisation operations. Scan the source image code looking for this info
+        int blockX = 0;
+        int blockY = 0;
+        if ( m_compilerOptions->m_textureLayoutStrategy
+             !=
+             CompilerOptions::TextureLayoutStrategy::None )
+        {
+            MUTABLE_CPUPROFILER_SCOPE(GetLayoutBlockSize);
+            op->Source->GetLayoutBlockSize( &blockX, &blockY );
+        }
+
+        if ( blockX && blockY )
+        {
+            int mipsX = (int)ceilf( logf( (float)blockX )/logf(2.0f) );
+            int mipsY = (int)ceilf( logf( (float)blockY )/logf(2.0f) );
+            op->BlockLevels = (uint8_t)std::max( mipsX, mipsY );
+        }
+        else
+        {
+            // No layout. Mipmap all the way down.
+            op->BlockLevels = 0;
+        }
+
+		op->AddressMode = node.m_settings.m_addressMode;
+		op->FilterType = node.m_settings.m_filterType;
+		op->SharpenFactor = node.m_settings.m_sharpenFactor;
+		op->DitherMipmapAlpha = node.m_settings.m_ditherMipmapAlpha;
+
+        res = op;
+
+        result.op = res;
+    }
+
+
+	//---------------------------------------------------------------------------------------------
+	void CodeGenerator::GenerateImage_Invert(IMAGE_GENERATION_RESULT& result, const NodeImageInvert* InNode)
+	{
+		const NodeImageInvert::Private& node = *InNode->GetPrivate();
+
+		Ptr<ASTOpFixed> op = new ASTOpFixed();
+		op->op.type = OP_TYPE::IM_INVERT;
+
+		// A image
+		Ptr<ASTOp> a;
+		if (Node* pBase = node.m_pBase.get())
+		{
+			a = Generate(pBase);
+		}
+		else
+		{
+			// This argument is required
+			a = GenerateMissingImageCode("Image Invert Color", EImageFormat::IF_RGB_UBYTE,
+				node.m_errorContext);
+		}
+		a = GenerateImageFormat(a, EImageFormat::IF_RGB_UBYTE);
+		a = GenerateImageSize
+		(a, FImageSize((uint16_t)m_imageState.back().m_imageRect.size[0],
+			(uint16_t)m_imageState.back().m_imageRect.size[1]));
+		op->SetChild(op->op.args.ImageInvert.base, a);
+
+		result.op = op;
+	}
+
+
+    //---------------------------------------------------------------------------------------------
+	void CodeGenerator::GenerateImage_Variation(IMAGE_GENERATION_RESULT& result, const NodeImageVariation* InNode)
+	{
+		const NodeImageVariation::Private& node = *InNode->GetPrivate();
+
+		Ptr<ASTOp> currentOp;
+
+        // Default case
+        if ( node.m_defaultImage )
+        {
+            IMAGE_GENERATION_RESULT branchResults;
+			GenerateImage(branchResults, node.m_defaultImage);
+			currentOp = branchResults.op;
+        }
+        else
+        {
+            // This argument is required
+            currentOp = GenerateMissingImageCode( "Variation default", EImageFormat::IF_RGBA_UBYTE, node.m_errorContext );
+        }
+
+        // Process variations in reverse order, since conditionals are built bottom-up.
+        for ( int t = int( node.m_variations.size() ) - 1; t >= 0; --t )
+        {
+            int tagIndex = -1;
+            const string& tag = node.m_variations[t].m_tag;
+            for ( int i = 0; i < int( m_firstPass.m_tags.size() ); ++i )
+            {
+                if ( m_firstPass.m_tags[i].tag == tag )
+                {
+                    tagIndex = i;
+                }
+            }
+
+            if ( tagIndex < 0 )
+            {
+                char buf[256];
+                mutable_snprintf( buf, 256, "Unknown tag found in image variation [%s].", tag.c_str() );
+
+                m_pErrorLog->GetPrivate()->Add( buf, ELMT_WARNING, node.m_errorContext );
+                continue;
+            }
+
+            Ptr<ASTOp> variationOp;
+            if ( node.m_variations[t].m_image )
+            {
+                variationOp = Generate( node.m_variations[t].m_image );
+            }
+            else
+            {
+                // This argument is required
+                variationOp = GenerateMissingImageCode( "Variation option", EImageFormat::IF_RGBA_UBYTE,
+                                                        node.m_errorContext );
+            }
+
+
+            Ptr<ASTOpConditional> conditional = new ASTOpConditional;
+            conditional->type = OP_TYPE::IM_CONDITIONAL;
+            conditional->no = currentOp;
+            conditional->yes = variationOp;
+            conditional->condition = m_firstPass.m_tags[tagIndex].genericCondition;
+
+            currentOp = conditional;
+        }
+
+        // Make sure all options are the same format and size
+        auto desc = currentOp->GetImageDesc( true );
+        if ( desc.m_format == EImageFormat::IF_NONE )
+        {
+            // TODO: Look for the most generic of the options?
+            // For now force a decently generic one
+            desc.m_format = EImageFormat::IF_RGBA_UBYTE;
+        }
+
+        if ( desc.m_size[0] != 0 && desc.m_size[1] != 0 )
+        {
+            Ptr<ASTOpFixed> sop = new ASTOpFixed();
+            sop->op.type = OP_TYPE::IM_RESIZE;
+            sop->op.args.ImageResize.size[0] = desc.m_size[0];
+            sop->op.args.ImageResize.size[1] = desc.m_size[1];
+            sop->SetChild( sop->op.args.ImageResize.source, currentOp );
+            currentOp = sop;
+        }
+
+        {
+            Ptr<ASTOpImagePixelFormat> fop = new ASTOpImagePixelFormat();
+            fop->Format = desc.m_format;
+            fop->Source = currentOp;
+            currentOp = fop;
+        }
+
+        result.op = currentOp;
+    }
+
+
+	//---------------------------------------------------------------------------------------------
+	void CodeGenerator::GenerateImage_Table(IMAGE_GENERATION_RESULT& result, const NodeImageTable* InNode)
+	{
+		const NodeImageTable::Private& node = *InNode->GetPrivate();
+
+		result.op = GenerateTableSwitch<NodeImageTable::Private, TCT_IMAGE, OP_TYPE::IM_SWITCH>(node,
+			[this](const NodeImageTable::Private& node, int colIndex, int row, ErrorLog* pErrorLog)
+			{
+				NodeImageConstantPtr pCell = new NodeImageConstant();
+
+				ImagePtrConst pImage = node.m_pTable->GetPrivate()->m_rows[row].m_values[colIndex].m_pProxyImage->Get();
+				if (!pImage)
+				{
+					char temp[256];
+					mutable_snprintf(temp, 256,
+						"Table has a missing image in column %d, row %d.", colIndex, row);
+					pErrorLog->GetPrivate()->Add(temp, ELMT_ERROR, node.m_errorContext);
+				}
+
+				pCell->SetValue(pImage);
+				return Generate(pCell);
+			});
+	}
+
+
+    //---------------------------------------------------------------------------------------------
+    ImagePtr CodeGenerator::GenerateMissingImage( EImageFormat format )
+    {
+        // Create the image node if it hasn't been created yet.
+        if (!m_missingImage[size_t(format)])
+        {
+            // Make a checkered debug image
+            const FImageSize size = MUTABLE_MISSING_IMAGE_DESC.m_size;
+
+            ImagePtr pImage = new Image( size[0], size[1], 1, format );
+
+            switch (format)
+            {
+            case EImageFormat::IF_L_UBYTE:
+                {
+                    uint8_t* pData = pImage->GetData();
+                    for ( int p=0; p<size[0]*size[1]; ++p )
+                    {
+                        if ((p+p/size[0])%2)
+                        {
+                            pData[0] = 255;
+                        }
+                        else
+                        {
+                            pData[0] = 64;
+                        }
+
+                        pData++;
+                    }
+                    break;
+                }
+
+            case EImageFormat::IF_RGB_UBYTE:
+                {
+                    uint8_t* pData = pImage->GetData();
+                    for ( int p=0; p<size[0]*size[1]; ++p )
+                    {
+                        if ((p+p/size[0])%2)
+                        {
+                            pData[0] = 255;
+                            pData[1] = 255;
+                            pData[2] = 64;
+                        }
+                        else
+                        {
+                            pData[0] = 64;
+                            pData[1] = 64;
+                            pData[2] = 255;
+                        }
+
+                        pData += 3;
+                    }
+                    break;
+                }
+
+            case EImageFormat::IF_BGRA_UBYTE:
+            case EImageFormat::IF_RGBA_UBYTE:
+                {
+                    uint8_t* pData = pImage->GetData();
+                    for ( int p=0; p<size[0]*size[1]; ++p )
+                    {
+                        if ((p+p/size[0])%2)
+                        {
+                            pData[0] = 255;
+                            pData[1] = 255;
+                            pData[2] = 64;
+                            pData[3] = 255;
+                        }
+                        else
+                        {
+                            pData[0] = 64;
+                            pData[1] = 64;
+                            pData[2] = 255;
+                            pData[3] = 128;
+                        }
+
+                        pData += 4;
+                    }
+                    break;
+                }
+
+            default:
+                check( false );
+                break;
+
+            }
+
+            m_missingImage[(size_t)format] = pImage;
+        }
+
+        return m_missingImage[(size_t)format].get();
+    }
+
+
+    //---------------------------------------------------------------------------------------------
+    Ptr<ASTOp> CodeGenerator::GenerateMissingImageCode( const char* strWhere,
+                                                     EImageFormat format,
+                                                     const void* errorContext )
+    {
+        // Log an error message
+        char buf[256];
+        mutable_snprintf
+            (
+                buf, 256,
+                "Required connection not found: %s",
+                strWhere
+            );
+        m_pErrorLog->GetPrivate()->Add( buf, ELMT_ERROR, errorContext );
+
+        // Make a checkered debug image
+        ImagePtr pImage = GenerateMissingImage( format );
+
+        NodeImageConstantPtr pNode = new NodeImageConstant();
+        pNode->SetValue( pImage.get() );
+
+        Ptr<ASTOp> result = Generate( pNode );
+
+        return result;
+    }
+
+
+    //---------------------------------------------------------------------------------------------
+    Ptr<ASTOp> CodeGenerator::GeneratePlainImageCode( const vec3<float>& colour )
+    {
+        const int size = 4;
+        ImagePtr pImage = new Image( size, size, 1, EImageFormat::IF_RGB_UBYTE );
+
+        uint8_t* pData = pImage->GetData();
+        for ( int p=0; p<size*size; ++p )
+        {
+            pData[0] = (uint8_t)std::min( 255.0f, std::max( 0.0f, 255*colour[0] ) );
+            pData[1] = (uint8_t)std::min( 255.0f, std::max( 0.0f, 255*colour[1] ) );
+            pData[2] = (uint8_t)std::min( 255.0f, std::max( 0.0f, 255*colour[2] ) );
+            pData += 3;
+        }
+
+        NodeImageConstantPtr pNode = new NodeImageConstant();
+        pNode->SetValue( pImage.get() );
+
+        Ptr<ASTOp> result = Generate( pNode );
+
+        return result;
+    }
+
+
+    //---------------------------------------------------------------------------------------------
+    Ptr<ASTOp> CodeGenerator::GenerateImageFormat( Ptr<ASTOp> at, EImageFormat format )
+    {
+        Ptr<ASTOp> result = at;
+
+        check( format != EImageFormat::IF_NONE );
+
+        if ( at && at->GetImageDesc().m_format != format )
+        {
+            // Generate the format change code
+            Ptr<ASTOpImagePixelFormat> op = new ASTOpImagePixelFormat();
+            op->Source = at;
+            op->Format = format;
+            result = op;
+        }
+
+        return result;
+    }
+
+
+    //---------------------------------------------------------------------------------------------
+    Ptr<ASTOp> CodeGenerator::GenerateImageUncompressed( Ptr<ASTOp> at )
+    {
+        Ptr<ASTOp> result = at;
+
+        if (at)
+        {
+            EImageFormat sourceFormat = at->GetImageDesc().m_format;
+            EImageFormat targetFormat = GetUncompressedFormat( sourceFormat );
+
+            if ( targetFormat != sourceFormat )
+            {
+                // Generate the format change code
+                Ptr<ASTOpImagePixelFormat> op = new ASTOpImagePixelFormat();
+                op->Source = at;
+                op->Format = targetFormat;
+                result = op;
+            }
+        }
+
+        return result;
+    }
+
+
+    //---------------------------------------------------------------------------------------------
+    Ptr<ASTOp> CodeGenerator::GenerateImageSize( Ptr<ASTOp> at, FImageSize size )
+    {
+        Ptr<ASTOp> result = at;
+
+        check( size[0]>0 && size[1]>0 );
+
+        if ( at->GetImageDesc().m_size != size )
+        {
+            Ptr<ASTOpFixed> op = new ASTOpFixed();
+            op->op.type = OP_TYPE::IM_RESIZE;
+            op->SetChild( op->op.args.ImageResize.source,at);
+            op->op.args.ImageResize.size[0] = size[0];
+            op->op.args.ImageResize.size[1] = size[1];
+            result = op;
+        }
+
+        return result;
+    }
+
+
+    //---------------------------------------------------------------------------------------------
+    FImageDesc CodeGenerator::CalculateImageDesc( const Node::Private& node )
+    {
+        ImageDescGenerator imageDescGenerator;
+        imageDescGenerator.Generate( node );
+        return imageDescGenerator.m_desc;
+    }
+
+
+
+    //---------------------------------------------------------------------------------------------
+    //! This class contains the support data to accelerate the GetImageDesc recursive function.
+    //! _If none is provided in the call, one will be created at that level and used from there on.
+    class GetImageDescContext
+    {
+    public:
+        vector<bool> m_visited;
+        vector<FImageDesc> m_results;
+    };
+
+ 
+}
