[mutable] Move Mutable plugin from the restricted folder to the experimental plugins folder.

diff --git a/Engine/Plugins/Experimental/Mutable/Source/MutableTools/Private/CodeGenerator_Mesh.cpp b/Engine/Plugins/Experimental/Mutable/Source/MutableTools/Private/CodeGenerator_Mesh.cpp
new file mode 100644
--- /dev/null
+++ b/Engine/Plugins/Experimental/Mutable/Source/MutableTools/Private/CodeGenerator_Mesh.cpp
@@ -0,0 +1,1427 @@
+// Copyright Epic Games, Inc. All Rights Reserved.
+
+#include "CodeGenerator.h"
+
+#include "Platform.h"
+#include "CompilerPrivate.h"
+#include "TablePrivate.h"
+#include "NodeLayoutPrivate.h"
+
+#include "NodeMeshPrivate.h"
+#include "NodeMeshConstantPrivate.h"
+#include "NodeMeshSwitchPrivate.h"
+#include "NodeMeshTablePrivate.h"
+#include "NodeMeshInterpolatePrivate.h"
+#include "NodeMeshMorphPrivate.h"
+#include "NodeMeshMakeMorphPrivate.h"
+#include "NodePatchMeshPrivate.h"
+#include "NodeMeshSubtractPrivate.h"
+#include "NodeMeshFormatPrivate.h"
+#include "NodeMeshFragmentPrivate.h"
+#include "NodeMeshTransformPrivate.h"
+#include "NodeMeshClipMorphPlanePrivate.h"
+#include "NodeMeshClipWithMeshPrivate.h"
+#include "NodeMeshApplyPosePrivate.h"
+#include "NodeMeshVariationPrivate.h"
+#include "NodeMeshGeometryOperationPrivate.h"
+#include "NodeMeshReshapePrivate.h"
+#include "NodeMeshClipDeformPrivate.h"
+
+#include "MutableTools/Private/ASTOpSwitch.h"
+#include "MutableTools/Private/ASTOpConditional.h"
+#include "MutableTools/Private/ASTOpConstantResource.h"
+#include "MutableTools/Private/ASTOpMeshFormat.h"
+#include "MutableTools/Private/ASTOpMeshRemapIndices.h"
+#include "MutableTools/Private/ASTOpMeshMorphReshape.h"
+#include "MutableTools/Private/ASTOpMeshApplyShape.h"
+#include "MutableTools/Private/ASTOpMeshBindShape.h"
+#include "MutableTools/Private/ASTOpMeshExtractLayoutBlocks.h"
+#include "MutableTools/Private/ASTOpMeshTransform.h"
+#include "MutableTools/Private/ASTOpMeshClipMorphPlane.h"
+#include "MutableTools/Private/ASTOpMeshGeometryOperation.h"
+#include "MutableTools/Private/ASTOpMeshApplyPose.h"
+#include "MutableTools/Private/ASTOpMeshClipDeform.h"
+
+#include "ModelPrivate.h"
+#include "Operations.h"
+
+
+namespace mu
+{
+
+    //---------------------------------------------------------------------------------------------
+    void CodeGenerator::PrepareForLayout( LayoutPtrConst pSourceLayout,
+                                             MeshPtr currentLayoutMesh,
+                                             size_t currentLayoutChannel,
+                                             const void* errorContext )
+    {
+        if (currentLayoutMesh->GetVertexCount()==0)
+        {
+            return;
+        }
+
+        // TODO: Done twice? Check new component code.
+		Ptr<const Layout> pLayout = AddLayout( pSourceLayout );
+        currentLayoutMesh->AddLayout( pLayout );
+
+        int buffer = -1;
+        int channel = -1;
+        currentLayoutMesh->GetVertexBuffers().FindChannel( MBS_TEXCOORDS,
+                                                              (int)currentLayoutChannel,
+                                                              &buffer,
+                                                              &channel );
+        check( buffer>=0 );
+        check( channel>=0 );
+
+        // Create the layout block vertex buffer
+        uint16_t* pLayoutData = 0;
+        {
+            int layoutBuf = currentLayoutMesh->GetVertexBuffers().GetBufferCount();
+            currentLayoutMesh->GetVertexBuffers().SetBufferCount( layoutBuf+1 );
+
+            // TODO
+            check( pLayout->GetBlockCount()<65535 );
+            MESH_BUFFER_SEMANTIC semantic = MBS_LAYOUTBLOCK;
+            int semanticIndex = int(currentLayoutChannel);
+            MESH_BUFFER_FORMAT format = MBF_UINT16;
+            int components = 1;
+            int offset = 0;
+            currentLayoutMesh->GetVertexBuffers().SetBuffer
+                    (
+                        layoutBuf,
+                        sizeof(uint16_t),
+                        1,
+                        &semantic, &semanticIndex,
+                        &format, &components,
+                        &offset
+                    );
+            pLayoutData = (uint16_t*)currentLayoutMesh->GetVertexBuffers().GetBufferData( layoutBuf );
+        }
+
+        // Get the information about the texture coordinates channel
+        MESH_BUFFER_SEMANTIC semantic;
+        int semanticIndex;
+        MESH_BUFFER_FORMAT format;
+        int components;
+        int offset;
+        currentLayoutMesh->GetVertexBuffers().GetChannel
+            ( buffer, channel, &semantic, &semanticIndex, &format, &components, &offset );
+        check( semantic == MBS_TEXCOORDS );
+
+        const uint8_t* pData = currentLayoutMesh->GetVertexBuffers().GetBufferData( buffer );
+        int elemSize = currentLayoutMesh->GetVertexBuffers().GetElementSize( buffer );
+        int channelOffset = currentLayoutMesh->GetVertexBuffers().GetChannelOffset( buffer, channel );
+        pData += channelOffset;
+
+        // Clear block data
+        for (int i=0; i< currentLayoutMesh->GetVertexBuffers().GetElementCount(); ++i)
+        {
+            pLayoutData[i] = 65535;
+        }
+
+        // TODO: Check overlapping layout blocks
+        // TODO: Check triangles crossing blocks
+        int inside = 0;
+        int temp = 0;
+        for ( int b=0; b<pLayout->GetBlockCount(); ++b )
+        {
+            FIntPoint grid = pLayout->GetGridSize();
+
+            box< vec2<int> > block;
+            pLayout->GetBlock
+                ( b, &block.min[0], &block.min[1], &block.size[0], &block.size[1] );
+
+            box< vec2<float> > rect;
+            rect.min[0] = ( (float)block.min[0] ) / (float) grid[0];
+            rect.min[1] = ( (float)block.min[1] ) / (float) grid[1];
+            rect.size[0] = ( (float)block.size[0] ) / (float) grid[0];
+            rect.size[1] = ( (float)block.size[1] ) / (float) grid[1];
+
+            const uint8_t* pVertices = pData;
+
+            if ( format == MBF_FLOAT32 )
+            {
+                for ( int v=0; v<currentLayoutMesh->GetVertexBuffers().GetElementCount(); ++v )
+                {
+                    vec2<float>* pUV = (vec2<float>*)pVertices;
+                    pVertices += elemSize;
+
+                    if ( pLayoutData[v]==65535 && rect.ContainsInclusive(*pUV) )
+                    {
+                        *pUV = rect.Homogenize( *pUV );
+
+                        // Set the value to the unique block id
+                        check( pLayout->m_blocks[b].m_id < 65535 );
+                        pLayoutData[v] =(uint16_t) pLayout->m_blocks[b].m_id;
+
+                        inside++;
+                    }
+                    else
+                    {
+                        temp++;
+                    }
+                }
+            }
+            else if ( format == MBF_FLOAT16 )
+            {
+                // TODO: Very slow?
+                for ( int v=0; v<currentLayoutMesh->GetVertexBuffers().GetElementCount(); ++v )
+                {
+                    float16* pUV = (float16*)pVertices;
+                    pVertices += elemSize;
+
+                    vec2<float> UV( halfToFloat(pUV[0]), halfToFloat(pUV[1]) );
+
+                    if ( pLayoutData[v]==65535 && rect.ContainsInclusive(UV) )
+                    {
+                        UV = rect.Homogenize( UV );
+
+                        // Set the value to the unique block id
+                        check( pLayout->m_blocks[b].m_id < 65535 );
+                        pLayoutData[v] = (uint16_t)pLayout->m_blocks[b].m_id;
+
+                        inside++;
+
+                        pUV[0] = floatToHalf( UV[0] );
+                        pUV[1] = floatToHalf( UV[1] );
+                    }
+                }
+            }
+        }
+
+        int outside = currentLayoutMesh->GetVertexBuffers().GetElementCount() - inside;
+        if (outside>0)
+        {
+            char buf[256];
+            mutable_snprintf
+                (
+                    buf, 256,
+                    "Source mesh has %d vertices not assigned to any layout block.",
+                    outside
+                );
+			int blockCount = pLayout->GetBlockCount();
+        
+            //m_pErrorLog->GetPrivate()->Add( buf, blockCount==1?ELMT_INFO:ELMT_WARNING, errorContext );
+            vector< float > unassignedUVs;
+            unassignedUVs.reserve(64);           
+            
+            const uint8_t* pVertices = pData;
+            for (int i=0; i<currentLayoutMesh->GetVertexBuffers().GetElementCount(); ++i)
+            {
+                vec2<float> UV;
+                if ( format == MBF_FLOAT32 )
+                {
+                    UV = *((vec2<float>*)pVertices);
+                }
+                else if ( format == MBF_FLOAT16 )
+                {
+                    float16* pUV = (float16*)pVertices;
+                    UV = vec2<float>( halfToFloat(pUV[0]), halfToFloat(pUV[1]) );
+                }
+
+                pVertices += elemSize;
+
+                if (pLayoutData[i]==65535)
+                {
+                    unassignedUVs.push_back(UV[0]);
+                    unassignedUVs.push_back(UV[1]);
+                }
+            }
+
+            ErrorLogMessageAttachedDataView attachedDataView;
+            attachedDataView.m_unassignedUVs = unassignedUVs.data();
+            attachedDataView.m_unassignedUVsSize = unassignedUVs.size();
+
+            m_pErrorLog->GetPrivate()->Add( buf, attachedDataView, blockCount==1?ELMT_INFO:ELMT_WARNING, errorContext );
+        }
+        
+        // Assign broken vertices to the first block
+        for (int i=0; i< currentLayoutMesh->GetVertexBuffers().GetElementCount(); ++i)
+        {
+            if ( pLayoutData[i]==65535 )
+            {
+                pLayoutData[i] = 0;
+            }
+        }
+    }
+
+
+    //---------------------------------------------------------------------------------------------
+    void CodeGenerator::GenerateLayout( 
+		MESH_GENERATION_RESULT& result, 
+		const NodeLayoutBlocksPtrConst& node,
+		uint32 currentLayoutChannel, const MeshPtr currentLayoutMesh )
+    {
+        if (!currentLayoutMesh)
+        {
+            m_pErrorLog->GetPrivate()->Add( "Generating a layout node without a parent mesh.",
+                                            ELMT_ERROR, node->GetPrivate()->m_errorContext );
+            return;
+        }
+
+        LayoutPtr pSourceLayout = node->GetPrivate()->m_pLayout;
+        result.layouts.push_back( pSourceLayout );
+
+        PrepareForLayout( pSourceLayout,
+                          currentLayoutMesh, currentLayoutChannel,
+                          node->GetPrivate()->m_errorContext  );
+    }
+
+
+    //---------------------------------------------------------------------------------------------
+    void CodeGenerator::GenerateMesh( MESH_GENERATION_RESULT& result, const NodeMeshPtrConst& Untyped )
+    {
+        if (!Untyped)
+        {
+            result = MESH_GENERATION_RESULT();
+            return;
+        }
+
+		// Clear bottom-up state
+		m_currentBottomUpState.m_address = nullptr;
+
+        // See if it was already generated
+		VISITED_MAP_KEY key = GetCurrentCacheKey(Untyped);
+        GeneratedMeshMap::ValueType* it = m_generatedMeshes.Find( key );
+        if ( it )
+        {
+            result = *it;
+            return;
+        }
+
+		const NodeMesh* Node = Untyped.get();
+
+        // Generate for each different type of node
+		switch (Untyped->GetMeshNodeType())
+		{
+		case NodeMesh::EType::Constant: GenerateMesh_Constant(result, static_cast<const NodeMeshConstant*>(Node)); break;
+		case NodeMesh::EType::Format: GenerateMesh_Format(result, static_cast<const NodeMeshFormat*>(Node)); break;
+		case NodeMesh::EType::Morph: GenerateMesh_Morph(result, static_cast<const NodeMeshMorph*>(Node)); break;
+		case NodeMesh::EType::MakeMorph: GenerateMesh_MakeMorph(result, static_cast<const NodeMeshMakeMorph*>(Node)); break;
+		case NodeMesh::EType::Fragment: GenerateMesh_Fragment(result, static_cast<const NodeMeshFragment*>(Node)); break;
+		case NodeMesh::EType::Interpolate: GenerateMesh_Interpolate(result, static_cast<const NodeMeshInterpolate*>(Node)); break;
+		case NodeMesh::EType::Switch: GenerateMesh_Switch(result, static_cast<const NodeMeshSwitch*>(Node)); break;
+		case NodeMesh::EType::Subtract: GenerateMesh_Subtract(result, static_cast<const NodeMeshSubtract*>(Node)); break;
+		case NodeMesh::EType::Transform: GenerateMesh_Transform(result, static_cast<const NodeMeshTransform*>(Node)); break;
+		case NodeMesh::EType::ClipMorphPlane: GenerateMesh_ClipMorphPlane(result, static_cast<const NodeMeshClipMorphPlane*>(Node)); break;
+		case NodeMesh::EType::ClipWithMesh: GenerateMesh_ClipWithMesh(result, static_cast<const NodeMeshClipWithMesh*>(Node)); break;
+		case NodeMesh::EType::ApplyPose: GenerateMesh_ApplyPose(result, static_cast<const NodeMeshApplyPose*>(Node)); break;
+		case NodeMesh::EType::Variation: GenerateMesh_Variation(result, static_cast<const NodeMeshVariation*>(Node)); break;
+		case NodeMesh::EType::Table: GenerateMesh_Table(result, static_cast<const NodeMeshTable*>(Node)); break;
+		case NodeMesh::EType::GeometryOperation: GenerateMesh_GeometryOperation(result, static_cast<const NodeMeshGeometryOperation*>(Node)); break;
+		case NodeMesh::EType::Reshape: GenerateMesh_Reshape(result, static_cast<const NodeMeshReshape*>(Node)); break;
+		case NodeMesh::EType::ClipDeform: GenerateMesh_ClipDeform(result, static_cast<const NodeMeshClipDeform*>(Node)); break;
+		case NodeMesh::EType::None: check(false);
+		}
+
+        // Cache the result
+        m_generatedMeshes.Add( key, result );
+    }
+
+
+    //---------------------------------------------------------------------------------------------
+    void CodeGenerator::GenerateMesh_Morph( MESH_GENERATION_RESULT& result, const NodeMeshMorph* morph )
+    {
+        NodeMeshMorph::Private& node = *morph->GetPrivate();
+
+        Ptr<ASTOpFixed> OpMorph = new ASTOpFixed();
+        OpMorph->op.type = OP_TYPE::ME_MORPH2;
+
+        // Factor
+        if ( node.m_pFactor )
+        {
+            OpMorph->SetChild( OpMorph->op.args.MeshMorph2.factor, Generate( node.m_pFactor.get() ) );
+        }
+        else
+        {
+            // This argument is required
+            OpMorph->SetChild( OpMorph->op.args.MeshMorph2.factor, GenerateMissingScalarCode( "Morph factor", 0.5f,
+                                                                   node.m_errorContext ) );
+        }
+
+        // Base
+        MESH_GENERATION_RESULT baseResult;
+        if ( node.m_pBase )
+        {
+            GenerateMesh( baseResult, node.m_pBase );
+            OpMorph->SetChild( OpMorph->op.args.MeshMorph2.base, baseResult.meshOp );
+        }
+        else
+        {
+            // This argument is required
+            m_pErrorLog->GetPrivate()->Add( "Mesh morph base node is not set.",
+                                            ELMT_ERROR, node.m_errorContext );
+        }
+
+        // TODO: Support more than MUTABLE_OP_MAX_MORPH2_TARGETS targets
+        if ( node.m_morphs.size()>MUTABLE_OP_MAX_MORPH2_TARGETS )
+        {
+            char temp[256];
+            mutable_snprintf( temp, 256,
+                              "A morph node has more targets [%d] than currently supported [%d].",
+                              node.m_morphs.size(),
+                              MUTABLE_OP_MAX_MORPH2_TARGETS );
+            m_pErrorLog->GetPrivate()->Add( temp, ELMT_WARNING, node.m_errorContext );
+        }
+
+        m_overrideLayoutsStack.push_back( baseResult.layouts );
+		m_activeTags.push_back({});
+
+        int count = 0;
+        for ( std::size_t t=0
+            ; t<node.m_morphs.size() && t<MUTABLE_OP_MAX_MORPH2_TARGETS
+            ; ++t )
+        {
+            if ( auto pA = node.m_morphs[t].get() )
+            {
+                MESH_GENERATION_RESULT targetResult;
+                GenerateMesh( targetResult, pA );
+
+                // TODO: Make sure that the target is a mesh with the morph format
+                Ptr<ASTOp> target = targetResult.meshOp;
+
+                // If the vertex indices are supposed to be relative in the targets, adjust them
+                if (node.m_vertexIndicesAreRelative)
+                {
+                    Ptr<ASTOpMeshRemapIndices> remapIndices = new ASTOpMeshRemapIndices;
+                    remapIndices->source = target;
+                    remapIndices->reference = baseResult.baseMeshOp;
+                    target = remapIndices;
+                }
+
+                OpMorph->SetChild( OpMorph->op.args.MeshMorph2.targets[count], target);
+                count++;
+            }
+        }
+ 
+        const bool bReshapeEnabled = node.m_reshapeSkeleton || node.m_reshapePhysicsVolumes;
+        
+        Ptr<ASTOpMeshMorphReshape> OpMorphReshape;
+        if ( bReshapeEnabled )
+        {
+		    Ptr<ASTOpMeshBindShape> OpBind = new ASTOpMeshBindShape();
+		    Ptr<ASTOpMeshApplyShape> OpApply = new ASTOpMeshApplyShape();
+
+			// Setting reshapeVertices to false the bind op will remove all mesh members except 
+			// PhysicsBodies and the Skeleton.
+            OpBind->m_reshapeVertices = false;
+		    OpBind->m_reshapeSkeleton = node.m_reshapeSkeleton;
+		    OpBind->m_deformAllBones = node.m_deformAllBones;
+		    OpBind->m_bonesToDeform = node.m_bonesToDeform;
+    	    OpBind->m_reshapePhysicsVolumes = node.m_reshapePhysicsVolumes; 
+			OpBind->m_physicsToDeform = node.m_physicsToDeform;
+			OpBind->m_deformAllPhysics = node.m_deformAllPhysics;
+			OpBind->m_bindingMethod = static_cast<uint32>(EShapeBindingMethod::ReshapeClosestProject);
+            
+			OpBind->Mesh = baseResult.meshOp;
+            OpBind->Shape = baseResult.meshOp;
+           
+			OpApply->m_reshapeVertices = OpBind->m_reshapeVertices;
+		    OpApply->m_reshapeSkeleton = OpBind->m_reshapeSkeleton;
+		    OpApply->m_reshapePhysicsVolumes = OpBind->m_reshapePhysicsVolumes;
+
+			OpApply->Mesh = OpBind;
+            OpApply->Shape = OpMorph;
+
+            OpMorphReshape = new ASTOpMeshMorphReshape();
+            OpMorphReshape->Morph = OpMorph;
+            OpMorphReshape->Reshape = OpApply;
+        }
+
+        m_overrideLayoutsStack.pop_back();
+		m_activeTags.pop_back();
+
+		if (OpMorphReshape)
+		{
+			result.meshOp = OpMorphReshape;
+		}
+		else
+		{
+			result.meshOp = OpMorph;
+		}
+        result.baseMeshOp = baseResult.baseMeshOp;
+
+        result.layouts = baseResult.layouts;
+    } 
+
+
+    //---------------------------------------------------------------------------------------------
+    void CodeGenerator::GenerateMesh_MakeMorph( MESH_GENERATION_RESULT& result,
+		const NodeMeshMakeMorph* morph )
+    {
+        NodeMeshMakeMorph::Private& node = *morph->GetPrivate();
+
+        Ptr<ASTOpFixed> op = new ASTOpFixed();
+        op->op.type = OP_TYPE::ME_DIFFERENCE;
+
+        // \todo Texcoords are broken?
+        op->op.args.MeshDifference.ignoreTextureCoords = 1;
+
+        // Base
+        MESH_GENERATION_RESULT baseResult;
+        if ( node.m_pBase )
+        {
+            GenerateMesh( baseResult, node.m_pBase );
+
+            op->SetChild( op->op.args.MeshDifference.base, baseResult.meshOp );
+        }
+        else
+        {
+            // This argument is required
+            m_pErrorLog->GetPrivate()->Add( "Mesh make morph base node is not set.",
+                                            ELMT_ERROR, node.m_errorContext );
+        }
+
+        // Target
+        m_overrideLayoutsStack.push_back( baseResult.layouts );
+		m_activeTags.push_back({});
+		if ( node.m_pTarget )
+        {
+            MESH_GENERATION_RESULT targetResult;
+            GenerateMesh( targetResult, node.m_pTarget );
+
+            op->SetChild( op->op.args.MeshDifference.target, targetResult.meshOp );
+        }
+        else
+        {
+            // This argument is required
+            m_pErrorLog->GetPrivate()->Add( "Mesh make morph target node is not set.",
+                                            ELMT_ERROR, node.m_errorContext );
+        }
+        m_overrideLayoutsStack.pop_back();
+		m_activeTags.pop_back();
+
+        result.meshOp = op;
+        result.baseMeshOp = baseResult.baseMeshOp;
+        result.layouts = baseResult.layouts;
+    }
+
+    //---------------------------------------------------------------------------------------------
+    void CodeGenerator::GenerateMesh_Fragment( MESH_GENERATION_RESULT& result,
+		const NodeMeshFragment* fragment )
+    {
+        NodeMeshFragment::Private& node = *fragment->GetPrivate();
+
+        MESH_GENERATION_RESULT baseResult;
+        if ( node.m_pMesh )
+        {
+            GenerateMesh( baseResult, node.m_pMesh );
+
+            if ( node.m_fragmentType==NodeMeshFragment::FT_LAYOUT_BLOCKS )
+            {
+                Ptr<ASTOpMeshExtractLayoutBlocks> op = new ASTOpMeshExtractLayoutBlocks();
+                result.meshOp = op;
+
+                op->source = baseResult.meshOp;
+
+                if ( baseResult.layouts.size()>(size_t)node.m_layoutOrGroup )
+                {
+                    LayoutPtrConst pSourceLayout = baseResult.layouts[ node.m_layoutOrGroup ];
+                    const Layout* pLayout = m_addedLayouts[ pSourceLayout.get() ].get();
+                    op->layout = (uint16_t)node.m_layoutOrGroup;
+
+                    for ( size_t i=0; i<node.m_blocks.size(); ++i )
+                    {
+                        if (node.m_blocks[i]>=0 && node.m_blocks[i]<pLayout->m_blocks.Num() )
+                        {
+                            int bid = pLayout->m_blocks[ node.m_blocks[i] ].m_id;
+                            op->blocks.push_back(bid);
+                        }
+                        else
+                        {
+                            m_pErrorLog->GetPrivate()->Add( "Internal layout block index error.",
+                                                            ELMT_ERROR, node.m_errorContext );
+                        }
+                    }
+                }
+                else
+                {
+                    // This argument is required
+                    m_pErrorLog->GetPrivate()->Add( "Missing layout in mesh fragment source.",
+                                                    ELMT_ERROR, node.m_errorContext );
+                }
+            }
+
+            else if ( node.m_fragmentType==NodeMeshFragment::FT_FACE_GROUP )
+            {
+                Ptr<ASTOpFixed> op = new ASTOpFixed();
+                result.meshOp = op;
+
+                op->op.type = OP_TYPE::ME_EXTRACTFACEGROUP;
+
+                op->SetChild( op->op.args.MeshExtractFaceGroup.source, baseResult.meshOp );
+                op->op.args.MeshExtractFaceGroup.group = node.m_layoutOrGroup;
+            }
+
+        }
+        else
+        {
+            // This argument is required
+            m_pErrorLog->GetPrivate()->Add( "Mesh fragment source is not set.",
+                                            ELMT_ERROR, node.m_errorContext );
+        }
+
+        result.baseMeshOp = baseResult.baseMeshOp;
+        result.layouts = baseResult.layouts;
+    }
+
+
+    //---------------------------------------------------------------------------------------------
+    void CodeGenerator::GenerateMesh_Interpolate( MESH_GENERATION_RESULT& result,
+		const NodeMeshInterpolate* interpolate )
+    {
+        NodeMeshInterpolate::Private& node = *interpolate->GetPrivate();
+
+        // Generate the code
+        Ptr<ASTOpFixed> op = new ASTOpFixed();
+        op->op.type = OP_TYPE::ME_INTERPOLATE;
+        result.meshOp = op;
+
+        // Factor
+        if ( Node* pFactor = node.m_pFactor.get() )
+        {
+            op->SetChild( op->op.args.MeshInterpolate.factor, Generate( pFactor ) );
+        }
+        else
+        {
+            // This argument is required
+            op->SetChild( op->op.args.MeshInterpolate.factor,
+                          GenerateMissingScalarCode( "Interpolation factor", 0.5f, node.m_errorContext ) );
+        }
+
+        //
+        Ptr<ASTOp> base = 0;
+        int count = 0;
+        for ( std::size_t t=0
+            ; t<node.m_targets.size() && t<MUTABLE_OP_MAX_INTERPOLATE_COUNT-1
+            ; ++t )
+        {
+            if ( NodeMesh* pA = node.m_targets[t].get() )
+            {
+                if ( count>0 )
+                {
+                    m_overrideLayoutsStack.push_back( result.layouts );
+                }
+
+                MESH_GENERATION_RESULT targetResult;
+                GenerateMesh( targetResult, pA );
+
+                if ( count>0 )
+                {
+                    m_overrideLayoutsStack.pop_back();
+                }
+
+                // The first target is the base
+                if (count==0)
+                {
+                    base = targetResult.meshOp;
+                    op->SetChild( op->op.args.MeshInterpolate.base, targetResult.meshOp );
+
+                    result.baseMeshOp = targetResult.baseMeshOp;
+                    result.layouts = targetResult.layouts;
+                }
+                else
+                {
+                    Ptr<ASTOpFixed> dop = new ASTOpFixed();
+                    dop->op.type = OP_TYPE::ME_DIFFERENCE;
+                    dop->SetChild( dop->op.args.MeshDifference.base, base );
+                    dop->SetChild( dop->op.args.MeshDifference.target, targetResult.meshOp );
+
+                    // \todo Texcoords are broken?
+                    dop->op.args.MeshDifference.ignoreTextureCoords = 1;
+
+                    if ( node.m_channels.size()>MUTABLE_OP_MAX_MORPH_CHANNELS )
+                    {
+                        char temp[256];
+                        mutable_snprintf( temp, 256,
+                                          "Morph uses too many channels [%d]. The maximum is [%d].",
+                                          node.m_channels.size(),
+                                          MUTABLE_OP_MAX_MORPH_CHANNELS );
+                        m_pErrorLog->GetPrivate()->Add( temp, ELMT_ERROR, node.m_errorContext );
+                    }
+
+                    for ( size_t c=0;
+                          c<node.m_channels.size() && c<MUTABLE_OP_MAX_MORPH_CHANNELS;
+                          ++c)
+                    {
+                        check( node.m_channels[c].semantic<256 );
+                        dop->op.args.MeshDifference.channelSemantic[c] =
+                                uint8_t( node.m_channels[c].semantic );
+
+                        check( node.m_channels[c].semanticIndex<256 );
+                        dop->op.args.MeshDifference.channelSemanticIndex[c] =
+                                uint8_t( node.m_channels[c].semanticIndex );
+                    }
+
+                    op->SetChild( op->op.args.MeshInterpolate.targets[count-1], dop );
+                }
+                count++;
+            }
+        }
+
+        // At least one mesh is required
+        if (!count)
+        {
+            // TODO
+            //op.args.MeshInterpolate.target[0] = GenerateMissingImageCode( "First mesh", IF_RGB_UBYTE );
+            m_pErrorLog->GetPrivate()->Add
+                ( "Mesh interpolation: at least the first mesh is required.",
+                  ELMT_ERROR, node.m_errorContext );
+        }
+    }
+
+
+    //---------------------------------------------------------------------------------------------
+    void CodeGenerator::GenerateMesh_Switch( MESH_GENERATION_RESULT& result, const NodeMeshSwitch* sw )
+    {
+        NodeMeshSwitch::Private& node = *sw->GetPrivate();
+
+        if (node.m_options.size() == 0)
+        {
+            // No options in the switch!
+            // TODO
+            result = MESH_GENERATION_RESULT();
+			return;
+        }
+
+        Ptr<ASTOpSwitch> op = new ASTOpSwitch();
+        op->type = OP_TYPE::ME_SWITCH;
+
+        // Factor
+        if ( node.m_pParameter )
+        {
+            op->variable = Generate( node.m_pParameter.get() );
+        }
+        else
+        {
+            // This argument is required
+            op->variable = GenerateMissingScalarCode( "Switch variable", 0.0f, node.m_errorContext );
+        }
+
+        // Options
+        for ( std::size_t t=0; t< node.m_options.size(); ++t )
+        {
+            if (t!=0)
+            {
+                m_overrideLayoutsStack.push_back( result.layouts );
+            }
+
+            if ( node.m_options[t] )
+            {
+                MESH_GENERATION_RESULT branchResults;
+                GenerateMesh( branchResults, node.m_options[t] );
+
+                auto branch = branchResults.meshOp;
+                op->cases.emplace_back((int16_t)t,op,branch);
+
+                if (t==0)
+                {
+                    result = branchResults;
+                }
+            }
+
+            if (t!=0)
+            {
+                m_overrideLayoutsStack.pop_back();
+            }
+        }
+
+        result.meshOp = op;
+    }
+
+
+	//---------------------------------------------------------------------------------------------
+	void CodeGenerator::GenerateMesh_Table(MESH_GENERATION_RESULT& result, const NodeMeshTable* TableNode)
+	{
+		//
+		MESH_GENERATION_RESULT NewResult = result;
+		int t = 0;
+
+		Ptr<ASTOp> Op = GenerateTableSwitch<NodeMeshTable::Private, TCT_MESH, OP_TYPE::ME_SWITCH>(*TableNode->GetPrivate(), 
+			[this, &NewResult, &t] (const NodeMeshTable::Private& node, int colIndex, int row, ErrorLog* pErrorLog)
+			{
+				NodeMeshConstantPtr pCell = new NodeMeshConstant();
+				MeshPtr pMesh = node.m_pTable->GetPrivate()->m_rows[row].m_values[colIndex].m_pMesh;
+
+				if (!pMesh)
+				{
+					char temp[256];
+					mutable_snprintf(temp, 256,
+						"Table has a missing mesh in column %d, row %d.", colIndex, row);
+					pErrorLog->GetPrivate()->Add(temp, ELMT_ERROR, node.m_errorContext);
+				}
+
+				pCell->SetValue(pMesh);
+
+				// TODO Take into account layout strategy
+				int numLayouts = (int)node.m_layouts.size();
+				pCell->SetLayoutCount(numLayouts);
+				for (int i = 0; i < numLayouts; ++i)
+				{
+					pCell->SetLayout(i, node.m_layouts[i]);
+				}
+
+				if (t != 0)
+				{
+					m_overrideLayoutsStack.push_back(NewResult.layouts);
+				}
+
+				MESH_GENERATION_RESULT branchResults;
+				GenerateMesh(branchResults, pCell);
+
+				if (t == 0)
+				{
+					NewResult = branchResults;
+				}
+				else
+				{
+					m_overrideLayoutsStack.pop_back();
+				}
+
+				++t;
+				return branchResults.meshOp;
+			});
+
+		NewResult.meshOp = Op;
+
+		result = NewResult;
+	}
+
+
+    //---------------------------------------------------------------------------------------------
+    void CodeGenerator::GenerateMesh_Variation( MESH_GENERATION_RESULT& result,
+		const NodeMeshVariation* va )
+    {
+        NodeMeshVariation::Private& node = *va->GetPrivate();
+
+        MESH_GENERATION_RESULT currentResult;
+        Ptr<ASTOp> currentMeshOp;
+
+        bool firstOptionProcessed = false;
+
+        // Default case
+        if ( node.m_defaultMesh )
+        {
+            MESH_GENERATION_RESULT branchResults;
+			GenerateMesh( branchResults, node.m_defaultMesh );
+            currentMeshOp = branchResults.meshOp;
+            currentResult = branchResults;
+            firstOptionProcessed = true;
+        }
+
+        // Process variations in reverse order, since conditionals are built bottom-up.
+        for ( int t = int(node.m_variations.size())-1; t >= 0; --t )
+        {
+            int tagIndex = -1;
+            const string& tag = node.m_variations[t].m_tag;
+            for ( int i = 0; i < int( m_firstPass.m_tags.size() ); ++i )
+            {
+                if ( m_firstPass.m_tags[i].tag==tag)
+                {
+                    tagIndex = i;
+                }
+            }
+
+            if ( tagIndex < 0 )
+            {
+                char buf[256];
+                mutable_snprintf( buf, 256, "Unknown tag found in mesh variation [%s].",
+                                  tag.c_str() );
+
+                m_pErrorLog->GetPrivate()->Add( buf, ELMT_WARNING, node.m_errorContext );
+                continue;
+            }
+
+            Ptr<ASTOp> variationMeshOp;
+            if ( node.m_variations[t].m_mesh )
+            {
+                if (firstOptionProcessed)
+                {
+                    m_overrideLayoutsStack.push_back( currentResult.layouts );
+                }
+         
+                MESH_GENERATION_RESULT branchResults;
+				GenerateMesh( branchResults, node.m_variations[t].m_mesh );
+
+                variationMeshOp = branchResults.meshOp;
+
+				if (firstOptionProcessed )
+				{
+					m_overrideLayoutsStack.pop_back();
+				}
+
+                if ( !firstOptionProcessed )
+                {
+                    firstOptionProcessed = true;                   
+                    currentResult = branchResults;
+                }
+            }
+
+
+            Ptr<ASTOpConditional> conditional = new ASTOpConditional;
+            conditional->type = OP_TYPE::ME_CONDITIONAL;
+            conditional->no = currentMeshOp;
+            conditional->yes = variationMeshOp;            
+            conditional->condition = m_firstPass.m_tags[tagIndex].genericCondition;
+
+            currentMeshOp = conditional;
+        }
+
+        result = currentResult;
+        result.meshOp = currentMeshOp;
+    }
+
+
+    //---------------------------------------------------------------------------------------------
+    void CodeGenerator::GenerateMesh_Constant( MESH_GENERATION_RESULT& result,
+		const NodeMeshConstant* constant )
+    {
+        NodeMeshConstant::Private& node = *constant->GetPrivate();
+
+        Ptr<ASTOpConstantResource> op = new ASTOpConstantResource();
+        op->type = OP_TYPE::ME_CONSTANT;
+		result.baseMeshOp = op;
+		result.meshOp = op;
+
+        MeshPtr pMesh = node.m_pValue.get();
+        if (pMesh)
+        {
+            // Clone the mesh
+            MeshPtr pCloned = pMesh->Clone();
+            pCloned->EnsureSurfaceData();
+            if ( m_overrideLayoutsStack.empty() )
+            {
+                // This means that we are processing a base mesh
+
+                // We will redefine the layouts
+				result.layouts.clear();
+
+                // Apply whatever transform is necessary for every layout
+                for ( std::size_t l=0; l<node.m_layouts.size();++l )
+                {
+					NodeLayoutPtr pLayoutNode = node.m_layouts[l];
+					// TODO: In a cleanup of the design of the layouts, we should remove this cast.
+					const NodeLayoutBlocks* TypedNode = dynamic_cast<NodeLayoutBlocks*>(pLayoutNode.get());
+					if (TypedNode)
+					{
+						GenerateLayout(result, TypedNode, l, pCloned);
+					}
+                }
+            }
+            else
+            {
+                // We need to apply the transform of the source layouts
+                for ( std::size_t l=0; l<m_overrideLayoutsStack.back().size(); ++l )
+                {
+                    PrepareForLayout( m_overrideLayoutsStack.back()[l],
+                                      pCloned, l,
+                                      node.m_errorContext  );
+
+                }
+
+				result.layouts = m_overrideLayoutsStack.back();
+
+                // Disabled: lots of misleading messages (when not using layouts?).
+                // Warn about unnecessary layouts
+//                if ( node.m_layouts.size() )
+//                {
+//                    m_pErrorLog->GetPrivate()->Add( "The layouts in this mesh are not required.",
+//                                                    ELMT_WARNING, node.m_errorContext );
+//                }
+            }
+
+            // See if we already have a mesh identical to this one, except for the internal data
+            // like vertex indices.
+            bool isDuplicated = false;
+            MeshPtrConst pCandidate;
+            for ( size_t i=0; !isDuplicated && i<m_constantMeshes.size(); ++i)
+            {
+                pCandidate = m_constantMeshes[i];
+
+                if ( pCandidate->IsSimilar(*pCloned) )
+                {
+                    isDuplicated = true;
+
+                    // Remap layouts from the source mesh to the the ones created for the mesh we will use instead.
+                    for (int l=0; l<pCandidate->GetLayoutCount(); ++l)
+                    {
+                        check(pCandidate->GetLayoutCount()==pCloned->GetLayoutCount());
+
+                        const Layout* pSourceLayoutValue = pCandidate->GetLayout(l);
+                        const Layout* pDestLayoutValue = pCloned->GetLayout(l);
+
+                        Ptr<const Layout> pDestLayoutKey = nullptr;
+                        for( const auto& e : m_addedLayouts )
+                        {
+                            if (e.second==pDestLayoutValue)
+                            {
+                                pDestLayoutKey = e.first;
+                            }
+                        }
+                        check(pDestLayoutKey);
+                        m_addedLayouts[pDestLayoutKey] = pSourceLayoutValue;
+                    }
+                }
+            }
+
+            if (isDuplicated)
+            {
+                op->SetValue( pCandidate, m_compilerOptions->m_optimisationOptions.m_useDiskCache );
+            }
+            else
+            {
+                // Enumerate the vertices uniquely unless they already have indices
+                int buf = -1;
+                int chan = -1;
+                pCloned->GetVertexBuffers().FindChannel( MBS_VERTEXINDEX,0,&buf,&chan );
+                bool hasVertexIndices = (buf>=0 && chan>=0);
+                if ( !hasVertexIndices )
+                {
+                    int newBuffer = pCloned->GetVertexBuffers().GetBufferCount();
+                    pCloned->GetVertexBuffers().SetBufferCount( newBuffer+1 );
+                    MESH_BUFFER_SEMANTIC semantic = MBS_VERTEXINDEX;
+                    int semanticIndex = 0;
+                    MESH_BUFFER_FORMAT format = MBF_UINT32;
+                    int components = 1;
+                    int offset = 0;
+                    pCloned->GetVertexBuffers().SetBuffer
+                            (
+                                newBuffer,
+                                sizeof(uint32_t),
+                                1,
+                                &semantic, &semanticIndex,
+                                &format, &components,
+                                &offset
+                            );
+                    uint32_t* pIdData = (uint32_t*)pCloned->GetVertexBuffers().GetBufferData( newBuffer );
+                    for ( int i=0; i<pMesh->GetVertexCount(); ++i )
+                    {
+                        (*pIdData++) = m_freeVertexIndex++;
+						check(m_freeVertexIndex < std::numeric_limits<uint32_t>::max());
+                    }
+                }
+
+                // Add the constant data
+                m_constantMeshes.push_back(pCloned);
+                op->SetValue( pCloned.get(), m_compilerOptions->m_optimisationOptions.m_useDiskCache );
+            }
+        }
+        else
+        {
+			result.layouts.clear();
+
+            // This data is required
+            MeshPtr pTempMesh = new Mesh();
+            op->SetValue( pTempMesh, m_compilerOptions->m_optimisationOptions.m_useDiskCache );
+            m_constantMeshes.push_back(pTempMesh);
+
+            // Log an error message
+            m_pErrorLog->GetPrivate()->Add( "Constant mesh not set.",
+                                            ELMT_WARNING, node.m_errorContext );
+        }
+
+		// Apply the modifier for the pre-normal operations stage.
+		BOTTOM_UP_STATE temp = m_currentBottomUpState;
+		if (!m_activeTags.empty())
+		{
+			bool bModifiersForBeforeOperations = true;
+			result.meshOp = ApplyMeshModifiers(op, m_activeTags.back(),
+				bModifiersForBeforeOperations, node.m_errorContext);
+		}
+		m_currentBottomUpState = temp;
+
+    }
+
+
+
+
+    //---------------------------------------------------------------------------------------------
+    void CodeGenerator::GenerateMesh_Format( MESH_GENERATION_RESULT& result,
+		const NodeMeshFormat* format )
+    {
+        NodeMeshFormat::Private& node = *format->GetPrivate();
+
+        if ( node.m_pSource )
+        {
+            MESH_GENERATION_RESULT baseResult;
+			GenerateMesh(baseResult, node.m_pSource);
+
+            Ptr<ASTOpMeshFormat> op = new ASTOpMeshFormat();
+            op->Source = baseResult.meshOp;
+            op->Buffers = 0;
+
+            if ( node.m_rebuildTangents )
+            {
+                op->Buffers |= OP::MeshFormatArgs::BT_REBUILD_TANGENTS;
+            }
+
+            MeshPtr pFormatMesh = new Mesh();
+
+            if (node.m_VertexBuffers.GetBufferCount())
+            {
+                op->Buffers |= OP::MeshFormatArgs::BT_VERTEX;
+                pFormatMesh->m_VertexBuffers = node.m_VertexBuffers;
+            }
+
+            if (node.m_IndexBuffers.GetBufferCount())
+            {
+				op->Buffers |= OP::MeshFormatArgs::BT_INDEX;
+                pFormatMesh->m_IndexBuffers = node.m_IndexBuffers;
+            }
+
+            if (node.m_FaceBuffers.GetBufferCount())
+            {
+                op->Buffers |= OP::MeshFormatArgs::BT_FACE;
+                pFormatMesh->m_FaceBuffers = node.m_FaceBuffers;
+            }
+
+            Ptr<ASTOpConstantResource> cop = new ASTOpConstantResource();
+            cop->type = OP_TYPE::ME_CONSTANT;
+            cop->SetValue( pFormatMesh, m_compilerOptions->m_optimisationOptions.m_useDiskCache );
+            op->Format = cop;
+
+            m_constantMeshes.push_back(pFormatMesh);
+
+            result.meshOp = op;
+            result.baseMeshOp = baseResult.baseMeshOp;
+            result.layouts = baseResult.layouts;
+        }
+        else
+        {
+            // Put something there
+            GenerateMesh( result, new NodeMeshConstant() );
+        }
+    }
+
+
+    //---------------------------------------------------------------------------------------------
+    void CodeGenerator::GenerateMesh_Subtract( MESH_GENERATION_RESULT& result,
+		const NodeMeshSubtract* subs )
+    {
+		// This node is deprecated.
+		check(false);        
+    }
+
+
+    //---------------------------------------------------------------------------------------------
+    void CodeGenerator::GenerateMesh_Transform( MESH_GENERATION_RESULT& result,
+                                                const NodeMeshTransform* trans )
+    {
+        const auto& node = *trans->GetPrivate();
+
+        Ptr<ASTOpMeshTransform> op = new ASTOpMeshTransform();
+
+        // Base
+        if (node.m_pSource)
+        {
+            GenerateMesh( result, node.m_pSource);
+            op->source = result.meshOp;
+        }
+        else
+        {
+            // This argument is required
+            m_pErrorLog->GetPrivate()->Add("Mesh transform base node is not set.",
+                ELMT_ERROR, node.m_errorContext);
+        }
+
+        op->matrix = node.m_transform;
+
+        result.meshOp = op;
+    }
+
+
+    //---------------------------------------------------------------------------------------------
+    void CodeGenerator::GenerateMesh_ClipMorphPlane( MESH_GENERATION_RESULT& result,
+                                                     const NodeMeshClipMorphPlane* clip )
+    {
+        const auto& node = *clip->GetPrivate();
+
+        Ptr<ASTOpMeshClipMorphPlane> op = new ASTOpMeshClipMorphPlane();
+
+        // Base
+        if (node.m_pSource)
+        {
+            GenerateMesh( result, node.m_pSource);
+            op->source = result.meshOp;
+        }
+        else
+        {
+            // This argument is required
+            m_pErrorLog->GetPrivate()->Add("Mesh clip-morph-plane source node is not set.",
+                ELMT_ERROR, node.m_errorContext);
+        }
+
+        // Morph to an ellipse
+        {
+            op->morphShape.type = (uint8_t)SHAPE::Type::Ellipse;
+            op->morphShape.position = node.m_origin;
+            op->morphShape.up = node.m_normal;
+            op->morphShape.size = vec3f(node.m_radius1, node.m_radius2, node.m_rotation); // TODO: Move rotation to ellipse rotation reference base instead of passing it directly
+
+                                                                                      // Generate a "side" vector.
+                                                                                      // \todo: make generic and move to the vector class
+            {
+                // Generate vector perpendicular to normal for ellipse rotation reference base
+                vec3f aux_base(0.f, 1.f, 0.f);
+
+                if (fabs(dot(node.m_normal, aux_base)) > 0.95f)
+                {
+                    aux_base = vec3f(0.f, 0.f, 1.f);
+                }
+
+                op->morphShape.side = cross(node.m_normal, aux_base);
+            }
+        }
+
+        // Selection by shape
+        if (node.m_vertexSelectionType== NodeMeshClipMorphPlane::Private::VS_SHAPE)
+        {
+            op->vertexSelectionType = OP::MeshClipMorphPlaneArgs::VS_SHAPE;
+            op->selectionShape.type = (uint8_t)SHAPE::Type::AABox;
+            op->selectionShape.position = node.m_selectionBoxOrigin;
+            op->selectionShape.size = node.m_selectionBoxRadius;
+        }
+        else if (node.m_vertexSelectionType == NodeMeshClipMorphPlane::Private::VS_BONE_HIERARCHY)
+        {
+            // Selection by bone hierarchy?
+            op->vertexSelectionType = OP::MeshClipMorphPlaneArgs::VS_BONE_HIERARCHY;
+            op->vertexSelectionBone = node.m_vertexSelectionBone;
+			op->vertexSelectionBoneMaxRadius = node.m_maxEffectRadius;
+        }
+        else
+        {
+            op->vertexSelectionType = OP::MeshClipMorphPlaneArgs::VS_NONE;
+        }
+
+        // Parameters
+        op->dist = node.m_dist;
+        op->factor = node.m_factor;
+
+        result.meshOp = op;
+    }
+
+
+    //---------------------------------------------------------------------------------------------
+    void CodeGenerator::GenerateMesh_ClipWithMesh( MESH_GENERATION_RESULT& result,
+                                                   const NodeMeshClipWithMesh* clip )
+    {
+        const auto& node = *clip->GetPrivate();
+
+        Ptr<ASTOpFixed> op = new ASTOpFixed();
+        op->op.type = OP_TYPE::ME_CLIPWITHMESH;
+
+        // Base
+        if (node.m_pSource)
+        {
+            GenerateMesh( result, node.m_pSource );
+            op->SetChild( op->op.args.MeshClipWithMesh.source, result.meshOp );
+        }
+        else
+        {
+            // This argument is required
+            m_pErrorLog->GetPrivate()->Add("Mesh clip-with-mesh source node is not set.",
+                ELMT_ERROR, node.m_errorContext);
+        }
+
+        // Clipping mesh
+        if (node.m_pClipMesh)
+        {
+			m_activeTags.push_back({});
+
+            MESH_GENERATION_RESULT clipResult;
+            GenerateMesh( clipResult, node.m_pClipMesh);
+            op->SetChild( op->op.args.MeshClipWithMesh.clipMesh, clipResult.meshOp );
+
+			m_activeTags.pop_back();
+		}
+        else
+        {
+            // This argument is required
+            m_pErrorLog->GetPrivate()->Add("Mesh clip-with-mesh clipping mesh node is not set.",
+                ELMT_ERROR, node.m_errorContext);
+        }
+
+        result.meshOp = op;
+    }
+
+	//---------------------------------------------------------------------------------------------
+	void CodeGenerator::GenerateMesh_ClipDeform(MESH_GENERATION_RESULT& Result, const NodeMeshClipDeform* ClipDeform)
+	{
+		const auto& Node = *ClipDeform->GetPrivate();
+
+		const Ptr<ASTOpMeshBindShape> OpBind = new ASTOpMeshBindShape();
+		const Ptr<ASTOpMeshClipDeform> OpClipDeform = new ASTOpMeshClipDeform();
+
+		// Base Mesh
+		if (Node.m_pBaseMesh)
+		{
+			GenerateMesh(Result, Node.m_pBaseMesh);
+			OpBind->Mesh = Result.meshOp;
+		}
+		else
+		{
+			// This argument is required
+			m_pErrorLog->GetPrivate()->Add("Mesh Clip Deform base mesh node is not set.",
+				ELMT_ERROR, Node.m_errorContext);
+		}
+
+		// Base Shape
+		if (Node.m_pClipShape)
+		{
+			MESH_GENERATION_RESULT baseResult;
+			GenerateMesh(baseResult, Node.m_pClipShape);
+			OpBind->Shape = baseResult.meshOp;
+			OpClipDeform->ClipShape = baseResult.meshOp;
+		}
+
+    	OpBind->m_discardInvalidBindings = false;
+		OpClipDeform->Mesh = OpBind;
+
+		Result.meshOp = OpClipDeform;
+	}
+
+    //---------------------------------------------------------------------------------------------
+    void CodeGenerator::GenerateMesh_ApplyPose( MESH_GENERATION_RESULT& result,
+                                                const NodeMeshApplyPose* pose )
+    {
+        const auto& node = *pose->GetPrivate();
+
+        Ptr<ASTOpMeshApplyPose> op = new ASTOpMeshApplyPose();
+
+        // Base
+        if (node.m_pBase)
+        {
+            GenerateMesh( result, node.m_pBase );
+            op->base = result.meshOp;
+        }
+        else
+        {
+            // This argument is required
+            m_pErrorLog->GetPrivate()->Add("Mesh apply-pose base node is not set.",
+                ELMT_ERROR, node.m_errorContext);
+        }
+
+        // Pose mesh
+        if (node.m_pPose)
+        {
+            // We don't need layouts for the pose mesh
+            m_overrideLayoutsStack.push_back( vector<LayoutPtrConst>() );
+			m_activeTags.push_back({});
+
+            MESH_GENERATION_RESULT poseResult;
+            GenerateMesh( poseResult, node.m_pPose );
+            op->pose = poseResult.meshOp;
+
+            m_overrideLayoutsStack.pop_back();
+			m_activeTags.pop_back();
+		}
+        else
+        {
+            // This argument is required
+            m_pErrorLog->GetPrivate()->Add("Mesh apply-pose pose node is not set.",
+                ELMT_ERROR, node.m_errorContext);
+        }
+
+        result.meshOp = op;
+    }
+
+
+	//---------------------------------------------------------------------------------------------
+	void CodeGenerator::GenerateMesh_GeometryOperation(MESH_GENERATION_RESULT& result, const NodeMeshGeometryOperation* geom)
+	{
+		const auto& node = *geom->GetPrivate();
+
+		Ptr<ASTOpMeshGeometryOperation> op = new ASTOpMeshGeometryOperation();
+
+		// Mesh A
+		if (node.m_pMeshA)
+		{
+			GenerateMesh(result, node.m_pMeshA);
+			op->meshA = result.meshOp;
+		}
+		else
+		{
+			// This argument is required
+			m_pErrorLog->GetPrivate()->Add("Mesh geometric op mesh-a node is not set.",
+				ELMT_ERROR, node.m_errorContext);
+		}
+
+		// Mesh B
+		if (node.m_pMeshB)
+		{
+			MESH_GENERATION_RESULT bResult;
+			GenerateMesh(bResult, node.m_pMeshB);
+			op->meshB = bResult.meshOp;
+		}
+
+		op->scalarA = Generate(node.m_pScalarA);
+		op->scalarB = Generate(node.m_pScalarB);
+
+		result.meshOp = op;
+	}
+
+
+	//---------------------------------------------------------------------------------------------
+	void CodeGenerator::GenerateMesh_Reshape(MESH_GENERATION_RESULT& result, const NodeMeshReshape* reshape)
+	{
+		const NodeMeshReshape::Private& node = *reshape->GetPrivate();
+
+		Ptr<ASTOpMeshBindShape> opBind = new ASTOpMeshBindShape();
+		Ptr<ASTOpMeshApplyShape> opApply = new ASTOpMeshApplyShape();
+
+		opBind->m_reshapeSkeleton = node.m_reshapeSkeleton;
+    	opBind->m_enableRigidParts = node.m_enableRigidParts;
+		opBind->m_deformAllBones = node.m_deformAllBones;
+		opBind->m_bonesToDeform = node.m_bonesToDeform;
+    	opBind->m_reshapePhysicsVolumes = node.m_reshapePhysicsVolumes;
+		opBind->m_deformAllPhysics = node.m_deformAllPhysics;
+		opBind->m_physicsToDeform = node.m_physicsToDeform;
+		opBind->m_reshapeVertices = true;
+		opBind->m_bindingMethod = static_cast<uint32>(EShapeBindingMethod::ReshapeClosestProject);
+
+		opApply->m_reshapeVertices = true;	
+		opApply->m_reshapeSkeleton = node.m_reshapeSkeleton;
+		opApply->m_reshapePhysicsVolumes = node.m_reshapePhysicsVolumes;
+
+		// Base Mesh
+		if (node.m_pBaseMesh)
+		{
+			GenerateMesh(result, node.m_pBaseMesh);
+			opBind->Mesh = result.meshOp;
+		}
+		else
+		{
+			// This argument is required
+			m_pErrorLog->GetPrivate()->Add("Mesh reshape base node is not set.", ELMT_ERROR, node.m_errorContext);
+		}
+
+		// Base and target shapes shouldn't have layouts or modifiers.
+		m_overrideLayoutsStack.push_back({});
+		m_activeTags.push_back({});
+
+		// Base Shape
+		if (node.m_pBaseShape)
+		{
+			MESH_GENERATION_RESULT baseResult;
+			GenerateMesh(baseResult, node.m_pBaseShape);
+			opBind->Shape = baseResult.meshOp;
+		}
+
+		opApply->Mesh = opBind;
+
+		// Target Shape
+		if (node.m_pTargetShape)
+		{
+			MESH_GENERATION_RESULT targetResult;
+			GenerateMesh(targetResult, node.m_pTargetShape);
+			opApply->Shape = targetResult.meshOp;
+		}
+
+		m_overrideLayoutsStack.pop_back();
+		m_activeTags.pop_back();
+
+		result.meshOp = opApply;
+	}
+
+}
+
diff --git a/Engine/Plugins/Experimental/Mutable/Source/MutableTools/Private/CodeGenerator_Projector.cpp b/Engine/Plugins/Experimental/Mutable/Source/MutableTools/Private/CodeGenerator_Projector.cpp
new file mode 100644
--- /dev/null
+++ b/Engine/Plugins/Experimental/Mutable/Source/MutableTools/Private/CodeGenerator_Projector.cpp
@@ -0,0 +1,169 @@
+// Copyright Epic Games, Inc. All Rights Reserved.
+
+#include "CodeGenerator.h"
+
+#include "Platform.h"
+#include "CompilerPrivate.h"
+
+#include "NodeProjectorPrivate.h"
+#include "MutableTools/Private/ASTOpConstantProjector.h"
+
+#include "ModelPrivate.h"
+#include "Operations.h"
+
+
+using namespace mu;
+
+
+//-------------------------------------------------------------------------------------------------
+PROJECTOR ProjectorFromNode( const NodeProjectorConstant::Private& node )
+{
+    PROJECTOR p;
+    p.type = node.m_type;
+    p.position[0] = node.m_position[0];
+    p.position[1] = node.m_position[1];
+    p.position[2] = node.m_position[2];
+    p.direction[0] = node.m_direction[0];
+    p.direction[1] = node.m_direction[1];
+    p.direction[2] = node.m_direction[2];
+    p.up[0] = node.m_up[0];
+    p.up[1] = node.m_up[1];
+    p.up[2] = node.m_up[2];
+    p.scale[0] = node.m_scale[0];
+    p.scale[1] = node.m_scale[1];
+    p.scale[2] = node.m_scale[2];
+    p.projectionAngle = node.m_projectionAngle;
+
+    return p;
+}
+
+
+//-------------------------------------------------------------------------------------------------
+void CodeGenerator::GenerateProjector( PROJECTOR_GENERATION_RESULT& result, const NodeProjectorPtrConst& untyped )
+{
+    if (!untyped)
+    {
+        result = PROJECTOR_GENERATION_RESULT();
+        return;
+    }
+
+    // See if it was already generated
+	VISITED_MAP_KEY key = GetCurrentCacheKey(untyped);
+	GeneratedProjectorsMap::ValueType* it = m_generatedProjectors.Find( key );
+    if ( it )
+    {
+        result = *it;
+        return;
+    }
+
+    // Generate for each different type of node
+    if ( auto Constant = dynamic_cast<const NodeProjectorConstant*>(untyped.get()) )
+    {
+        GenerateProjector_Constant( result, Constant );
+    }
+    else if ( auto Param = dynamic_cast<const NodeProjectorParameter*>(untyped.get()) )
+    {
+        GenerateProjector_Parameter( result, Param );
+    }
+    else
+    {
+        check(false);
+    }
+
+
+    // Cache the result
+    m_generatedProjectors.Add( key, result );
+}
+
+
+//-------------------------------------------------------------------------------------------------
+void CodeGenerator::GenerateProjector_Constant( PROJECTOR_GENERATION_RESULT& result,
+                                                const Ptr<const NodeProjectorConstant>& constant )
+{
+    const NodeProjectorConstant::Private& node = *constant->GetPrivate();
+
+    Ptr<ASTOpConstantProjector> op = new ASTOpConstantProjector();
+    op->value = ProjectorFromNode( node );
+
+    result.op = op;
+    result.type = op->value.type;
+}
+
+
+//-------------------------------------------------------------------------------------------------
+void CodeGenerator::GenerateProjector_Parameter( PROJECTOR_GENERATION_RESULT& result,
+                                                const Ptr<const NodeProjectorParameter>& paramn )
+{
+    const NodeProjectorParameter::Private& node = *paramn->GetPrivate();
+
+    Ptr<ASTOpParameter> op;
+
+    auto it = m_nodeVariables.find(node.m_pNode);
+    if ( it == m_nodeVariables.end() )
+    {
+        PARAMETER_DESC param;
+        param.m_name = node.m_name;
+        param.m_uid = node.m_uid;
+        param.m_type = PARAMETER_TYPE::T_PROJECTOR;
+
+        PROJECTOR p = ProjectorFromNode( node );
+        param.m_defaultValue.m_projector = p;
+
+        op = new ASTOpParameter();
+        op->type = OP_TYPE::PR_PARAMETER;
+        op->parameter = param;
+
+        // Generate the code for the ranges
+        for ( std::size_t a=0; a<node.m_ranges.size(); ++a )
+        {
+            RANGE_GENERATION_RESULT rangeResult;
+            GenerateRange( rangeResult, node.m_ranges[a] );
+            op->ranges.emplace_back
+                    ( op.get(), rangeResult.sizeOp, rangeResult.rangeName, rangeResult.rangeUID );
+        }
+
+        m_nodeVariables[node.m_pNode] = op;
+    }
+    else
+    {
+        op = it->second;
+    }
+
+    result.type = op->parameter.m_defaultValue.m_projector.type;
+    result.op = op;
+}
+
+
+//-------------------------------------------------------------------------------------------------
+void CodeGenerator::GenerateMissingProjectorCode( PROJECTOR_GENERATION_RESULT& result,
+                                                  const void* errorContext )
+{
+    // Log an error message
+    m_pErrorLog->GetPrivate()->Add( "Required projector connection not found.",
+                                    ELMT_ERROR, errorContext );
+
+    PROJECTOR p;
+    p.type = PROJECTOR_TYPE::PLANAR;
+
+    p.direction[0] = 1;
+    p.direction[1] = 0;
+    p.direction[2] = 0;
+
+    p.up[0] = 0;
+    p.up[1] = 0;
+    p.up[2] = 1;
+
+    p.position[0] = 0.0f;
+    p.position[1] = 0.0f;
+    p.position[2] = 0.0f;
+
+    p.scale[0] = 1.0f;
+    p.scale[1] = 1.0f;
+
+    Ptr<ASTOpConstantProjector> op = new ASTOpConstantProjector();
+    op->value = p;
+
+    result.op = op;
+    result.type = p.type;
+}
+
diff --git a/Engine/Plugins/Experimental/Mutable/Source/MutableTools/Private/CodeGenerator_Scalar.cpp b/Engine/Plugins/Experimental/Mutable/Source/MutableTools/Private/CodeGenerator_Scalar.cpp
new file mode 100644
--- /dev/null
+++ b/Engine/Plugins/Experimental/Mutable/Source/MutableTools/Private/CodeGenerator_Scalar.cpp
@@ -0,0 +1,466 @@
+// Copyright Epic Games, Inc. All Rights Reserved.
+
+#include "CodeGenerator.h"
+
+#include "Platform.h"
+#include "CompilerPrivate.h"
+
+#include "NodeScalarPrivate.h"
+#include "NodeScalarConstantPrivate.h"
+#include "NodeScalarParameterPrivate.h"
+#include "NodeScalarSwitchPrivate.h"
+#include "NodeScalarEnumParameterPrivate.h"
+#include "NodeScalarCurvePrivate.h"
+#include "NodeScalarArithmeticOperationPrivate.h"
+#include "NodeScalarVariationPrivate.h"
+#include "NodeScalarTablePrivate.h"
+
+#include "MutableTools/Private/ASTOpScalarCurve.h"
+#include "MutableTools/Private/ASTOpConditional.h"
+
+#include "ModelPrivate.h"
+#include "Operations.h"
+
+
+using namespace mu;
+
+
+//-------------------------------------------------------------------------------------------------
+void CodeGenerator::GenerateScalar( SCALAR_GENERATION_RESULT& result, const NodeScalarPtrConst& untyped )
+{
+    if (!untyped)
+    {
+        result = SCALAR_GENERATION_RESULT();
+        return;
+    }
+
+    // See if it was already generated
+	VISITED_MAP_KEY key = GetCurrentCacheKey(untyped);
+	GeneratedScalarsMap::ValueType* it = m_generatedScalars.Find( key );
+    if ( it )
+    {
+        result = *it;
+        return;
+    }
+
+
+    // Generate for each different type of node
+    if ( auto Constant = dynamic_cast<const NodeScalarConstant*>(untyped.get()) )
+    {
+        GenerateScalar_Constant( result, Constant );
+    }
+    else if ( auto Param = dynamic_cast<const NodeScalarParameter*>(untyped.get()) )
+    {
+        GenerateScalar_Parameter( result, Param );
+    }
+	else if (auto Switch = dynamic_cast<const NodeScalarSwitch*>(untyped.get()))
+	{
+		GenerateScalar_Switch(result, Switch);
+	}
+	else if (auto EnumParam = dynamic_cast<const NodeScalarEnumParameter*>(untyped.get()))
+	{
+		GenerateScalar_EnumParameter(result, EnumParam);
+	}
+	else if (auto Curve = dynamic_cast<const NodeScalarCurve*>(untyped.get()))
+	{
+		GenerateScalar_Curve(result, Curve);
+	}
+	else if (auto Arithmetic = dynamic_cast<const NodeScalarArithmeticOperation*>(untyped.get()))
+	{
+		GenerateScalar_Arithmetic(result, Arithmetic);
+	}
+	else if (auto Variation = dynamic_cast<const NodeScalarVariation*>(untyped.get()))
+	{
+		GenerateScalar_Variation(result, Variation);
+	}
+	else if (auto Table = dynamic_cast<const NodeScalarTable*>(untyped.get()))
+	{
+		GenerateScalar_Table(result, Table);
+	}
+	else
+    {
+        check(false);
+        mu::Halt();
+    }
+
+    // Cache the result
+    m_generatedScalars.Add( key, result );
+}
+
+
+//-------------------------------------------------------------------------------------------------
+void CodeGenerator::GenerateScalar_Constant(SCALAR_GENERATION_RESULT& result, const Ptr<const NodeScalarConstant>& Typed )
+{
+    const NodeScalarConstant::Private& node = *Typed->GetPrivate();
+
+	Ptr<ASTOpFixed> op = new ASTOpFixed();
+	op->op.type = OP_TYPE::SC_CONSTANT;
+	op->op.args.ScalarConstant.value = node.m_value;
+
+    result.op = op;
+}
+
+
+//-------------------------------------------------------------------------------------------------
+void CodeGenerator::GenerateScalar_Parameter( SCALAR_GENERATION_RESULT& result, const Ptr<const NodeScalarParameter>& Typed)
+{
+    const NodeScalarParameter::Private& node = *Typed->GetPrivate();
+
+    Ptr<ASTOpParameter> op;
+
+	auto it = m_nodeVariables.find(node.m_pNode);
+	if (it == m_nodeVariables.end())
+	{
+		PARAMETER_DESC param;
+		param.m_name = node.m_name;
+		param.m_uid = node.m_uid;
+		param.m_type = PARAMETER_TYPE::T_FLOAT;
+		param.m_defaultValue.m_float = node.m_defaultValue;
+		param.m_detailedType = node.m_detailedType;
+
+		op = new ASTOpParameter();
+		op->type = OP_TYPE::SC_PARAMETER;
+		op->parameter = param;
+
+		// Generate the code for the ranges
+		for (std::size_t a = 0; a < node.m_ranges.size(); ++a)
+		{
+			RANGE_GENERATION_RESULT rangeResult;
+			GenerateRange(rangeResult, node.m_ranges[a]);
+			op->ranges.emplace_back
+			(op.get(), rangeResult.sizeOp, rangeResult.rangeName, rangeResult.rangeUID);
+		}
+
+		// Generate the code for the additional images in the parameter
+		for (std::size_t a = 0; a < node.m_additionalImages.size(); ++a)
+		{
+			Ptr<ASTOp> descAd;
+			if (node.m_additionalImages[a])
+			{
+				// We take whatever size will be produced
+				IMAGE_STATE newState;
+				FImageDesc desc = CalculateImageDesc(*node.m_additionalImages[a]->GetBasePrivate());
+				newState.m_imageSize = desc.m_size;
+				newState.m_imageRect.min[0] = 0;
+				newState.m_imageRect.min[1] = 0;
+				newState.m_imageRect.size = desc.m_size;
+				newState.m_layoutBlock = -1;
+				m_imageState.push_back(newState);
+
+				// Generate
+				descAd = Generate(node.m_additionalImages[a]);
+				check(descAd);
+
+				// Restore rect
+				m_imageState.pop_back();
+			}
+			op->additionalImages.emplace_back(op, descAd);
+		}
+
+		m_nodeVariables[node.m_pNode] = op;
+	}
+	else
+	{
+		op = it->second;
+	}
+
+    result.op = op;
+}
+
+
+//-------------------------------------------------------------------------------------------------
+void CodeGenerator::GenerateScalar_EnumParameter(SCALAR_GENERATION_RESULT& result, const Ptr<const NodeScalarEnumParameter>& Typed)
+{
+	const NodeScalarEnumParameter::Private& node = *Typed->GetPrivate();
+
+	Ptr<ASTOpParameter> op;
+
+	auto it = m_nodeVariables.find(node.m_pNode);
+	if (it == m_nodeVariables.end())
+	{
+		PARAMETER_DESC param;
+		param.m_name = node.m_name;
+		param.m_uid = node.m_uid;
+		param.m_type = PARAMETER_TYPE::T_INT;
+		param.m_defaultValue.m_int = node.m_defaultValue;
+		param.m_detailedType = node.m_detailedType;
+
+		param.m_possibleValues.resize(node.m_options.size());
+		for (size_t i = 0; i < node.m_options.size(); ++i)
+		{
+			param.m_possibleValues[i].m_value = (int16_t)node.m_options[i].value;
+			param.m_possibleValues[i].m_name = node.m_options[i].name;
+		}
+
+		op = new ASTOpParameter();
+		op->type = OP_TYPE::NU_PARAMETER;
+		op->parameter = param;
+
+		// Generate the code for the ranges
+		for (std::size_t a = 0; a < node.m_ranges.size(); ++a)
+		{
+			RANGE_GENERATION_RESULT rangeResult;
+			GenerateRange(rangeResult, node.m_ranges[a]);
+			op->ranges.emplace_back
+			(op.get(), rangeResult.sizeOp, rangeResult.rangeName, rangeResult.rangeUID);
+		}
+
+		m_nodeVariables[node.m_pNode] = op;
+
+	}
+	else
+	{
+		op = it->second;
+	}
+
+	result.op = op;
+}
+
+
+//-------------------------------------------------------------------------------------------------
+void CodeGenerator::GenerateScalar_Switch(SCALAR_GENERATION_RESULT& result, const Ptr<const NodeScalarSwitch>& Typed)
+{
+	const NodeScalarSwitch::Private& node = *Typed->GetPrivate();
+
+	if (node.m_options.size() == 0)
+	{
+		// No options in the switch!
+		Ptr<ASTOp> missingOp = GenerateMissingScalarCode("Switch option",
+			1.0f,
+			node.m_errorContext);
+		result.op = missingOp;
+		return;
+	}
+
+	Ptr<ASTOpSwitch> op = new ASTOpSwitch();
+	op->type = OP_TYPE::SC_SWITCH;
+
+	// Variable value
+	if (node.m_pParameter)
+	{
+		SCALAR_GENERATION_RESULT ChildResult;
+		GenerateScalar(ChildResult,node.m_pParameter.get());
+		op->variable = ChildResult.op;
+	}
+	else
+	{
+		// This argument is required
+		op->variable = GenerateMissingScalarCode("Switch variable", 0.0f, node.m_errorContext);
+	}
+
+	// Options
+	for (std::size_t t = 0; t < node.m_options.size(); ++t)
+	{
+		Ptr<ASTOp> branch;
+		if (node.m_options[t])
+		{
+			SCALAR_GENERATION_RESULT ChildResult;
+			GenerateScalar(ChildResult,node.m_options[t].get());
+			branch = ChildResult.op;
+		}
+		else
+		{
+			// This argument is required
+			branch = GenerateMissingScalarCode("Switch option", 1.0f, node.m_errorContext);
+		}
+		op->cases.push_back(ASTOpSwitch::CASE((int16_t)t, op, branch));
+	}
+
+	result.op = op;
+}
+
+
+//-------------------------------------------------------------------------------------------------
+void CodeGenerator::GenerateScalar_Variation(SCALAR_GENERATION_RESULT& result, const Ptr<const NodeScalarVariation>& Typed)
+{
+	const NodeScalarVariation::Private& node = *Typed->GetPrivate();
+
+	Ptr<ASTOp> op;
+
+	// Default case
+	if (node.m_defaultScalar)
+	{
+		MESH_GENERATION_RESULT branchResults;
+
+		SCALAR_GENERATION_RESULT ChildResult;
+		GenerateScalar(ChildResult,node.m_defaultScalar);
+		op = ChildResult.op;
+	}
+	else
+	{
+		// This argument is required
+		op = GenerateMissingScalarCode("Variation default", 0.0f, node.m_errorContext);
+	}
+
+	// Process variations in reverse order, since conditionals are built bottom-up.
+	for (int t = int(node.m_variations.size()) - 1; t >= 0; --t)
+	{
+		int tagIndex = -1;
+		const string& tag = node.m_variations[t].m_tag;
+		for (int i = 0; i < int(m_firstPass.m_tags.size()); ++i)
+		{
+			if (m_firstPass.m_tags[i].tag == tag)
+			{
+				tagIndex = i;
+			}
+		}
+
+		if (tagIndex < 0)
+		{
+			char buf[256];
+			mutable_snprintf(buf, 256, "Unknown tag found in image variation [%s].",
+				tag.c_str());
+
+			m_pErrorLog->GetPrivate()->Add(buf, ELMT_WARNING, node.m_errorContext);
+			continue;
+		}
+
+		Ptr<ASTOp> variationOp;
+		if (node.m_variations[t].m_scalar)
+		{
+			SCALAR_GENERATION_RESULT ChildResult;
+			GenerateScalar(ChildResult,node.m_variations[t].m_scalar);
+			variationOp = ChildResult.op;
+		}
+		else
+		{
+			// This argument is required
+			variationOp = GenerateMissingScalarCode("Variation option", 0.0f,
+				node.m_errorContext);
+		}
+
+
+		Ptr<ASTOpConditional> conditional = new ASTOpConditional;
+		conditional->type = OP_TYPE::SC_CONDITIONAL;
+		conditional->no = op;
+		conditional->yes = variationOp;
+		conditional->condition = m_firstPass.m_tags[tagIndex].genericCondition;
+
+		op = conditional;
+	}
+
+	result.op = op;
+}
+
+
+//-------------------------------------------------------------------------------------------------
+void CodeGenerator::GenerateScalar_Curve(SCALAR_GENERATION_RESULT& result, const Ptr<const NodeScalarCurve>& Typed)
+{
+	const NodeScalarCurve::Private& node = *Typed->GetPrivate();
+
+	Ptr<ASTOpScalarCurve> op = new ASTOpScalarCurve();
+
+	// T
+	if (Node* pA = node.m_input_scalar.get())
+	{
+		op->time = Generate(pA);
+	}
+	else
+	{
+		op->time = CodeGenerator::GenerateMissingScalarCode( "Curve T", 0.5f, node.m_errorContext );
+	}
+
+	op->curve = node.m_curve;
+
+	result.op = op;
+}
+
+
+//-------------------------------------------------------------------------------------------------
+void CodeGenerator::GenerateScalar_Arithmetic(SCALAR_GENERATION_RESULT& result, const Ptr<const NodeScalarArithmeticOperation>& Typed)
+{
+	const NodeScalarArithmeticOperation::Private& node = *Typed->GetPrivate();
+
+	Ptr<ASTOpFixed> op = new ASTOpFixed();
+	op->op.type = OP_TYPE::SC_ARITHMETIC;
+
+	switch (node.m_operation)
+	{
+	case NodeScalarArithmeticOperation::AO_ADD: op->op.args.ScalarArithmetic.operation = OP::ArithmeticArgs::ADD; break;
+	case NodeScalarArithmeticOperation::AO_SUBTRACT: op->op.args.ScalarArithmetic.operation = OP::ArithmeticArgs::SUBTRACT; break;
+	case NodeScalarArithmeticOperation::AO_MULTIPLY: op->op.args.ScalarArithmetic.operation = OP::ArithmeticArgs::MULTIPLY; break;
+	case NodeScalarArithmeticOperation::AO_DIVIDE: op->op.args.ScalarArithmetic.operation = OP::ArithmeticArgs::DIVIDE; break;
+	default:
+		checkf(false, TEXT("Unknown arithmetic operation."));
+		op->op.args.ScalarArithmetic.operation = OP::ArithmeticArgs::NONE;
+		break;
+	}
+
+	// A
+	if (Node* pA = node.m_pA.get())
+	{
+		op->SetChild(op->op.args.ScalarArithmetic.a, Generate(pA));
+	}
+	else
+	{
+		op->SetChild(op->op.args.ScalarArithmetic.a,
+			CodeGenerator::GenerateMissingScalarCode
+			(
+				"ScalarArithmetic A",
+				1.0f,
+				node.m_errorContext
+			)
+		);
+	}
+
+	// B
+	if (Node* pB = node.m_pB.get())
+	{
+		op->SetChild(op->op.args.ScalarArithmetic.b, Generate(pB));
+	}
+	else
+	{
+		op->SetChild(op->op.args.ScalarArithmetic.b,
+			CodeGenerator::GenerateMissingScalarCode
+			(
+				"ScalarArithmetic B",
+				1.0f,
+				node.m_errorContext
+			)
+		);
+	}
+
+	result.op = op;
+}
+
+
+//-------------------------------------------------------------------------------------------------
+void CodeGenerator::GenerateScalar_Table(SCALAR_GENERATION_RESULT& result, const Ptr<const NodeScalarTable>& Typed)
+{
+	const NodeScalarTable::Private& node = *Typed->GetPrivate();
+
+	Ptr<ASTOp> Op = GenerateTableSwitch<NodeScalarTable::Private, TCT_SCALAR, OP_TYPE::SC_SWITCH>(node,
+			[this](const NodeScalarTable::Private& node, int colIndex, int row, ErrorLog* pErrorLog)
+			{
+				NodeScalarConstantPtr pCell = new NodeScalarConstant();
+				float scalar = node.m_pTable->GetPrivate()->m_rows[row].m_values[colIndex].m_scalar;
+				pCell->SetValue(scalar);
+				return Generate(pCell);
+			});
+
+	result.op = Op;
+}
+
+
+//---------------------------------------------------------------------------------------------
+mu::Ptr<ASTOp> CodeGenerator::GenerateMissingScalarCode(const char* strWhere, float value, const void* errorContext)
+{
+	// Log a warning
+	char buf[256];
+	mutable_snprintf
+	(
+		buf, 256,
+		"Required connection not found: %s",
+		strWhere
+	);
+	m_pErrorLog->GetPrivate()->Add(buf, ELMT_ERROR, errorContext);
+
+	// Create a constant node
+	NodeScalarConstantPtr pNode = new NodeScalarConstant();
+	pNode->SetValue(value);
+
+	Ptr<ASTOp> result = Generate(pNode);
+
+	return result;
+}
+
diff --git a/Engine/Plugins/Experimental/Mutable/Source/MutableTools/Private/CodeGenerator_SecondPass.cpp b/Engine/Plugins/Experimental/Mutable/Source/MutableTools/Private/CodeGenerator_SecondPass.cpp
new file mode 100644
--- /dev/null
+++ b/Engine/Plugins/Experimental/Mutable/Source/MutableTools/Private/CodeGenerator_SecondPass.cpp
@@ -0,0 +1,661 @@
+// Copyright Epic Games, Inc. All Rights Reserved.
+
+#include "CodeGenerator_SecondPass.h"
+#include "CodeGenerator_FirstPass.h"
+#include "CodeGenerator.h"
+
+#include "Platform.h"
+#include "CompilerPrivate.h"
+
+#include "NodeSurfaceNewPrivate.h"
+#include "NodeSurfaceEditPrivate.h"
+
+#include "NodeComponentPrivate.h"
+#include "NodeComponentNewPrivate.h"
+#include "NodeComponentEditPrivate.h"
+
+#include "NodeLODPrivate.h"
+#include "NodeObjectNewPrivate.h"
+#include "NodeObjectStatePrivate.h"
+#include "NodeObjectGroupPrivate.h"
+
+#include "NodeMeshClipMorphPlane.h"
+#include "NodeMeshClipWithMesh.h"
+
+#include "NodeModifierPrivate.h"
+#include "NodeModifierMeshClipMorphPlanePrivate.h"
+#include "NodeModifierMeshClipWithMeshPrivate.h"
+#include "NodeComponentNewPrivate.h"
+
+#include "MutableTools/Private/ASTOpConstantBool.h"
+
+#include "ModelPrivate.h"
+#include "Operations.h"
+#include "SystemPrivate.h"
+
+#include "MutableTrace.h"
+
+using namespace mu;
+
+
+//---------------------------------------------------------------------------------------------
+SecondPassGenerator::SecondPassGenerator( 
+    FirstPassGenerator* firstPass,
+    const CompilerOptions::Private* options )
+{
+    check(firstPass);
+    check(options);
+    m_pFirstPass = firstPass;
+    m_pCompilerOptions = options;
+
+    // Default conditions when there is no restriction accumulated.
+    CONDITION_CONTEXT noCondition;
+    m_currentCondition.push_back(noCondition);
+}
+
+
+//---------------------------------------------------------------------------------------------
+mu::Ptr<ASTOp> SecondPassGenerator::GenerateTagCondition( size_t tagIndex,
+                                                      const set<size_t>& posSurf,
+                                                      const set<size_t>& negSurf,
+                                                      const set<size_t>& posTag,
+                                                      const set<size_t>& negTag )
+{
+    auto& t = m_pFirstPass->m_tags[tagIndex];
+
+    // If this tag is already in the list of positive tags, return true as condition
+    if ( posTag.find(tagIndex)!=posTag.end() )
+    {
+        return m_opPool.Add( new ASTOpConstantBool(true) );
+    }
+
+    // If this tag is already in the list of negative tags, return false as condition
+    if ( negTag.find(tagIndex)!=negTag.end() )
+    {
+        return m_opPool.Add( new ASTOpConstantBool(false) );
+    }
+
+    // Cached?
+    CONDITION_GENERATION_KEY key;
+    key.tagOrSurfIndex = tagIndex;
+//    key.negSurf = negSurf;
+//    key.posSurf = posSurf;
+//    key.negTag = negTag;
+//    key.posTag = posTag;
+	for (auto s : negTag) { if (m_tagsPerTag[tagIndex].count(s) > 0) { key.negTag.insert(s); } }
+	for (auto s : posTag) { if (m_tagsPerTag[tagIndex].count(s) > 0) { key.posTag.insert(s); } }
+	for (auto s : negSurf) { if (m_surfacesPerTag[tagIndex].count(s) > 0) { key.negSurf.insert(s); } }
+	for (auto s : posSurf) { if (m_surfacesPerTag[tagIndex].count(s) > 0) { key.posSurf.insert(s); } }
+
+    {
+        auto it = m_tagConditionGenerationCache.find( key );
+        if (it!=m_tagConditionGenerationCache.end())
+        {
+            return it->second;
+        }
+    }
+
+    Ptr<ASTOp> c;
+
+    // Condition expression for all the surfaces that activate the tag
+    for ( auto surfIndex: t.surfaces )
+    {
+        if ( posSurf.find(surfIndex) != posSurf.end() )
+        {
+            // This surface is already a positive requirement higher up in the condition so
+            // we can ignore it here.
+            continue;
+        }
+
+        if ( negSurf.find(surfIndex) != negSurf.end() )
+        {
+            // This surface is a negative requirement higher up in the condition so
+            // this branch never be true.
+            continue;
+        }
+
+        const auto& surface = m_pFirstPass->surfaces[surfIndex];
+
+        auto positiveTags = posTag;
+        positiveTags.insert( tagIndex );
+
+        Ptr<ASTOp> surfCondition = GenerateSurfaceCondition( surfIndex,
+                                                             posSurf,
+                                                             negSurf,
+                                                             positiveTags,
+                                                             negTag );
+
+        // If the surface is a constant false, we can skip adding it
+        if ( auto* constOp = dynamic_cast<const ASTOpConstantBool*>(surfCondition.get()) )
+        {
+            if (constOp->value==false)
+            {
+                continue;
+            }
+        }
+
+        Ptr<ASTOp> fullCondition;
+        if (surfCondition)
+        {
+            Ptr<ASTOpFixed> f = new ASTOpFixed();
+            f->op.type = OP_TYPE::BO_AND;
+            f->SetChild(f->op.args.BoolBinary.a, surface.objectCondition);
+            f->SetChild(f->op.args.BoolBinary.b, surfCondition);
+            fullCondition = m_opPool.Add( f );
+        }
+        else
+        {
+            fullCondition = m_opPool.Add( surface.objectCondition );
+        }
+
+
+        if (!c)
+        {
+            c = fullCondition;
+        }
+        else
+        {
+            Ptr<ASTOpFixed> o = new ASTOpFixed();
+            o->op.type = OP_TYPE::BO_OR;
+            o->SetChild( o->op.args.BoolBinary.a, fullCondition) ;
+            o->SetChild( o->op.args.BoolBinary.b, c );
+            c = m_opPool.Add( o );
+        }
+
+        // Optimise the condition now.
+        //PartialOptimise( c, m_pCompilerOptions->m_optimisationOptions );
+    }
+
+
+    // Condition expression for all the edit-surfaces that activate the tag
+    for ( auto editKey: t.edits )
+    {
+        if ( negSurf.find(editKey.first) != negSurf.end() )
+        {
+            // The surface in the edit is a negative requirement higher up in the condition so
+            // this branch never be true.
+            continue;
+        }
+
+        Ptr<ASTOp> surfCondition;
+        if ( posSurf.find(editKey.first) != posSurf.end() )
+        {
+            // This surface in the edit is already a positive requirement higher up in the condition
+            // so we don't need that part of the condition
+        }
+        else
+        {
+            auto positiveTags = posTag;
+            positiveTags.insert( tagIndex );
+
+            surfCondition = GenerateSurfaceCondition( editKey.first,
+                                                      posSurf,
+                                                      negSurf,
+                                                      positiveTags,
+                                                      negTag );
+        }
+
+        const auto& surface = m_pFirstPass->surfaces[editKey.first];
+        const auto& edit = surface.edits[editKey.second];
+
+        // Combine object and surface conditions
+        Ptr<ASTOp> fullCondition;
+        if (surfCondition)
+        {
+            Ptr<ASTOpFixed> f = new ASTOpFixed();
+            f->op.type = OP_TYPE::BO_AND;
+            f->SetChild(f->op.args.BoolBinary.a, surface.objectCondition);
+            f->SetChild(f->op.args.BoolBinary.b, surfCondition);
+            fullCondition = m_opPool.Add( f );
+        }
+        else
+        {
+            fullCondition = m_opPool.Add( surface.objectCondition );
+        }
+
+        // Combine with edit condition
+        if (fullCondition)
+        {
+            Ptr<ASTOpFixed> f = new ASTOpFixed();
+            f->op.type = OP_TYPE::BO_AND;
+            f->SetChild(f->op.args.BoolBinary.a, edit.condition);
+            f->SetChild(f->op.args.BoolBinary.b, fullCondition);
+            fullCondition = m_opPool.Add( f );
+        }
+        else
+        {
+            fullCondition = m_opPool.Add( edit.condition );
+        }
+
+
+        if (!c)
+        {
+            c = fullCondition;
+        }
+        else
+        {
+            Ptr<ASTOpFixed> o = new ASTOpFixed();
+            o->op.type = OP_TYPE::BO_OR;
+            o->SetChild( o->op.args.BoolBinary.a, fullCondition) ;
+            o->SetChild( o->op.args.BoolBinary.b, c );
+            c = m_opPool.Add( o );
+        }
+    }
+
+
+    m_tagConditionGenerationCache.insert( std::make_pair<>(key,c) );
+
+    return c;
+}
+
+//---------------------------------------------------------------------------------------------
+mu::Ptr<ASTOp> SecondPassGenerator::GenerateSurfaceCondition( size_t surfIndex,
+                                                    const set<size_t>& posSurf,
+                                                    const set<size_t>& negSurf,
+                                                    const set<size_t>& posTag,
+                                                    const set<size_t>& negTag )
+{
+    const auto& surf = m_pFirstPass->surfaces[surfIndex];
+
+    // If this surface is already in the list of positive surfaces, return true as condition
+    if ( posSurf.find(surfIndex)!=posSurf.end() )
+    {
+        return m_opPool.Add( new ASTOpConstantBool(true) );
+    }
+
+    // If this surface is already in the list of negative surfaces, return false as condition
+    if ( negSurf.find(surfIndex)!=negSurf.end() )
+    {
+        return m_opPool.Add( new ASTOpConstantBool(false) );
+    }
+
+    Ptr<ASTOp> c;
+
+    for (const auto& t: surf.positiveTags)
+    {
+        auto it = std::find_if( m_pFirstPass->m_tags.begin(), m_pFirstPass->m_tags.end(),
+                                [&](const FirstPassGenerator::TAG& e) { return e.tag==t; } );
+        if (it==m_pFirstPass->m_tags.end())
+        {
+            // This could happen if a tag is in a variation but noone defines it.
+            // This surface depends on a tag that will never be active, so it will never be used.
+            return m_opPool.Add( new ASTOpConstantBool(false) );
+        }
+
+        size_t tagIndex = it - m_pFirstPass->m_tags.begin();
+
+        set<size_t> positiveSurfacesVisited = posSurf;
+        positiveSurfacesVisited.insert( surfIndex );
+
+        Ptr<ASTOp> tagCondition = GenerateTagCondition( tagIndex,
+                                                        positiveSurfacesVisited,
+                                                        negSurf,
+                                                        posTag,
+                                                        negTag );
+        // TODO: Optimise the tag condition here
+
+        
+
+        // If the tag is a constant ...
+        bool isConstant = false;
+        bool constantValue = false;
+        if ( auto* constOp = dynamic_cast<const ASTOpConstantBool*>(tagCondition.get()) )
+        {
+            isConstant = true;
+            constantValue = constOp->value;
+        }
+
+        if (tagCondition && !isConstant)
+        {
+            if (!c)
+            {
+                c = tagCondition;
+            }
+            else
+            {
+                Ptr<ASTOpFixed> o = new ASTOpFixed;
+                o->op.type = OP_TYPE::BO_AND;
+                o->SetChild( o->op.args.BoolBinary.a, tagCondition );
+                o->SetChild( o->op.args.BoolBinary.b, c );
+                c = m_opPool.Add( o );
+            }
+        }
+        else if (constantValue==true)
+        {
+            // No need to add it to the AND
+        }
+        else //if (constantValue==false)
+        {
+            // Entire expression will be false
+            Ptr<ASTOpConstantBool> f = new ASTOpConstantBool( false );
+
+            // No need to evaluate anything else.
+            c = m_opPool.Add( f );
+            break;
+        }
+    }
+
+
+    for (const auto& t: surf.negativeTags)
+    {
+        auto it = std::find_if( m_pFirstPass->m_tags.begin(), m_pFirstPass->m_tags.end(),
+                                [&](const FirstPassGenerator::TAG& e) { return e.tag==t; } );
+        if (it==m_pFirstPass->m_tags.end())
+        {
+            // This could happen if a tag is in a variation but noone defines it.
+            continue;
+        }
+
+        size_t tagIndex = it - m_pFirstPass->m_tags.begin();
+
+        set<size_t> positiveSurfacesVisited = negSurf;
+        set<size_t> negativeSurfacesVisited = posSurf;
+        negativeSurfacesVisited.insert( surfIndex );
+        set<size_t> positiveTagsVisited = negTag;
+        set<size_t> negativeTagsVisited = posTag;
+        Ptr<ASTOp> tagCondition = GenerateTagCondition( tagIndex,
+                                                        positiveSurfacesVisited,
+                                                        negativeSurfacesVisited,
+                                                        positiveTagsVisited,
+                                                        negativeTagsVisited );
+
+        // TODO: Optimise the tag condition here
+
+        // If the tag is a constant ...
+        bool isConstant = false;
+        bool constantValue = false;
+        if ( auto* constOp = dynamic_cast<const ASTOpConstantBool*>(tagCondition.get()) )
+        {
+            isConstant = true;
+            constantValue = constOp->value;
+        }
+
+
+        if (!isConstant && tagCondition)
+        {
+            Ptr<ASTOpFixed> n = new ASTOpFixed;
+            n->op.type = OP_TYPE::BO_NOT;
+            n->SetChild( n->op.args.BoolNot.source, tagCondition );
+
+            if (!c)
+            {
+                c = m_opPool.Add( n );
+            }
+            else
+            {
+                Ptr<ASTOpFixed> o = new ASTOpFixed;
+                o->op.type = OP_TYPE::BO_AND;
+                o->SetChild( o->op.args.BoolBinary.a, n);
+                o->SetChild( o->op.args.BoolBinary.b, c);
+                c = m_opPool.Add( o );
+            }
+        }
+        else if (isConstant && constantValue==true)
+        {
+            // No expression here means always true which becomes always false
+            Ptr<ASTOpConstantBool> f = new ASTOpConstantBool( false );
+
+            // No need to evaluate anything else.
+            c = m_opPool.Add( f );
+            break;
+        }
+    }
+
+    return c;
+}
+
+//---------------------------------------------------------------------------------------------
+mu::Ptr<ASTOp> SecondPassGenerator::GenerateModifierCondition( size_t modIndex )
+{
+    const auto& mod = m_pFirstPass->modifiers[modIndex];
+
+    Ptr<ASTOp> c;
+
+    bool done = false;
+    for (const auto& t: mod.positiveTags)
+    {
+        auto it = std::find_if( m_pFirstPass->m_tags.begin(), m_pFirstPass->m_tags.end(),
+                                [&](const FirstPassGenerator::TAG& e) { return e.tag==t; } );
+        if (it==m_pFirstPass->m_tags.end())
+        {
+            // This could happen if a tag is in a variation but noone defines it.
+
+            // Entire expression will be false
+            Ptr<ASTOpConstantBool> f = new ASTOpConstantBool( false );
+            c = m_opPool.Add( f );
+            done = true;
+
+            continue;
+        }
+
+        size_t tagIndex = it - m_pFirstPass->m_tags.begin();
+
+        set<size_t> empty;
+        Ptr<ASTOp> tagCondition = GenerateTagCondition( tagIndex, empty, empty, empty, empty );
+
+        // TODO: Optimise the tag condition here
+
+        // If the tag is a constant ...
+        bool isConstant = false;
+        bool constantValue = false;
+        if ( auto* constOp = dynamic_cast<const ASTOpConstantBool*>(tagCondition.get()) )
+        {
+            isConstant = true;
+            constantValue = constOp->value;
+        }
+
+        if (tagCondition && !isConstant)
+        {
+            if (!c)
+            {
+                c = tagCondition;
+            }
+            else
+            {
+                Ptr<ASTOpFixed> o = new ASTOpFixed;
+                o->op.type = OP_TYPE::BO_AND;
+                o->SetChild( o->op.args.BoolBinary.a, tagCondition );
+                o->SetChild( o->op.args.BoolBinary.b, c );
+                c = m_opPool.Add( o );
+            }
+        }
+        else if (constantValue==true)
+        {
+            // No need to add it to the AND
+        }
+        else //if (constantValue==false)
+        {
+            // Entire expression will be false
+            Ptr<ASTOpConstantBool> f = new ASTOpConstantBool( false );
+
+            // No need to evaluate anything else.
+            c = m_opPool.Add( f );
+            done = true;
+            break;
+        }
+    }
+
+    if (!done)
+    {
+        for (const auto& t: mod.negativeTags)
+        {
+            auto it = std::find_if( m_pFirstPass->m_tags.begin(), m_pFirstPass->m_tags.end(),
+                                    [&](const FirstPassGenerator::TAG& e) { return e.tag==t; } );
+            if (it==m_pFirstPass->m_tags.end())
+            {
+                // This could happen if a tag is in a variation but noone defines it.
+                continue;
+            }
+
+            size_t tagIndex = it - m_pFirstPass->m_tags.begin();
+
+            set<size_t> empty;
+            Ptr<ASTOp> tagCondition = GenerateTagCondition( tagIndex, empty, empty, empty, empty );
+
+            // TODO: Optimise the tag condition here
+
+            // If the tag is a constant ...
+            bool isConstant = false;
+            bool constantValue = false;
+            if ( auto* constOp = dynamic_cast<const ASTOpConstantBool*>(tagCondition.get()) )
+            {
+                isConstant = true;
+                constantValue = constOp->value;
+            }
+
+
+            if (!isConstant && tagCondition)
+            {
+                Ptr<ASTOpFixed> n = new ASTOpFixed;
+                n->op.type = OP_TYPE::BO_NOT;
+                n->SetChild( n->op.args.BoolNot.source, tagCondition );
+
+                if (!c)
+                {
+                    c = m_opPool.Add( n );
+                }
+                else
+                {
+                    Ptr<ASTOpFixed> o = new ASTOpFixed;
+                    o->op.type = OP_TYPE::BO_AND;
+                    o->SetChild( o->op.args.BoolBinary.a, n);
+                    o->SetChild( o->op.args.BoolBinary.b, c);
+                    c = m_opPool.Add( o );
+                }
+            }
+            else if (isConstant && constantValue==true)
+            {
+                // No expression here means always true which becomes always false
+                Ptr<ASTOpConstantBool> f = new ASTOpConstantBool( false );
+
+                // No need to evaluate anything else.
+                c = m_opPool.Add( f );
+                done = true;
+                break;
+            }
+        }
+    }
+
+    //PartialOptimise( c, m_pCompilerOptions->m_optimisationOptions );
+
+    return c;
+}
+
+
+//---------------------------------------------------------------------------------------------
+bool SecondPassGenerator::Generate(
+    ErrorLogPtr pErrorLog,
+    const Node::Private* root)
+{
+    MUTABLE_CPUPROFILER_SCOPE(SecondPassGenerate);
+
+    check(root);
+    m_pErrorLog = pErrorLog;
+
+    // Find the list of surfaces every tag depends on
+	m_surfacesPerTag.clear();
+    m_surfacesPerTag.resize(m_pFirstPass->m_tags.size());
+	m_tagsPerTag.clear();
+    m_tagsPerTag.resize(m_pFirstPass->m_tags.size());
+    for( size_t t=0; t<m_pFirstPass->m_tags.size(); ++t )
+    {
+        set<size_t> pendingSurfs;
+        for( size_t s: m_pFirstPass->m_tags[t].surfaces )
+        {
+            pendingSurfs.insert(s);
+        }
+
+		set<size_t> processedSurfs;
+        while( !pendingSurfs.empty() )
+        {
+            size_t cs = *pendingSurfs.begin();
+            pendingSurfs.erase( pendingSurfs.begin() );
+
+			if (std::find(processedSurfs.begin(), processedSurfs.end(), cs)
+				 != processedSurfs.end())
+			{
+				continue;
+			}
+
+			processedSurfs.insert(cs);
+
+            m_surfacesPerTag[t].insert(cs);
+
+            auto& csurf = m_pFirstPass->surfaces[cs];
+            for ( auto sct: csurf.positiveTags )
+            {
+                auto it = std::find_if( m_pFirstPass->m_tags.begin(), m_pFirstPass->m_tags.end(),
+                                        [&](const FirstPassGenerator::TAG& e) { return e.tag==sct; } );
+                if (it==m_pFirstPass->m_tags.end())
+                {
+                    // This could happen if a tag is in a variation but noone defines it.
+                    continue;
+                }
+
+                size_t ct = it - m_pFirstPass->m_tags.begin();
+
+				m_tagsPerTag[t].insert(ct);
+
+                for( size_t s: m_pFirstPass->m_tags[ct].surfaces )
+                {
+                    if (m_surfacesPerTag[t].find(s)!=m_surfacesPerTag[t].end())
+                    {
+                        pendingSurfs.insert(s);
+                    }
+                }
+            }
+            for ( auto sct: csurf.negativeTags )
+            {
+                auto it = std::find_if( m_pFirstPass->m_tags.begin(), m_pFirstPass->m_tags.end(),
+                                        [&](const FirstPassGenerator::TAG& e) { return e.tag==sct; } );
+                if (it==m_pFirstPass->m_tags.end())
+                {
+                    // This could happen if a tag is in a variation but noone defines it.
+                    continue;
+                }
+
+                size_t ct = it - m_pFirstPass->m_tags.begin();
+
+				m_tagsPerTag[t].insert(ct);
+
+                for( size_t s: m_pFirstPass->m_tags[ct].surfaces )
+                {
+                    if (m_surfacesPerTag[t].find(s)!=m_surfacesPerTag[t].end())
+                    {
+                        pendingSurfs.insert(s);
+                    }
+                }
+            }
+        }
+    }
+
+    // Create the conditions for every surface, modifier and individual tag.
+    m_tagConditionGenerationCache.clear();
+
+    for ( size_t s=0; s<m_pFirstPass->surfaces.size(); ++s )
+    {
+        set<size_t> empty;
+        Ptr<ASTOp> c = GenerateSurfaceCondition( s, empty, empty, empty, empty );
+        m_pFirstPass->surfaces[s].surfaceCondition = c;
+    }
+
+    for ( size_t m=0; m<m_pFirstPass->modifiers.size(); ++m )
+    {
+        Ptr<ASTOp> c = GenerateModifierCondition( m );
+        m_pFirstPass->modifiers[m].surfaceCondition = c;
+    }
+
+    for ( size_t s = 0; s < m_pFirstPass->m_tags.size(); ++s )
+    {
+        set<size_t> empty;
+        Ptr<ASTOp> c = GenerateTagCondition( s, empty, empty, empty, empty );
+        m_pFirstPass->m_tags[s].genericCondition = c;
+    }
+
+
+    m_pFirstPass = nullptr;
+
+    return true;
+}
+
+
diff --git a/Engine/Plugins/Experimental/Mutable/Source/MutableTools/Private/CodeGenerator_SecondPass.h b/Engine/Plugins/Experimental/Mutable/Source/MutableTools/Private/CodeGenerator_SecondPass.h
new file mode 100644
--- /dev/null
+++ b/Engine/Plugins/Experimental/Mutable/Source/MutableTools/Private/CodeGenerator_SecondPass.h
@@ -0,0 +1,129 @@
+// Copyright Epic Games, Inc. All Rights Reserved.
+
+#pragma once
+
+#include "Visitor.h"
+#include "CodeGenerator_FirstPass.h"
+
+#include "AST.h"
+#include "CompilerPrivate.h"
+#include "NodeLayoutPrivate.h"
+#include "ErrorLogPrivate.h"
+#include "NodeObjectNewPrivate.h"
+
+#include "Table.h"
+
+#include "NodeBool.h"
+
+#include "NodeComponent.h"
+#include "NodeComponentNew.h"
+#include "NodeComponentEdit.h"
+
+#include "NodeModifier.h"
+#include "NodeModifierMeshClipMorphPlane.h"
+#include "NodeModifierMeshClipWithMesh.h"
+
+#include "NodeSurface.h"
+#include "NodeSurfaceNew.h"
+#include "NodeSurfaceEdit.h"
+
+#include "NodePatchImage.h"
+#include "NodeLOD.h"
+#include "NodeObjectNew.h"
+#include "NodeObjectState.h"
+#include "NodeObjectGroup.h"
+#include "NodePatchMesh.h"
+
+#include "Operations.h"
+#include "ModelPrivate.h"
+#include "ImagePrivate.h"
+
+
+namespace mu
+{
+
+	//---------------------------------------------------------------------------------------------
+	//! Second pass of the code generation process.
+    //! Solves surface and modifier conditions from tags and variations
+	//---------------------------------------------------------------------------------------------
+    class SecondPassGenerator : public Base
+	{
+	public:
+
+		SecondPassGenerator( FirstPassGenerator* firstPass,
+            const CompilerOptions::Private* options  );
+
+		// Return true on success.
+        bool Generate( ErrorLogPtr pErrorLog,
+                       const Node::Private* root );
+
+	private:
+
+        FirstPassGenerator* m_pFirstPass = nullptr;
+        const CompilerOptions::Private *m_pCompilerOptions = nullptr;
+
+
+        struct CONDITION_CONTEXT
+        {
+            Ptr<ASTOp> surfaceCondition;
+            FirstPassGenerator::StateCondition stateCondition;
+        };
+        vector< CONDITION_CONTEXT > m_currentCondition;
+
+
+        //!
+        ErrorLogPtr m_pErrorLog;
+
+        //!
+        struct CONDITION_GENERATION_KEY
+        {
+            size_t tagOrSurfIndex = 0;
+            set<size_t> posSurf;
+            set<size_t> negSurf;
+            set<size_t> posTag;
+            set<size_t> negTag;
+
+            inline bool operator<(const CONDITION_GENERATION_KEY& o) const
+            {
+                if (tagOrSurfIndex<o.tagOrSurfIndex) return true;
+                if (tagOrSurfIndex>o.tagOrSurfIndex) return false;
+                if (posSurf<o.posSurf) return true;
+                if (posSurf>o.posSurf) return false;
+                if (negSurf<o.negSurf) return true;
+                if (negSurf>o.negSurf) return false;
+                if (posTag<o.posTag) return true;
+                if (posTag>o.posTag) return false;
+                if (negTag<o.negTag) return true;
+                if (negTag>o.negTag) return false;
+                return false;
+            }
+        };
+
+        // List of surfaces that activate or deactivate every tag, or another surface that activates a tag in this set.
+		vector< set<size_t> > m_surfacesPerTag;
+		vector< set<size_t> > m_tagsPerTag;
+
+        std::map<CONDITION_GENERATION_KEY,Ptr<ASTOp>> m_tagConditionGenerationCache;
+
+        UniqueOpPool m_opPool;
+
+        //!
+        Ptr<ASTOp> GenerateTagCondition( size_t tagIndex,
+                                         const set<size_t>& posSurf,
+                                         const set<size_t>& negSurf,
+                                         const set<size_t>& posTag,
+                                         const set<size_t>& negTag );
+
+        //!
+        Ptr<ASTOp> GenerateSurfaceCondition( size_t surfIndex,
+                                             const set<size_t>& posSurf,
+                                             const set<size_t>& negSurf,
+                                             const set<size_t>& posTag,
+                                             const set<size_t>& negTag );
+
+        //!
+        Ptr<ASTOp> GenerateModifierCondition( size_t modIndex );
+    };
+
+}
+
