[mutable] Move Mutable plugin from the restricted folder to the experimental plugins folder.

diff --git a/Engine/Plugins/Experimental/Mutable/Source/MutableTools/Private/CodeGenerator_String.cpp b/Engine/Plugins/Experimental/Mutable/Source/MutableTools/Private/CodeGenerator_String.cpp
new file mode 100644
--- /dev/null
+++ b/Engine/Plugins/Experimental/Mutable/Source/MutableTools/Private/CodeGenerator_String.cpp
@@ -0,0 +1,114 @@
+// Copyright Epic Games, Inc. All Rights Reserved.
+
+#include "CodeGenerator.h"
+
+#include "Platform.h"
+#include "CompilerPrivate.h"
+
+#include "MutableTools/Private/NodeStringPrivate.h"
+#include "MutableTools/Private/NodeStringConstantPrivate.h"
+#include "MutableTools/Private/NodeStringParameterPrivate.h"
+
+#include "MutableTools/Private/ASTOpConstantString.h"
+
+#include "ModelPrivate.h"
+#include "Operations.h"
+
+
+using namespace mu;
+
+
+//-------------------------------------------------------------------------------------------------
+void CodeGenerator::GenerateString( STRING_GENERATION_RESULT& result, const NodeStringPtrConst& untyped )
+{
+    if (!untyped)
+    {
+        result = STRING_GENERATION_RESULT();
+        return;
+    }
+
+    // See if it was already generated
+	VISITED_MAP_KEY key = GetCurrentCacheKey(untyped);
+	GeneratedStringsMap::ValueType* it = m_generatedStrings.Find( key );
+    if ( it )
+    {
+        result = *it;
+        return;
+    }
+
+    // Generate for each different type of node
+    if ( auto Constant = dynamic_cast<const NodeStringConstant*>(untyped.get()) )
+    {
+        GenerateString_Constant( result, Constant );
+    }
+    else if ( auto Param = dynamic_cast<const NodeStringParameter*>(untyped.get()) )
+    {
+        GenerateString_Parameter( result, Param );
+    }
+	else
+    {
+        check(false);
+    }
+
+    // Cache the result
+    m_generatedStrings.Add( key, result );
+}
+
+
+//-------------------------------------------------------------------------------------------------
+void CodeGenerator::GenerateString_Constant(STRING_GENERATION_RESULT& result, const Ptr<const NodeStringConstant>& Typed )
+{
+    const NodeStringConstant::Private& node = *Typed->GetPrivate();
+
+	Ptr<ASTOpConstantString> op = new ASTOpConstantString();
+	op->value = node.m_value;
+
+    result.op = op;
+}
+
+
+//-------------------------------------------------------------------------------------------------
+void CodeGenerator::GenerateString_Parameter( STRING_GENERATION_RESULT& result, const Ptr<const NodeStringParameter>& Typed)
+{
+    const NodeStringParameter::Private& node = *Typed->GetPrivate();
+
+	Ptr<ASTOpParameter> op;
+
+	auto it = m_nodeVariables.find(node.m_pNode);
+	if (it == m_nodeVariables.end())
+	{
+		PARAMETER_DESC param;
+		param.m_name = node.m_name;
+		param.m_uid = node.m_uid;
+		param.m_type = PARAMETER_TYPE::T_FLOAT;
+		param.m_detailedType = node.m_detailedType;
+		// param.m_defaultValue.m_text = node.m_defaultValue;
+		int len = std::min(MUTABLE_MAX_STRING_PARAM_LENGTH, int(node.m_defaultValue.size()));
+		FMemory::Memcpy(param.m_defaultValue.m_text, node.m_defaultValue.c_str(), len);
+		param.m_defaultValue.m_text[len] = 0;
+
+
+		op = new ASTOpParameter();
+		op->type = OP_TYPE::ST_PARAMETER;
+		op->parameter = param;
+
+		// Generate the code for the ranges
+		for (std::size_t a = 0; a < node.m_ranges.size(); ++a)
+		{
+			RANGE_GENERATION_RESULT rangeResult;
+			GenerateRange(rangeResult, node.m_ranges[a]);
+			op->ranges.emplace_back(
+				op.get(), rangeResult.sizeOp, rangeResult.rangeName, rangeResult.rangeUID);
+		}
+
+		m_nodeVariables[node.m_pNode] = op;
+	}
+	else
+	{
+		op = it->second;
+	}
+
+    result.op = op;
+}
+
+
diff --git a/Engine/Plugins/Experimental/Mutable/Source/MutableTools/Private/CodeOptimiser.cpp b/Engine/Plugins/Experimental/Mutable/Source/MutableTools/Private/CodeOptimiser.cpp
new file mode 100644
--- /dev/null
+++ b/Engine/Plugins/Experimental/Mutable/Source/MutableTools/Private/CodeOptimiser.cpp
@@ -0,0 +1,1361 @@
+// Copyright Epic Games, Inc. All Rights Reserved.
+
+#include "CodeOptimiser.h"
+#include "ErrorLogPrivate.h"
+#include "AST.h"
+#include "StreamsPrivate.h"
+
+#include "ModelPrivate.h"
+#include "SystemPrivate.h"
+#include "Operations.h"
+#include "OpMeshMerge.h"
+
+#include "MutableTools/Private/ASTOpInstanceAdd.h"
+#include "MutableTools/Private/ASTOpConditional.h"
+#include "MutableTools/Private/ASTOpSwitch.h"
+#include "MutableTools/Private/ASTOpConstantResource.h"
+#include "MutableTools/Private/ASTOpConstantBool.h"
+#include "MutableTools/Private/ASTOpMeshApplyShape.h"
+#include "MutableTools/Private/ASTOpMeshBindShape.h"
+#include "MutableTools/Private/ASTOpMeshClipMorphPlane.h"
+#include "MutableTools/Private/ASTOpMeshApplyPose.h"
+
+#include <unordered_set>
+#include <shared_mutex>
+
+using namespace mu;
+
+//-------------------------------------------------------------------------------------------------
+//-------------------------------------------------------------------------------------------------
+//-------------------------------------------------------------------------------------------------
+namespace
+{
+
+    struct MESH_ENTRY
+    {
+        MeshPtrConst mesh;
+		mu::Ptr<ASTOpConstantResource> op;
+
+        bool operator==(const MESH_ENTRY& o) const
+        {
+            return mesh==o.mesh || *mesh==*o.mesh;
+        }
+    };
+
+    struct custom_mesh_hash
+    {
+		uint64 operator()( const MeshPtrConst& k ) const
+        {
+			uint64 h =  std::hash<uint64>()(k->GetVertexCount());
+            hash_combine(h, k->GetFaceCount());
+            return h;
+        }
+    };
+
+    struct custom_mesh_equal
+    {
+        bool operator()( const MeshPtrConst& a, const MeshPtrConst& b ) const
+        {
+            return a==b || *a==*b;
+        }
+    };
+
+    struct custom_image_hash
+    {
+		uint64 operator()( const ImagePtrConst& k ) const
+        {
+			uint64 h =  std::hash<uint64>()(k->GetSizeX());
+            hash_combine(h, k->GetSizeY());
+            hash_combine(h, k->GetLODCount());
+            return h;
+        }
+    };
+
+    struct custom_image_equal
+    {
+        bool operator()( const ImagePtrConst& a, const ImagePtrConst& b ) const
+        {
+            return a==b || *a==*b;
+        }
+    };
+
+    struct custom_layout_hash
+    {
+		uint64 operator()( const LayoutPtrConst& k ) const
+        {
+			uint64 h =  std::hash<uint64>()(k->GetBlockCount());
+            FIntPoint s = k->GetGridSize();
+            hash_combine(h, s[0]);
+            hash_combine(h, s[1]);
+            return h;
+        }
+    };
+
+    struct custom_layout_equal
+    {
+        bool operator()( const LayoutPtrConst& a, const LayoutPtrConst& b ) const
+        {
+            return a==b || *a==*b;
+        }
+    };
+}
+
+
+//-------------------------------------------------------------------------------------------------
+void mu::DuplicatedDataRemoverAST( ASTOpList& roots )
+{
+    MUTABLE_CPUPROFILER_SCOPE(DuplicatedDataRemoverAST);
+
+    std::vector<Ptr<ASTOpConstantResource>> allMeshOps;
+    std::vector<Ptr<ASTOpConstantResource>> allImageOps;
+    std::vector<Ptr<ASTOpConstantResource>> allLayoutOps;
+
+    // Gather constants
+    {
+        //MUTABLE_CPUPROFILER_SCOPE(Gather);
+
+        ASTOp::Traverse_TopRandom_Unique_NonReentrant( roots, [&](Ptr<ASTOp> n)
+        {
+            switch ( n->GetOpType() )
+            {
+
+            case OP_TYPE::ME_CONSTANT:
+            {
+                auto typedNode = dynamic_cast<ASTOpConstantResource*>(n.get());
+                check(typedNode);
+
+                if (typedNode)
+                {
+                    allMeshOps.push_back(typedNode);
+                }
+                break;
+            }
+
+            case OP_TYPE::IM_CONSTANT:
+            {
+                auto typedNode = dynamic_cast<ASTOpConstantResource*>(n.get());
+                check(typedNode);
+
+                if (typedNode)
+                {
+                    allImageOps.push_back(typedNode);
+                }
+                break;
+            }
+
+            case OP_TYPE::LA_CONSTANT:
+            {
+                auto typedNode = dynamic_cast<ASTOpConstantResource*>(n.get());
+                check(typedNode);
+
+                if (typedNode)
+                {
+                    allLayoutOps.push_back(typedNode);
+                }
+                break;
+            }
+
+            //    These should be part of the duplicated code removal, in AST.
+            //            // Names
+            //            case OP_TYPE::IN_ADDMESH:
+            //            case OP_TYPE::IN_ADDIMAGE:
+            //            case OP_TYPE::IN_ADDVECTOR:
+            //            case OP_TYPE::IN_ADDSCALAR:
+            //            case OP_TYPE::IN_ADDCOMPONENT:
+            //            case OP_TYPE::IN_ADDSURFACE:
+            //            {
+            //                OP::ADDRESS value = program.m_code[at].args.InstanceAdd.name;
+            //                program.m_code[at].args.InstanceAdd.name = m_oldToNewStrings[value];
+            //                break;
+            //            }
+
+
+            default:
+                break;
+            }
+
+            return true;
+        });
+    }
+
+
+    // Compare meshes
+    {
+        //MUTABLE_CPUPROFILER_SCOPE(CompareMeshes);
+
+        std::unordered_multimap< size_t, MESH_ENTRY > meshes;
+
+        for (auto& typedNode: allMeshOps)
+        {
+            size_t key = typedNode->GetValueHash();
+
+            Ptr<ASTOp> found;
+
+            auto r = meshes.equal_range(key);
+            if (r.first!=r.second)
+            {
+                MeshPtrConst mesh = static_cast<const Mesh*>( typedNode->GetValue().get() );
+
+                for ( auto it=r.first; it!=r.second; ++it )
+                {
+                    if (!it->second.mesh)
+                    {
+                        it->second.mesh = static_cast<const Mesh*>( it->second.op->GetValue().get() );
+                    }
+
+                    if ( custom_mesh_equal()( mesh, it->second.mesh ) )
+                    {
+                        found = it->second.op;
+                        break;
+                    }
+                }
+            }
+
+            if (found)
+            {
+                ASTOp::Replace(typedNode,found);
+            }
+            else
+            {
+                MESH_ENTRY e;
+                e.op = typedNode;
+                meshes.insert( std::make_pair<>(key, e) );
+            }
+        }
+    }
+
+    // Compare images
+    {
+        //MUTABLE_CPUPROFILER_SCOPE(CompareImages);
+
+        // TODO Optimise like the mesh compare above
+        std::unordered_map< ImagePtrConst, Ptr<ASTOp>, custom_image_hash, custom_image_equal > images;
+
+        for (auto& typedNode: allImageOps)
+        {
+            ImagePtrConst image = static_cast<const Image*>( typedNode->GetValue().get() );
+
+            auto it = images.find(image);
+            if (it!=images.end())
+            {
+                ASTOp::Replace(typedNode,it->second);
+            }
+            else
+            {
+                images.insert( std::make_pair<>(image, typedNode) );
+            }
+        }
+    }
+
+    // Compare layouts
+    {
+        //MUTABLE_CPUPROFILER_SCOPE(CompareLayouts);
+
+        // TODO Optimise like the mesh compare above
+        std::unordered_map< LayoutPtrConst, Ptr<ASTOp>, custom_layout_hash, custom_layout_equal > constants;
+
+        for (auto& typedNode: allLayoutOps)
+        {
+            LayoutPtrConst r = static_cast<const Layout*>( typedNode->GetValue().get() );
+
+            auto it = constants.find(r);
+            if (it!=constants.end())
+            {
+                ASTOp::Replace(typedNode,it->second);
+            }
+            else
+            {
+                constants.insert( std::make_pair<>(r, typedNode) );
+            }
+        }
+    }
+
+}
+
+
+//-------------------------------------------------------------------------------------------------
+//-------------------------------------------------------------------------------------------------
+//-------------------------------------------------------------------------------------------------
+namespace
+{
+    struct op_pointer_hash
+    {
+        inline size_t operator() (const mu::Ptr<ASTOp>& n) const
+        {
+            return n->Hash();
+        }
+    };
+
+    struct op_pointer_equal
+    {
+        inline bool operator() (const mu::Ptr<ASTOp>& a, const mu::Ptr<ASTOp>& b) const
+        {
+            return *a==*b;
+        }
+    };
+}
+
+void mu::DuplicatedCodeRemoverAST( ASTOpList& roots )
+{
+    MUTABLE_CPUPROFILER_SCOPE(DuplicatedCodeRemoverAST);
+
+    // Visited nodes, per type
+    std::unordered_set<Ptr<ASTOp>,op_pointer_hash,op_pointer_equal> visited[int(OP_TYPE::COUNT)];
+
+    ASTOp::Traverse_BottomUp_Unique_NonReentrant( roots, [&](Ptr<ASTOp>& n)
+    {
+        auto& container = visited[(int)n->GetOpType()];
+
+        // debug
+//        for( const auto& o : container )
+//        {
+//            if (*n==*o)
+//            {
+//                auto thisHash = n->Hash();
+//                auto otherHash = o->Hash();
+//                check(thisHash==otherHash);
+//            }
+//        }
+
+        // Insert will tell us if it was already there
+        auto it = container.insert(n);
+        if( !it.second )
+        {
+            // It wasn't inserted, so it was already there
+            ASTOp::Replace(n,*it.first);
+        }
+    });
+}
+
+
+//---------------------------------------------------------------------------------------------
+//---------------------------------------------------------------------------------------------
+//---------------------------------------------------------------------------------------------
+bool FindOpTypeVisitor::Apply( PROGRAM& program, OP::ADDRESS rootAt )
+{
+    bool found = false;
+
+    // If the program added new operations, we haven't visited them.
+    m_visited.resize( program.m_opAddress.size(), 0 );
+
+    m_pending.clear();
+    m_pending.reserve( program.m_opAddress.size()/4 );
+    m_pending.push_back( std::make_pair(false,rootAt) );
+
+    // Don't early out to be able to complete parent op cached flags
+    while ( !m_pending.empty() )
+    {
+        std::pair<bool,int> item = m_pending.back();
+        m_pending.pop_back();
+        OP::ADDRESS at = item.second;
+
+        // Not cached?
+        if (m_visited[at]<=1)
+        {
+            if (item.first)
+            {
+                // Item indicating we finished with all the children of a parent
+                check(m_visited[at]==1);
+
+                bool subtreeFound = false;
+                ForEachReference( program, at, [&](OP::ADDRESS ref)
+                {
+                    subtreeFound = subtreeFound || m_visited[ref]==3;
+                });
+
+                m_visited[at] = subtreeFound ? 3 : 2;
+            }
+            else if (!found)
+            {
+                if ( std::find( m_typesToFind.begin(), m_typesToFind.end(), program.GetOpType(at))
+                     !=
+                     m_typesToFind.end() )
+                {
+                    m_visited[at] = 3; // visited, ops found
+                    found = true;
+                }
+                else
+                {
+                    check(m_visited[at]==0);
+
+                    m_visited[at] = 1;
+                    m_pending.push_back( std::make_pair(true,at) );
+
+                    ForEachReference( program, at, [&](OP::ADDRESS ref)
+                    {
+                        if (ref && m_visited[ref]==0)
+                        {
+                            m_pending.push_back( std::make_pair(false,ref) );
+                        }
+                    });
+                }
+            }
+            else
+            {
+                // We won't process it.
+                m_visited[at] = 0;
+            }
+        }
+    }
+
+    return m_visited[rootAt]==3;
+}
+
+
+//---------------------------------------------------------------------------------------------
+//---------------------------------------------------------------------------------------------
+//---------------------------------------------------------------------------------------------
+IsConstantVisitor::IsConstantVisitor()
+{
+    vector<OP_TYPE> parameterTypes;
+    parameterTypes.push_back( OP_TYPE::BO_PARAMETER );
+    parameterTypes.push_back( OP_TYPE::NU_PARAMETER );
+    parameterTypes.push_back( OP_TYPE::SC_PARAMETER );
+    parameterTypes.push_back( OP_TYPE::CO_PARAMETER );
+    parameterTypes.push_back( OP_TYPE::PR_PARAMETER );
+    parameterTypes.push_back( OP_TYPE::IM_PARAMETER );
+    m_findOpTypeVisitor= std::make_unique<FindOpTypeVisitor>(parameterTypes);
+}
+
+
+//---------------------------------------------------------------------------------------------
+bool IsConstantVisitor::Apply( PROGRAM& program, OP::ADDRESS at )
+{
+    return !m_findOpTypeVisitor->Apply(program,at);
+}
+
+
+//---------------------------------------------------------------------------------------------
+//---------------------------------------------------------------------------------------------
+//---------------------------------------------------------------------------------------------
+class ConstantTask : public TaskManager::Task
+{
+public:
+    // input
+	mu::Ptr<ASTOp> m_source;
+    bool m_useDiskCache = false;
+    int m_imageCompressionQuality = 0;
+
+    std::shared_timed_mutex* m_codeAccessMutex = nullptr;
+
+private:
+
+    //
+	mu::Ptr<ASTOp> m_result;
+
+public:
+
+    // mu::Task interface
+    void Run() override
+    {
+        MUTABLE_CPUPROFILER_SCOPE(ConstantTask_Run);
+
+        // We need the clone because linking modifies ASTOp state
+        m_codeAccessMutex->lock_shared();
+		mu::Ptr<ASTOp> cloned = ASTOp::DeepClone( m_source );
+        m_codeAccessMutex->unlock_shared();
+
+        OP_TYPE type = cloned->GetOpType();
+        DATATYPE dtype = GetOpDataType(type);
+
+        SettingsPtr pSettings = new Settings;
+        pSettings->SetProfile( false );
+        pSettings->SetImageCompressionQuality( m_imageCompressionQuality );
+        SystemPtr pSystem = new System( pSettings );
+
+		// Don't generate mips suring linking here.
+		FLinkerOptions LinkerOptions;
+		LinkerOptions.MinTextureResidentMipCount = 255;
+
+        ModelPtrConst model = new Model;
+        ASTOp::FullLink( cloned, model->GetPrivate()->m_program, &LinkerOptions);
+        OP::ADDRESS at = cloned->linkedAddress;
+
+        PROGRAM::STATE state;
+        state.m_root = at;
+        model->GetPrivate()->m_program.m_states.push_back(state);
+
+        ParametersPtr localParams = model->NewParameters();
+        pSystem->GetPrivate()->BeginBuild( model );
+
+        // Calculate the value and replace this op by a constant
+        switch( dtype )
+        {
+        case DT_MESH:
+        {
+            MUTABLE_CPUPROFILER_SCOPE(ConstantMesh);
+
+			mu::Ptr<const Mesh> pMesh = pSystem->GetPrivate()->BuildMesh( model, localParams.get(), at );
+
+            if (pMesh)
+            {
+				mu::Ptr<ASTOpConstantResource> constantOp = new ASTOpConstantResource();
+                constantOp->type = OP_TYPE::ME_CONSTANT;
+                constantOp->SetValue( pMesh, m_useDiskCache );
+                m_result = constantOp;
+              }
+            break;
+        }
+
+        case DT_IMAGE:
+        {
+            MUTABLE_CPUPROFILER_SCOPE(ConstantImage);
+
+			mu::Ptr<const Image> pImage = pSystem->GetPrivate()->BuildImage( model, localParams.get(), at, 0 );
+
+            if (pImage)
+            {
+				mu::Ptr<ASTOpConstantResource> constantOp = new ASTOpConstantResource();
+                constantOp->type = OP_TYPE::IM_CONSTANT;
+                constantOp->SetValue( pImage, m_useDiskCache );
+                m_result = constantOp;
+            }
+            break;
+        }
+
+        case DT_LAYOUT:
+        {
+            MUTABLE_CPUPROFILER_SCOPE(ConstantLayout);
+
+			mu::Ptr<const Layout> pLayout = pSystem->GetPrivate()->BuildLayout( model, localParams.get(), at );
+
+            if (pLayout)
+            {
+				mu::Ptr<ASTOpConstantResource> constantOp = new ASTOpConstantResource();
+                constantOp->type = OP_TYPE::LA_CONSTANT;
+                constantOp->SetValue( pLayout, m_useDiskCache );
+                m_result = constantOp;
+            }
+            break;
+        }
+
+        case DT_BOOL:
+        {
+            MUTABLE_CPUPROFILER_SCOPE(ConstantBool);
+
+            bool value = pSystem->GetPrivate()->BuildBool( model, localParams.get(), at );
+
+            {
+				mu::Ptr<ASTOpConstantBool> constantOp = new ASTOpConstantBool();
+                constantOp->value = value;
+                m_result = constantOp;
+            }
+            break;
+        }
+
+        case DT_COLOUR:
+        {
+            MUTABLE_CPUPROFILER_SCOPE(ConstantBool);
+
+            float r=0.0f, g=0.0f, b=0.0f;
+            pSystem->GetPrivate()->BuildColour( model, localParams.get(), at, &r,&g,&b );
+
+            {
+				mu::Ptr<ASTOpFixed> constantOp = new ASTOpFixed();
+                constantOp->op.type = OP_TYPE::CO_CONSTANT;
+                constantOp->op.args.ColourConstant.value[0] = r;
+                constantOp->op.args.ColourConstant.value[1] = g;
+                constantOp->op.args.ColourConstant.value[2] = b;
+                constantOp->op.args.ColourConstant.value[3] = 1.0f;
+                m_result = constantOp;
+            }
+            break;
+        }
+
+        case DT_INT:
+        case DT_SCALAR:
+        case DT_STRING:
+        case DT_PROJECTOR:
+            // TODO
+            break;
+
+        default:
+            break;
+        }
+
+		pSystem->GetPrivate()->EndBuild();
+
+    }
+
+    // TaskManager::Task interface
+    void Complete() override
+    {
+        // This runs in the managing thread
+        m_codeAccessMutex->lock();
+        ASTOp::Replace( m_source, m_result );
+        m_codeAccessMutex->unlock();
+    }
+};
+
+
+//---------------------------------------------------------------------------------------------
+bool mu::ConstantGeneratorAST( const CompilerOptions::Private* options, Ptr<ASTOp>& root,
+                               TaskManager* pTaskManager )
+{
+    MUTABLE_CPUPROFILER_SCOPE(ConstantGenerator);
+
+    bool modified = false;
+
+    // don't do this if constant optimization has been disabled, usually for debugging.
+    if (!options->m_optimisationOptions.m_constReduction)
+    {
+        return false;
+    }
+
+
+    // Calculate constant-subtree and special-op flags
+    {
+        MUTABLE_CPUPROFILER_SCOPE(ConstantGenerator_CalculateFlags);
+
+        ASTOp::Traverse_BottomUp_Unique( root, [&](Ptr<ASTOp>& n)
+        {
+            bool constantSubtree = true;
+            switch (n->GetOpType())
+            {
+            case OP_TYPE::BO_PARAMETER:
+            case OP_TYPE::NU_PARAMETER:
+            case OP_TYPE::SC_PARAMETER:
+            case OP_TYPE::CO_PARAMETER:
+            case OP_TYPE::PR_PARAMETER:
+            case OP_TYPE::IM_PARAMETER:
+                constantSubtree = false;
+                break;
+            default:
+                // Propagate from children
+                n->ForEachChild( [&constantSubtree](ASTChild& c)
+                {
+                    if (c)
+                    {
+                        constantSubtree = constantSubtree && c->m_constantSubtree;
+                    }
+                });
+                break;
+            }
+            n->m_constantSubtree = constantSubtree;
+
+            // We avoid generating constants for these operations, to avoid the memory
+            // explosion.
+            // TODO: Make compiler options for some of them
+            // TODO: Some of them are worth if the code below them is unique.
+            bool hasSpecialOpInSubtree = false;
+            switch (n->GetOpType())
+            {
+            case OP_TYPE::IM_BLANKLAYOUT:
+            case OP_TYPE::IM_COMPOSE:
+            //case OP_TYPE::IM_RASTERMESH:            // TODO review this one
+            case OP_TYPE::ME_MERGE:
+            case OP_TYPE::ME_CLIPWITHMESH:
+            case OP_TYPE::ME_CLIPMORPHPLANE:
+            case OP_TYPE::ME_APPLYPOSE:
+            case OP_TYPE::ME_REMOVEMASK:
+                hasSpecialOpInSubtree = true;
+                break;
+
+            default:
+                // Propagate from children
+                n->ForEachChild( [&](ASTChild& c)
+                {
+                    if (c)
+                    {
+                        hasSpecialOpInSubtree = hasSpecialOpInSubtree || c->m_hasSpecialOpInSubtree;
+                    }
+                });
+                break;
+            }
+            n->m_hasSpecialOpInSubtree = hasSpecialOpInSubtree;
+        });
+
+    }
+
+    std::shared_timed_mutex codeAccessMutex;
+
+    vector< Ptr<ASTOp> > roots;
+    roots.push_back(root);
+
+    // Generate constant operations
+    ASTOp::Traverse_TopDown_Unique_Imprecise( roots, [&](Ptr<ASTOp>& n)
+    {
+        bool recurse = true;
+
+        OP_TYPE type = n->GetOpType();
+        DATATYPE dtype = GetOpDataType(type);
+
+        bool constant = false;
+        if (dtype!=DT_INSTANCE)
+        {
+            constant = n->m_constantSubtree;
+        }
+
+        // See if it is a case of instructions we want to avoid, but with special parameters that make
+        // them ok
+        bool specialCase = false;
+        if (constant)
+        {
+            //const OP& op = program.m_code[at];
+            // TODO
+//                if (n->GetOpType()==OP_TYPE::IM_COMPOSE)
+//                {
+//                    // Get the layout
+//                    // In case we fill all of the target image anyway
+//                    ParametersPtr pParams = m_pModel->NewParameters();
+
+//                    const auto& args = op.args.ImageCompose;
+//                    LayoutPtrConst pLayout = m_pSystem->GetPrivate()->BuildLayout
+//                        ( m_pModel.get(), pParams.get(), args.layout );
+
+//                    if ( pLayout &&
+//                         pLayout->GetPrivate()->IsSingleBlockAndFull() )
+//                    {
+//                        specialCase = true;
+//                    }
+//                }
+
+//                else
+                if (n->GetOpType()==OP_TYPE::IM_RASTERMESH)
+            {
+                const auto& args = dynamic_cast<const ASTOpFixed*>(n.get())->op.args.ImageRasterMesh;
+                if ( !args.image )
+                {
+                    specialCase = true;
+                }
+            }
+        }
+
+        // See if the subtree contains operations that we don't want to collapse even if
+        // they are constant
+        if ( constant && !specialCase)
+        {
+            if ( n->m_hasSpecialOpInSubtree )
+            {
+                constant = false;
+            }
+        }
+
+        // If children are constant
+        if ( constant
+             && type!=OP_TYPE::BO_CONSTANT
+             && type!=OP_TYPE::NU_CONSTANT
+             && type!=OP_TYPE::SC_CONSTANT
+             && type!=OP_TYPE::CO_CONSTANT
+             && type!=OP_TYPE::IM_CONSTANT
+             && type!=OP_TYPE::ME_CONSTANT
+             && type!=OP_TYPE::LA_CONSTANT
+             && type!=OP_TYPE::PR_CONSTANT
+             && type!=OP_TYPE::NONE
+            )
+        {
+            switch( dtype )
+            {
+            case DT_MESH:
+            case DT_IMAGE:
+            case DT_LAYOUT:
+            case DT_BOOL:
+            case DT_COLOUR:
+            {
+                recurse = false;
+                modified = true;
+
+                ConstantTask* constantTask = new ConstantTask;
+                constantTask->m_useDiskCache =
+                    options->m_optimisationOptions.m_useDiskCache;
+                constantTask->m_source = n;
+                constantTask->m_codeAccessMutex = &codeAccessMutex;
+                constantTask->m_imageCompressionQuality =
+                    options->m_imageCompressionQuality;
+
+                if ( pTaskManager )
+                {
+                    pTaskManager->AddTask( constantTask );
+                }
+                else
+                {
+                    constantTask->Run();
+                    constantTask->Complete();
+                    delete constantTask;
+                }
+
+                break;
+            }
+
+            case DT_INT:
+            case DT_SCALAR:
+            case DT_STRING:
+            case DT_PROJECTOR:
+                // TODO
+                break;
+
+            case DT_INSTANCE:
+                // nothing to do
+                break;
+
+            default:
+                break;
+            }
+
+        }
+
+        return recurse;
+    });
+
+    if (pTaskManager)
+    {
+        pTaskManager->CompleteTasks();
+    }
+
+    return modified;
+}
+
+
+//-------------------------------------------------------------------------------------------------
+//-------------------------------------------------------------------------------------------------
+//-------------------------------------------------------------------------------------------------
+CodeOptimiser::CodeOptimiser( CompilerOptionsPtr options, vector<STATE_COMPILATION_DATA>& states )
+    : m_states( states )
+{
+    m_options = options;
+    m_pModelReport = new ModelReport;
+}
+
+
+//-------------------------------------------------------------------------------------------------
+void CodeOptimiser::FullOptimiseAST( ASTOpList& roots,
+                                     TaskManager* pTaskManager )
+{
+    bool modified = true;
+    int numIterations = 0;
+    while (modified && (!m_optimizeIterationsMax || (m_optimizeIterationsLeft>0) || !numIterations) )
+    {
+        --m_optimizeIterationsLeft;
+        ++numIterations;
+        UE_LOG(LogMutableCore, Verbose, TEXT("Main optimise iteration %d, max %d, left %d"),
+                numIterations, m_optimizeIterationsMax, m_optimizeIterationsLeft);
+
+        modified = false;
+
+        // All kind of optimisations that depend on the meaning of each operation
+        // \TODO: We are doing it for all states.
+        UE_LOG(LogMutableCore, Verbose, TEXT(" - semantic optimiser"));
+        modified |= SemanticOptimiserAST( roots, m_options->GetPrivate()->m_optimisationOptions );
+        //AXE_INT_VALUE("Mutable", Verbose, "ast size", (int64_t)ASTOp::CountNodes(roots));
+		UE_LOG(LogMutableCore, Verbose, TEXT("(int) %s : %ld"), TEXT("ast size"), int64(ASTOp::CountNodes(roots)));
+		ASTOp::LogHistogram(roots);
+
+        UE_LOG(LogMutableCore, Verbose, TEXT(" - sink optimiser"));
+        modified |= SinkOptimiserAST( roots, m_options->GetPrivate()->m_optimisationOptions );
+        //AXE_INT_VALUE("Mutable", Verbose, "ast size", (int64_t)ASTOp::CountNodes(roots));
+        ASTOp::LogHistogram(roots);
+
+        // Image size operations are treated separately
+        UE_LOG(LogMutableCore, Verbose, TEXT(" - size optimiser"));
+        modified |= SizeOptimiserAST( roots );
+        //AXE_INT_VALUE("Mutable", Verbose, "ast size", (int64_t)ASTOp::CountNodes(roots));
+
+        // We cannot run the logic optimiser here because the constant memory explodes.
+        // TODO: Conservative logic optimiser
+        if (!modified)
+        {
+//                AXE_INT_VALUE("Mutable", Verbose, "program size", (int64_t)program.m_opAddress.size());
+//                UE_LOG(LogMutableCore, Verbose, " - logic optimiser");
+//                LocalLogicOptimiser log;
+//                modified |= log.Apply( program, (int)s );
+        }
+    }
+
+//    ASTOp::LogHistogram(roots);
+
+    UE_LOG(LogMutableCore, Verbose, TEXT(" - duplicated code remover"));
+    DuplicatedCodeRemoverAST( roots );
+	UE_LOG(LogMutableCore, Verbose, TEXT("(int) %s : %ld"), TEXT("ast size"), int64(ASTOp::CountNodes(roots)));
+
+    ASTOp::LogHistogram(roots);
+
+    UE_LOG(LogMutableCore, Verbose, TEXT(" - duplicated data remover"));
+    DuplicatedDataRemoverAST( roots );
+	UE_LOG(LogMutableCore, Verbose, TEXT("(int) %s : %ld"), TEXT("ast size"), int64(ASTOp::CountNodes(roots)));
+
+    ASTOp::LogHistogram(roots);
+
+    // Generate constants
+    for ( auto& r: roots )
+    {
+		UE_LOG(LogMutableCore, Verbose, TEXT("(int) %s : %ld"), TEXT("ast size"), int64(ASTOp::CountNodes(roots)));
+		UE_LOG(LogMutableCore, Verbose, TEXT(" - constant generator"));
+
+        // Constant subtree generation
+        modified = ConstantGeneratorAST( m_options->GetPrivate(), r, pTaskManager );
+    }
+
+	UE_LOG(LogMutableCore, Verbose, TEXT("(int) %s : %ld"), TEXT("ast size"), int64(ASTOp::CountNodes(roots)));
+
+//    ASTOp::LogHistogram(roots);
+
+    // \todo: this code deduplication here, takes very long and doesn't do much... investigate!
+//    UE_LOG(LogMutableCore, Verbose, " - duplicated code remover");
+//    DuplicatedCodeRemoverAST( roots );
+//    AXE_INT_VALUE("Mutable", Verbose, "ast size", (int64_t)ASTOp::CountNodes(roots));
+
+    ASTOp::LogHistogram(roots);
+
+    UE_LOG(LogMutableCore, Verbose, TEXT(" - duplicated data remover"));
+    DuplicatedDataRemoverAST( roots );
+	UE_LOG(LogMutableCore, Verbose, TEXT("(int) %s : %ld"), TEXT("ast size"), int64(ASTOp::CountNodes(roots)));
+
+    //if (!modified)
+    {
+        //for ( size_t s=0;  s<program.m_states.size(); ++s )
+        {
+            ASTOp::LogHistogram(roots);
+            UE_LOG(LogMutableCore, Verbose, TEXT(" - logic optimiser"));
+            modified |= LocalLogicOptimiserAST( roots );
+            ASTOp::LogHistogram(roots);
+        }
+
+        //AXE_INT_VALUE("Mutable", Verbose, "ast size", (int64_t)ASTOp::CountNodes(roots));
+    }
+
+    ASTOp::LogHistogram(roots);
+}
+
+
+//---------------------------------------------------------------------------------------------
+// The state represents if there is a parent operation requiring skeleton for current mesh subtree.
+class CollectAllMeshesForSkeletonVisitorAST : public Visitor_TopDown_Unique_Const<uint8_t>
+{
+public:
+
+    CollectAllMeshesForSkeletonVisitorAST( const ASTOpList& roots  )
+    {
+        Traverse( roots, false );
+    }
+
+    // List of meshes that require a skeleton
+    vector<mu::Ptr<ASTOpConstantResource>> m_meshesRequiringSkeleton;
+
+private:
+
+    // Visitor_TopDown_Unique_Const<uint8_t> interface
+    bool Visit( const mu::Ptr<ASTOp>& node ) override
+    {
+        // \todo: refine to avoid instruction branches with irrelevant skeletons.
+
+        uint8_t currentProtected = GetCurrentState();
+
+        switch (node->GetOpType())
+        {
+
+        case OP_TYPE::ME_CONSTANT:
+        {
+			mu::Ptr<ASTOpConstantResource> typedOp = dynamic_cast<ASTOpConstantResource*>(node.get());
+            check(typedOp);
+
+            if (currentProtected)
+            {
+                if ( std::find(m_meshesRequiringSkeleton.begin(), m_meshesRequiringSkeleton.end(), typedOp)
+                     ==
+                     m_meshesRequiringSkeleton.end() )
+                {
+                    m_meshesRequiringSkeleton.push_back(typedOp);
+                }
+            }
+
+            return false;
+        }
+
+        case OP_TYPE::ME_CLIPMORPHPLANE:
+        {
+            auto typedOp = dynamic_cast<ASTOpMeshClipMorphPlane*>(node.get());
+            if (typedOp)
+            {
+                if (typedOp->vertexSelectionType
+                        ==
+                        OP::MeshClipMorphPlaneArgs::VS_BONE_HIERARCHY)
+                {
+                    // We need the skeleton for the source mesh
+                    RecurseWithState( typedOp->source.child(), true );
+                    return false;
+                }
+            }
+
+            return true;
+        }
+
+		case OP_TYPE::ME_APPLYPOSE:
+		{
+			auto typedOp = dynamic_cast<ASTOpMeshApplyPose*>(node.get());
+			check(typedOp);
+
+			if (typedOp)
+			{
+				// We need the skeleton for both meshes
+				RecurseWithState(typedOp->base.child(), true);
+				RecurseWithState(typedOp->pose.child(), true);
+				return false;
+			}
+
+			break;
+		}
+
+		case OP_TYPE::ME_BINDSHAPE:
+		{
+			auto typedOp = dynamic_cast<ASTOpMeshBindShape*>(node.get());
+			check(typedOp);
+
+			if (typedOp && typedOp->m_reshapeSkeleton)
+			{
+				RecurseWithState(typedOp->Mesh.child(), true);
+				return false;
+			}
+
+			break;
+		}
+
+		case OP_TYPE::ME_APPLYSHAPE:
+		{
+			auto typedOp = dynamic_cast<ASTOpMeshApplyShape*>(node.get());
+			check(typedOp);
+
+			if (typedOp && typedOp->m_reshapeSkeleton)
+			{
+				RecurseWithState(typedOp->Mesh.child(), true);
+				return false;
+			}
+
+			break;
+		}
+
+        default:
+            break;
+        }
+
+        return true;
+    }
+
+};
+
+
+//---------------------------------------------------------------------------------------------
+// This stores an ADD_MESH op with the child meshes collected and the final skeleton to use
+// for this op.
+struct ADDMESH_SKELETON
+{
+	mu::Ptr<ASTOp> m_pAddMeshOp;
+    vector<mu::Ptr<ASTOpConstantResource>> m_contributingMeshes;
+	mu::Ptr<Skeleton> m_pFinalSkeleton;
+
+    ADDMESH_SKELETON( const mu::Ptr<ASTOp>& pAddMeshOp,
+                      vector<mu::Ptr<ASTOpConstantResource>>& contributingMeshes,
+                      const mu::Ptr<Skeleton>& pFinalSkeleton )
+    {
+        m_pAddMeshOp = pAddMeshOp;
+        m_contributingMeshes = std::move(contributingMeshes);
+        m_pFinalSkeleton = pFinalSkeleton;
+    }
+};
+
+
+//---------------------------------------------------------------------------------------------
+void SkeletonCleanerAST( vector<mu::Ptr<ASTOp>>& roots, const MODEL_OPTIMIZATION_OPTIONS& options )
+{
+    // This collects all the meshes that require a skeleton because they are used in operations
+    // that require it.
+    CollectAllMeshesForSkeletonVisitorAST requireSkeletonCollector( roots );
+
+    vector<ADDMESH_SKELETON> replacementsFound;
+
+    ASTOp::Traverse_TopDown_Unique_Imprecise( roots, [&](mu::Ptr<ASTOp>& at )
+    {
+        // Only recurse instance construction ops.
+        bool processChildren = GetOpDataType(at->GetOpType())==DT_INSTANCE;
+
+        if ( at->GetOpType() == OP_TYPE::IN_ADDMESH )
+        {
+            auto typedNode = dynamic_cast<ASTOpInstanceAdd*>(at.get());
+			mu::Ptr<ASTOp> meshRoot = typedNode->value.child();
+
+            if (meshRoot)
+            {
+                // Gather constant meshes contributing to the final mesh
+                vector<mu::Ptr<ASTOpConstantResource>> subtreeMeshes;
+                vector<mu::Ptr<ASTOp>> tempRoots;
+                tempRoots.push_back(meshRoot);
+                ASTOp::Traverse_TopDown_Unique_Imprecise( tempRoots, [&](mu::Ptr<ASTOp>& lat )
+                {
+                    // \todo: refine to avoid instruction branches with irrelevant skeletons.
+                    if ( lat->GetOpType() == OP_TYPE::ME_CONSTANT )
+                    {
+						mu::Ptr<ASTOpConstantResource> typedOp = dynamic_cast<ASTOpConstantResource*>(lat.get());
+                        check(typedOp);
+
+                        if ( std::find(subtreeMeshes.begin(), subtreeMeshes.end(), typedOp)
+                             ==
+                             subtreeMeshes.end() )
+                        {
+                            subtreeMeshes.push_back(typedOp);
+                        }
+                    }
+                    return true;
+                });
+
+                // Create a mesh just with the unified skeleton
+                SkeletonPtr pFinalSkeleton = new Skeleton;
+                for (const auto& meshAt: subtreeMeshes)
+                {
+                    MeshPtrConst pMesh = static_cast<const Mesh*>(meshAt->GetValue().get());
+					mu::Ptr<const Skeleton> sourceSkeleton = pMesh ? pMesh->GetSkeleton() : nullptr;
+                    if (sourceSkeleton)
+                    {
+                        ExtendSkeleton(pFinalSkeleton.get(),sourceSkeleton.get());
+                    }
+                }
+
+                replacementsFound.emplace_back( at, subtreeMeshes, pFinalSkeleton );
+            }
+        }
+
+        return processChildren;
+    });
+
+
+    // Iterate all meshes again
+    ASTOp::Traverse_TopDown_Unique_Imprecise( roots, [&](mu::Ptr<ASTOp>& at )
+    {
+        if (at->GetOpType()==OP_TYPE::ME_CONSTANT)
+        {
+            auto typedOp = dynamic_cast<ASTOpConstantResource*>(at.get());
+            check(typedOp);
+            if (!typedOp)
+            {
+                return true;
+            }
+
+            for( auto& rep: replacementsFound )
+            {
+                if ( std::find( rep.m_contributingMeshes.begin(),
+                                rep.m_contributingMeshes.end(),
+                                at)
+                     !=
+                     rep.m_contributingMeshes.end() )
+                {
+					mu::Ptr<const Mesh> pMesh =  static_cast<const Mesh*>(typedOp->GetValue().get());
+                    pMesh->CheckIntegrity();
+
+					mu::Ptr<Mesh> pNewMesh = MeshRemapSkeleton( pMesh.get(),
+                                                            rep.m_pFinalSkeleton.get());
+
+                    // It returns null if there is no need to remap.
+                    if (pNewMesh)
+                    {
+                        pNewMesh->CheckIntegrity();
+						mu::Ptr<ASTOpConstantResource> newOp = new ASTOpConstantResource();
+                        newOp->type = OP_TYPE::ME_CONSTANT;
+                        newOp->SetValue( pNewMesh, options.m_useDiskCache );
+
+                        ASTOp::Replace( at, newOp );
+                    }
+                }
+            }
+        }
+        return true;
+    });
+
+//    for( auto& rep: replacementsFound )
+//    {
+//        check(rep.m_pAddMeshOp->GetOpType()==OP_TYPE::IN_ADDMESH);
+//        auto newAddMeshRootTyped = dynamic_cast<ASTOpInstanceAdd*>(rep.m_pAddMeshOp.get());
+//        auto meshRoot = newAddMeshRootTyped->value.child();
+
+//        // Add new skeleton constant instruction
+//        Ptr<ASTOpConstantResource> cop = new ASTOpConstantResource();
+//        cop->type = OP_TYPE::ME_CONSTANT;
+//        cop->SetValue( rep.m_pFinalSkeleton, options.m_useDiskCache );
+
+//        // Add new "apply skeleton" instruction
+//        Ptr<ASTOpFixed> skop = new ASTOpFixed();
+//        skop->op.type = OP_TYPE::ME_SETSKELETON;
+//        skop->SetChild( skop->op.args.MeshSetSkeleton.source, meshRoot);
+//        skop->SetChild( skop->op.args.MeshSetSkeleton.skeleton, cop);
+
+//        // We need a new mesh add instruction, it is not safe to edit the current one.
+//        auto newOp = mu::Clone<ASTOpInstanceAdd>(newAddMeshRootTyped);
+//        newOp->value = skop;
+
+//        // TODO BLEH UNSAFE?
+//        ASTOp::Replace(rep.m_pAddMeshOp,newOp);
+//    }
+}
+
+
+//---------------------------------------------------------------------------------------------
+void CodeOptimiser::OptimiseAST( TaskManager* pTaskManager )
+{
+    MUTABLE_CPUPROFILER_SCOPE(OptimiseAST);
+
+    // Gather all the roots
+    vector<Ptr<ASTOp>> roots;
+    for(const auto& s:m_states)
+    {
+        roots.push_back(s.root);
+    }
+
+	UE_LOG(LogMutableCore, Verbose, TEXT("(int) %s : %ld"), TEXT("ast size"), int64(ASTOp::CountNodes(roots)));
+
+    if ( m_options->GetPrivate()->m_optimisationOptions.m_enabled )
+    {
+        // We use 4 times the count because at the time we moved to sharing this count it
+        // was being used 4 times, and we want to keep the tests consistent.
+        m_optimizeIterationsLeft =
+            m_options->GetPrivate()->m_optimisationOptions.m_maxOptimisationLoopCount * 4;
+        m_optimizeIterationsMax = m_optimizeIterationsLeft;
+
+        // The first duplicated data remover has the special mission of removing
+        // duplicated data (meshes) that may have been specified in the source
+        // data, before we make it diverge because of different uses, like layout
+        // creation
+        UE_LOG(LogMutableCore, Verbose, TEXT(" - duplicated data remover"));
+        DuplicatedDataRemoverAST( roots );
+		UE_LOG(LogMutableCore, Verbose, TEXT("(int) %s : %ld"), TEXT("ast size"), int64(ASTOp::CountNodes(roots)));
+
+        ASTOp::LogHistogram(roots);
+
+        UE_LOG(LogMutableCore, Verbose, TEXT(" - duplicated code remover"));
+        DuplicatedCodeRemoverAST( roots );
+		UE_LOG(LogMutableCore, Verbose, TEXT("(int) %s : %ld"), TEXT("ast size"), int64(ASTOp::CountNodes(roots)));
+
+        // Special optimization stages
+        if ( m_options->GetPrivate()->m_optimisationOptions.m_uniformizeSkeleton )
+        {
+            UE_LOG(LogMutableCore, Verbose, TEXT(" - skeleton cleaner"));
+            ASTOp::LogHistogram(roots);
+
+            SkeletonCleanerAST( roots, m_options->GetPrivate()->m_optimisationOptions );
+			UE_LOG(LogMutableCore, Verbose, TEXT("(int) %s : %ld"), TEXT("ast size"), int64(ASTOp::CountNodes(roots)));
+			ASTOp::LogHistogram(roots);
+        }
+
+        // First optimisation stage. It tries to resolve all the image sizes. This is necessary
+        // because some operations cannot be applied correctly until the image size is known
+        // like the grow-map generation.
+        bool modified = true;
+        int numIterations = 0;
+        while (modified && (!m_optimizeIterationsMax || (m_optimizeIterationsLeft>0) || !numIterations) )
+        {
+            MUTABLE_CPUPROFILER_SCOPE(FirstStage);
+
+            --m_optimizeIterationsLeft;
+            ++numIterations;
+            UE_LOG(LogMutableCore, Verbose, TEXT("First optimise iteration %d, max %d, left %d"),
+                    numIterations, m_optimizeIterationsMax, m_optimizeIterationsLeft);
+
+            modified = false;
+
+            UE_LOG(LogMutableCore, Verbose, TEXT(" - size optimiser"));
+            modified |= SizeOptimiserAST( roots );
+        }
+
+		UE_LOG(LogMutableCore, Verbose, TEXT("(int) %s : %ld"), TEXT("ast size"), int64(ASTOp::CountNodes(roots)));
+
+        // Main optimisation stage
+        {
+            MUTABLE_CPUPROFILER_SCOPE(MainStage);
+            FullOptimiseAST( roots, pTaskManager );
+			UE_LOG(LogMutableCore, Verbose, TEXT("(int) %s : %ld"), TEXT("ast size"), int64(ASTOp::CountNodes(roots)));
+		}
+
+//            // Third optimisation stage: introduce image compose conditions
+//            {
+//                // TODO: We have disable this since the new operation "layout remove block" is used.
+//                // Proper support needs to be implemented for this case.
+//    //                MUTABLE_CPUPROFILER_SCOPE(ImageComposeConditionGenerator);
+//    //                AXE_INT_VALUE("Mutable", Verbose, "program size", (int64_t)program.m_opAddress.size());
+//    //                UE_LOG(LogMutableCore, Verbose, " - image compose conditions generator");
+//    //                ImageComposeConditionGenerator composeGen( program );
+//            }
+
+
+        {
+            MUTABLE_CPUPROFILER_SCOPE(ConditionalsStage);
+
+            //program.Check();
+
+            // Optimise the conditionals added by the image compose condition generator
+            // also the constants generator at the end of previous fulloptimiseAST.
+            FullOptimiseAST( roots, pTaskManager );
+        }
+
+        // Analyse mesh constants to see which of them are in optimised mesh formats, and set the flags.
+        ASTOp::Traverse_BottomUp_Unique_NonReentrant( roots, [&](Ptr<ASTOp>& n)
+        {
+            if (n->GetOpType()==OP_TYPE::ME_CONSTANT)
+            {
+                auto typed = dynamic_cast<ASTOpConstantResource*>(n.get());
+                auto pMesh = static_cast<const Mesh*>(typed->GetValue().get());
+                pMesh->ResetStaticFormatFlags();
+                typed->SetValue( pMesh,
+                                 m_options->GetPrivate()->m_optimisationOptions.m_useDiskCache );
+            }
+        });
+
+        // Make sure we didn't lose track of pointers
+        for ( size_t s=0;  s<m_states.size(); ++s )
+        {
+            check( std::find(roots.begin(), roots.end(), m_states[s].root )
+                            != roots.end() );
+        }
+
+        ASTOp::LogHistogram(roots);
+
+        {
+            MUTABLE_CPUPROFILER_SCOPE(StatesStage);
+
+            // Optimise for every state
+            OptimiseStatesAST( pTaskManager );
+
+            // Optimise the data formats (TODO)
+            //OperationFlagGenerator flagGen( pResult.get() );
+        }
+
+        ASTOp::LogHistogram(roots);
+    }
+
+//        // Minimal optimisation of constant subtrees
+    else if ( m_options->GetPrivate()->m_optimisationOptions.m_constReduction )
+    {
+        // The first duplicated data remover has the special mission of removing
+        // duplicated data (meshes) that may have been specified in the source
+        // data, before we make it diverge because of different uses, like layout
+        // creation
+        UE_LOG(LogMutableCore, Verbose, TEXT(" - duplicated data remover"));
+        DuplicatedDataRemoverAST( roots );
+        //AXE_INT_VALUE("Mutable", Verbose, "ast size", (int64_t)ASTOp::CountNodes(roots));
+
+        UE_LOG(LogMutableCore, Verbose, TEXT(" - duplicated code remover"));
+        DuplicatedCodeRemoverAST( roots );
+		UE_LOG(LogMutableCore, Verbose, TEXT("(int) %s : %ld"), TEXT("ast size"), int64(ASTOp::CountNodes(roots)));
+
+        for ( size_t s=0;  s<m_states.size(); ++s )
+        {
+            UE_LOG(LogMutableCore, Verbose, TEXT(" - constant generator"));
+            ConstantGeneratorAST( m_options->GetPrivate(), m_states[s].root, pTaskManager );
+            //AXE_INT_VALUE("Mutable", Verbose, "ast size", (int64_t)ASTOp::CountNodes(roots));
+        }
+
+        UE_LOG(LogMutableCore, Verbose, TEXT(" - duplicated data remover"));
+        DuplicatedDataRemoverAST( roots );
+        //AXE_INT_VALUE("Mutable", Verbose, "ast size", (int64_t)ASTOp::CountNodes(roots));
+
+        UE_LOG(LogMutableCore, Verbose, TEXT(" - duplicated code remover"));
+        DuplicatedCodeRemoverAST( roots );
+        //AXE_INT_VALUE("Mutable", Verbose, "ast size", (int64_t)ASTOp::CountNodes(roots));
+
+        // Make sure we didn't lose track of pointers
+        for ( size_t s=0;  s<m_states.size(); ++s )
+        {
+            check( std::find(roots.begin(), roots.end(), m_states[s].root )
+                            != roots.end() );
+        }
+    }
+
+    ASTOp::LogHistogram(roots);
+
+}
+
diff --git a/Engine/Plugins/Experimental/Mutable/Source/MutableTools/Private/CodeOptimiser.h b/Engine/Plugins/Experimental/Mutable/Source/MutableTools/Private/CodeOptimiser.h
new file mode 100644
--- /dev/null
+++ b/Engine/Plugins/Experimental/Mutable/Source/MutableTools/Private/CodeOptimiser.h
@@ -0,0 +1,325 @@
+// Copyright Epic Games, Inc. All Rights Reserved.
+
+#pragma once
+
+#include "Operations.h"
+#include "ModelPrivate.h"
+#include "ImagePrivate.h"
+#include "CodeVisitor.h"
+#include "CompilerPrivate.h"
+#include "AST.h"
+#include "TaskManager.h"
+
+#include <memory>
+
+namespace mu
+{
+
+    //---------------------------------------------------------------------------------------------
+    //! Code optimiser
+    //---------------------------------------------------------------------------------------------
+    class CodeOptimiser : public Base
+    {
+    public:
+        CodeOptimiser( CompilerOptionsPtr options, vector<STATE_COMPILATION_DATA>& states );
+
+        //! Optimise the virtual machine code, using several transforms.
+        void OptimiseAST( TaskManager* );
+
+        //! Report about states external gpu usage
+        ModelReportPtr m_pModelReport;
+
+    private:
+
+        CompilerOptionsPtr m_options;
+
+        //!
+        vector<STATE_COMPILATION_DATA>& m_states;
+
+        //! The max number of optimize iterations is shared across several stages now.
+        //! This is how many are left
+        int m_optimizeIterationsMax=0;
+        int m_optimizeIterationsLeft=0;
+
+        //! Full optimisation pass
+        void FullOptimiseAST( ASTOpList& roots,
+                              TaskManager* pTaskManager );
+
+        //! Optimise the code of a model for a specific state, generating new instructions and
+        //! state information.
+        void OptimiseStatesAST( TaskManager* pTaskManager );
+
+    };
+
+
+    //---------------------------------------------------------------------------------------------
+    //! Scan a subtree and return true if any op of a type in the provided subset is found.
+    //! Intermediate data is used between calls to apply, so don't remove program code or directly
+    //! change the instructions. Adding new instructions is ok.
+    //---------------------------------------------------------------------------------------------
+    class FindOpTypeVisitor
+    {
+    public:
+
+        FindOpTypeVisitor(const vector<OP_TYPE>& types)
+            : m_typesToFind(types)
+        {
+        }
+
+        bool Apply( PROGRAM& program, OP::ADDRESS at );
+
+
+    private:
+
+        vector<OP_TYPE> m_typesToFind;
+
+        vector< std::pair<bool,OP::ADDRESS> > m_pending;
+
+        // 0 not visited
+        // 1 children pending
+        // 2 visited and not found
+        // 3 visited and found
+        vector<uint8_t> m_visited;
+    };
+
+
+    //---------------------------------------------------------------------------------------------
+    //! Scan a subtree and return true if the expression is constant.
+    //! It has the same restrictions that FindOpTypeVisitor
+    //---------------------------------------------------------------------------------------------
+    class IsConstantVisitor
+    {
+    public:
+
+        IsConstantVisitor();
+
+        bool Apply( PROGRAM& program, OP::ADDRESS at );
+
+    private:
+
+        std::unique_ptr<FindOpTypeVisitor> m_findOpTypeVisitor;
+    };
+
+
+    //---------------------------------------------------------------------------------------------
+    //! ConstantGenerator replaces constant subtrees of operations with an equivalent single
+	//! constant value operation. 
+    //---------------------------------------------------------------------------------------------
+    extern bool ConstantGeneratorAST( const CompilerOptions::Private* options, Ptr<ASTOp>& root,
+                                      TaskManager* pTaskManager );
+
+    //---------------------------------------------------------------------------------------------
+    //! \TODO: shapes, projectors, others? but not switches (they must be unique)
+    //---------------------------------------------------------------------------------------------
+    extern void DuplicatedDataRemoverAST( ASTOpList& roots );
+
+
+    //---------------------------------------------------------------------------------------------
+    //! Mark all the duplicated code instructions to point at the same operation, leaving the
+    //! copies unreachable.
+    //---------------------------------------------------------------------------------------------
+    extern void DuplicatedCodeRemoverAST( ASTOpList& roots );
+
+
+    //---------------------------------------------------------------------------------------------
+    //! All kinds of optimisations that depend on the meaning of each operation
+    //---------------------------------------------------------------------------------------------
+    extern bool SemanticOptimiserAST( ASTOpList& roots,
+                                      const MODEL_OPTIMIZATION_OPTIONS& optimisationOptions );
+
+    //---------------------------------------------------------------------------------------------
+    //! Semantic operator that reorders instructions moving expensive ones down to the
+    //! leaves of the expressions trying to turn them into constants.
+    //---------------------------------------------------------------------------------------------
+    extern bool SinkOptimiserAST( ASTOpList& roots, const MODEL_OPTIMIZATION_OPTIONS& );
+
+    //---------------------------------------------------------------------------------------------
+    //!
+    //---------------------------------------------------------------------------------------------
+    extern bool SizeOptimiserAST( ASTOpList& roots );
+
+    //---------------------------------------------------------------------------------------------
+    //!
+    //---------------------------------------------------------------------------------------------
+    extern bool LocalLogicOptimiserAST(ASTOpList& roots);
+
+    //---------------------------------------------------------------------------------------------
+    //! Discard all LODs beyond the given lod count.
+    //---------------------------------------------------------------------------------------------
+    class LODCountReducerAST : public Visitor_TopDown_Unique_Cloning
+    {
+    public:
+
+        LODCountReducerAST( Ptr<ASTOp>& root, int lodCount );
+
+    protected:
+
+        Ptr<ASTOp> Visit( Ptr<ASTOp> at, bool& processChildren ) override;
+
+        int m_lodCount;
+    };
+
+
+    //---------------------------------------------------------------------------------------------
+    //! Scan the code in the given subtree and return true if a state runtime parameters is
+    //! found.
+    //! Intermediate data is used betwen calls to apply, so don't remove program code or directly
+    //! change the instructions. Adding new instructions is ok.
+    //---------------------------------------------------------------------------------------------
+    class RuntimeParameterVisitorAST
+    {
+    public:
+
+        RuntimeParameterVisitorAST(const STATE_COMPILATION_DATA* pState);
+
+        bool HasAny( const Ptr<ASTOp>& root );
+
+    private:
+
+        const STATE_COMPILATION_DATA* m_pState;
+
+        //!
+        struct PENDING_ITEM
+        {
+            //! 0: indicate subtree pending
+            //! 1: indicate children finished
+            uint8_t itemType;
+
+            //! 0: everything is relevant
+            //! 1: only layouts are relevant
+            uint8_t onlyLayoutsRelevant;
+
+            //! Address to visit
+            Ptr<ASTOp> at;
+        };
+
+        //!
+        vector< PENDING_ITEM > m_pending;
+
+        //! Possible op state
+        enum class OP_STATE : uint8_t
+        {
+            NOT_VISITED = 0,
+            CHILDREN_PENDING_FULL,
+            CHILDREN_PENDING_PARTIAL,
+            VISITED_HASRUNTIME,
+            VISITED_FULL_DOESNTHAVERUNTIME,
+            VISITED_PARTIAL_DOESNTHAVERUNTIME
+        };
+
+        std::unordered_map<Ptr<ASTOp>,OP_STATE> m_visited;
+
+        //!
+        void AddIfNeeded( const PENDING_ITEM& item );
+
+    };
+
+
+    //---------------------------------------------------------------------------------------------
+    //! Remove all texture compression operations that would happen for runtime parameter changes.
+    //---------------------------------------------------------------------------------------------
+    class RuntimeTextureCompressionRemoverAST : public Visitor_TopDown_Unique_Cloning
+    {
+    public:
+
+        RuntimeTextureCompressionRemoverAST( STATE_COMPILATION_DATA* state );
+
+    protected:
+
+        Ptr<ASTOp> Visit( Ptr<ASTOp> at, bool& processChildren ) override;
+
+    private:
+
+        RuntimeParameterVisitorAST m_hasRuntimeParamVisitor;
+
+    };
+
+    //---------------------------------------------------------------------------------------------
+    //! Restructure the code to move operations involving runtime parameters as high as possible
+    //! in the code tree.
+    //---------------------------------------------------------------------------------------------
+    class ParameterOptimiserAST : public Visitor_TopDown_Unique_Cloning
+    {
+    private:
+
+        STATE_COMPILATION_DATA& m_stateProps;
+        // unused const GPU_PLATFORM_PROPS& m_gpuPlatformProps;
+
+    public:
+
+        ParameterOptimiserAST( STATE_COMPILATION_DATA &s,
+                               const MODEL_OPTIMIZATION_OPTIONS& optimisationOptions );
+
+        bool Apply();
+
+
+    private:
+
+        Ptr<ASTOp> Visit( Ptr<ASTOp> at, bool& processChildren ) override;
+
+        bool m_modified;
+
+        MODEL_OPTIMIZATION_OPTIONS m_optimisationOptions;
+
+        RuntimeParameterVisitorAST m_hasRuntimeParamVisitor;
+
+    };
+
+
+    //---------------------------------------------------------------------------------------------
+    //! Some masks are optional. If they are null, replace them by a white plain image of the right
+    //! size.
+    //---------------------------------------------------------------------------------------------
+    extern Ptr<ASTOp> EnsureValidMask( Ptr<ASTOp> mask, Ptr<ASTOp> base );
+
+    //---------------------------------------------------------------------------------------------
+    //! Return true if two non-zero pixels of the masks overlap.
+    //---------------------------------------------------------------------------------------------
+    extern bool AreMasksOverlapping( const PROGRAM& program, OP::ADDRESS a, OP::ADDRESS b );
+
+
+    //---------------------------------------------------------------------------------------------
+    //! Calculate all the parameters found relevant under a particular operation. This may not
+    //! incldue all the parameters in the subtree (if because of the operations they are not
+    //! relevant)
+    //! It has an internal cache, so don't reuse if the program changes.
+    //---------------------------------------------------------------------------------------------
+    class SubtreeRelevantParametersVisitorAST : public Base
+    {
+    public:
+
+        void Run( Ptr<ASTOp> root );
+
+        //! After Run, list of relevant parameters.
+        std::unordered_set< string > m_params;
+
+    private:
+
+        struct STATE
+        {
+            Ptr<ASTOp> op;
+            bool onlyLayoutIsRelevant=false;
+
+            STATE( Ptr<ASTOp> o=nullptr, bool l=false) : op(o), onlyLayoutIsRelevant(l) {}
+
+            bool operator==(const STATE& o) const
+            {
+                return  op == o.op &&
+                        onlyLayoutIsRelevant == o.onlyLayoutIsRelevant;
+            }
+        };
+
+        struct state_hash
+        {
+            std::size_t operator()(const STATE& k) const
+            {
+                return std::hash<const void*>()(k.op.get());
+            }
+        };
+
+        // Result cache
+        // \todo optimise by storing unique lists separately and an index here.
+        std::unordered_map< STATE, std::unordered_set< string >, state_hash > m_resultCache;
+    };
+
+}
