[mutable] Move Mutable plugin from the restricted folder to the experimental plugins folder.

diff --git a/Engine/Plugins/Experimental/Mutable/Source/MutableTools/Private/CodeOptimiser_GPU.cpp b/Engine/Plugins/Experimental/Mutable/Source/MutableTools/Private/CodeOptimiser_GPU.cpp
new file mode 100644
--- /dev/null
+++ b/Engine/Plugins/Experimental/Mutable/Source/MutableTools/Private/CodeOptimiser_GPU.cpp
@@ -0,0 +1,1111 @@
+// Copyright Epic Games, Inc. All Rights Reserved.
+
+#include "CodeOptimiser_GPU.h"
+#include "CodeOptimiser.h"
+#include "CodeGenerator.h"
+#include "ModelReportPrivate.h"
+
+#include "CodeVisitor.h"
+#include "ModelPrivate.h"
+#include "SystemPrivate.h"
+#include "Operations.h"
+
+#define MUTABLE_GPU_CODE_BUFFER_SIZE	512
+
+
+namespace mu
+{
+
+    //---------------------------------------------------------------------------------------------
+    //---------------------------------------------------------------------------------------------
+    //---------------------------------------------------------------------------------------------
+//    class GPUSubtreeGenerator
+//    {
+//    public:
+
+//        virtual ~GPUSubtreeGenerator() {}
+
+//        //! Generate a full program to run internally an expression subtree on the GPU.
+//        virtual GPU_PROGRAM GenerateProgram( OP::ADDRESS at,
+//                              PROGRAM& program,
+//                              int state,
+//                              const GPU_PLATFORM_PROPS& platformOptions,
+//                              const GPU_STATE_PROPS& stateOptions ) = 0;
+
+//        //! Generate a method to run an expression subtree on the GPU externally (integrated on an
+//        //! engine shader).
+//        virtual string GenerateImage( const string& imageName,
+//                                      OP::ADDRESS at,
+//                                      PROGRAM& program,
+//                                      int state,
+//                                      const GPU_PLATFORM_PROPS& platformOptions,
+//                                      const GPU_STATE_PROPS& stateOptions ) = 0;
+
+//        virtual void Clear()
+//        {
+//            m_floatParams.clear();
+//            m_vectorParams.clear();
+//            m_imageParams.clear();
+//            m_currentVar = "";
+//            m_expressionCode = "";
+//        }
+
+//        //! Additional result data
+//        vector<GPU_PROGRAM::PARAM_DATA> m_imageParams;
+//        vector<GPU_PROGRAM::PARAM_DATA> m_vectorParams;
+//        vector<GPU_PROGRAM::PARAM_DATA> m_floatParams;
+//        string m_currentImageName;
+//        string m_currentVar;
+//        string m_expressionCode;
+//    };
+
+
+    //---------------------------------------------------------------------------------------------
+    //---------------------------------------------------------------------------------------------
+    //---------------------------------------------------------------------------------------------
+//    class GPUSubtreeGenerator_GLSL : public GPUSubtreeGenerator, public ConstCodeVisitor
+//    {
+//    public:
+
+//        GPUSubtreeGenerator_GLSL()
+//        {
+//            m_currentLayoutBlock = -1;
+//        }
+
+
+//        //-----------------------------------------------------------------------------------------
+//        // GPUSubtreeGenerator interface
+//        //-----------------------------------------------------------------------------------------
+//        virtual void Clear() override
+//        {
+//            GPUSubtreeGenerator::Clear();
+//            m_generated.clear();
+//            m_usedVars.clear();
+//            m_currentTexCoords = MUTABLE_GPU_CODE_UVS;
+//        }
+
+
+//        //-----------------------------------------------------------------------------------------
+//        string GenerateImage( const string& imageName,
+//                              OP::ADDRESS at,
+//                              PROGRAM& program,
+//                              int state,
+//                              const GPU_PLATFORM_PROPS& platformOptions,
+//                              const GPU_STATE_PROPS& //stateOptions
+//                              ) override
+//        {
+//            check( platformOptions.type == GPU_GL4 );
+//            (void)platformOptions;
+//            m_state = state;
+//            m_currentImageName = imageName;
+
+//            // Clear state
+//            Clear();
+
+//            Traverse( at, program, false );
+
+//            string imageCode = "// begin mutable image["+imageName+"]\n";
+
+//            // Add every source image of the fragment program
+//            size_t sourceImageCount = m_imageParams.size();
+//            for ( size_t i=0; i<sourceImageCount; ++i )
+//            {
+//                char temp[128];
+//                mutable_snprintf( temp, 128, "uniform sampler2D %s;\n",
+//                                  m_imageParams[i].name.c_str() );
+//                imageCode += temp;
+//            }
+
+//            // Add every source vector of the fragment program
+//            size_t sourceVectorCount = m_vectorParams.size();
+//            for ( size_t i=0; i<sourceVectorCount; ++i )
+//            {
+//                char temp[128];
+//                mutable_snprintf( temp, 128, "uniform vec4 %s;\n",
+//                                  m_vectorParams[i].name.c_str() );
+//                imageCode += temp;
+//            }
+
+//            // Add every source float of the fragment program
+//            size_t sourceFloatCount = m_floatParams.size();
+//            for ( size_t i=0; i<sourceFloatCount; ++i )
+//            {
+//                char temp[128];
+//                mutable_snprintf( temp, 128, "uniform float %s;\n",
+//                                  m_floatParams[i].name.c_str() );
+//                imageCode += temp;
+//            }
+
+//            imageCode += "\nvec4 "+imageName+" ( vec2 "+string(MUTABLE_GPU_CODE_UVS)+" )\n";
+//            imageCode += "{\n";
+//            imageCode += m_expressionCode;
+//            imageCode += "\n\treturn "+m_currentVar+";\n";
+//            imageCode += "}\n";
+//            imageCode += "// end of mutable image\n\n";
+
+//            return imageCode;
+//        }
+
+
+//        //-----------------------------------------------------------------------------------------
+//        // ConstCodeVisitor interface
+//        //-----------------------------------------------------------------------------------------
+//        virtual void Visit( OP::ADDRESS at, PROGRAM& program ) override
+//        {
+//            (void)at;
+//            (void)program;
+////            const OP& op = program.m_code[at];
+
+////            if ( !m_generated[at].size() )
+////            {
+////                char temp[MUTABLE_GPU_CODE_BUFFER_SIZE];
+
+////                // See if we can translate the operation to the gpu
+////                bool gpuizable = GetOpDesc( (OP_TYPE)op.type ).gpuizable;
+
+////                // If the op is a state constant, don't translate it.
+////                gpuizable &= !program.m_states[m_state].IsUpdateCache(at);
+
+////                if ( gpuizable )
+////                {
+////                    // Generate more code
+////                    switch( op.type )
+////                    {
+
+////                    case OP_TYPE::IM_MULTIPLY:
+////                    {
+////                        string result = GetNewVar( "multiply" );
+
+////                        Visit( op.args.ImageLayer.base, program );
+////                        string baseVar = m_currentVar;
+
+////                        Visit( op.args.ImageLayer.blended, program );
+////                        string blendedVar = m_currentVar ;
+
+////                        string multiplyVar = GetNewVar("multiply_unmasked");
+////                        mutable_snprintf( temp, MUTABLE_GPU_CODE_BUFFER_SIZE,
+////                                  "\nvec4 %s = %s*%s;\n",
+////                                  multiplyVar.c_str(),
+////                                  baseVar.c_str(),
+////                                  blendedVar.c_str() );
+////                        m_expressionCode += temp;
+
+////                        if ( op.args.ImageLayer.mask )
+////                        {
+////                            Visit( op.args.ImageLayer.mask, program );
+////                            string maskVar = m_currentVar;
+
+////                            mutable_snprintf( temp, MUTABLE_GPU_CODE_BUFFER_SIZE,
+////                                      "vec4 %s = %s*(1.0-%s.rrrr) + %s*%s.rrrr;\n",
+////                                      result.c_str(),
+////                                      baseVar.c_str(),
+////                                      maskVar.c_str(),
+////                                      multiplyVar.c_str(),
+////                                      maskVar.c_str() );
+
+////                            m_expressionCode += temp;
+////                        }
+////                        else
+////                        {
+////                            mutable_snprintf( temp, MUTABLE_GPU_CODE_BUFFER_SIZE,
+////                                      "vec4 %s = %s;\n",
+////                                      result.c_str(),
+////                                      multiplyVar.c_str() );
+////                            m_expressionCode += temp;
+////                        }
+
+////                        m_currentVar = result;
+
+////                        break;
+////                    }
+
+////                    case OP_TYPE::IM_BLEND:
+////                    {
+////                        string result = GetNewVar( "blend" );
+
+////                        Visit( op.args.ImageLayer.base, program );
+////                        string baseVar = m_currentVar;
+
+////                        Visit( op.args.ImageLayer.blended, program );
+////                        string blendedVar = m_currentVar ;
+
+////                        if ( op.args.ImageLayer.mask )
+////                        {
+////                            Visit( op.args.ImageLayer.mask, program );
+////                            string maskVar = m_currentVar;
+
+////                            mutable_snprintf( temp, MUTABLE_GPU_CODE_BUFFER_SIZE,
+////                                      "vec4 %s = %s*(1.0-%s.rrrr) + %s*%s.rrrr;\n",
+////                                      result.c_str(),
+////                                      baseVar.c_str(),
+////                                      maskVar.c_str(),
+////                                      blendedVar.c_str(),
+////                                      maskVar.c_str() );
+
+////                            m_expressionCode += temp;
+////                        }
+////                        else
+////                        {
+////                            mutable_snprintf( temp, MUTABLE_GPU_CODE_BUFFER_SIZE,
+////                                      "vec4 %s = %s;\n",
+////                                      result.c_str(),
+////                                      blendedVar.c_str() );
+////                            m_expressionCode += temp;
+////                        }
+
+////                        m_currentVar = result;
+
+////                        break;
+////                    }
+
+////                    case OP_TYPE::IM_BLENDCOLOUR:
+////                    {
+////                        string result = GetNewVar( "blendcol" );
+
+////                        Visit( op.args.ImageLayerColour.base, program );
+////                        string baseVar = m_currentVar;
+
+////                        Visit( op.args.ImageLayerColour.colour, program );
+////                        string blendedVar = m_currentVar ;
+
+////                        if ( op.args.ImageLayerColour.mask )
+////                        {
+////                            Visit( op.args.ImageLayerColour.mask, program );
+////                            string maskVar = m_currentVar;
+
+////                            mutable_snprintf( temp, MUTABLE_GPU_CODE_BUFFER_SIZE,
+////                                      "vec4 %s = %s*(1.0-%s.rrrr) + %s*%s.rrrr;\n",
+////                                      result.c_str(),
+////                                      baseVar.c_str(),
+////                                      maskVar.c_str(),
+////                                      blendedVar.c_str(),
+////                                      maskVar.c_str() );
+
+////                            m_expressionCode += temp;
+////                        }
+////                        else
+////                        {
+////                            mutable_snprintf( temp, MUTABLE_GPU_CODE_BUFFER_SIZE,
+////                                      "vec4 %s = %s;\n",
+////                                      result.c_str(),
+////                                      blendedVar.c_str() );
+////                            m_expressionCode += temp;
+////                        }
+
+////                        m_currentVar = result;
+
+////                        break;
+////                    }
+
+////                    case OP_TYPE::IM_SOFTLIGHTCOLOUR:
+////                    {
+////                        string result = GetNewVar( "softlight" );
+
+////                        Visit( op.args.ImageLayerColour.base, program );
+////                        string baseVar = m_currentVar;
+
+////                        Visit( op.args.ImageLayerColour.colour, program );
+////                        string colourVar = m_currentVar;
+
+////                        string softVar = GetNewVar("softlight_unmasked");
+////                        mutable_snprintf( temp, MUTABLE_GPU_CODE_BUFFER_SIZE,
+////                                  "\nvec4 %s = (1.0-%s)*%s*%s + %s*(1.0-(1.0-%s)*(1.0-%s));\n",
+////                                  softVar.c_str(),
+////                                  baseVar.c_str(),
+////                                  baseVar.c_str(),
+////                                  colourVar.c_str(),
+////                                  baseVar.c_str(),
+////                                  baseVar.c_str(),
+////                                  colourVar.c_str() );
+////                        m_expressionCode += temp;
+
+////                        if ( op.args.ImageLayerColour.mask )
+////                        {
+////                            Visit( op.args.ImageLayerColour.mask, program );
+////                            string maskVar = m_currentVar;
+
+////                            mutable_snprintf( temp, MUTABLE_GPU_CODE_BUFFER_SIZE,
+////                                      "vec4 %s = %s*(1.0-%s.rrrr) + %s*%s.rrrr;\n",
+////                                      result.c_str(),
+////                                      baseVar.c_str(),
+////                                      maskVar.c_str(),
+////                                      softVar.c_str(),
+////                                      maskVar.c_str() );
+////                            m_expressionCode += temp;
+////                        }
+////                        else
+////                        {
+////                            mutable_snprintf( temp, MUTABLE_GPU_CODE_BUFFER_SIZE,
+////                                      "vec4 %s = %s;\n",
+////                                      result.c_str(),
+////                                      softVar.c_str() );
+////                            m_expressionCode += temp;
+////                        }
+
+////                        m_currentVar = result;
+////                        break;
+////                    }
+
+////                    case OP_TYPE::IM_INTERPOLATE3:
+////                    {
+////                        string result = GetNewVar( "interpolate3" );
+
+////                        Visit( op.args.ImageInterpolate3.factor1, program );
+////                        string factor1Var = m_currentVar;
+
+////                        Visit( op.args.ImageInterpolate3.factor2, program );
+////                        string factor2Var = m_currentVar ;
+
+////                        Visit( op.args.ImageInterpolate3.target0, program );
+////                        string target0Var = m_currentVar;
+
+////                        Visit( op.args.ImageInterpolate3.target1, program );
+////                        string target1Var = m_currentVar;
+
+////                        Visit( op.args.ImageInterpolate3.target2, program );
+////                        string target2Var = m_currentVar;
+
+////                        mutable_snprintf( temp, MUTABLE_GPU_CODE_BUFFER_SIZE,
+////                                          "vec4 %s = %s*(1.0-%s-%s) + %s*%s + %s*%s;\n",
+////                                          result.c_str(),
+////                                          target0Var.c_str(),
+////                                          factor1Var.c_str(),
+////                                          factor2Var.c_str(),
+////                                          target1Var.c_str(),
+////                                          factor1Var.c_str(),
+////                                          target2Var.c_str(),
+////                                          factor2Var.c_str() );
+
+////                        m_expressionCode += temp;
+
+////                        m_currentVar = result;
+
+////                        break;
+////                    }
+
+////                    case OP_TYPE::IM_COMPOSE:
+////                    {
+////                        string result = GetNewVar( "compose" );
+
+////                        // Base image
+////                        Visit( op.args.ImageCompose.base, program );
+////                        string baseVar = m_currentVar;
+
+////                        // Generate the block transform vector
+////                        int oldLayoutBlock = m_currentLayoutBlock;
+////                        m_currentLayoutBlock = op.args.ImageCompose.blockIndex;
+////                        Visit( op.args.ImageCompose.layout, program );
+////                        m_currentLayoutBlock = oldLayoutBlock;
+////                        string transformVar = m_currentVar;
+
+////                        // GLSL like this:
+////                        // vec2 blockCoords = (uvs.xy-transform.xy)/transform.zw;
+////                        // vec4 outsideTest = vec4( blockCoords.x<0, blockCoords.y<0, blockCoords.x>=1, blockCoords.y>=1 );
+////                        // if ( dot(outsideTest,outsideTest) )
+////                        // {
+////                        //     block code;
+////                        //     masking code;
+////                        // }
+////                        // else
+////                        // {
+////                        //     just set the base result;
+////                        // }
+////                        string blockCoordsVar = GetNewVar( "blockCoords" );
+////                        mutable_snprintf( temp, MUTABLE_GPU_CODE_BUFFER_SIZE,
+////                                  "vec2 %s = (%s.xy-%s.xy)/%s.zw;\n",
+////                                  blockCoordsVar.c_str(),
+////                                  m_currentTexCoords.c_str(),
+////                                  transformVar.c_str(),
+////                                  transformVar.c_str() );
+////                        m_expressionCode += temp;
+
+////                        string testVar = GetNewVar( "outsideTest" );
+////                        mutable_snprintf( temp, MUTABLE_GPU_CODE_BUFFER_SIZE,
+////                                          "vec4 %s = vec4( %s.x<0, %s.y<0, %s.x>=1, %s.y>=1 );\n",
+////                                          testVar.c_str(),
+////                                          blockCoordsVar.c_str(),
+////                                          blockCoordsVar.c_str(),
+////                                          blockCoordsVar.c_str(),
+////                                          blockCoordsVar.c_str()
+////                                          );
+////                        m_expressionCode += temp;
+
+////                        // Declare the result
+////                        mutable_snprintf( temp, MUTABLE_GPU_CODE_BUFFER_SIZE,
+////                                          "vec4 %s;\n",
+////                                          result.c_str()
+////                                          );
+////                        m_expressionCode += temp;
+
+////                        // "Inside" conditional
+////                        // TODO: variables declared inside the conditional will generate wrong
+////                        // code if reused outside by the internal map, as they are not accessible
+////                        // outside the "if" scope.
+////                        mutable_snprintf( temp, MUTABLE_GPU_CODE_BUFFER_SIZE,
+////                                          "if ( dot(%s,%s)==0 )\n{\n",
+////                                          testVar.c_str(),
+////                                          testVar.c_str()
+////                                          );
+////                        m_expressionCode += temp;
+
+////                        // Block image
+////                        string oldTexCoords = m_currentTexCoords;
+////                        m_currentTexCoords = blockCoordsVar;
+////                        Visit( op.args.ImageCompose.blockImage, program );
+////                        string blockVar = m_currentVar;
+
+////                        // Block mask
+////                        if ( op.args.ImageCompose.mask )
+////                        {
+////                            Visit( op.args.ImageCompose.mask, program );
+////                            string maskVar = m_currentVar;
+
+////                            mutable_snprintf( temp, MUTABLE_GPU_CODE_BUFFER_SIZE,
+////                                      "%s = %s*(1.0-%s.rrrr) + %s*%s.rrrr;\n",
+////                                      result.c_str(),
+////                                      baseVar.c_str(),
+////                                      maskVar.c_str(),
+////                                      blockVar.c_str(),
+////                                      maskVar.c_str() );
+////                            m_expressionCode += temp;
+////                        }
+////                        else
+////                        {
+////                            mutable_snprintf( temp, MUTABLE_GPU_CODE_BUFFER_SIZE,
+////                                      "%s = %s;\n",
+////                                      result.c_str(),
+////                                      blockVar.c_str() );
+////                            m_expressionCode += temp;
+////                        }
+
+////                        m_currentTexCoords = oldTexCoords;
+
+////                        // If outside
+////                        // TODO: if there is no block mask, skip the base when possible
+////                        m_expressionCode += "}\nelse\n{\n";
+////                        mutable_snprintf( temp, MUTABLE_GPU_CODE_BUFFER_SIZE,
+////                                  "%s = %s;\n",
+////                                  result.c_str(),
+////                                  baseVar.c_str() );
+////                        m_expressionCode += temp;
+////                        m_expressionCode += "}\n";
+
+////                        m_currentVar = result;
+////                        break;
+////                    }
+
+////                    default:
+////                        check( false );
+////                    }
+////                }
+
+////                else
+////                {
+////                    // We cannot generate gpu code below this, make a parameter for what is left
+////                    switch ( GetOpType( op.type ) )
+////                    {
+////                    case DT_IMAGE:
+////                    {
+////                        mutable_snprintf( temp, MUTABLE_GPU_CODE_BUFFER_SIZE, "%s_%s_%d",
+////                                          m_currentImageName.c_str(),
+////                                          MUTABLE_GPU_IMAGE_PARAM, m_imageParams.size() );
+////                        string varName = temp;
+
+////                        m_currentVar = GetNewVar( "imgParam" );
+////                        mutable_snprintf( temp, MUTABLE_GPU_CODE_BUFFER_SIZE,
+////                                          "vec4 %s = texture( %s, %s );\n",
+////                                          m_currentVar.c_str(),
+////                                          varName.c_str(),
+////                                          m_currentTexCoords.c_str() );
+////                        m_expressionCode += temp;
+
+////                        // Make sure the image is in a gpu-usable format
+////                        FImageDesc desc = GetImageDesc( program, at );
+////                        if ( desc.m_format==IF_L_UBYTE_RLE )
+////                        {
+////                            OP conv;
+////                            conv.type = OP_TYPE::IM_PIXELFORMAT;
+////                            conv.args.ImagePixelFormat.format = IF_L_UBYTE;
+////                            conv.args.ImagePixelFormat.source = at;
+////                            at = program.AddOp( conv );
+////                        }
+////                        else if ( desc.m_format==IF_RGB_UBYTE_RLE )
+////                        {
+////                            OP conv;
+////                            conv.type = OP_TYPE::IM_PIXELFORMAT;
+////                            conv.args.ImagePixelFormat.format = IF_RGB_UBYTE;
+////                            conv.args.ImagePixelFormat.source = at;
+////                            at = program.AddOp( conv );
+////                        }
+////                        else if ( desc.m_format==IF_RGBA_UBYTE_RLE )
+////                        {
+////                            OP conv;
+////                            conv.type = OP_TYPE::IM_PIXELFORMAT;
+////                            conv.args.ImagePixelFormat.format = IF_RGBA_UBYTE;
+////                            conv.args.ImagePixelFormat.source = at;
+////                            at = program.AddOp( conv );
+////                        }
+
+////                        m_imageParams.push_back( GPU_PROGRAM::PARAM_DATA( at, varName ) );
+////                        break;
+////                    }
+
+////                    case DT_COLOUR:
+////                        mutable_snprintf( temp, MUTABLE_GPU_CODE_BUFFER_SIZE, "%s_%s_%d",
+////                                          m_currentImageName.c_str(),
+////                                          MUTABLE_GPU_VECTOR_PARAM, m_vectorParams.size() );
+////                        m_currentVar = temp;
+////                        m_vectorParams.push_back( GPU_PROGRAM::PARAM_DATA( at, temp ) );
+////                        break;
+
+////                    case DT_SCALAR:
+////                        mutable_snprintf( temp, MUTABLE_GPU_CODE_BUFFER_SIZE, "%s_%s_%d",
+////                                          m_currentImageName.c_str(),
+////                                          MUTABLE_GPU_FLOAT_PARAM, m_floatParams.size() );
+////                        m_currentVar = temp;
+////                        m_floatParams.push_back( GPU_PROGRAM::PARAM_DATA( at, temp ) );
+////                        break;
+
+////                    case DT_LAYOUT:
+////                    {
+////                        OP trans;
+////                        trans.type = OP_TYPE::CO_LAYOUTBLOCKTRANSFORM;
+////                        trans.args.ColourLayoutBlockTransform.layout = at;
+////                        trans.args.ColourLayoutBlockTransform.block = (uint16_t)m_currentLayoutBlock;
+////                        at = program.AddOp( trans );
+
+////                        mutable_snprintf( temp, 128, "%s_%s_%d", m_currentImageName.c_str(),
+////                                          MUTABLE_GPU_VECTOR_PARAM, m_vectorParams.size() );
+////                        m_currentVar = temp;
+////                        m_vectorParams.push_back( GPU_PROGRAM::PARAM_DATA( at, temp ) );
+////                        break;
+////                    }
+
+////                    default:
+////                        check( false );
+////                    }
+////                }
+
+////                m_generated[at] = m_currentVar;
+////            }
+
+////            else
+////            {
+////                // Reuse previously generated code
+////                m_currentVar = m_generated[at];
+////            }
+
+//        }
+
+//    protected:
+
+//        int m_state;
+
+//        //! Instructions that have been already generated and the variable they reside in.
+//        map<OP::ADDRESS,string> m_generated;
+
+//        vector<string> m_usedVars;
+
+//        //! State when generating image expressions
+//        string m_currentTexCoords;
+
+//        //! State when generating layout expressions
+//        int m_currentLayoutBlock;
+
+//        //! Generate an unused variable name based on the provided one.
+//        string GetNewVar( const string& name )
+//        {
+//            int suffix = 0;
+//            string candidate = name;
+
+//            while( std::find( m_usedVars.begin(), m_usedVars.end(), candidate )
+//                   !=
+//                   m_usedVars.end() )
+//            {
+//                ++suffix;
+
+//                char temp[128];
+//                mutable_snprintf( temp, 128, "%s_%d", name.c_str(), suffix );
+//                candidate = temp;
+//            }
+
+//            m_usedVars.push_back( candidate );
+//            return candidate;
+//        }
+//    };
+
+
+//    //---------------------------------------------------------------------------------------------
+//    //---------------------------------------------------------------------------------------------
+//    //---------------------------------------------------------------------------------------------
+//    class GPUSubtreeGenerator_GL4 : public GPUSubtreeGenerator_GLSL
+//    {
+//    public:
+
+//        GPU_PROGRAM GenerateProgram( OP::ADDRESS at,
+//                              PROGRAM& program,
+//                              int state,
+//                              const GPU_PLATFORM_PROPS& platformOptions,
+//                              const GPU_STATE_PROPS& //stateOptions
+//                              )
+//        {
+//            check( platformOptions.type == GPU_GL4 );
+//            (void)platformOptions;
+//            m_state = state;
+//            m_currentImageName = "";
+
+//            // Clear state
+//            Clear();
+
+//            Traverse( at, program, false );
+
+//            GPU_PROGRAM gpuProgram;
+//            gpuProgram.m_floatParams = m_floatParams;
+//            gpuProgram.m_vectorParams = m_vectorParams;
+//            gpuProgram.m_imageParams = m_imageParams;
+
+////            // Find out what pixel format we need to generate
+////            IMAGE_FORMAT format = GetImageDesc( program, at ).m_format;
+
+////            // Convert to a usable format
+////            switch ( format )
+////            {
+////            case IF_RGBA_UBYTE: gpuProgram.m_format = format; break;
+////            case IF_RGB_UBYTE: gpuProgram.m_format = IF_RGBA_UBYTE; break;
+////            default:
+////                break;
+////            }
+
+
+////            // Generate the size instruction
+////            OP sizeOp;
+////            sizeOp.type = OP_TYPE::CO_IMAGESIZE;
+////            sizeOp.args.ColourImageSize.image = at;
+////            gpuProgram.m_size = program.AddOp( sizeOp );
+
+
+////            // Complete the code
+////            string params =
+////                    "\n"
+////                    "#version 330\n"
+////                    "\n";
+
+////            char temp[MUTABLE_GPU_CODE_BUFFER_SIZE];
+
+////            // Add every source image of the fragment program
+////            for ( size_t i=0; i<m_imageParams.size(); ++i )
+////            {
+////                mutable_snprintf( temp, MUTABLE_GPU_CODE_BUFFER_SIZE, "uniform sampler2D %s;\n",
+////                                  m_imageParams[i].name.c_str() );
+////                params += temp;
+////            }
+
+////            // Add every source vector of the fragment program
+////            for ( size_t i=0; i<m_vectorParams.size(); ++i )
+////            {
+////                mutable_snprintf( temp, MUTABLE_GPU_CODE_BUFFER_SIZE, "uniform vec4 %s;\n",
+////                                  m_vectorParams[i].name.c_str() );
+////                params += temp;
+////            }
+
+////            // Add every source float of the fragment program
+////            for ( size_t i=0; i<m_floatParams.size(); ++i )
+////            {
+////                mutable_snprintf( temp, MUTABLE_GPU_CODE_BUFFER_SIZE, "uniform float %s;\n",
+////                                  m_floatParams[i].name.c_str() );
+////                params += temp;
+////            }
+
+////            string header =
+////                    "\n"
+////                    "in vec2 "+string(MUTABLE_GPU_CODE_UVS)+";\n"
+////                    "out vec4 result;\n"
+////                    "\n"
+////                    "void main(void)\n"
+////                    "{\n";
+
+////            string footer =
+////                    "\n"
+////                    "result = "+m_currentVar+";\n"
+////                    "}\n";
+
+////            gpuProgram.m_fragmentCode = params + header + m_expressionCode + footer;
+
+////            gpuProgram.m_vertexCode =
+////                    "#version 330\n"
+////                    "\n"
+////                    "in vec2 position;"
+////                    "out vec2 "+string(MUTABLE_GPU_CODE_UVS)+";"
+////                    ""
+////                    "void main()"
+////                    "{"
+////                    "	gl_Position = vec4( position.xy, 0.0, 1.0 );"
+////                    "	"+string(MUTABLE_GPU_CODE_UVS)+" = position.xy*0.5+0.5;"
+////                    "}";
+
+////            gpuProgram.m_platform = SS_GL4;
+
+//            return gpuProgram;
+//        }
+
+//    };
+
+
+//    //---------------------------------------------------------------------------------------------
+//    //---------------------------------------------------------------------------------------------
+//    //---------------------------------------------------------------------------------------------
+//    class GPUSubtreeGenerator_GLES2 : public GPUSubtreeGenerator_GLSL
+//    {
+//    public:
+
+
+//        GPU_PROGRAM GenerateProgram( OP::ADDRESS at,
+//                                     PROGRAM& program,
+//                                     int state,
+//                                     const GPU_PLATFORM_PROPS& platformOptions,
+//                                     const GPU_STATE_PROPS& //stateOptions
+//                                     ) override
+//        {
+//            check( platformOptions.type == GPU_GLES2 );
+//            (void)platformOptions;
+
+//            m_state = state;
+//            m_currentImageName = "";
+
+//            // Clear state
+//            Clear();
+
+//            Traverse( at, program, false );
+
+//            GPU_PROGRAM gpuProgram;
+//            gpuProgram.m_floatParams = m_floatParams;
+//            gpuProgram.m_vectorParams = m_vectorParams;
+//            gpuProgram.m_imageParams = m_imageParams;
+
+////            // Find out what pixel format we need to generate
+////            IMAGE_FORMAT format = GetImageDesc( program, at ).m_format;
+
+////            // Convert to a usable format
+////            switch ( format )
+////            {
+////            case IF_RGBA_UBYTE: gpuProgram.m_format = format; break;
+////            case IF_RGB_UBYTE: gpuProgram.m_format = IF_RGBA_UBYTE; break;
+////            default:
+////                break;
+////            }
+
+
+////            // Generate the size instruction
+////            OP sizeOp;
+////            sizeOp.type = OP_TYPE::CO_IMAGESIZE;
+////            sizeOp.args.ColourImageSize.image = at;
+////            gpuProgram.m_size = program.AddOp( sizeOp );
+
+
+////            // Complete the code
+////            string params =
+////                    "\n"
+////                    "precision mediump float;\n"
+////                    "\n";
+
+////            char temp[MUTABLE_GPU_CODE_BUFFER_SIZE];
+
+////            // Add every source image of the fragment program
+////            for ( size_t i=0; i<m_imageParams.size(); ++i )
+////            {
+////                mutable_snprintf( temp, MUTABLE_GPU_CODE_BUFFER_SIZE, "uniform sampler2D %s;\n",
+////                                  m_imageParams[i].name.c_str() );
+////                params += temp;
+////            }
+
+////            // Add every source vector of the fragment program
+////            for ( size_t i=0; i<m_vectorParams.size(); ++i )
+////            {
+////                mutable_snprintf( temp, MUTABLE_GPU_CODE_BUFFER_SIZE, "uniform vec4 %s;\n",
+////                                  m_vectorParams[i].name.c_str() );
+////                params += temp;
+////            }
+
+////            // Add every source float of the fragment program
+////            for ( size_t i=0; i<m_floatParams.size(); ++i )
+////            {
+////                mutable_snprintf( temp, MUTABLE_GPU_CODE_BUFFER_SIZE, "uniform float %s;\n",
+////                                  m_floatParams[i].name.c_str() );
+////                params += temp;
+////            }
+
+////            string header =
+////                    "\n"
+////                    "varying vec2 texCoords;\n"
+////                    "\n"
+////                    "void main(void)\n"
+////                    "{\n";
+
+////            string footer =
+////                    "\n"
+////                    "gl_FragData[0] = "+m_currentVar+";\n"
+////                    "}\n";
+
+////            gpuProgram.m_fragmentCode = params + header + m_expressionCode + footer;
+
+////            gpuProgram.m_vertexCode =
+////                    "\n"
+////                    "attribute vec2 position;"
+////                    "varying vec2 texCoords;"
+////                    ""
+////                    "void main()"
+////                    "{"
+////                    "	gl_Position = vec4( position.xy, 0.0, 1.0 );"
+////                    "	texCoords = position.xy*0.5+0.5;"
+////                    "}";
+
+////            gpuProgram.m_platform = SS_GLES2;
+
+//            return gpuProgram;
+//        }
+
+//    };
+
+
+//    //---------------------------------------------------------------------------------------------
+//    //---------------------------------------------------------------------------------------------
+//    //---------------------------------------------------------------------------------------------
+//    GPUTranslator::GPUTranslator( PROGRAM& program, int s,
+//                                  const GPU_PLATFORM_PROPS& platformOptions,
+//                                  const GPU_STATE_PROPS& , //stateOptions,
+//                                  ModelReportPtr pReport )
+//        //: m_platformOptions( platformOptions )
+//        //, m_stateOptions( stateOptions )
+//        : m_pReport( pReport )
+//    {
+//        m_state = s;
+//        m_visited.resize( program.m_opAddress.size(), 0 );
+
+//        // Create the platfomr specific GPU code generator.
+//        m_pGPUSubtreeGen = 0;
+//        switch ( platformOptions.type )
+//        {
+//        case GPU_GL4:
+//            m_pGPUSubtreeGen = new GPUSubtreeGenerator_GL4;
+//            break;
+
+//        case GPU_GLES2:
+//            m_pGPUSubtreeGen = new GPUSubtreeGenerator_GLES2;
+//            break;
+
+//        default:
+//            check( false );
+//            break;
+//        }
+
+//        // Clone the code before optimising it for the gpu: this way we keep the original
+//        // code to be used by other states.
+//        //CloneVisitor( program, program.m_states[s].m_root, s );
+
+//        Traverse( program.m_states[s].m_root, program );
+//    }
+
+
+//    //---------------------------------------------------------------------------------------------
+//    GPUTranslator::~GPUTranslator()
+//    {
+//        delete m_pGPUSubtreeGen;
+//        m_pGPUSubtreeGen = 0;
+//    }
+
+//    //---------------------------------------------------------------------------------------------
+//    OP::ADDRESS GPUTranslator::Visit( OP::ADDRESS at, PROGRAM& program )
+//    {
+//        // Visit only once
+//        OP::ADDRESS sourceAt = at;
+//        m_visited.resize( program.m_opAddress.size(), 0 );
+//        if ( m_visited[at] )
+//        {
+//            return m_visited[at];
+//        }
+
+//        OP op = program.m_code[at];
+
+//        bool gpuizable = GetOpDesc( (OP_TYPE)op.type ).gpuizable;
+//        gpuizable &= !program.m_states[m_state].IsUpdateCache(at);
+
+////        if ( gpuizable && m_stateOptions.m_internal )
+////        {
+////            // Find out what pixel format we need to generate
+////            IMAGE_FORMAT format = GetImageDesc( program, at ).m_format;
+
+////            // Create the fragment program
+////            GPU_PROGRAM gpuProgram =
+////                    m_pGPUSubtreeGen->GenerateProgram( at, program,
+////                                                       m_state, m_platformOptions, m_stateOptions );
+
+////            // Create a GPU image instruction
+////            OP gpuOp;
+////            gpuOp.type = OP_TYPE::IM_GPU;
+////            gpuOp.args.ImageGPU.program = (OP::ADDRESS)program.m_gpuPrograms.size();
+////            program.m_gpuPrograms.push_back( gpuProgram );
+////            at = program.AddOp( gpuOp );
+
+////            // Check that we have the expected pixel format, as gpu instructions cannot generate
+////            // some of them.
+////            if ( format != gpuProgram.m_format )
+////            {
+////                OP formatOp;
+////                formatOp.type = OP_TYPE::IM_PIXELFORMAT;
+////                formatOp.args.ImagePixelFormat.source = at;
+////                formatOp.args.ImagePixelFormat.format = format;
+////                at = program.AddOp( formatOp );
+////            }
+
+////            // Recurse on the program parameters
+////            //TODO
+////        }
+////        else if ( m_stateOptions.m_external )
+////        {
+////            if ( op.type == OP_TYPE::IN_ADDLOD )
+////            {
+////                for (int t=0;t<MUTABLE_OP_MAX_ADD_COUNT;++t)
+////                {
+////                    if ( op.args.InstanceAddLOD.lod[t] )
+////                    {
+////                        m_currentLOD = REPORT_STATE::LOD();
+////                        op.args.InstanceAddLOD.lod[t] = Visit( op.args.InstanceAddLOD.lod[t], program );
+////                        m_pReport->GetPrivate()->m_states[m_state].m_lods.push_back(m_currentLOD);
+////                    }
+////                }
+
+////                at = program.AddOp(op);
+////            }
+////            else if ( op.type == OP_TYPE::IN_ADDCOMPONENT )
+////            {
+////                op.args.InstanceAdd.instance = Visit( op.args.InstanceAdd.instance, program );
+
+////                if ( op.args.InstanceAdd.value )
+////                {
+////                    m_currentComponent = REPORT_STATE::COMPONENT();
+
+////                    op.args.InstanceAdd.value = Visit( op.args.InstanceAdd.value, program );
+
+////                    m_currentComponent.m_name =
+////                            program.m_constantStrings[ op.args.InstanceAdd.name ];
+////                    m_currentLOD.m_components.push_back(m_currentComponent);
+////                }
+
+////                at = program.AddOp(op);
+////            }
+////            else if ( op.type == OP_TYPE::IN_ADDIMAGE )
+////            {
+////                OP::ADDRESS lastInstance = op.args.InstanceAdd.instance;
+
+////                // Visit the instance
+////                lastInstance = Visit( lastInstance, program );
+
+////                // GPU-ize or visit every image
+////                // We actually unfold the images too
+////                if ( op.args.InstanceAdd.value )
+////                {
+////                    if ( m_stateOptions.m_external )
+////                    {
+////                        REPORT_STATE::IMAGE imageData;
+////                        imageData.m_name = program.m_constantStrings[ op.args.InstanceAdd.name ];
+
+////                        // Create the fragment program
+////                        string imageCode =
+////                                m_pGPUSubtreeGen->GenerateImage( imageData.m_name,
+////                                                                 op.args.InstanceAdd.value,
+////                                                                 program,
+////                                                                 m_state,
+////                                                                 m_platformOptions,
+////                                                                 m_stateOptions );
+
+////                        // Add every source image of the fragment program
+////                        size_t sourceImageCount = m_pGPUSubtreeGen->m_imageParams.size();
+////                        for ( size_t i=0; i<sourceImageCount; ++i )
+////                        {
+////                            OP nop;
+////                            nop.type = OP_TYPE::IN_ADDIMAGE;
+////                            nop.args.InstanceAdd.instance = lastInstance;
+////                            nop.args.InstanceAdd.value =
+////                                    Visit( m_pGPUSubtreeGen->m_imageParams[i].address, program );
+
+////                            nop.args.InstanceAdd.name = program.AddConstant
+////                                    ( m_pGPUSubtreeGen->m_imageParams[i].name.c_str() );
+////                            lastInstance = program.AddOp( nop );
+
+////                            imageData.m_sourceImages.push_back
+////                                    ( m_pGPUSubtreeGen->m_imageParams[i].name );
+////                        }
+
+////                        // Add every source vector of the fragment program
+////                        size_t sourceVectorCount = m_pGPUSubtreeGen->m_vectorParams.size();
+////                        for ( size_t i=0; i<sourceVectorCount; ++i )
+////                        {
+////                            OP nop;
+////                            nop.type = OP_TYPE::IN_ADDVECTOR;
+////                            nop.args.InstanceAdd.instance = lastInstance;
+////                            nop.args.InstanceAdd.value =
+////                                    Visit( m_pGPUSubtreeGen->m_vectorParams[i].address, program );
+
+////                            nop.args.InstanceAdd.name = program.AddConstant
+////                                    ( m_pGPUSubtreeGen->m_vectorParams[i].name.c_str() );
+////                            lastInstance = program.AddOp( nop );
+
+////                            imageData.m_sourceVectors.push_back
+////                                    ( m_pGPUSubtreeGen->m_vectorParams[i].name );
+////                        }
+
+////                        // Add every source float of the fragment program
+////                        size_t sourceFloatCount = m_pGPUSubtreeGen->m_floatParams.size();
+////                        for ( size_t i=0; i<sourceFloatCount; ++i )
+////                        {
+////                            OP nop;
+////                            nop.type = OP_TYPE::IN_ADDSCALAR;
+////                            nop.args.InstanceAdd.instance = lastInstance;
+////                            nop.args.InstanceAdd.value =
+////                                    Visit( m_pGPUSubtreeGen->m_floatParams[i].address, program );
+
+////                            nop.args.InstanceAdd.name = program.AddConstant
+////                                    ( m_pGPUSubtreeGen->m_floatParams[i].name.c_str() );
+////                            lastInstance = program.AddOp( nop );
+
+////                            imageData.m_sourceScalars.push_back
+////                                    ( m_pGPUSubtreeGen->m_floatParams[i].name );
+////                        }
+
+////                        imageData.m_fragmentCode = imageCode;
+////                        m_currentComponent.m_images.push_back(imageData);
+////                    }
+////                    else
+////                    {
+////                        // Recurse and add the image
+////                        OP nop;
+////                        nop.type = OP_TYPE::IN_ADDIMAGE;
+////                        nop.args.InstanceAdd.instance = lastInstance;
+////                        nop.args.InstanceAdd.value =
+////                                Visit( op.args.InstanceAdd.value, program );
+////                        nop.args.InstanceAdd.name = op.args.InstanceAdd.name;
+////                        lastInstance = program.AddOp( nop );
+////                    }
+////                }
+
+////                at = lastInstance;
+////            }
+////            else
+////            {
+////                // We cannot make a gpu program here. Let's see below.
+////                at = Recurse( at, program );
+////            }
+////        }
+////        else
+////        {
+////            // We cannot make a gpu program here. Let's see below.
+////            at = Recurse( at, program );
+////        }
+
+//        m_visited[sourceAt] = at;
+
+//        return at;
+//    }
+
+}
+
diff --git a/Engine/Plugins/Experimental/Mutable/Source/MutableTools/Private/CodeOptimiser_GPU.h b/Engine/Plugins/Experimental/Mutable/Source/MutableTools/Private/CodeOptimiser_GPU.h
new file mode 100644
--- /dev/null
+++ b/Engine/Plugins/Experimental/Mutable/Source/MutableTools/Private/CodeOptimiser_GPU.h
@@ -0,0 +1,58 @@
+// Copyright Epic Games, Inc. All Rights Reserved.
+
+#pragma once
+
+#include "Operations.h"
+#include "ModelPrivate.h"
+#include "ImagePrivate.h"
+#include "MeshPrivate.h"
+#include "CodeVisitor.h"
+
+#include "CompilerPrivate.h"
+#include "ModelReportPrivate.h"
+
+
+namespace mu
+{
+    //!
+    class GPUSubtreeGenerator;
+
+
+    //---------------------------------------------------------------------------------------------
+    //!
+    //---------------------------------------------------------------------------------------------
+//    class GPUTranslator : public CodeVisitor
+//    {
+//    public:
+
+//        GPUTranslator( PROGRAM& program,
+//                       int state,
+//                       const GPU_PLATFORM_PROPS& platformOptions,
+//                       const GPU_STATE_PROPS& stateOptions,
+//                       ModelReportPtr pReport );
+
+//        virtual ~GPUTranslator();
+
+//        virtual OP::ADDRESS Visit( OP::ADDRESS at, PROGRAM& program );
+
+//    private:
+
+//        vector<OP::ADDRESS> m_visited;
+
+//        int m_state;
+//        //const GPU_PLATFORM_PROPS& m_platformOptions;
+//        //const GPU_STATE_PROPS& m_stateOptions;
+//        REPORT_STATE::LOD m_currentLOD;
+//        REPORT_STATE::COMPONENT m_currentComponent;
+
+//        //!
+//        ModelReportPtr m_pReport;
+
+//        //!
+//        GPUSubtreeGenerator* m_pGPUSubtreeGen;
+
+//    };
+
+
+
+}
