[mutable] Move Mutable plugin from the restricted folder to the experimental plugins folder.

diff --git a/Engine/Plugins/Experimental/Mutable/Source/MutableTools/Private/CodeOptimiser_Logic.cpp b/Engine/Plugins/Experimental/Mutable/Source/MutableTools/Private/CodeOptimiser_Logic.cpp
new file mode 100644
--- /dev/null
+++ b/Engine/Plugins/Experimental/Mutable/Source/MutableTools/Private/CodeOptimiser_Logic.cpp
@@ -0,0 +1,1359 @@
+// Copyright Epic Games, Inc. All Rights Reserved.
+
+#include "CodeOptimiser.h"
+
+#include "ModelPrivate.h"
+#include "MutableTrace.h"
+#include "SystemPrivate.h"
+#include "Operations.h"
+
+#include "MutableTools/Private/ASTOpConditional.h"
+#include "MutableTools/Private/ASTOpInstanceAdd.h"
+#include "MutableTools/Private/ASTOpSwitch.h"
+
+
+
+namespace mu
+{
+
+
+    //---------------------------------------------------------------------------------------------
+    //---------------------------------------------------------------------------------------------
+    //---------------------------------------------------------------------------------------------
+//    IntExpressionPtr GetIntExpression( OP::ADDRESS at, const PROGRAM& program )
+//    {
+//        IntExpressionPtr pRes = new IntExpression;
+
+//        const OP& op = program.m_code[at];
+//        OP_TYPE type = op.type;
+//        switch ( type )
+//        {
+//        case OP_TYPE::NU_CONSTANT:
+//            pRes->type = IntExpression::IET_CONSTANT;
+//            pRes->value = op.args.IntConstant.value;
+//            break;
+
+//        case OP_TYPE::NU_PARAMETER:
+//            pRes->type = IntExpression::IET_PARAMETER;
+//            pRes->parameter = op.args.Parameter.variable;
+//            break;
+
+//        default:
+//            //check( false );
+//            pRes->type = IntExpression::IET_UNKNOWN;
+//            break;
+//        }
+
+//        return pRes;
+//    }
+
+
+    //---------------------------------------------------------------------------------------------
+    //---------------------------------------------------------------------------------------------
+    //---------------------------------------------------------------------------------------------
+//    BoolExpressionPtr GetBoolExpression( OP::ADDRESS at,
+//                                         const PROGRAM& program,
+//                                         BOOL_EXPRESSION_CACHE* cache )
+//    {
+//        if (cache && (size_t)at<=cache->m_values.size() && cache->m_values[at])
+//        {
+//            return cache->m_values[at];
+//        }
+
+//        BoolExpressionPtr pRes = new BoolExpression;
+
+//        const OP& op = program.m_code[at];
+//        OP_TYPE type = (OP_TYPE)op.type;
+//        switch ( type )
+//        {
+//        case OP_TYPE::BO_CONSTANT:
+//            pRes->type = op.args.BoolConstant.value
+//                ? BoolExpression::BET_TRUE
+//                : BoolExpression::BET_FALSE;
+//            break;
+
+//        case OP_TYPE::BO_PARAMETER:
+//            pRes->type = BoolExpression::BET_PARAMETER;
+//            pRes->parameter = op.args.Parameter.variable;
+//            break;
+
+//        case OP_TYPE::BO_AND:
+//            pRes->type = BoolExpression::BET_AND;
+//            pRes->a = GetBoolExpression( op.args.BoolBinary.a, program, cache );
+//            pRes->b = GetBoolExpression( op.args.BoolBinary.b, program, cache );
+//            break;
+
+//        case OP_TYPE::BO_OR:
+//            pRes->type = BoolExpression::BET_OR;
+//            pRes->a = GetBoolExpression( op.args.BoolBinary.a, program, cache );
+//            pRes->b = GetBoolExpression( op.args.BoolBinary.b, program, cache );
+//            break;
+
+//        case OP_TYPE::BO_NOT:
+//            pRes->type = BoolExpression::BET_NOT;
+//            pRes->a = GetBoolExpression( op.args.BoolNot.source, program, cache );
+//            break;
+
+//        case OP_TYPE::BO_EQUAL_INT_CONST:
+//            pRes->type = BoolExpression::BET_INT_EQUAL_CONST;
+//            pRes->intExp = GetIntExpression( op.args.BoolEqualScalarConst.value, program );
+//            pRes->intValue = op.args.BoolEqualScalarConst.constant;
+//            break;
+
+//        default:
+//            check( false );
+//            pRes->type = BoolExpression::BET_UNKNOWN;
+//            break;
+//        }
+
+//        if (cache)
+//        {
+//            if (cache->m_values.size()<program.m_opAddress.size())
+//            {
+//                cache->m_values.resize(program.m_opAddress.size());
+//            }
+//            cache->m_values[at]=pRes;
+//        }
+
+//        return pRes;
+//    }
+
+
+    //---------------------------------------------------------------------------------------------
+    //---------------------------------------------------------------------------------------------
+    //---------------------------------------------------------------------------------------------
+//    void LogicCallstackGather::Apply( PROGRAM& program, int state )
+//    {
+//        MUTABLE_CPUPROFILER_SCOPE(LogicCallstackGather);
+//        Traverse( program.m_states[state].m_root, program );
+//    }
+
+
+//    //---------------------------------------------------------------------------------------------
+//    void LogicCallstackGather::Visit( OP::ADDRESS at, PROGRAM& program )
+//    {
+//        // We traverse top down
+//        OP_TYPE type = (OP_TYPE)program.m_code[at].type;
+//        OP::ARGS args = program.m_code[at].args;
+
+//        switch ( type )
+//        {
+
+//        //-----------------------------------------------------------------------------------------
+//        case OP_TYPE::NU_CONDITIONAL:
+//        case OP_TYPE::SC_CONDITIONAL:
+//        case OP_TYPE::CO_CONDITIONAL:
+//        case OP_TYPE::IM_CONDITIONAL:
+//        case OP_TYPE::ME_CONDITIONAL:
+//        case OP_TYPE::LA_CONDITIONAL:
+//        case OP_TYPE::IN_CONDITIONAL:
+//        {
+//            m_factsPerCallstack[at].push_back( m_facts );
+
+//            BoolExpressionPtr pCond = GetBoolExpression( args.Conditional.condition, program,
+//                                                         &m_boolExpressionCache );
+
+//            // Check if the condition has a value with the current callstack. If it does, we skip
+//            // the contradictory paths, as they will never happen.
+//            BoolExpression::TYPE result = pCond->Evaluate( m_facts );
+
+//            // Recurse the condition without adding extra facts
+//            Visit( args.Conditional.condition, program );
+
+//            if ( result != BoolExpression::BET_FALSE )
+//            {
+//                // yes branch
+//                m_facts.push_back( pCond );
+//                Visit( args.Conditional.yes, program );
+//                m_facts.pop_back();
+//            }
+
+//            if ( result != BoolExpression::BET_TRUE )
+//            {
+//                // no branch
+//                BoolExpressionPtr pNot = new BoolExpression;
+//                pNot->type = BoolExpression::BET_NOT;
+//                pNot->a = pCond;
+
+//                m_facts.push_back( pNot );
+//                Visit( args.Conditional.no, program );
+//                m_facts.pop_back();
+//            }
+
+//            break;
+//        }
+
+//        case OP_TYPE::NU_SWITCH:
+//        case OP_TYPE::SC_SWITCH:
+//        case OP_TYPE::CO_SWITCH:
+//        case OP_TYPE::IM_SWITCH:
+//        case OP_TYPE::VO_SWITCH:
+//        case OP_TYPE::ME_SWITCH:
+//        case OP_TYPE::LA_SWITCH:
+//        {
+//            m_factsPerCallstack[at].push_back( m_facts );
+
+//            IntExpressionPtr pValue = GetIntExpression( args.Switch.variable, program );
+
+//            const auto& options = program.m_constantSwitches[args.Switch.options].m_options;
+//            for ( size_t i=0; i<options.size(); ++i )
+//            {
+//                if ( options[i].at )
+//                {
+//                    BoolExpressionPtr pCond = new BoolExpression;
+//                    pCond->type = BoolExpression::BET_INT_EQUAL_CONST;
+//                    pCond->intExp = pValue;
+//                    pCond->intValue = args.Switch.conditions[i];
+
+//                    // Check if the condition has a value with the current callstack. If it does, we
+//                    // skip the contradictory paths, as they will never happen.
+//                    BoolExpression::TYPE result = pCond->Evaluate( m_facts );
+
+//                    if (result!=BoolExpression::BET_FALSE)
+//                    {
+//                        // We don't know if it is true or false. Add its value as a fact and
+//                        // recurse.
+//                        m_facts.push_back( pCond );
+//                        Visit( options[i].at, program );
+//                        m_facts.pop_back();
+//                    }
+//                }
+//            }
+
+//            if ( args.Switch.def )
+//            {
+//                // TODO: Add facts negating the options in the switch
+//                Visit( args.Switch.def, program );
+//            }
+//            break;
+//        }
+
+//        default:
+//            // Normal recursion
+//            Recurse( at, program );
+//            break;
+
+//        }
+//    }
+
+
+    //---------------------------------------------------------------------------------------------
+    //---------------------------------------------------------------------------------------------
+    //---------------------------------------------------------------------------------------------
+//    bool ClonelessLogicOptimiser::Apply( PROGRAM& program, int state )
+//    {
+//        MUTABLE_CPUPROFILER_SCOPE(ClonelessLogicOptimiser);
+
+//        // Gather all the conditions for every op
+//        m_callstacks.Apply( program, state );
+
+//        m_modified = false;
+//        Traverse( program.m_states[state].m_root, program );
+
+//        return m_modified;
+//    }
+
+
+//    //---------------------------------------------------------------------------------------------
+//    OP::ADDRESS ClonelessLogicOptimiser::Visit( OP::ADDRESS at, PROGRAM& program )
+//    {
+//        // We traverse top down
+
+//        OP_TYPE type = (OP_TYPE)program.m_code[at].type;
+//        OP::ARGS args = program.m_code[at].args;
+
+//        switch ( type )
+//        {
+
+//        //-----------------------------------------------------------------------------------------
+//        case OP_TYPE::NU_CONDITIONAL:
+//        case OP_TYPE::SC_CONDITIONAL:
+//        case OP_TYPE::CO_CONDITIONAL:
+//        case OP_TYPE::IM_CONDITIONAL:
+//        case OP_TYPE::ME_CONDITIONAL:
+//        case OP_TYPE::LA_CONDITIONAL:
+//        case OP_TYPE::IN_CONDITIONAL:
+//        {
+//            BOOL_EXPRESSION_CACHE localBoolExpressionCache;
+//            BoolExpressionPtr pCond = GetBoolExpression( args.Conditional.condition, program, &localBoolExpressionCache );
+
+//            // Check if the condition has a value with all the possible callstacks on this op
+//            vector< vector<BoolExpressionPtr> >& callstacks = m_callstacks.m_factsPerCallstack[at];
+//            check( callstacks.size() );
+
+//            BoolExpression::TYPE result = pCond->Evaluate( callstacks[0] );
+//            for ( size_t c=1; c<callstacks.size(); ++c )
+//            {
+//                BoolExpression::TYPE thisResult = pCond->Evaluate( callstacks[c] );
+//                if (result!=thisResult)
+//                {
+//                    result = BoolExpression::BET_UNKNOWN;
+//                    break;
+//                }
+//            }
+
+//            if ( result==BoolExpression::BET_TRUE )
+//            {
+//                m_modified = true;
+//                m_callstacks.m_factsPerCallstack[at] = m_callstacks.m_factsPerCallstack[program.m_code[at].args.Conditional.yes];
+//                at = program.m_code[at].args.Conditional.yes;
+//                at = Visit( at, program );
+//            }
+//            else if ( result==BoolExpression::BET_FALSE )
+//            {
+//                m_modified = true;
+//                m_callstacks.m_factsPerCallstack[at] = m_callstacks.m_factsPerCallstack[program.m_code[at].args.Conditional.no];
+//                at = program.m_code[at].args.Conditional.no;
+//                at = Visit( at, program );
+//            }
+//            else
+//            {
+//                at = Recurse( at, program );
+//            }
+
+//            break;
+//        }
+
+////        case OP_TYPE::NU_SWITCH:
+////        case OP_TYPE::SC_SWITCH:
+////        case OP_TYPE::CO_SWITCH:
+////        case OP_TYPE::IM_SWITCH:
+////        case OP_TYPE::VO_SWITCH:
+////        case OP_TYPE::ME_SWITCH:
+////        case OP_TYPE::LA_SWITCH:
+////        {
+////            IntExpressionPtr pValue = GetIntExpression( args.Switch.variable, program );
+
+////            vector< vector<BoolExpressionPtr> >& callstacks = m_callstacks.m_factsPerCallstack[at];
+////            check( callstacks.size() );
+
+////            for ( int i=0;i<MUTABLE_OP_MAX_SWITCH_OPTIONS; ++i )
+////            {
+////                if ( args.Switch.values[i] )
+////                {
+////                    BoolExpressionPtr pCond = new BoolExpression;
+////                    pCond->type = BoolExpression::BET_INT_EQUAL_CONST;
+////                    pCond->intExp = pValue;
+////                    pCond->intValue = args.Switch.conditions[i];
+
+////                    // Check if the condition has a value with all the possible callstacks on this op
+////                    BoolExpression::TYPE result = pCond->Evaluate( callstacks[0] );
+////                    for ( size_t c=1; c<callstacks.size(); ++c )
+////                    {
+////                        BoolExpression::TYPE thisResult = pCond->Evaluate( callstacks[c] );
+////                        if (result!=thisResult)
+////                        {
+////                            result = BoolExpression::BET_UNKNOWN;
+////                            break;
+////                        }
+////                    }
+
+////                    if ( result==BoolExpression::BET_TRUE )
+////                    {
+////                        // This is the only possible branch.
+////                        m_modified = true;
+////                        at = args.Switch.values[i];
+////                        //at = Visit( at, program );
+////                        return at;
+////                        break;
+////                    }
+////                    else if ( result==BoolExpression::BET_FALSE )
+////                    {
+////                        // This option is impossible
+////                        args.Switch.values[i] = 0;
+////                        args.Switch.conditions[i] = 0;
+////                    }
+////                    else
+////                    {
+////                        // We don't know if it is true or false. Add its value as a fact and
+////                        // recurse.
+////                        //args.Switch.values[i] = Visit( args.Switch.values[i], program );
+////                    }
+////                }
+////            }
+
+////            if ( args.Switch.def )
+////            {
+////                // TODO: Add facts negating the options in the switch
+////                //args.Switch.def = Visit( args.Switch.def, program );
+////            }
+
+////            program.m_code[at].args = args;
+////            break;
+////        }
+
+//        default:
+//            // Normal recursion
+//            at = Recurse( at, program );
+//            break;
+
+//        }
+
+//        return at;
+//    }
+
+
+//    //---------------------------------------------------------------------------------------------
+//    //---------------------------------------------------------------------------------------------
+//    //---------------------------------------------------------------------------------------------
+//    bool LogicOptimiser::Apply( PROGRAM& program, int state )
+//    {
+//        m_initialCodeSize = program.m_opAddress.size();
+//        m_modified = false;
+//        Traverse( program.m_states[state].m_root, program );
+
+//        return m_modified;
+//    }
+
+
+//    //---------------------------------------------------------------------------------------------
+//    OP::ADDRESS LogicOptimiser::Visit( OP::ADDRESS at, PROGRAM& program )
+//    {
+//        // Sanity stop
+//        const int sanityFactor = 4;
+//        if ( program.m_opAddress.size()>m_initialCodeSize*sanityFactor)
+//        {
+//            return at;
+//        }
+
+//        // We traverse top down
+
+//        OP_TYPE type = (OP_TYPE)program.m_code[at].type;
+//        OP::ARGS args = program.m_code[at].args;
+
+//        // We only apply it to instances, for performance reasons.
+//        if (GetOpType(type)!=DT_INSTANCE)
+//        {
+//            return at;
+//        }
+
+//        switch ( type )
+//        {
+
+//        //-----------------------------------------------------------------------------------------
+////        case OP_TYPE::NU_CONDITIONAL:
+////        case OP_TYPE::SC_CONDITIONAL:
+////        case OP_TYPE::CO_CONDITIONAL:
+////        case OP_TYPE::IM_CONDITIONAL:
+////        case OP_TYPE::ME_CONDITIONAL:
+////        case OP_TYPE::LA_CONDITIONAL:
+//        case OP_TYPE::IN_CONDITIONAL:
+//        {
+//            BoolExpressionPtr pCond = GetBoolExpression( args.Conditional.condition, program, &m_boolExpressionCache );
+
+//            // Check if the condition has a value with the current facts
+//            // TODO: This is a simple adhoc check which is also quite imprecise.
+//            BoolExpression::TYPE result = pCond->Evaluate( m_facts );
+//            if ( result==BoolExpression::BET_TRUE )
+//            {
+//                m_modified = true;
+//                at = program.m_code[at].args.Conditional.yes;
+//                at = Visit( at, program );
+//            }
+//            else if ( result==BoolExpression::BET_FALSE )
+//            {
+//                m_modified = true;
+//                at = program.m_code[at].args.Conditional.no;
+//                at = Visit( at, program );
+//            }
+//            else
+//            {
+//                // Recurse after modifying the current facts
+
+//                // TODO: Recurse condition too?
+
+//                // yes
+//                m_facts.push_back( pCond );
+//                args.Conditional.yes = Visit( args.Conditional.yes, program );
+//                m_facts.pop_back();
+
+//                // no
+//                BoolExpressionPtr pNot = new BoolExpression;
+//                pNot->type = BoolExpression::BET_NOT;
+//                pNot->a = pCond;
+
+//                m_facts.push_back( pNot );
+//                args.Conditional.no = Visit( args.Conditional.no, program );
+//                m_facts.pop_back();
+
+//                // update op if necessary
+//                if ( mutable_memcmp( &args, &program.m_code[at].args, sizeof(OP::ARGS) ) )
+//                {
+//                    OP op = program.m_code[at];
+//                    op.args = args;
+//                    at = program.AddOp( op );
+//                }
+//            }
+
+//            break;
+//        }
+
+////        case OP_TYPE::NU_SWITCH:
+////        case OP_TYPE::SC_SWITCH:
+////        case OP_TYPE::CO_SWITCH:
+////        case OP_TYPE::IM_SWITCH:
+////        case OP_TYPE::VO_SWITCH:
+////        case OP_TYPE::ME_SWITCH:
+////        case OP_TYPE::LA_SWITCH:
+////        {
+////            IntExpressionPtr pValue = GetIntExpression( args.Switch.variable, program );
+
+////            for ( int i=0;i<MUTABLE_OP_MAX_SWITCH_OPTIONS; ++i )
+////            {
+////                if ( args.Switch.values[i] )
+////                {
+////                    BoolExpressionPtr pCond = new BoolExpression;
+////                    pCond->type = BoolExpression::BET_INT_EQUAL_CONST;
+////                    pCond->intExp = pValue;
+////                    pCond->intValue = args.Switch.conditions[i];
+
+////                    BoolExpression::TYPE result = pCond->Evaluate( m_facts );
+
+////                    if ( result==BoolExpression::BET_TRUE )
+////                    {
+////                        // This is the only possible branch.
+////                        m_modified = true;
+////                        at = args.Switch.values[i];
+////                        at = Visit( at, program );
+////                        return at;
+////                        break;
+////                    }
+////                    else if ( result==BoolExpression::BET_FALSE )
+////                    {
+////                        // This option is impossible
+////                        args.Switch.values[i] = 0;
+////                        args.Switch.conditions[i] = 0;
+////                    }
+////                    else
+////                    {
+////                        // We don't know if it is true or false. Add its value as a fact and
+////                        // recurse.
+////                        m_facts.push_back( pCond );
+////                        args.Switch.values[i] = Visit( args.Switch.values[i], program );
+////                        m_facts.pop_back();
+////                    }
+////                }
+////            }
+
+////            if ( args.Switch.def )
+////            {
+////                // TODO: Add facts negating the options in the switch
+////                args.Switch.def = Visit( args.Switch.def, program );
+////            }
+
+////            // update op if necessary
+////            if ( mutable_memcmp( &args, &program.m_code[at].args, sizeof(OP::ARGS) ) )
+////            {
+////                m_modified = true;
+////                OP op = program.m_code[at];
+////                op.args = args;
+////                at = program.AddOp( op );
+////            }
+////            break;
+////        }
+
+//        default:
+//            // Normal recursion
+//            at = Recurse( at, program );
+//            break;
+
+//        }
+
+//        return at;
+//    }
+
+
+    //---------------------------------------------------------------------------------------------
+    //---------------------------------------------------------------------------------------------
+    //---------------------------------------------------------------------------------------------
+//    IntExpression::IntExpression()
+//    {
+//        type = IET_UNKNOWN;
+//        parameter = 0;
+//        value = 0;
+//    }
+
+
+//    //---------------------------------------------------------------------------------------------
+//    OP::ADDRESS IntExpression::GenerateCode( PROGRAM& program, PROGRAM::OP_CACHE* cache )
+//    {
+//        OP::ADDRESS at = 0;
+
+//        switch ( type )
+//        {
+
+//        case IET_CONSTANT:
+//        {
+//            OP op;
+//            op.type = OP_TYPE::NU_CONSTANT;
+//            op.args.IntConstant.value = value;
+//            at = program.AddOpUnique( op, cache );
+//            break;
+//        }
+
+//        case IET_PARAMETER:
+//        {
+//            OP op;
+//            op.type = OP_TYPE::NU_PARAMETER;
+//            op.args.Parameter.variable = parameter;
+//            at = program.AddOpUnique( op, cache );
+//            break;
+//        }
+
+//        default:
+//            check( false );
+//            break;
+
+//        }
+
+//        return at;
+//    }
+
+
+//    //---------------------------------------------------------------------------------------------
+//    //---------------------------------------------------------------------------------------------
+//    //---------------------------------------------------------------------------------------------
+//    BoolExpression::BoolExpression()
+//    {
+//        type = BET_UNKNOWN;
+//        parameter = 0;
+//        intValue = 0;
+//    }
+
+
+//    //---------------------------------------------------------------------------------------------
+//    BoolExpression::TYPE BoolExpression::Evaluate( vector<BoolExpressionPtr>& facts )
+//    {
+//        TYPE result = BET_UNKNOWN;
+
+//        switch ( type )
+//        {
+//        case BET_UNKNOWN:
+//            break;
+
+//        case BET_TRUE:
+//            result = BET_TRUE;
+//            break;
+
+//        case BET_FALSE:
+//            result = BET_FALSE;
+//            break;
+
+//        case BET_PARAMETER:
+//            for ( size_t f=0; f<facts.size(); ++f )
+//            {
+//                switch ( EvaluateParameter( facts[f], parameter ) )
+//                {
+//                case BET_TRUE:
+//                    return BET_TRUE;
+//                    break;
+
+//                case BET_FALSE:
+//                    return BET_FALSE;
+//                    break;
+
+//                default:
+//                    break;
+//                }
+//            }
+//            break;
+
+//        case BET_NOT:
+//            for ( size_t f=0; f<facts.size(); ++f )
+//            {
+//                result = a->Evaluate( facts );
+//                switch ( result )
+//                {
+//                case BET_TRUE:
+//                    return BET_FALSE;
+//                    break;
+
+//                case BET_FALSE:
+//                    return BET_TRUE;
+//                    break;
+
+//                default:
+//                    break;
+//                }
+//            }
+//            break;
+
+//        case BET_AND:
+//        {
+//            TYPE resultA = BET_UNKNOWN;
+//            TYPE resultB = BET_UNKNOWN;
+//            for ( size_t f=0; f<facts.size(); ++f )
+//            {
+//                if ( resultA == BET_UNKNOWN )
+//                {
+//                    resultA = a->Evaluate( facts );
+//                }
+//                if ( resultB == BET_UNKNOWN )
+//                {
+//                    resultB = b->Evaluate( facts );
+//                }
+//            }
+
+//            if ( resultA==BET_TRUE && resultB==BET_TRUE )
+//            {
+//                result = BET_TRUE;
+//            }
+//            else if ( resultA==BET_FALSE || resultB==BET_FALSE )
+//            {
+//                result = BET_FALSE;
+//            }
+//            break;
+//        }
+
+//        case BET_OR:
+//        {
+//            TYPE resultA = BET_UNKNOWN;
+//            TYPE resultB = BET_UNKNOWN;
+//            for ( size_t f=0; f<facts.size(); ++f )
+//            {
+//                if ( resultA == BET_UNKNOWN )
+//                {
+//                    resultA = a->Evaluate( facts );
+//                }
+//                if ( resultB == BET_UNKNOWN )
+//                {
+//                    resultB = b->Evaluate( facts );
+//                }
+//            }
+
+//            if ( resultA==BET_TRUE || resultB==BET_TRUE )
+//            {
+//                result = BET_TRUE;
+//            }
+//            else if ( resultA==BET_FALSE && resultB==BET_FALSE )
+//            {
+//                result = BET_FALSE;
+//            }
+//            break;
+//        }
+
+//        case BET_INT_EQUAL_CONST:
+//            result = EvaluateIntEquality( facts, intExp, intValue );
+//            break;
+
+//        default:
+//            check( false );
+//            break;
+//        }
+
+//        return result;
+//    }
+
+
+//    //---------------------------------------------------------------------------------------------
+//    OP::ADDRESS BoolExpression::GenerateCode( PROGRAM& program, PROGRAM::OP_CACHE* cache )
+//    {
+//        if (m_generated)
+//        {
+//            return m_generated;
+//        }
+
+//        OP::ADDRESS at = 0;
+
+//        switch ( type )
+//        {
+
+//        case BET_TRUE:
+//        {
+//            OP op;
+//            op.type = OP_TYPE::BO_CONSTANT;
+//            op.args.BoolConstant.value = true;
+//            at = program.AddOpUnique( op, cache );
+//            break;
+//        }
+
+//        case BET_FALSE:
+//        {
+//            OP op;
+//            op.type = OP_TYPE::BO_CONSTANT;
+//            op.args.BoolConstant.value = false;
+//            at = program.AddOpUnique( op, cache );
+//            break;
+//        }
+
+//        case BET_NOT:
+//        {
+//            OP op;
+//            op.type = OP_TYPE::BO_NOT;
+//            op.args.BoolNot.source = a->GenerateCode( program, cache );
+//            at = program.AddOpUnique( op, cache );
+//            break;
+//        }
+
+//        case BET_AND:
+//        {
+//            OP op;
+//            op.type = OP_TYPE::BO_AND;
+//            op.args.BoolBinary.a = a->GenerateCode( program, cache );
+//            op.args.BoolBinary.b = b->GenerateCode( program, cache );
+//            at = program.AddOpUnique( op, cache );
+//            break;
+//        }
+
+//        case BET_OR:
+//        {
+//            OP op;
+//            op.type = OP_TYPE::BO_OR;
+//            op.args.BoolBinary.a = a->GenerateCode( program, cache );
+//            op.args.BoolBinary.b = b->GenerateCode( program, cache );
+//            at = program.AddOpUnique( op, cache );
+//            break;
+//        }
+
+
+//        case BET_PARAMETER:
+//        {
+//            OP op;
+//            op.type = OP_TYPE::BO_PARAMETER;
+//            op.args.Parameter.variable = parameter;
+//            at = program.AddOpUnique( op, cache );
+//            break;
+//        }
+
+
+//        case BET_INT_EQUAL_CONST:
+//        {
+//            OP op;
+//            op.type = OP_TYPE::BO_EQUAL_INT_CONST;
+//            op.args.BoolEqualScalarConst.value = intExp->GenerateCode( program, cache );
+//            op.args.BoolEqualScalarConst.constant = (int16_t)intValue;
+//            at = program.AddOpUnique( op, cache );
+//            break;
+//        }
+
+
+//        default:
+//            check( false );
+//            break;
+
+//        }
+
+//        m_generated = at;
+
+//        return at;
+//    }
+
+
+//    //---------------------------------------------------------------------------------------------
+//    BoolExpression::TYPE BoolExpression::EvaluateParameter( BoolExpressionPtr fact, OP::ADDRESS parameter)
+//    {
+//        TYPE result = BET_UNKNOWN;
+
+//        switch (fact->type)
+//        {
+//        case BET_PARAMETER:
+//            if ( parameter==fact->parameter )
+//            {
+//                result = BET_TRUE;
+//            }
+//            break;
+
+//        case BET_NOT:
+//            result = EvaluateParameter( fact->a, parameter );
+//            if ( result == BET_TRUE )
+//            {
+//                result = BET_FALSE;
+//            }
+//            else if ( result == BET_FALSE )
+//            {
+//                result = BET_TRUE;
+//            }
+//            break;
+
+//        case BET_AND:
+//            result = EvaluateParameter( fact->a, parameter );
+//            if ( result == BET_UNKNOWN )
+//            {
+//                result = EvaluateParameter( fact->b, parameter );
+//            }
+//            break;
+
+//        case BET_OR:
+//        {
+//            TYPE resultA = EvaluateParameter( fact->a, parameter );
+//            TYPE resultB = EvaluateParameter( fact->b, parameter );
+//            if ( resultA==BET_TRUE && resultB==BET_TRUE )
+//            {
+//                result = BET_TRUE;
+//            }
+//            else if ( resultA==BET_FALSE && resultB==BET_FALSE )
+//            {
+//                result = BET_FALSE;
+//            }
+//            else
+//            {
+//                // TODO: Check true+false value
+//                result = BET_UNKNOWN;
+//            }
+//            break;
+//        }
+
+//        case BET_INT_EQUAL_CONST:
+//        case BET_UNKNOWN:
+//        case BET_TRUE:
+//        case BET_FALSE:
+//            break;
+
+//        default:
+//            check( false);
+//            break;
+
+//        }
+
+//        return result;
+//    }
+
+
+//    //---------------------------------------------------------------------------------------------
+//    BoolExpression::TYPE BoolExpression::EvaluateIntEquality
+//        ( const vector<BoolExpressionPtr>& facts, IntExpressionPtr intExp, int value )
+//    {
+//        TYPE result = BET_UNKNOWN;
+
+//        switch (intExp->type)
+//        {
+//        case IntExpression::IET_CONSTANT:
+//            if ( value==intExp->value )
+//            {
+//                result = BET_TRUE;
+//            }
+//            else
+//            {
+//                result = BET_FALSE;
+//            }
+//            break;
+
+//        case IntExpression::IET_PARAMETER:
+//            for ( size_t f=0; (result==BET_UNKNOWN) && f<facts.size(); ++f )
+//            {
+//                result = EvaluateIntEquality(facts[f].get(),intExp,value);
+//            }
+//            break;
+
+//        default:
+//            break;
+
+//        }
+
+//        return result;
+//    }
+
+
+//    //---------------------------------------------------------------------------------------------
+//    BoolExpression::TYPE BoolExpression::EvaluateIntEquality
+//        ( const BoolExpression* fact, IntExpressionPtr intExp, int value )
+//    {
+//        TYPE result = BET_UNKNOWN;
+
+//        if( fact->type == BET_INT_EQUAL_CONST )
+//        {
+//            if ( fact->intExp->type == IntExpression::IET_PARAMETER
+//                &&
+//                fact->intExp->parameter == intExp->parameter
+//                )
+//            {
+//                if ( fact->intValue == value )
+//                {
+//                    result = BET_TRUE;
+//                }
+//                else
+//                {
+//                    result = BET_FALSE;
+//                }
+//            }
+//        }
+
+//        else if( fact->type == BET_NOT )
+//        {
+//            if (fact->a->type==BET_INT_EQUAL_CONST)
+//            {
+//                TYPE notResult = EvaluateIntEquality( fact->a.get(), intExp, value );
+//                if (notResult==BET_TRUE)
+//                {
+//                    result = BET_FALSE;
+//                }
+//                else
+//                {
+//                    // We cannot know, because there are other values.
+//                }
+//            }
+//            else
+//            {
+//                MUTABLE_CPUPROFILER_SCOPE(Dummy);
+//            }
+//        }
+
+//        else
+//        {
+//            MUTABLE_CPUPROFILER_SCOPE(Dummy);
+//        }
+
+//        return result;
+//    }
+
+
+    //---------------------------------------------------------------------------------------------
+    //---------------------------------------------------------------------------------------------
+    //---------------------------------------------------------------------------------------------
+    bool LocalLogicOptimiserAST(ASTOpList& roots)
+    {
+        MUTABLE_CPUPROFILER_SCOPE(LocalLogicOptimiserAST);
+
+        bool modified = false;
+
+        // The separate steps should not be combined into one traversal
+
+        // Unwrap some typical code daisy-chains
+        //-----------------------------------------------------------------------------------------
+        {
+            MUTABLE_CPUPROFILER_SCOPE(Unwrap);
+            ASTOp::Traverse_TopDown_Unique_Imprecise( roots, [&](Ptr<ASTOp>& n)
+            {
+                bool recurse = true;
+
+                if (n->GetOpType()==OP_TYPE::IN_CONDITIONAL)
+                {
+                    auto topConditional = dynamic_cast<ASTOpConditional*>(n.get());
+                    auto yes = topConditional->yes.child();
+                    if (yes->GetOpType()==OP_TYPE::IN_ADDSURFACE)
+                    {
+                        auto addSurface = dynamic_cast<ASTOpInstanceAdd*>(yes.get());
+
+                        bool ended = false;
+                        while (!ended)
+                        {
+                            ended = true;
+                            auto base = addSurface->instance.child();
+                            if ( base && base->GetOpType()==OP_TYPE::IN_CONDITIONAL )
+                            {
+                                auto bottomConditional = dynamic_cast<ASTOpConditional*>(base.get());
+
+                                // Are the two conditions exclusive?
+                                bool conditionaAreExclusive = false;
+
+                                ASTOpList facts;
+                                facts.push_back(topConditional->condition.child());
+
+                                // Check if the child condition has a value with the current facts
+                                const auto& pChildCond = bottomConditional->condition.child();
+                                ASTOp::BOOL_EVAL_RESULT result;
+                                {
+                                    //MUTABLE_CPUPROFILER_SCOPE(EvaluateBool);
+                                    result = pChildCond->EvaluateBool( facts );
+                                }
+                                if ( result==ASTOp::BET_FALSE )
+                                {
+                                    conditionaAreExclusive = true;
+                                }
+
+                                if (conditionaAreExclusive)
+                                {
+                                    if (addSurface->GetParentCount()==1)
+                                    {
+                                        // Directly modify the instruction to skip the impossible child option
+                                        addSurface->instance = bottomConditional->no.child();
+                                    }
+                                    else
+                                    {
+                                        // Other parents may not impose the same condition that allows
+                                        // the optimisation.
+                                        auto newAddSurface = mu::Clone<ASTOpInstanceAdd>(addSurface);
+                                        newAddSurface->instance = bottomConditional->no.child();
+                                        topConditional->yes = newAddSurface;
+                                    }
+
+                                    modified = true;
+                                    ended = false;
+                                }
+                            }
+                        }
+                    }
+                }
+
+                else if (n->GetOpType()==OP_TYPE::ME_CONDITIONAL)
+                {
+                    auto topConditional = dynamic_cast<ASTOpConditional*>(n.get());
+                    auto yes = topConditional->yes.child();
+                    if ( yes->GetOpType()==OP_TYPE::ME_MERGE )
+                    {
+                        Ptr<ASTOpFixed> merge = dynamic_cast<ASTOpFixed*>(yes.get());
+
+                        bool ended = false;
+                        while (!ended)
+                        {
+                            ended = true;
+                            auto base = merge->children[merge->op.args.MeshMerge.base].child();
+                            if ( base && base->GetOpType()==OP_TYPE::ME_CONDITIONAL )
+                            {
+                                auto bottomConditional = dynamic_cast<ASTOpConditional*>(base.get());
+
+                                // Are the two conditions exclusive?
+                                bool conditionaAreExclusive = false;
+
+                                ASTOpList facts;
+                                facts.push_back(topConditional->condition.child());
+
+                                // Check if the child condition has a value with the current facts
+                                auto pChildCond = bottomConditional->condition.child();
+                                ASTOp::BOOL_EVAL_RESULT result = pChildCond->EvaluateBool( facts );
+                                if ( result==ASTOp::BET_FALSE )
+                                {
+                                    conditionaAreExclusive = true;
+                                }
+
+                                if (conditionaAreExclusive)
+                                {
+                                    if (merge->GetParentCount()==1)
+                                    {
+                                        // Directly modify the instruction to skip the impossible child option
+                                        merge->SetChild(merge->op.args.MeshMerge.base, bottomConditional->no.child() );
+                                    }
+                                    else
+                                    {
+                                        // Other parents may not impose the same condition that allows
+                                        // the optimisation.
+                                        auto newMerge = mu::Clone<ASTOpFixed>(merge);
+                                        newMerge->SetChild(newMerge->op.args.MeshMerge.base, bottomConditional->no.child());
+                                        topConditional->yes = newMerge;
+                                        merge = newMerge;
+                                    }
+
+                                    modified = true;
+                                    ended = false;
+                                }
+                            }
+                        }
+                    }
+
+// todo, since the move to aggregated remove mask op
+//                    else if ( yes->GetOpType()==OP_TYPE::ME_REMOVEMASK )
+//                    {
+//                        Ptr<ASTOpFixed> remove = dynamic_cast<ASTOpFixed*>(yes.get());
+
+//                        bool ended = false;
+//                        while (!ended)
+//                        {
+//                            ended = true;
+//                            auto base = remove->children[remove->op.args.MeshRemoveMask.source].child();
+//                            if ( base && base->GetOpType()==OP_TYPE::ME_CONDITIONAL )
+//                            {
+//                                auto bottomConditional = dynamic_cast<ASTOpConditional*>(base.get());
+
+//                                // Are the two conditions exclusive?
+//                                bool conditionaAreExclusive = false;
+
+//                                ASTOpList facts;
+//                                facts.push_back(topConditional->condition.child());
+
+//                                // Check if the child condition has a value with the current facts
+//                                auto pChildCond = bottomConditional->condition.child();
+//                                ASTOp::BOOL_EVAL_RESULT result = pChildCond->EvaluateBool( facts );
+//                                if ( result==ASTOp::BET_FALSE )
+//                                {
+//                                    conditionaAreExclusive = true;
+//                                }
+
+//                                if (conditionaAreExclusive)
+//                                {
+//                                    if (remove->GetParentCount()==1)
+//                                    {
+//                                        // Directly modify the instruction to skip the impossible child option
+//                                        remove->SetChild(remove->op.args.MeshRemoveMask.source, bottomConditional->no.child() );
+//                                    }
+//                                    else
+//                                    {
+//                                        // Other parents may not impose the same condition that allows
+//                                        // the optimisation.
+//                                        auto newRemove = mu::Clone<ASTOpFixed>(remove);
+//                                        newRemove->SetChild(newRemove->op.args.MeshRemoveMask.source, bottomConditional->no.child());
+//                                        topConditional->yes = newRemove;
+//                                    }
+
+//                                    modified = true;
+//                                    ended = false;
+//                                }
+//                            }
+//                        }
+//                    }
+                }
+
+
+                return recurse;
+            });
+        }
+
+
+        // See if we can turn conditional chains into switches: all conditions must be integer
+        // comparison with the same variable.
+        //-----------------------------------------------------------------------------------------
+        {
+            MUTABLE_CPUPROFILER_SCOPE(ConditionalToSwitch);
+            ASTOp::Traverse_TopDown_Unique_Imprecise( roots, [&](Ptr<ASTOp>& n)
+            {
+                bool recurse = true;
+
+                auto topConditional = dynamic_cast<ASTOpConditional*>(n.get());
+                if (topConditional && topConditional->condition)
+                {
+
+                    if ( topConditional->condition->GetOpType()==OP_TYPE::BO_EQUAL_INT_CONST
+                         &&
+                         topConditional->no
+                         &&
+                         topConditional->no->GetOpType()==topConditional->GetOpType()
+                         )
+                    {
+                        Ptr<ASTOpSwitch> switchOp = new ASTOpSwitch();
+                        switchOp->type = GetSwitchForType(GetOpDataType(topConditional->GetOpType()));
+
+                        auto firstCompare = dynamic_cast<ASTOpFixed*>(topConditional->condition.child().get());
+                        switchOp->variable = firstCompare->children[firstCompare->op.args.BoolEqualScalarConst.value].child();
+
+                        auto current = n;
+                        while(current)
+                        {
+                            bool valid = false;
+
+                            auto conditional = dynamic_cast<ASTOpConditional*>(current.get());
+                            if ( conditional
+                                 &&
+                                 conditional->GetOpType()==topConditional->GetOpType()
+                                 &&
+                                 conditional->condition
+                                 &&
+                                 conditional->condition->GetOpType()==OP_TYPE::BO_EQUAL_INT_CONST )
+                            {
+                                auto compare = dynamic_cast<ASTOpFixed*>(conditional->condition.child().get());
+                                check(compare);
+                                if (compare)
+                                {
+                                    auto compareValue = compare->children[compare->op.args.BoolEqualScalarConst.value].child();
+                                    check(compare);
+                                    if ( compareValue == switchOp->variable.child() )
+                                    {
+                                        switchOp->cases.emplace_back( compare->op.args.BoolEqualScalarConst.constant,
+                                                                      switchOp,
+                                                                      conditional->yes.child() );
+
+                                        current = conditional->no.child();
+                                        valid = true;
+                                    }
+                                }
+                            }
+
+                            if (!valid)
+                            {
+                                switchOp->def = current;
+                                current = nullptr;
+                            }
+                        }
+
+                        const int MIN_CONDITIONS_TO_CREATE_SWITCH = 3;
+                        if (switchOp->cases.size()>=MIN_CONDITIONS_TO_CREATE_SWITCH)
+                        {
+                            ASTOp::Replace(n,switchOp);
+                            n = switchOp;
+                            modified = true;
+                        }
+                    }
+
+                }
+
+                return recurse;
+            });
+        }
+
+
+        // Float operations up switches, to tidy up the code and reduce its size.
+        //-----------------------------------------------------------------------------------------
+        {
+            MUTABLE_CPUPROFILER_SCOPE(FloatSwitches);
+
+            ASTOp::Traverse_TopDown_Unique_Imprecise( roots, [&](Ptr<ASTOp>& n)
+            {
+                bool recurse = true;
+
+                auto topSwitch = dynamic_cast<ASTOpSwitch*>(n.get());
+                if (topSwitch)
+                {
+                    bool first = true;
+                    auto caseType = OP_TYPE::NONE;
+                    for (const auto& c:topSwitch->cases)
+                    {
+                        if ( c.branch )
+                        {
+                            if (first)
+                            {
+                                first = false;
+                                caseType = c.branch->GetOpType();
+                            }
+                            else
+                            {
+                                if (caseType!=c.branch->GetOpType())
+                                {
+                                    caseType=OP_TYPE::NONE;
+                                    break;
+                                }
+                            }
+                        }
+                    }
+
+                    // todo, since the move to aggregated remove mask op
+//                    if (caseType==OP_TYPE::ME_REMOVEMASK)
+//                    {
+//                        // Is the source of all remove operations the default option in the switch?
+//                        auto source = topSwitch->def.child();
+//                        for (const auto& c:topSwitch->cases)
+//                        {
+//                            if ( c.branch )
+//                            {
+//                                auto typedRemove = dynamic_cast<const ASTOpFixed*>(c.branch.child().get());
+//                                check( typedRemove );
+//                                if (typedRemove->children[typedRemove->op.args.MeshRemoveMask.source].child()!=source)
+//                                {
+//                                    source = nullptr;
+//                                    break;
+//                                }
+//                            }
+//                        }
+
+//                        if (source)
+//                        {
+//                            auto newSwitch = mu::Clone<ASTOpSwitch>(topSwitch);
+//                            for (auto& c: newSwitch->cases)
+//                            {
+//                                if ( c.branch )
+//                                {
+//                                    auto typedRemove = dynamic_cast<const ASTOpFixed*>(c.branch.child().get());
+//                                    c.branch = typedRemove->children[typedRemove->op.args.MeshRemoveMask.mask].child();
+//                                }
+//                            }
+//                            newSwitch->def = nullptr;
+
+//                            Ptr<ASTOpFixed> newRemove = new ASTOpFixed();
+//                            newRemove->op.type = OP_TYPE::ME_REMOVEMASK;
+//                            newRemove->SetChild( newRemove->op.args.MeshRemoveMask.mask, newSwitch );
+//                            newRemove->SetChild( newRemove->op.args.MeshRemoveMask.source, source);
+
+//                            ASTOp::Replace(n,newRemove);
+//                            n = newRemove;
+//                            modified = true;
+//                        }
+//                    }
+
+                }
+
+                return recurse;
+            });
+        }
+
+        return modified;
+    }
+
+}
+
