[mutable] Move Mutable plugin from the restricted folder to the experimental plugins folder.

diff --git a/Engine/Plugins/Experimental/Mutable/Source/MutableTools/Private/CodeOptimiser_Parameters.cpp b/Engine/Plugins/Experimental/Mutable/Source/MutableTools/Private/CodeOptimiser_Parameters.cpp
new file mode 100644
--- /dev/null
+++ b/Engine/Plugins/Experimental/Mutable/Source/MutableTools/Private/CodeOptimiser_Parameters.cpp
@@ -0,0 +1,2227 @@
+// Copyright Epic Games, Inc. All Rights Reserved.
+
+#include "CodeOptimiser.h"
+#include "CodeOptimiser_GPU.h"
+#include "NodeObjectNewPrivate.h"
+#include "CodeGenerator.h"
+#include "DataPacker.h"
+#include "ModelReportPrivate.h"
+#include "CompilerPrivate.h"
+
+#include "CodeVisitor.h"
+#include "ModelPrivate.h"
+#include "SystemPrivate.h"
+#include "Operations.h"
+
+#include "MutableTools/Private/ASTOpImagePixelFormat.h"
+#include "MutableTools/Private/ASTOpImageMipmap.h"
+#include "MutableTools/Private/ASTOpImageCompose.h"
+#include "MutableTools/Private/ASTOpAddLOD.h"
+#include "MutableTools/Private/ASTOpSwitch.h"
+#include "MutableTools/Private/ASTOpConditional.h"
+#include "MutableTools/Private/ASTOpInstanceAdd.h"
+
+#include <array>
+
+#include "MutableTrace.h"
+
+
+namespace mu
+{
+
+
+    //---------------------------------------------------------------------------------------------
+    //---------------------------------------------------------------------------------------------
+    //---------------------------------------------------------------------------------------------
+    RuntimeParameterVisitorAST::RuntimeParameterVisitorAST(const STATE_COMPILATION_DATA* pState)
+        : m_pState(pState)
+    {
+    }
+
+
+    bool RuntimeParameterVisitorAST::HasAny( const Ptr<ASTOp>& root )
+    {
+        if (!m_pState->nodeState.m_runtimeParams.size())
+        {
+            return false;
+        }
+
+        // Shortcut flag: if true we already found a runtime parameter, don't process new ops,
+        // but still store the results of processed ops.
+        bool found = false;
+
+        m_pending.clear();
+
+        PENDING_ITEM start;
+        start.at = root;
+        start.itemType = 0;
+        start.onlyLayoutsRelevant = 0;
+        m_pending.push_back( start );
+
+        // Don't early out to be able to complete parent op cached flags
+        while ( m_pending.size() )
+        {
+            PENDING_ITEM item = m_pending.back();
+            m_pending.pop_back();
+            Ptr<ASTOp> at = item.at;
+
+			if (!at)
+			{
+				continue;
+			}
+
+            // Not cached?
+            if ( m_visited[at]!=OP_STATE::VISITED_HASRUNTIME
+                 &&
+                 m_visited[at]!=OP_STATE::VISITED_FULL_DOESNTHAVERUNTIME )
+            {
+                if (item.itemType)
+                {
+                    // Item indicating we finished with all the children of a parent
+                    check( m_visited[at]==OP_STATE::CHILDREN_PENDING_FULL
+                                    ||
+                                    m_visited[at]==OP_STATE::CHILDREN_PENDING_PARTIAL
+                                    ||
+                                    m_visited[at]==OP_STATE::VISITED_PARTIAL_DOESNTHAVERUNTIME );
+
+                    bool subtreeFound = false;
+                    at->ForEachChild( [&](ASTChild& ref)
+                    {
+                        subtreeFound = subtreeFound || m_visited[ref.child()]==OP_STATE::VISITED_HASRUNTIME;
+                    });
+
+                    if (subtreeFound)
+                    {
+                        m_visited[at] = OP_STATE::VISITED_HASRUNTIME;
+                    }
+                    else
+                    {
+                        m_visited[at] = item.onlyLayoutsRelevant
+                                ? OP_STATE::VISITED_PARTIAL_DOESNTHAVERUNTIME
+                                : OP_STATE::VISITED_FULL_DOESNTHAVERUNTIME;
+                    }
+                }
+
+                else if (!found)
+                {
+                    // We need to process the subtree
+                    check( m_visited[at]==OP_STATE::NOT_VISITED
+                                    ||
+                                    ( m_visited[at]==OP_STATE::VISITED_PARTIAL_DOESNTHAVERUNTIME
+                                      &&
+                                      item.onlyLayoutsRelevant==0 )
+                                    );
+
+                    // Request the processing of the end of this instruction
+                    PENDING_ITEM endItem = item;
+                    endItem.itemType = 1;
+                    m_pending.push_back( endItem );
+                    m_visited[at] = item.onlyLayoutsRelevant
+                            ? OP_STATE::CHILDREN_PENDING_PARTIAL
+                            : OP_STATE::CHILDREN_PENDING_FULL;
+
+                    // Is it a special op type?
+                    switch ( at->GetOpType() )
+                    {
+
+                    case OP_TYPE::BO_PARAMETER:
+                    case OP_TYPE::NU_PARAMETER:
+                    case OP_TYPE::SC_PARAMETER:
+                    case OP_TYPE::CO_PARAMETER:
+                    case OP_TYPE::PR_PARAMETER:
+                    case OP_TYPE::IM_PARAMETER:
+                    {
+                        auto typed = dynamic_cast<const ASTOpParameter*>(at.get());
+                        const auto& params = m_pState->nodeState.m_runtimeParams;
+                        if ( std::find(params.cbegin(), params.cend(), typed->parameter.m_name)
+                             !=
+                             params.cend() )
+                        {
+                            found = true;
+                            m_visited[at] = OP_STATE::VISITED_HASRUNTIME;
+                        }
+                        break;
+                    }
+
+                    case OP_TYPE::ME_INTERPOLATE:
+                    {
+                        auto typed = dynamic_cast<const ASTOpFixed*>(at.get());
+                        PENDING_ITEM childItem;
+                        childItem.itemType = 0;
+                        childItem.onlyLayoutsRelevant = item.onlyLayoutsRelevant;
+
+                        if ( item.onlyLayoutsRelevant==0 )
+                        {
+                            childItem.at = typed->children[typed->op.args.MeshInterpolate.factor].child();
+                            AddIfNeeded(childItem);
+                        }
+
+                        childItem.at = typed->children[typed->op.args.MeshInterpolate.base].child();
+                        AddIfNeeded(childItem);
+
+                        for (int t=0;t<MUTABLE_OP_MAX_INTERPOLATE_COUNT-1;++t)
+                        {
+                            childItem.at = typed->children[typed->op.args.MeshInterpolate.targets[t]].child();
+                            AddIfNeeded(childItem);
+                        }
+
+                        break;
+                    }
+
+                    default:
+                    {
+                        at->ForEachChild([&](ASTChild& ref)
+                        {
+                            PENDING_ITEM childItem;
+                            childItem.itemType = 0;
+                            childItem.at = ref.child();
+                            childItem.onlyLayoutsRelevant = item.onlyLayoutsRelevant;
+                            AddIfNeeded(childItem);
+                        });
+                        break;
+                    }
+
+                    }
+
+                }
+
+                else
+                {
+                    // We won't process it.
+                    m_visited[at] = OP_STATE::NOT_VISITED;
+                }
+            }
+        }
+
+        return m_visited[root]==OP_STATE::VISITED_HASRUNTIME;
+    }
+
+
+    void RuntimeParameterVisitorAST::AddIfNeeded( const PENDING_ITEM& item )
+    {
+        if (item.at)
+        {
+            if (m_visited[item.at]==OP_STATE::NOT_VISITED)
+            {
+                m_pending.push_back( item );
+            }
+            else if (m_visited[item.at]==OP_STATE::VISITED_PARTIAL_DOESNTHAVERUNTIME
+                     &&
+                     item.onlyLayoutsRelevant==0)
+            {
+                m_pending.push_back( item );
+            }
+            else if (m_visited[item.at]==OP_STATE::CHILDREN_PENDING_PARTIAL
+                     &&
+                     item.onlyLayoutsRelevant==0)
+            {
+                m_pending.push_back( item );
+            }
+        }
+    }
+
+
+    //---------------------------------------------------------------------------------------------
+    //---------------------------------------------------------------------------------------------
+    //---------------------------------------------------------------------------------------------
+    Ptr<ASTOp> EnsureValidMask( Ptr<ASTOp> mask, Ptr<ASTOp> base )
+    {
+        if ( !mask )
+        {
+            Ptr<ASTOpFixed> whiteOp = new ASTOpFixed;
+            whiteOp->op.type = OP_TYPE::CO_CONSTANT;
+            whiteOp->op.args.ColourConstant.value[0] = 1;
+            whiteOp->op.args.ColourConstant.value[1] = 1;
+            whiteOp->op.args.ColourConstant.value[2] = 1;
+
+            Ptr<ASTOpFixed> wplainOp = new ASTOpFixed;
+            wplainOp->op.type = OP_TYPE::IM_PLAINCOLOUR;
+            wplainOp->SetChild( wplainOp->op.args.ImagePlainColour.colour, whiteOp );
+            wplainOp->op.args.ImagePlainColour.format = EImageFormat::IF_L_UBYTE;
+            wplainOp->op.args.ImagePlainColour.size[0] = 4;
+            wplainOp->op.args.ImagePlainColour.size[1] = 4;
+
+            Ptr<ASTOpFixed> wresizeOp = new ASTOpFixed;
+            wresizeOp->op.type = OP_TYPE::IM_RESIZELIKE;
+            wresizeOp->SetChild( wresizeOp->op.args.ImageResizeLike.source, wplainOp );
+            wresizeOp->SetChild( wresizeOp->op.args.ImageResizeLike.sizeSource, base );
+
+            mask = wresizeOp;
+        }
+
+        return mask;
+    }
+
+
+    //---------------------------------------------------------------------------------------------
+    //---------------------------------------------------------------------------------------------
+    //---------------------------------------------------------------------------------------------
+    ParameterOptimiserAST::ParameterOptimiserAST(
+            STATE_COMPILATION_DATA& s,
+            const MODEL_OPTIMIZATION_OPTIONS& optimisationOptions
+            )
+            : m_stateProps(s)
+            , m_modified(false)
+            , m_optimisationOptions(optimisationOptions)
+            , m_hasRuntimeParamVisitor(&s)
+    {
+    }
+
+
+    bool ParameterOptimiserAST::Apply()
+    {
+        MUTABLE_CPUPROFILER_SCOPE(ParameterOptimiserAST);
+
+        m_modified = false;
+
+        // Optimise the cloned tree
+        Traverse( m_stateProps.root );
+
+        return m_modified;
+    }
+
+
+    Ptr<ASTOp> ParameterOptimiserAST::Visit( Ptr<ASTOp> at, bool& processChildren )
+    {
+        // Only process children if there are runtime parameters in the subtree
+        processChildren = m_hasRuntimeParamVisitor.HasAny(at);
+
+        OP_TYPE type = at->GetOpType();
+        switch ( type )
+        {
+        //-------------------------------------------------------------------------------------
+        // Be careful with changing merge options and "mergesurfaces" flags
+//        case OP_TYPE::ME_MERGE:
+//        {
+//            OP::MeshMergeArgs mergeArgs = program.m_code[at].args.MeshMerge;
+
+//            RuntimeParameterVisitor paramVis;
+
+//            switch ( program.m_code[ mergeArgs.base ].type )
+//            {
+//            case OP_TYPE::ME_CONDITIONAL:
+//            {
+//                OP::ADDRESS conditionAt =
+//                        program.m_code[ mergeArgs.base ].args.Conditional.condition;
+//                bool conditionConst = !paramVis.HasAny( m_pModel.get(),
+//                                                        m_state,
+//                                                        conditionAt );
+//                if ( conditionConst )
+//                {
+//                    // TODO: this may unfold mesh combinations of some models increasing the size of
+//                    // the model data. Make this optimisation optional.
+//                    m_modified = true;
+
+//                    OP yesOp = program.m_code[at];
+//                    yesOp.args.MeshMerge.base = program.m_code[ mergeArgs.base ].args.Conditional.yes;
+
+//                    OP noOp = program.m_code[at];
+//                    noOp.args.MeshMerge.base = program.m_code[ mergeArgs.base ].args.Conditional.no;
+
+//                    OP op = program.m_code[ mergeArgs.base ];
+//                    op.args.Conditional.yes = program.AddOp( yesOp );
+//                    op.args.Conditional.no = program.AddOp( noOp );
+//                    at = program.AddOp( op );
+//                }
+
+//                break;
+//            }
+
+//            case OP_TYPE::ME_MERGE:
+//            {
+//                OP::ADDRESS childBaseAt = program.m_code[ mergeArgs.base ].args.MeshMerge.base;
+//                bool childBaseConst = !paramVis.HasAny( m_pModel.get(),
+//                                                        m_state,
+//                                                        childBaseAt );
+
+//                OP::ADDRESS childAddAt = program.m_code[ mergeArgs.base ].args.MeshMerge.added;
+//                bool childAddConst = !paramVis.HasAny( m_pModel.get(),
+//                                                       m_state,
+//                                                       childAddAt );
+
+//                bool addConst = !paramVis.HasAny( m_pModel.get(),
+//                                                  m_state,
+//                                                  mergeArgs.added );
+
+//                if ( !childBaseConst && childAddConst && addConst )
+//                {
+//                    m_modified = true;
+
+//                    OP bottom = program.m_code[at];
+//                    bottom.args.MeshMerge.base = childAddAt;
+
+//                    OP top = program.m_code[ mergeArgs.base ];
+//                    top.args.MeshMerge.added = program.AddOp( bottom );
+//                    at = program.AddOp( top );
+//                }
+//                break;
+//            }
+
+//            default:
+//                break;
+//            }
+
+//            break;
+//        }
+
+        //-------------------------------------------------------------------------------------
+        //-------------------------------------------------------------------------------------
+        //-------------------------------------------------------------------------------------
+        case OP_TYPE::IM_CONDITIONAL:
+        {
+            auto typedAt = dynamic_cast<const ASTOpConditional*>(at.get());
+
+            // If the condition is not runtime, but the branches are, try to move the
+            // conditional down
+            bool optimised = false;
+
+            if ( !m_hasRuntimeParamVisitor.HasAny( typedAt->condition.child() ) )
+            {
+                OP_TYPE yesType = typedAt->yes->GetOpType();
+                OP_TYPE noType = typedAt->no->GetOpType();
+
+                bool yesHasAny = m_hasRuntimeParamVisitor.HasAny( typedAt->yes.child() );
+                bool noHasAny = m_hasRuntimeParamVisitor.HasAny( typedAt->no.child() );
+
+                if ( !optimised && yesHasAny && noHasAny && yesType==noType)
+                {
+                    switch (yesType)
+                    {
+                        case OP_TYPE::IM_COMPOSE:
+                        {
+						const ASTOpImageCompose* typedYes = dynamic_cast<const ASTOpImageCompose*>(typedAt->yes.child().get());
+						const ASTOpImageCompose* typedNo = dynamic_cast<const ASTOpImageCompose*>(typedAt->no.child().get());
+                        if ( typedYes->BlockIndex
+                             ==
+                             typedNo->BlockIndex
+                             &&
+                             (
+                                 (typedYes->Mask.child().get() != nullptr)
+                                  ==
+                                 (typedNo->Mask.child().get() != nullptr)
+                                )
+                             )
+                            {
+                                // Move the conditional down
+                                auto compOp = mu::Clone<ASTOpImageCompose>(typedYes);
+
+                                auto baseCond = mu::Clone<ASTOpConditional>(typedAt);
+                                baseCond->yes = typedYes->Base.child();
+                                baseCond->no = typedNo->Base.child();
+                                compOp->Base = baseCond;
+
+                                auto blockCond = mu::Clone<ASTOpConditional>(typedAt);
+                                blockCond->yes = typedYes->BlockImage.child();
+                                blockCond->no = typedNo->BlockImage.child();
+                                compOp->BlockImage = blockCond;
+
+                                if (typedYes->Mask)
+                                {
+                                    auto maskCond = mu::Clone<ASTOpConditional>(typedAt);
+                                    maskCond->yes = typedYes->Mask.child();
+                                    maskCond->no = typedNo->Mask.child();
+                                    compOp->Mask = maskCond;
+                                }
+
+                                auto layCond = mu::Clone<ASTOpConditional>(typedAt);
+                                layCond->type = OP_TYPE::LA_CONDITIONAL;
+                                layCond->yes = typedYes->Layout.child();
+                                layCond->no = typedNo->Layout.child();
+                                compOp->Layout = layCond;
+
+
+                                at = compOp;
+                                optimised = true;
+                            }
+                            break;
+                        }
+
+                    default:
+                        break;
+
+                    }
+                }
+
+                if ( !optimised && yesHasAny )
+                {
+                    switch (yesType)
+                    {
+                        case OP_TYPE::IM_LAYERCOLOUR:
+                        {
+                            optimised = true;
+
+                            auto typedYes = dynamic_cast<const ASTOpFixed*>(typedAt->yes.child().get());
+
+                            Ptr<ASTOpFixed> blackOp = new ASTOpFixed;
+                            blackOp->op.type = OP_TYPE::CO_CONSTANT;
+                            blackOp->op.args.ColourConstant.value[0] = 0;
+                            blackOp->op.args.ColourConstant.value[1] = 0;
+                            blackOp->op.args.ColourConstant.value[2] = 0;
+
+                            Ptr<ASTOpFixed> plainOp = new ASTOpFixed;
+                            plainOp->op.type = OP_TYPE::IM_PLAINCOLOUR;
+                            plainOp->SetChild( plainOp->op.args.ImagePlainColour.colour, blackOp );
+                            plainOp->op.args.ImagePlainColour.format = EImageFormat::IF_L_UBYTE;
+                            plainOp->op.args.ImagePlainColour.size[0] = 4;
+                            plainOp->op.args.ImagePlainColour.size[1] = 4;
+
+                            Ptr<ASTOpFixed> resizeOp = new ASTOpFixed;
+                            resizeOp->op.type = OP_TYPE::IM_RESIZELIKE;
+                            resizeOp->SetChild( resizeOp->op.args.ImageResizeLike.source, plainOp );
+                            resizeOp->SetChild( resizeOp->op.args.ImageResizeLike.sizeSource,
+                                                typedYes->children[ typedYes->op.args.ImageLayerColour.base ] );
+
+                            auto maskOp = mu::Clone<ASTOpConditional>(typedAt);
+                            maskOp->no = resizeOp;
+
+                            // If there is no mask (because it is optional), we need to make a
+                            // white plain image
+                            maskOp->yes = EnsureValidMask
+                                    ( typedYes->children[ typedYes->op.args.ImageLayerColour.mask ].child(),
+                                      typedYes->children[ typedYes->op.args.ImageLayerColour.base ].child() );
+
+                            auto baseOp = mu::Clone<ASTOpConditional>(typedAt);
+                            baseOp->yes = typedYes->children[ typedYes->op.args.ImageLayerColour.base ].child();
+
+                            auto softOp = mu::Clone<ASTOpFixed>(typedYes);
+                            softOp->SetChild( softOp->op.args.ImageLayerColour.base, baseOp );
+                            softOp->SetChild( softOp->op.args.ImageLayerColour.mask, maskOp );
+
+                            at = softOp;
+                            break;
+                        }
+
+                        // TODO
+                        // It seems this is not worth since it replaces a conditional by a compose
+                        // (but only at build time, not update?) and it introduces the use of masks
+                        // and resize likes... plus masks can't always be used if BC formats.
+//						case OP_TYPE::IM_COMPOSE:
+//						{
+//							optimised = true;
+
+//							OP blackOp;
+//							blackOp.type = OP_TYPE::CO_CONSTANT;
+//							blackOp.args.ColourConstant.value[0] = 0;
+//							blackOp.args.ColourConstant.value[1] = 0;
+//							blackOp.args.ColourConstant.value[2] = 0;
+
+//							OP plainOp;
+//							plainOp.type = OP_TYPE::IM_PLAINCOLOUR;
+//							plainOp.args.ImagePlainColour.colour = program.AddOp( blackOp );
+//							plainOp.args.ImagePlainColour.format = IF_L_UBYTE;
+
+//							OP resizeOp;
+//							resizeOp.type = OP_TYPE::IM_RESIZELIKE;
+//							resizeOp.args.ImageResizeLike.source = program.AddOp( plainOp );
+//							resizeOp.args.ImageResizeLike.sizeSource =
+//									program.m_code[args.yes].args.ImageCompose.blockImage;
+
+//							OP maskOp = program.m_code[at];
+//							maskOp.args.Conditional.no = program.AddOp( resizeOp );
+
+//							// If there is no mask (because it is optional), we need to make a
+//							// white plain image
+//							maskOp.args.Conditional.yes = EnsureValidMask
+//									( program.m_code[args.yes].args.ImageCompose.mask,
+//									  program.m_code[args.yes].args.ImageCompose.base,
+//									  program );
+
+//							OP baseOp = program.m_code[at];
+//							baseOp.args.Conditional.yes =
+//									program.m_code[args.yes].args.ImageCompose.base;
+
+//							OP composeOp = program.m_code[args.yes];
+//							composeOp.args.ImageCompose.base = program.AddOp( baseOp );
+//							composeOp.args.ImageCompose.mask = program.AddOp( maskOp );
+
+//							at = program.AddOp( composeOp );
+
+//							// Process the new children
+//							at = Recurse( at, program );
+//							break;
+//						}
+
+                        default:
+                            break;
+
+                    }
+                }
+
+                else if ( !optimised && noHasAny )
+                {
+                    switch (noType)
+                    {
+                        case OP_TYPE::IM_LAYERCOLOUR:
+                        {
+                            optimised = true;
+
+                            auto typedNo = dynamic_cast<const ASTOpFixed*>(typedAt->no.child().get());
+
+                            Ptr<ASTOpFixed> blackOp = new ASTOpFixed;
+                            blackOp->op.type = OP_TYPE::CO_CONSTANT;
+                            blackOp->op.args.ColourConstant.value[0] = 0;
+                            blackOp->op.args.ColourConstant.value[1] = 0;
+                            blackOp->op.args.ColourConstant.value[2] = 0;
+
+                            Ptr<ASTOpFixed> plainOp = new ASTOpFixed;
+                            plainOp->op.type = OP_TYPE::IM_PLAINCOLOUR;
+                            plainOp->SetChild( plainOp->op.args.ImagePlainColour.colour, blackOp );
+                            plainOp->op.args.ImagePlainColour.format = EImageFormat::IF_L_UBYTE;
+                            plainOp->op.args.ImagePlainColour.size[0] = 4;
+                            plainOp->op.args.ImagePlainColour.size[1] = 4;
+
+                            Ptr<ASTOpFixed> resizeOp = new ASTOpFixed;
+                            resizeOp->op.type = OP_TYPE::IM_RESIZELIKE;
+                            resizeOp->SetChild( resizeOp->op.args.ImageResizeLike.source, plainOp );
+                            resizeOp->SetChild( resizeOp->op.args.ImageResizeLike.sizeSource,
+                                                typedNo->children[ typedNo->op.args.ImageLayerColour.base ] );
+
+                            auto maskOp = mu::Clone<ASTOpConditional>(typedAt);
+                            maskOp->no = resizeOp;
+
+                            // If there is no mask (because it is optional), we need to make a
+                            // white plain image
+                            maskOp->no = EnsureValidMask
+                                    ( typedNo->children[ typedNo->op.args.ImageLayerColour.mask ].child(),
+                                      typedNo->children[ typedNo->op.args.ImageLayerColour.base ].child() );
+
+                            auto baseOp = mu::Clone<ASTOpConditional>(typedAt);
+                            baseOp->no = typedNo->children[ typedNo->op.args.ImageLayerColour.base ].child();
+
+                            auto softOp = mu::Clone<ASTOpFixed>(typedNo);
+                            softOp->SetChild( softOp->op.args.ImageLayerColour.base, baseOp );
+                            softOp->SetChild( softOp->op.args.ImageLayerColour.mask, maskOp );
+
+                            at = softOp;
+                            break;
+                        }
+
+                        default:
+                            break;
+
+                    }
+                }
+            }
+
+            m_modified |= optimised;
+
+            break;
+        }
+
+
+        //-------------------------------------------------------------------------------------
+        case OP_TYPE::IM_SWITCH:
+        {
+            // If the switch is not runtime, but the branches are, try to move the
+            // switch down
+//            bool optimised = false;
+
+//            OP::ADDRESS variable = program.m_code[at].args.Switch.variable;
+
+//            if ( !m_hasRuntimeParamVisitor.HasAny( variable, program ) )
+//            {
+//                SWITCH_CHAIN chain = GetSwitchChain( program, at );
+
+//                bool branchHasAny = false;
+//                OP_TYPE branchType = (OP_TYPE)program.m_code[program.m_code[at].args.Switch.values[0]].type;
+//                for ( map<uint16_t,OP::ADDRESS>::const_iterator it=chain.cases.begin();
+//                      it != chain.cases.end();
+//                      ++it )
+//                {
+//                    if ( program.m_code[it->second].type != branchType )
+//                    {
+//                        branchType = OP_TYPE::NONE;
+//                    }
+//                    else
+//                    {
+//                        if (!branchHasAny)
+//                        {
+//                            branchHasAny = m_hasRuntimeParamVisitor.HasAny( it->second, program );
+//                        }
+//                    }
+//                }
+
+//                if ( chain.def && program.m_code[chain.def].type != branchType )
+//                {
+//                    branchType = OP_TYPE::NONE;
+//                }
+
+//                // Some branch in runtime
+//                if ( branchHasAny )
+//                {
+//                    switch ( branchType )
+//                    {
+
+//                    // TODO: Other operations
+//                    case OP_TYPE::IM_BLEND:
+//                    case OP_TYPE::IM_MULTIPLY:
+//                    {
+//                        // Move the switch down the base
+//                        OP::ADDRESS baseAt = 0;
+//                        for ( map<uint16_t,OP::ADDRESS>::const_iterator it=chain.cases.begin();
+//                              it != chain.cases.end();
+//                              )
+//                        {
+//                            OP bsw;
+//                            bsw.type = OP_TYPE::IM_SWITCH;
+//                            bsw.args.Switch.variable = variable;
+
+//                            for ( int b=0;
+//                                  it != chain.cases.end() && b<MUTABLE_OP_MAX_SWITCH_OPTIONS;
+//                                  ++b )
+//                            {
+//                                bsw.args.Switch.conditions[b] = it->first;
+//                                bsw.args.Switch.values[b] =
+//                                        program.m_code[it->second].args.ImageLayer.base;
+//                                ++it;
+//                            }
+
+//                            bsw.args.Switch.def = baseAt;
+//                            baseAt = program.AddOp( bsw );
+//                        }
+
+//                        // Move the switch down the mask
+//                        OP::ADDRESS maskAt = 0;
+//                        for ( map<uint16_t,OP::ADDRESS>::const_iterator it=chain.cases.begin();
+//                              it != chain.cases.end();
+//                              )
+//                        {
+//                            OP bsw;
+//                            bsw.type = OP_TYPE::IM_SWITCH;
+//                            bsw.args.Switch.variable = variable;
+
+//                            for ( int b=0;
+//                                  it != chain.cases.end() && b<MUTABLE_OP_MAX_SWITCH_OPTIONS;
+//                                  ++b )
+//                            {
+//                                bsw.args.Switch.conditions[b] = it->first;
+//                                bsw.args.Switch.values[b] =
+//                                        program.m_code[it->second].args.ImageLayer.mask;
+//                                ++it;
+//                            }
+
+//                            bsw.args.Switch.def = maskAt;
+//                            maskAt = program.AddOp( bsw );
+//                        }
+
+//                        // Move the switch down the blended
+//                        OP::ADDRESS blendedAt = 0;
+//                        for ( map<uint16_t,OP::ADDRESS>::const_iterator it=chain.cases.begin();
+//                              it != chain.cases.end();
+//                              )
+//                        {
+//                            OP bsw;
+//                            bsw.type = OP_TYPE::IM_SWITCH;
+//                            bsw.args.Switch.variable = variable;
+
+//                            for ( int b=0;
+//                                  it != chain.cases.end() && b<MUTABLE_OP_MAX_SWITCH_OPTIONS;
+//                                  ++b )
+//                            {
+//                                bsw.args.Switch.conditions[b] = it->first;
+//                                bsw.args.Switch.values[b] =
+//                                        program.m_code[it->second].args.ImageLayer.blended;
+//                                ++it;
+//                            }
+
+//                            bsw.args.Switch.def = blendedAt;
+//                            blendedAt = program.AddOp( bsw );
+//                        }
+
+//                        OP top;
+//                        top.type = branchType;
+//                        top.args.ImageLayer.base = baseAt;
+//                        top.args.ImageLayer.mask = maskAt;
+//                        top.args.ImageLayer.blended = blendedAt;
+//                        at = program.AddOp( top );
+//                        optimised = true;
+//                        break;
+//                    }
+
+
+//                    // TODO: Other operations
+//                    case OP_TYPE::IM_SOFTLIGHTCOLOUR:
+//                    {
+//                        // Move the switch down the base
+//                        OP::ADDRESS baseAt = 0;
+//                        for ( map<uint16_t,OP::ADDRESS>::const_iterator it=chain.cases.begin();
+//                              it != chain.cases.end();
+//                              )
+//                        {
+//                            OP bsw;
+//                            bsw.type = OP_TYPE::IM_SWITCH;
+//                            bsw.args.Switch.variable = variable;
+
+//                            for ( int b=0;
+//                                  it != chain.cases.end() && b<MUTABLE_OP_MAX_SWITCH_OPTIONS;
+//                                  ++b )
+//                            {
+//                                bsw.args.Switch.conditions[b] = it->first;
+//                                bsw.args.Switch.values[b] =
+//                                        program.m_code[it->second].args.ImageLayerColour.base;
+//                                ++it;
+//                            }
+
+//                            bsw.args.Switch.def = baseAt;
+//                            baseAt = program.AddOp( bsw );
+//                        }
+
+//                        // Move the switch down the mask
+//                        OP::ADDRESS maskAt = 0;
+//                        for ( map<uint16_t,OP::ADDRESS>::const_iterator it=chain.cases.begin();
+//                              it != chain.cases.end();
+//                              )
+//                        {
+//                            OP bsw;
+//                            bsw.type = OP_TYPE::IM_SWITCH;
+//                            bsw.args.Switch.variable = variable;
+
+//                            for ( int b=0;
+//                                  it != chain.cases.end() && b<MUTABLE_OP_MAX_SWITCH_OPTIONS;
+//                                  ++b )
+//                            {
+//                                bsw.args.Switch.conditions[b] = it->first;
+//                                bsw.args.Switch.values[b] =
+//                                        program.m_code[it->second].args.ImageLayerColour.mask;
+//                                ++it;
+//                            }
+
+//                            bsw.args.Switch.def = maskAt;
+//                            maskAt = program.AddOp( bsw );
+//                        }
+
+//                        // Move the switch down the colour
+//                        OP::ADDRESS colourAt = 0;
+//                        for ( map<uint16_t,OP::ADDRESS>::const_iterator it=chain.cases.begin();
+//                              it != chain.cases.end();
+//                              )
+//                        {
+//                            OP bsw;
+//                            bsw.type = OP_TYPE::CO_SWITCH;
+//                            bsw.args.Switch.variable = variable;
+
+//                            for ( int b=0;
+//                                  it != chain.cases.end() && b<MUTABLE_OP_MAX_SWITCH_OPTIONS;
+//                                  ++b )
+//                            {
+//                                bsw.args.Switch.conditions[b] = it->first;
+//                                bsw.args.Switch.values[b] =
+//                                        program.m_code[it->second].args.ImageLayerColour.colour;
+//                                ++it;
+//                            }
+
+//                            bsw.args.Switch.def = colourAt;
+//                            colourAt = program.AddOp( bsw );
+//                        }
+
+//                        OP top;
+//                        top.type = branchType;
+//                        top.args.ImageLayerColour.base = baseAt;
+//                        top.args.ImageLayerColour.mask = maskAt;
+//                        top.args.ImageLayerColour.colour = colourAt;
+//                        at = program.AddOp( top );
+//                        optimised = true;
+//                        break;
+//                    }
+
+
+//                    default:
+//                        break;
+
+//                    }
+//                }
+
+//            }
+
+//            m_modified |= optimised;
+
+            break;
+        }
+
+
+        //-----------------------------------------------------------------------------------------
+        case OP_TYPE::IM_COMPOSE:
+        {
+			const ASTOpImageCompose* typedAt = dynamic_cast<const ASTOpImageCompose*>(at.get());
+
+            auto blockAt = typedAt->BlockImage.child();
+            auto baseAt = typedAt->Base.child();
+            auto layoutAt = typedAt->Layout.child();
+
+			if (!blockAt)
+			{
+				at = baseAt;
+				break;
+			}
+
+            OP_TYPE blockType = blockAt->GetOpType();
+            OP_TYPE baseType = baseAt->GetOpType();
+
+            bool baseHasRuntime = m_hasRuntimeParamVisitor.HasAny( baseAt );
+            bool blockHasRuntime = m_hasRuntimeParamVisitor.HasAny( blockAt );
+            bool layoutHasRuntime = m_hasRuntimeParamVisitor.HasAny( layoutAt );
+
+            bool optimised = false;
+
+            // Try to optimise base and block together, if possible
+            if ( blockHasRuntime && baseHasRuntime && !layoutHasRuntime )
+            {
+                if ( baseType == blockType )
+                {
+                    switch ( blockType )
+                    {
+                    case OP_TYPE::IM_LAYERCOLOUR:
+                    {
+                        optimised = true;
+
+                        auto typedBaseAt = dynamic_cast<const ASTOpFixed*>(baseAt.get());
+                        auto typedBlockAt = dynamic_cast<const ASTOpFixed*>(blockAt.get());
+
+                        // The mask is a compose of the block mask on the base mask, but if none has
+                        // a mask we don't need to make one.
+                        auto baseImage =
+                            typedBaseAt->children[typedBaseAt->op.args.ImageLayerColour.base]
+                                .child();
+                        auto baseMask =
+                            typedBaseAt->children[typedBaseAt->op.args.ImageLayerColour.mask]
+                                .child();
+                        auto blockImage =
+                            typedBlockAt->children[typedBlockAt->op.args.ImageLayerColour.base]
+                                .child();
+                        auto blockMask =
+                            typedBlockAt->children[typedBlockAt->op.args.ImageLayerColour.mask]
+                                .child();
+
+                        Ptr<ASTOpImageCompose> maskOp;
+                        if (baseMask || blockMask)
+                        {
+                            // \TODO: BLEH! This may create a discrepancy of number of mips between
+                            // the base image and the mask This is for now solved with emergy fix
+                            // c36adf47-e40d-490f-b709-41142bafad78
+                            auto newBaseMask = EnsureValidMask(baseMask, baseImage);
+                            auto newBlockMask = EnsureValidMask(blockMask, blockImage);
+
+                            maskOp = mu::Clone<ASTOpImageCompose>(typedAt);
+                            maskOp->Base = newBaseMask;
+                            maskOp->BlockImage = newBlockMask;
+                        }
+
+                        // The base is composition of the bases of both layer effect
+                        auto baseOp = mu::Clone<ASTOpImageCompose>(typedAt);
+                        baseOp->Base = baseImage;
+                        baseOp->BlockImage = blockImage;
+
+                        Ptr<ASTOpFixed> nop = mu::Clone<ASTOpFixed>(blockAt);
+                        nop->SetChild( nop->op.args.ImageLayerColour.mask, maskOp );
+                        nop->SetChild( nop->op.args.ImageLayerColour.base, baseOp );
+
+                        // Done
+                        at = nop;
+                        break;
+                    }
+
+                    case OP_TYPE::IM_LAYER:
+                    {
+                        optimised = true;
+
+                        auto typedBaseAt = dynamic_cast<const ASTOpFixed*>(baseAt.get());
+                        auto typedBlockAt = dynamic_cast<const ASTOpFixed*>(blockAt.get());
+
+                        // The mask is a compose of the block mask on the base mask, but if none has
+                        // a mask we don't need to make one.
+                        auto baseImage =
+                            typedBaseAt->children[typedBaseAt->op.args.ImageLayer.base]
+                                .child();
+                        auto baseBlended =
+                            typedBaseAt->children[typedBaseAt->op.args.ImageLayer.blended]
+                                .child();
+                        auto baseMask =
+                            typedBaseAt->children[typedBaseAt->op.args.ImageLayer.mask]
+                                .child();
+                        auto blockImage =
+                            typedBlockAt->children[typedBlockAt->op.args.ImageLayer.base]
+                                .child();
+                        auto blockBlended =
+                            typedBlockAt->children[typedBlockAt->op.args.ImageLayer.blended]
+                                .child();
+                        auto blockMask =
+                            typedBlockAt->children[typedBlockAt->op.args.ImageLayer.mask]
+                                .child();
+
+                        Ptr<ASTOpImageCompose> maskOp;
+                        if (baseMask || blockMask)
+                        {
+                            // \TODO: BLEH! This may create a discrepancy of number of mips between
+                            // the base image and the mask This is for now solved with emergy fix
+                            // c36adf47-e40d-490f-b709-41142bafad78
+                            auto newBaseMask = EnsureValidMask(baseMask, baseImage);
+                            auto newBlockMask = EnsureValidMask(blockMask, blockImage);
+
+                            maskOp = mu::Clone<ASTOpImageCompose>(typedAt);
+                            maskOp->Base = newBaseMask;
+                            maskOp->BlockImage = newBlockMask;
+                        }
+
+                        // The base is composition of the bases of both layer effect
+                        auto baseOp = mu::Clone<ASTOpImageCompose>(typedAt);
+                        baseOp->Base = baseImage;
+                        baseOp->BlockImage = blockImage;
+
+                        // The base is composition of the bases of both layer effect
+                        auto blendedOp = mu::Clone<ASTOpImageCompose>(typedAt);
+                        blendedOp->Base = baseBlended;
+                        blendedOp->BlockImage = blockBlended;
+
+                        Ptr<ASTOpFixed> nop = mu::Clone<ASTOpFixed>(blockAt);
+                        nop->SetChild(nop->op.args.ImageLayer.mask, maskOp);
+                        nop->SetChild(nop->op.args.ImageLayer.base, baseOp);
+                        nop->SetChild(nop->op.args.ImageLayer.blended, blendedOp);
+
+                        // Done
+                        at = nop;
+                        break;
+                    }
+
+                    default:
+                        break;
+                    }
+                }
+            }
+
+
+            // Swap two composes
+            if ( !optimised && baseHasRuntime && !blockHasRuntime
+                 &&
+                 baseType == OP_TYPE::IM_COMPOSE )
+            {
+                auto typedBaseAt = dynamic_cast<const ASTOpImageCompose*>(baseAt.get());
+
+                auto baseBlockAt = typedBaseAt->BlockImage.child();
+                bool baseBlockHasAny = m_hasRuntimeParamVisitor.HasAny( baseBlockAt );
+                if ( baseBlockHasAny )
+                {
+                    optimised = true;
+
+                    // Swap
+                    auto childCompose = mu::Clone<ASTOpImageCompose>(at);
+                    childCompose->Base = typedBaseAt->Base.child();
+
+                    auto parentCompose = mu::Clone<ASTOpImageCompose>(baseAt);
+                    parentCompose->Base = childCompose;
+
+                    at = parentCompose;
+                }
+            }
+
+
+            // Try to optimise the block
+            // This optimisation requires a lot of memory for every target. Use only if
+            // we are optimising for GPU processing.
+            if ( !optimised && blockHasRuntime && !baseHasRuntime
+                 //&& m_stateProps.m_gpu.m_external
+                 // TODO BLEH
+                 // Only worth in case of more than one block using the same operation. Move this
+                 // optimisation to that test.
+                 //&& false
+                 )
+            {
+                switch ( blockType )
+                {
+                case OP_TYPE::IM_LAYERCOLOUR:
+                {
+                    optimised = true;
+
+                    auto typedBlockAt = dynamic_cast<const ASTOpFixed*>(blockAt.get());
+
+                    auto blockImage =
+                        typedBlockAt->children[typedBlockAt->op.args.ImageLayerColour.base].child();
+                    auto blockMask =
+                        typedBlockAt->children[typedBlockAt->op.args.ImageLayerColour.mask].child();
+
+                    // The mask is a compose of the layer mask on a black image, however if there is
+                    // no mask and the base of the layer opertation is a blanklayout, we can skip
+                    // generating a mask.
+                    Ptr<ASTOpImageCompose> maskOp;
+                    if (blockMask || baseType!=OP_TYPE::IM_BLANKLAYOUT)
+                    {
+                        maskOp = mu::Clone<ASTOpImageCompose>(at);
+                        auto newMaskBlock = EnsureValidMask(blockMask, blockImage);
+                        maskOp->BlockImage = newMaskBlock;
+
+                        Ptr<ASTOpFixed> blackOp = new ASTOpFixed;
+                        blackOp->op.type = OP_TYPE::CO_CONSTANT;
+                        blackOp->op.args.ColourConstant.value[0] = 0;
+                        blackOp->op.args.ColourConstant.value[1] = 0;
+                        blackOp->op.args.ColourConstant.value[2] = 0;
+
+                        Ptr<ASTOpFixed> plainOp = new ASTOpFixed;
+                        plainOp->op.type = OP_TYPE::IM_PLAINCOLOUR;
+                        plainOp->SetChild( plainOp->op.args.ImagePlainColour.colour, blackOp );
+                        plainOp->op.args.ImagePlainColour.format = EImageFormat::IF_L_UBYTE;
+                        plainOp->op.args.ImagePlainColour.size[0] = 4;
+                        plainOp->op.args.ImagePlainColour.size[1] = 4;
+
+                        Ptr<ASTOpFixed> baseResizeOp = new ASTOpFixed;
+                        baseResizeOp->op.type = OP_TYPE::IM_RESIZELIKE;
+                        baseResizeOp->SetChild( baseResizeOp->op.args.ImageResizeLike.sizeSource, baseAt );
+                        baseResizeOp->SetChild( baseResizeOp->op.args.ImageResizeLike.source, plainOp );
+
+                        maskOp->Base = baseResizeOp;
+                    }
+
+                    // The base is composition of the layer base on the compose base
+                    auto baseOp = mu::Clone<ASTOpImageCompose>(at);
+                    baseOp->BlockImage = typedBlockAt->children[typedBlockAt->op.args.ImageLayerColour.base].child();
+
+                    auto nop = mu::Clone<ASTOpFixed>(blockAt);
+                    nop->SetChild( nop->op.args.ImageLayerColour.mask, maskOp );
+                    nop->SetChild( nop->op.args.ImageLayerColour.base, baseOp );
+
+                    // Done
+                    at = nop;
+                    break;
+                }
+
+                case OP_TYPE::IM_LAYER:
+                {
+                    optimised = true;
+
+                    auto typedBlockAt = dynamic_cast<const ASTOpFixed*>(blockAt.get());
+
+                    auto blockImage =
+                        typedBlockAt->children[typedBlockAt->op.args.ImageLayer.base].child();
+                    auto blockBlended =
+                         typedBlockAt->children[typedBlockAt->op.args.ImageLayer.blended].child();
+                    auto blockMask =
+                        typedBlockAt->children[typedBlockAt->op.args.ImageLayer.mask].child();
+
+                    // The mask is a compose of the layer mask on a black image, however if there is
+                    // no mask and the base of the layer opertation is a blanklayout, we can skip
+                    // generating a mask.
+                    Ptr<ASTOpImageCompose> maskOp;
+                    if (blockMask || baseType != OP_TYPE::IM_BLANKLAYOUT)
+                    {
+                        maskOp = mu::Clone<ASTOpImageCompose>(at);
+                        auto newMaskBlock = EnsureValidMask(blockMask, blockImage);
+                        maskOp->BlockImage = newMaskBlock;
+
+
+                        Ptr<ASTOpFixed> blackOp = new ASTOpFixed;
+                        blackOp->op.type = OP_TYPE::CO_CONSTANT;
+                        blackOp->op.args.ColourConstant.value[0] = 0;
+                        blackOp->op.args.ColourConstant.value[1] = 0;
+                        blackOp->op.args.ColourConstant.value[2] = 0;
+
+                        Ptr<ASTOpFixed> plainOp = new ASTOpFixed;
+                        plainOp->op.type = OP_TYPE::IM_PLAINCOLOUR;
+                        plainOp->SetChild( plainOp->op.args.ImagePlainColour.colour, blackOp );
+                        plainOp->op.args.ImagePlainColour.format = EImageFormat::IF_L_UBYTE;
+                        plainOp->op.args.ImagePlainColour.size[0] = 4;
+                        plainOp->op.args.ImagePlainColour.size[1] = 4;
+
+                        Ptr<ASTOpFixed> baseResizeOp = new ASTOpFixed;
+                        baseResizeOp->op.type = OP_TYPE::IM_RESIZELIKE;
+                        baseResizeOp->SetChild( baseResizeOp->op.args.ImageResizeLike.sizeSource, baseAt );
+                        baseResizeOp->SetChild( baseResizeOp->op.args.ImageResizeLike.source, plainOp );
+
+                        maskOp->Base = baseResizeOp;
+                    }
+
+                    // The blended is a compose of the blended image on a blank image
+                    auto blendedOp = mu::Clone<ASTOpImageCompose>(at);
+                    {
+                        blendedOp->BlockImage = blockBlended;
+
+                        Ptr<ASTOpFixed> blackOp = new ASTOpFixed;
+                        blackOp->op.type = OP_TYPE::CO_CONSTANT;
+                        blackOp->op.args.ColourConstant.value[0] = 0;
+                        blackOp->op.args.ColourConstant.value[1] = 0;
+                        blackOp->op.args.ColourConstant.value[2] = 0;
+
+                        Ptr<ASTOpFixed> plainOp = new ASTOpFixed;
+                        plainOp->op.type = OP_TYPE::IM_PLAINCOLOUR;
+                        plainOp->SetChild( plainOp->op.args.ImagePlainColour.colour, blackOp );
+                        FImageDesc blendedDesc = baseAt->GetImageDesc();
+                        plainOp->op.args.ImagePlainColour.format = blendedDesc.m_format;
+                        plainOp->op.args.ImagePlainColour.size[0] = 4;
+                        plainOp->op.args.ImagePlainColour.size[1] = 4;
+
+                        Ptr<ASTOpFixed> resizeOp = new ASTOpFixed;
+                        resizeOp->op.type = OP_TYPE::IM_RESIZELIKE;
+                        resizeOp->SetChild( resizeOp->op.args.ImageResizeLike.sizeSource, baseAt );
+                        resizeOp->SetChild( resizeOp->op.args.ImageResizeLike.source, plainOp );
+
+                        blendedOp->Base = resizeOp;
+                    }
+
+                    // The base is composition of the softlight base on the compose base
+                    auto baseOp = mu::Clone<ASTOpImageCompose>(at);
+                    baseOp->BlockImage = typedBlockAt->children[typedBlockAt->op.args.ImageLayer.base].child();
+
+                    auto nop = mu::Clone<ASTOpFixed>(blockAt);
+                    nop->SetChild( nop->op.args.ImageLayer.base, baseOp );
+                    nop->SetChild( nop->op.args.ImageLayer.mask, maskOp );
+                    nop->SetChild( nop->op.args.ImageLayer.blended, blendedOp );
+
+                    // Done
+                    at = nop;
+                    break;
+                }
+
+
+//                case OP_TYPE::IM_INTERPOLATE:
+//                {
+//                    optimised = true;
+//                    OP op = program.m_code[blockAt];
+
+//                    // The targets are composition of the block targets on the compose base
+//                    for ( int t=0; t<MUTABLE_OP_MAX_INTERPOLATE_COUNT; ++t )
+//                    {
+//                        if ( op.args.ImageInterpolate.targets[t] )
+//                        {
+//                            OP targetOp = program.m_code[at];
+//                            targetOp.args.ImageCompose.blockImage =
+//                                    op.args.ImageInterpolate.targets[t];
+//                            op.args.ImageInterpolate.targets[t] = program.AddOp( targetOp );
+//                        }
+//                    }
+
+//                    // Done
+//                    at = program.AddOp( op );
+
+//                    // Reprocess the new children
+//                    at = Recurse( at, program );
+
+//                    break;
+//                }
+
+                default:
+                    break;
+
+                }
+            }
+
+
+            // Try to optimise the base
+            if ( !optimised && baseHasRuntime /*&& m_stateProps.nodeState.m_optimisation.m_gpu.m_external*/ )
+            {
+                switch ( baseType )
+                {
+                case OP_TYPE::IM_LAYERCOLOUR:
+                {
+                    optimised = true;
+
+                    auto typedBaseAt = dynamic_cast<const ASTOpFixed*>(baseAt.get());
+
+                    auto maskOp = mu::Clone<ASTOpImageCompose>(at);
+                    {
+                        Ptr<ASTOpFixed> blackOp = new ASTOpFixed;
+                        blackOp->op.type = OP_TYPE::CO_CONSTANT;
+                        blackOp->op.args.ColourConstant.value[0] = 0;
+                        blackOp->op.args.ColourConstant.value[1] = 0;
+                        blackOp->op.args.ColourConstant.value[2] = 0;
+
+                        Ptr<ASTOpFixed> plainOp = new ASTOpFixed;
+                        plainOp->op.type = OP_TYPE::IM_PLAINCOLOUR;
+                        plainOp->SetChild( plainOp->op.args.ImagePlainColour.colour, blackOp );
+                        plainOp->op.args.ImagePlainColour.format = EImageFormat::IF_L_UBYTE; //TODO: FORMAT_LIKE
+                        plainOp->op.args.ImagePlainColour.size[0] = 4;
+                        plainOp->op.args.ImagePlainColour.size[1] = 4;
+
+                        Ptr<ASTOpFixed> blockResizeOp = new ASTOpFixed;
+                        blockResizeOp->op.type = OP_TYPE::IM_RESIZELIKE;
+                        blockResizeOp->SetChild( blockResizeOp->op.args.ImageResizeLike.sizeSource, blockAt );
+                        blockResizeOp->SetChild( blockResizeOp->op.args.ImageResizeLike.source, plainOp );
+
+                        // Blank out the block from the mask
+                        auto newMaskBase = EnsureValidMask(
+                                    typedBaseAt->children[ typedBaseAt->op.args.ImageLayerColour.mask].child(),
+                                    baseAt );
+                        maskOp->Base = newMaskBase;
+                        maskOp->BlockImage = blockResizeOp;
+                    }
+
+                    // The base is composition of the softlight base on the compose base
+                    auto baseOp = mu::Clone<ASTOpImageCompose>(at);
+                    baseOp->Base = typedBaseAt->children[typedBaseAt->op.args.ImageLayerColour.base].child();
+
+                    auto nop = mu::Clone<ASTOpFixed>(baseAt);
+                    nop->SetChild( nop->op.args.ImageLayerColour.base, baseOp );
+                    nop->SetChild( nop->op.args.ImageLayerColour.mask, maskOp );
+
+                    // Done
+                    at = nop;
+                    break;
+                }
+
+                case OP_TYPE::IM_LAYER:
+                {
+                    optimised = true;
+
+                    auto typedBaseAt = dynamic_cast<const ASTOpFixed*>(baseAt.get());
+
+                    auto maskOp = mu::Clone<ASTOpImageCompose>(at);
+                    {
+                        Ptr<ASTOpFixed> blackOp = new ASTOpFixed;
+                        blackOp->op.type = OP_TYPE::CO_CONSTANT;
+                        blackOp->op.args.ColourConstant.value[0] = 0;
+                        blackOp->op.args.ColourConstant.value[1] = 0;
+                        blackOp->op.args.ColourConstant.value[2] = 0;
+
+                        Ptr<ASTOpFixed> plainOp = new ASTOpFixed;
+                        plainOp->op.type = OP_TYPE::IM_PLAINCOLOUR;
+                        plainOp->SetChild( plainOp->op.args.ImagePlainColour.colour, blackOp );
+                        plainOp->op.args.ImagePlainColour.format = EImageFormat::IF_L_UBYTE; //TODO: FORMAT_LIKE
+                        plainOp->op.args.ImagePlainColour.size[0] = 4;
+                        plainOp->op.args.ImagePlainColour.size[1] = 4;
+
+                        Ptr<ASTOpFixed> blockResizeOp = new ASTOpFixed;
+                        blockResizeOp->op.type = OP_TYPE::IM_RESIZELIKE;
+                        blockResizeOp->SetChild( blockResizeOp->op.args.ImageResizeLike.sizeSource, blockAt );
+                        blockResizeOp->SetChild( blockResizeOp->op.args.ImageResizeLike.source, plainOp );
+
+                        // Blank out the block from the mask
+                        auto newMaskBase = EnsureValidMask(
+                                    typedBaseAt->children[ typedBaseAt->op.args.ImageLayer.mask ].child(),
+                                    baseAt );
+                        maskOp->Base = newMaskBase;
+                        maskOp->BlockImage = blockResizeOp;
+                    }
+
+                    // The base is composition of the effect base on the compose base
+                    auto baseOp = mu::Clone<ASTOpImageCompose>(at);
+                    baseOp->Base = typedBaseAt->children[typedBaseAt->op.args.ImageLayer.base].child();
+
+                    auto nop = mu::Clone<ASTOpFixed>(baseAt);
+                    nop->SetChild( nop->op.args.ImageLayer.base, baseOp );
+                    nop->SetChild( nop->op.args.ImageLayer.mask, maskOp );
+
+                    // Done
+                    at = nop;
+                    break;
+                }
+
+
+//                case OP_TYPE::IM_INTERPOLATE:
+//                {
+//                    optimised = true;
+//                    OP op = program.m_code[baseAt];
+
+//                    // The targets are composition of the blocks on the compose base targets
+//                    for ( int t=0; t<MUTABLE_OP_MAX_INTERPOLATE_COUNT; ++t )
+//                    {
+//                        if ( op.args.ImageInterpolate.targets[t] )
+//                        {
+//                            OP targetOp = program.m_code[at];
+//                            targetOp.args.ImageCompose.base =
+//                                    op.args.ImageInterpolate.targets[t];
+//                            op.args.ImageInterpolate.targets[t] = program.AddOp( targetOp );
+//                        }
+//                    }
+
+//                    // Done
+//                    at = program.AddOp( op );
+
+//                    // Reprocess the new children
+//                    at = Recurse( at, program );
+//                    break;
+//                }
+
+                default:
+                    break;
+
+                }
+            }
+
+
+            m_modified = m_modified || optimised;
+            break;
+        }
+
+/*
+        //-----------------------------------------------------------------------------------------
+        // TODO: Other ops?
+        case OP_TYPE::IM_BLEND:
+        {
+            OP op = program.m_code[at];
+
+            if ( !m_hasRuntimeParamVisitor.HasAny( op.args.ImageLayer.mask, program ) )
+            {
+                // If both the base and the blended have the same image layer operation with
+                // similar parameters, we can move that operation up.
+                OP_TYPE baseType = program.m_code[ op.args.ImageLayer.base ].type;
+                OP_TYPE blendedType = program.m_code[ op.args.ImageLayer.blended ].type;
+                if ( baseType == blendedType )
+                {
+                    switch ( baseType )
+                    {
+                    case OP_TYPE::IM_BLENDCOLOUR:
+                    case OP_TYPE::IM_SOFTLIGHTCOLOUR:
+                    case OP_TYPE::IM_HARDLIGHTCOLOUR:
+                    case OP_TYPE::IM_BURNCOLOUR:
+                    case OP_TYPE::IM_SCREENCOLOUR:
+                    case OP_TYPE::IM_OVERLAYCOLOUR:
+                    case OP_TYPE::IM_DODGECOLOUR:
+                    case OP_TYPE::IM_MULTIPLYCOLOUR:
+                    {
+                        if ( m_optimisationOptions.m_optimiseOverlappedMasks )
+                        {
+                            OP::ADDRESS maskAt = op.args.ImageLayer.mask;
+                            OP::ADDRESS baseMaskAt =
+                                    program.m_code[ op.args.ImageLayer.base ].args.ImageLayerColour.mask;
+                            OP::ADDRESS blendedMaskAt =
+                                    program.m_code[ op.args.ImageLayer.blended ].args.ImageLayerColour.mask;
+                            OP::ADDRESS baseColourAt =
+                                    program.m_code[ op.args.ImageLayer.base ].args.ImageLayerColour.colour;
+                            OP::ADDRESS blendedColourAt =
+                                    program.m_code[ op.args.ImageLayer.blended ].args.ImageLayerColour.colour;
+
+                            // Check extra conditions on the masks
+                            if ( maskAt && baseMaskAt && blendedMaskAt
+                                 && (baseColourAt==blendedColourAt)
+                                 && !AreMasksOverlapping( program, baseMaskAt, blendedMaskAt )
+                                 && !AreMasksOverlapping( program, baseMaskAt, maskAt ) )
+                            {
+                                // We can apply the transform
+                                m_modified = true;
+
+                                OP baseOp;
+                                baseOp.type = OP_TYPE::IM_BLEND;
+                                baseOp.args.ImageLayer.base =
+                                        program.m_code[ op.args.ImageLayer.base ].args.ImageLayerColour.base;
+                                baseOp.args.ImageLayer.blended =
+                                        program.m_code[ op.args.ImageLayer.blended ].args.ImageLayerColour.base;
+                                baseOp.args.ImageLayer.mask = maskAt;
+
+                                // TODO: Find out why these are not equivalent
+    //							OP maskOp;
+    //							maskOp.type = OP_TYPE::IM_SCREEN;
+    //							maskOp.args.ImageLayer.base = baseMaskAt;
+    //							maskOp.args.ImageLayer.blended = blendedMaskAt;
+                                OP maskOp;
+                                maskOp.type = OP_TYPE::IM_BLEND;
+                                maskOp.args.ImageLayer.base = baseMaskAt;
+                                maskOp.args.ImageLayer.blended = blendedMaskAt;
+                                maskOp.args.ImageLayer.mask = blendedMaskAt;
+                                maskOp.args.ImageLayer.flags = OP::ImageLayerArgs::F_BINARY_MASK;
+
+                                OP top;
+                                top.type = baseType;
+                                top.args.ImageLayerColour.base = program.AddOp( baseOp );
+                                top.args.ImageLayerColour.mask = program.AddOp( maskOp );
+                                top.args.ImageLayerColour.colour = baseColourAt;
+                                at = program.AddOp( top );
+                            }
+                        }
+
+                        break;
+                    }
+
+                    default:
+                        break;
+
+                    }
+                }
+            }
+            break;
+        }
+
+        */
+
+        //-----------------------------------------------------------------------------------------
+        // Sink the mipmap if worth it.
+        case OP_TYPE::IM_MIPMAP:
+        {
+            auto typedAt = dynamic_cast<const ASTOpImageMipmap*>(at.get());
+
+            auto sourceOp = typedAt->Source.child();
+
+            switch ( sourceOp->GetOpType() )
+            {
+            case OP_TYPE::IM_LAYERCOLOUR:
+            {
+                auto typedSource = dynamic_cast<const ASTOpFixed*>(sourceOp.get());
+
+                bool colourHasRuntime = m_hasRuntimeParamVisitor.HasAny
+                        ( typedSource->children[ typedSource->op.args.ImageLayerColour.colour].child() );
+
+                if (colourHasRuntime)
+                {
+                    m_modified = true;
+
+                    auto top = mu::Clone<ASTOpFixed>(sourceOp);
+
+                    auto baseOp = mu::Clone<ASTOpImageMipmap>(at);
+                    baseOp->Source = typedSource->children[ typedSource->op.args.ImageLayerColour.base ].child();
+                    top->SetChild( top->op.args.ImageLayerColour.base, baseOp );
+
+                    auto sourceMaskOp = typedSource->children[typedSource->op.args.ImageLayerColour.mask].child();
+                    if (sourceMaskOp)
+                    {
+                        auto maskOp = mu::Clone<ASTOpImageMipmap>(at);
+                        maskOp->Source = sourceMaskOp;
+                        top->SetChild( top->op.args.ImageLayerColour.mask, maskOp );
+                    }
+
+                    at = top;
+                }
+
+                break;
+            }
+
+            default:
+                break;
+            }
+
+            break;
+        }
+
+        default:
+            break;
+
+        }
+
+        return at;
+    }
+
+
+    class AccumulateAllImageFormatsOpAST
+            : public Visitor_TopDown_Unique_Const< std::array<uint8_t, size_t(EImageFormat::IF_COUNT)> >
+    {
+    public:
+
+        void Run( const ASTOpList& roots )
+        {
+            MUTABLE_CPUPROFILER_SCOPE(AccumulateAllImageFormatsOpAST);
+
+            // Initially, all formats are supported
+            m_allSupported.fill(1);
+
+            // The initial traversal state is no format supported
+            m_initialState.fill(0);
+            Traverse( roots, m_initialState );
+        }
+
+
+        bool Visit( const Ptr<ASTOp>& at ) override
+        {
+            bool recurse = false;
+
+            const std::array<uint8_t, (size_t)EImageFormat::IF_COUNT>& currentFormats = GetCurrentState();
+
+            // Remove unsupported formats
+            if (GetOpDataType( at->GetOpType() )==DT_IMAGE)
+            {
+                auto it = m_supportedFormats.find(at);
+                if (it==m_supportedFormats.end())
+                {
+                    // Default to all supported
+                    m_supportedFormats[at] = m_allSupported;
+                    it = m_supportedFormats.find(at);
+                }
+
+                for ( unsigned f=0; f< (unsigned)EImageFormat::IF_COUNT; ++f )
+                {
+                    if ( !currentFormats[f] )
+                    {
+                        it->second[f] = 0;
+                    }
+                }
+            }
+
+            switch ( at->GetOpType() )
+            {
+            // TODO: Code shared with the constant data format optimisation visitor
+            case OP_TYPE::IM_LAYERCOLOUR:
+            {
+                auto typedAt = dynamic_cast<const ASTOpFixed*>(at.get());
+
+                RecurseWithCurrentState( typedAt->children[typedAt->op.args.ImageLayerColour.base].child() );
+                RecurseWithCurrentState( typedAt->children[typedAt->op.args.ImageLayerColour.colour].child() );
+
+                if ( typedAt->children[typedAt->op.args.ImageLayerColour.mask] )
+                {
+                    std::array<uint8_t, (size_t)EImageFormat::IF_COUNT> newState;
+                    newState.fill(0);
+                    newState[(size_t)EImageFormat::IF_L_UBYTE ] = 1;
+                    newState[(size_t)EImageFormat::IF_L_UBYTE_RLE ] = 1;
+
+                    RecurseWithState( typedAt->children[typedAt->op.args.ImageLayerColour.mask].child(), newState );
+                }
+                break;
+            }
+
+            case OP_TYPE::IM_LAYER:
+            {
+                auto typedAt = dynamic_cast<const ASTOpFixed*>(at.get());
+
+                RecurseWithCurrentState( typedAt->children[typedAt->op.args.ImageLayer.base].child() );
+                RecurseWithCurrentState( typedAt->children[typedAt->op.args.ImageLayer.blended].child() );
+
+                std::array<uint8_t, (size_t)EImageFormat::IF_COUNT> newState;
+                newState.fill(0);
+                // TODO
+                //newState[ IF_L_UBYTE ] = 1;
+                //newState[ IF_L_UBYTE_RLE ] = 1;
+
+                if ( typedAt->op.args.ImageLayer.mask )
+                {
+                    RecurseWithState( typedAt->children[typedAt->op.args.ImageLayer.mask].child(), newState );
+                }
+                break;
+            }
+
+            case OP_TYPE::IM_DISPLACE:
+            {
+                auto typedAt = dynamic_cast<const ASTOpFixed*>(at.get());
+
+                RecurseWithCurrentState( typedAt->children[typedAt->op.args.ImageDisplace.source].child() );
+
+                std::array<uint8_t, (size_t)EImageFormat::IF_COUNT> newState;
+                newState.fill(0);
+                newState[(size_t)EImageFormat::IF_L_UBYTE ] = 1;
+                newState[(size_t)EImageFormat::IF_L_UBYTE_RLE ] = 1;
+                RecurseWithState( typedAt->children[typedAt->op.args.ImageDisplace.displacementMap].child(), newState );
+                break;
+            }
+
+            default:
+                SetCurrentState( m_initialState );
+                recurse = true;
+                break;
+            }
+
+            return recurse;
+        }
+
+
+        bool IsSupportedFormat( Ptr<ASTOp> at, EImageFormat format ) const
+        {
+            auto it = m_supportedFormats.find(at);
+            if (it==m_supportedFormats.end())
+            {
+                return false;
+            }
+
+            return it->second[(size_t)format]!=0;
+        }
+
+
+    private:
+
+        //! Formats known to be supported for every instruction.
+        //! IF_COUNT*code.size() entries
+        std::unordered_map< Ptr<ASTOp>, std::array<uint8_t, (size_t)EImageFormat::IF_COUNT> > m_supportedFormats;
+
+        //! Constant convenience initial value
+        std::array<uint8_t, (size_t)EImageFormat::IF_COUNT> m_initialState;
+
+        //! Constant convenience initial value
+        std::array<uint8_t, (size_t)EImageFormat::IF_COUNT> m_allSupported;
+
+    };
+
+
+    //---------------------------------------------------------------------------------------------
+    void SubtreeRelevantParametersVisitorAST::Run( Ptr<ASTOp> root )
+    {
+        // Cached?
+        auto it = m_resultCache.find( STATE(root,false) );
+        if (it!=m_resultCache.end())
+        {
+            m_params = it->second;
+            return;
+        }
+
+        // Not cached
+        {
+            MUTABLE_CPUPROFILER_SCOPE(SubtreeRelevantParametersVisitorAST);
+
+            m_params.clear();
+
+            // The state is the onlyLayoutRelevant flag
+            ASTOp::Traverse_TopDown_Unique_Imprecise_WithState<bool>( root, false,
+                [&]( Ptr<ASTOp>& at, bool& state, vector<pair<Ptr<ASTOp>,bool>>& pending )
+            {
+                (void)state;
+
+                switch ( at->GetOpType() )
+                {
+                case OP_TYPE::NU_PARAMETER:
+                case OP_TYPE::SC_PARAMETER:
+                case OP_TYPE::BO_PARAMETER:
+                case OP_TYPE::CO_PARAMETER:
+                case OP_TYPE::PR_PARAMETER:
+                case OP_TYPE::IM_PARAMETER:
+                {
+                    auto typedAt = dynamic_cast<const ASTOpParameter*>(at.get());
+                    m_params.insert(typedAt->parameter.m_name);
+
+                    // Not interested in the parameters from the parameters decorators.
+                    return false;
+
+                    break;
+                }
+
+                case OP_TYPE::LA_REMOVEBLOCKS:
+                {
+                    // Manually choose how to recurse this op
+                    auto pTyped = dynamic_cast<const ASTOpFixed*>( at.get() );
+
+                    if ( auto& source = pTyped->children[pTyped->op.args.LayoutRemoveBlocks.source] )
+                    {
+                        pending.push_back( std::make_pair<>(source.m_child, state) );
+                    }
+
+                    // For that mesh we only want to know about the layouts
+                    if ( auto& mesh = pTyped->children[pTyped->op.args.LayoutRemoveBlocks.mesh] )
+                    {
+                        pending.push_back( std::make_pair<>(mesh.m_child, true) );
+                    }
+
+                    return false;
+                }
+
+                case OP_TYPE::ME_MORPH2:
+                {
+                    // Manually choose how to recurse this op
+                    auto pTyped = dynamic_cast<const ASTOpFixed*>( at.get() );
+
+                    if ( auto& base = pTyped->children[pTyped->op.args.MeshMorph2.base] )
+                    {
+                        pending.push_back( std::make_pair<>(base.m_child, state) );
+                    }
+
+                    // Mesh morphs don't modify the layouts, so we can ignore the factor and morphs
+                    if (!state)
+                    {
+                        if ( auto& factor = pTyped->children[pTyped->op.args.MeshMorph2.factor] )
+                        {
+                            pending.push_back( std::make_pair<>(factor.m_child, state) );
+                        }
+
+                        for (int t=0;t<MUTABLE_OP_MAX_MORPH2_TARGETS;++t)
+                        {
+                            if ( auto& target = pTyped->children[pTyped->op.args.MeshMorph2.targets[t]] )
+                            {
+                                pending.push_back( std::make_pair<>(target.m_child, state) );
+                            }
+                        }
+                    }
+
+                    return false;
+                }
+
+                default:
+                    break;
+                }
+
+                return true;
+            });
+
+            m_resultCache[STATE(root,false)] = m_params;
+        }
+    }
+
+
+    //---------------------------------------------------------------------------------------------
+    //! Mark all the instructions that don't depend on runtime parameters but are below
+    //! instructions that do.
+    //! Also detect which instructions are the root of a resource that is dynamic in this state.
+    //! Visitor state is:
+    //!   .first IsResourceRoot
+    //!   .second ParentIsRuntime
+    //---------------------------------------------------------------------------------------------
+    class StateCacheDetectorAST : public Visitor_TopDown_Unique_Const< std::pair<bool,bool> >
+    {
+    public:
+
+        StateCacheDetectorAST( STATE_COMPILATION_DATA* pState )
+            : m_hasRuntimeParamVisitor( pState )
+        {
+            ASTOpList roots;
+            roots.push_back(pState->root);
+            Traverse( roots, std::make_pair<>(false,false) );
+
+            pState->m_updateCache.clear();
+            pState->m_dynamicResources.clear();
+
+            for( auto i=m_cache.cbegin();
+                 i != m_cache.cend();
+                 ++i )
+            {
+                if ( i->second )
+                {
+                    pState->m_updateCache.push_back( i->first );
+                }
+            }
+
+            for( auto i=m_dynamicResourceRoot.cbegin();
+                 i != m_dynamicResourceRoot.cend();
+                 ++i )
+            {
+                if ( i->second )
+                {
+                    // Generate the list of relevant parameters
+                    SubtreeRelevantParametersVisitorAST subtreeParams;
+                    subtreeParams.Run( i->first );
+
+                    pState->m_dynamicResources.emplace_back( i->first, std::move(subtreeParams.m_params) );
+                }
+            }
+        }
+
+
+        bool Visit( const Ptr<ASTOp>& at ) override
+        {
+            bool thisIsRuntime = m_hasRuntimeParamVisitor.HasAny( at );
+
+            bool resourceRoot = GetCurrentState().first;
+            bool parentIsRuntime = GetCurrentState().second;
+
+            OP_TYPE type = at->GetOpType();
+            if ( GetOpDesc( type ).cached )
+            {
+                // If parent is runtime, but we are not
+                if ( (!thisIsRuntime) && parentIsRuntime
+                     &&
+                     // Resource roots are special, and they don't need to be marked as updateCache
+                     // since the dynamicResource flag takes care of everything.
+                     !resourceRoot )
+                {
+                    // We want to cache this result to update the instances.
+                    // Mark this as update cache
+                    m_cache[at] = true;
+                }
+            }
+
+            if ( !m_cache[at] && resourceRoot && thisIsRuntime )
+            {
+                m_dynamicResourceRoot[at] = true;
+            }
+
+            if ( !m_cache[at] && thisIsRuntime )
+            {
+                switch( type )
+                {
+                case OP_TYPE::IN_ADDIMAGE:
+                case OP_TYPE::IN_ADDMESH:
+                case OP_TYPE::IN_ADDVECTOR:
+                case OP_TYPE::IN_ADDSCALAR:
+                case OP_TYPE::IN_ADDSTRING:
+                {
+                    auto typedAt = dynamic_cast<const ASTOpInstanceAdd*>(at.get());
+
+                    std::pair<bool,bool> newState;
+                    newState.first = false; //resource root
+                    newState.second = thisIsRuntime;
+                    RecurseWithState( typedAt->instance.child(), newState );
+
+                    if ( typedAt->value )
+                    {
+                        newState.first = true; //resource root
+                        newState.second = thisIsRuntime;
+                        RecurseWithState( typedAt->value.child(), newState );
+                    }
+                    return false;
+                }
+
+                default:
+                {
+                    std::pair<bool,bool> newState;
+                    newState.first = false; //resource root
+                    newState.second = thisIsRuntime;
+                    SetCurrentState(newState);
+                    return true;
+                }
+
+                }
+            }
+
+            return false;
+        }
+
+    private:
+
+        //!
+        std::map<Ptr<ASTOp>,bool> m_cache;
+        std::map<Ptr<ASTOp>,bool> m_dynamicResourceRoot;
+
+        RuntimeParameterVisitorAST m_hasRuntimeParamVisitor;
+    };
+
+
+    //---------------------------------------------------------------------------------------------
+    //! Find out what images can be compressed during build phase of an instance so that the update
+    //! cache can be smaller (and some update operations faster)
+    //---------------------------------------------------------------------------------------------
+    class StateCacheFormatOptimiserAST : public Visitor_TopDown_Unique_Cloning
+    {
+    public:
+
+        StateCacheFormatOptimiserAST( STATE_COMPILATION_DATA& state,
+                                   const AccumulateAllImageFormatsOpAST& opFormats )
+            : m_state(state)
+            , m_opFormats(opFormats)
+        {
+            Traverse( state.root );
+        }
+
+
+    protected:
+
+        Ptr<ASTOp> Visit( Ptr<ASTOp> at, bool& processChildren ) override
+        {
+            processChildren = true;
+
+            auto it = std::find(m_state.m_updateCache.cbegin(), m_state.m_updateCache.cend(), at);
+            bool isUpdateCache = it!=m_state.m_updateCache.cend();
+
+            if ( isUpdateCache )
+            {
+                // Its children cannot be update-cache, so no need to process them.
+                processChildren = false;
+
+                // See if we can convert it to a more efficient format
+                if ( GetOpDataType( at->GetOpType() )==DT_IMAGE )
+                {
+                    FImageDesc desc = at->GetImageDesc();
+
+                    if ( desc.m_format!=EImageFormat::IF_L_UBYTE_RLE
+                         &&
+                         m_opFormats.IsSupportedFormat(at, EImageFormat::IF_L_UBYTE_RLE) )
+                    {
+                        Ptr<ASTOpImagePixelFormat> op = new ASTOpImagePixelFormat;
+                        op->Format = EImageFormat::IF_L_UBYTE_RLE;
+                        // Note: we have to clone here, to avoid a loop with the visitor system
+                        // that updates visited children before processing a node.
+						ASTOp::MapChildFunc Identity = [](const Ptr<ASTOp>& o) {return o;};
+						op->Source = at->Clone(Identity);
+
+                        at = op;
+                    }
+                }
+
+            }
+
+            return at;
+        }
+
+    private:
+
+        STATE_COMPILATION_DATA& m_state;
+
+        const AccumulateAllImageFormatsOpAST& m_opFormats;
+    };
+
+
+    //---------------------------------------------------------------------------------------------
+    //---------------------------------------------------------------------------------------------
+    //---------------------------------------------------------------------------------------------
+    RuntimeTextureCompressionRemoverAST::RuntimeTextureCompressionRemoverAST(
+            STATE_COMPILATION_DATA* pState
+            )
+        : m_hasRuntimeParamVisitor(pState)
+    {
+        Traverse( pState->root );
+    }
+
+
+    Ptr<ASTOp> RuntimeTextureCompressionRemoverAST::Visit( Ptr<ASTOp> at, bool& processChildren )
+    {
+        OP_TYPE type = at->GetOpType();
+        processChildren = GetOpDataType(type)==DT_INSTANCE;
+
+        // TODO: Finer grained: what if the runtime parameter just selectes between compressed
+        // textures? We don't want them uncompressed.
+        if( type==OP_TYPE::IN_ADDIMAGE )
+        {
+            auto typedAt = dynamic_cast<ASTOpInstanceAdd*>(at.get());
+            auto imageAt = typedAt->value.child();
+
+            // Does it have a runtime parameter in its subtree?
+            bool hasRuntimeParameter = m_hasRuntimeParamVisitor.HasAny( imageAt );
+
+            if (hasRuntimeParameter)
+            {
+                FImageDesc imageDesc = imageAt->GetImageDesc( true );
+
+                // Is it a compressed format?
+				EImageFormat format = imageDesc.m_format;
+				EImageFormat uncompressedFormat = GetUncompressedFormat( format );
+                bool isCompressedFormat = (uncompressedFormat != format);
+
+                if (isCompressedFormat)
+                {
+                    Ptr<ASTOpInstanceAdd> newAt = mu::Clone<ASTOpInstanceAdd>(at);
+
+                    // Add a new format operation to uncompress the image
+                    Ptr<ASTOpImagePixelFormat> fop = new ASTOpImagePixelFormat;
+                    fop->Format = uncompressedFormat;
+                    fop->FormatIfAlpha = uncompressedFormat;
+                    fop->Source = imageAt;
+
+                    newAt->value = fop;
+                    at = newAt;
+                }
+            }
+        }
+
+        return at;
+    }
+
+
+    //---------------------------------------------------------------------------------------------
+    //---------------------------------------------------------------------------------------------
+    //---------------------------------------------------------------------------------------------
+    LODCountReducerAST::LODCountReducerAST( Ptr<ASTOp>& root, int lodCount )
+    {
+        m_lodCount = lodCount;
+        Traverse( root );
+    }
+
+
+    Ptr<ASTOp> LODCountReducerAST::Visit( Ptr<ASTOp> at, bool& processChildren )
+    {
+        processChildren = true;
+
+        if( at->GetOpType()==OP_TYPE::IN_ADDLOD )
+        {
+            auto* typedAt = dynamic_cast<ASTOpAddLOD*>(at.get());
+
+            if (typedAt->lods.size()>(size_t)m_lodCount)
+            {
+                Ptr<ASTOpAddLOD> newAt = mu::Clone<ASTOpAddLOD>(at);
+                while( newAt->lods.size()>(size_t)m_lodCount )
+                {
+                    newAt->lods.pop_back();
+                }
+                at = newAt;
+            }
+
+            processChildren = false;
+        }
+
+        return at;
+    }
+
+
+    //---------------------------------------------------------------------------------------------
+    //---------------------------------------------------------------------------------------------
+    //---------------------------------------------------------------------------------------------
+    void CodeOptimiser::OptimiseStatesAST( TaskManager* pTaskManager )
+    {
+        MUTABLE_CPUPROFILER_SCOPE(OptimiseStatesAST);
+
+         for ( size_t s=0; s<m_states.size(); ++s )
+         {
+            // Remove the unnecessary lods
+            if (m_states[s].optimisationFlags.m_onlyFirstLOD)
+            {
+                LODCountReducerAST(m_states[s].root, m_states[s].optimisationFlags.m_firstLOD + 1);
+            }
+
+            // If a state has no runtime parameters, skip its optimisation alltogether
+            if (m_states[s].nodeState.m_runtimeParams.size())
+            {
+                // Remove unnecessary image compression
+                if (m_states[s].optimisationFlags.m_avoidRuntimeCompression)
+                {
+                    MUTABLE_CPUPROFILER_SCOPE(RuntimeTextureCompressionRemover);
+                    RuntimeTextureCompressionRemoverAST r(&m_states[s]);
+                }
+
+                // Promote the intructions that depend on runtime parameters, and sink new
+                // format instructions.
+                bool modified = true;
+                int numIterations = 0;
+                while (modified && (!m_optimizeIterationsMax || m_optimizeIterationsLeft>0 || !numIterations ))
+                {
+                    modified = false;
+
+                    ++numIterations;
+                    --m_optimizeIterationsLeft;
+                    UE_LOG(LogMutableCore, Verbose, TEXT("State optimise iteration %d, max %d, left %d"),
+                            numIterations, m_optimizeIterationsMax, m_optimizeIterationsLeft);
+
+                    UE_LOG(LogMutableCore, Verbose, TEXT(" - before parameter optimiser"));
+
+                    ParameterOptimiserAST param( m_states[s],
+                                                 m_options->GetPrivate()->m_optimisationOptions );
+                    modified = param.Apply();
+
+                    vector<Ptr<ASTOp>> roots;
+                    roots.push_back(m_states[s].root);
+
+                    UE_LOG(LogMutableCore, Verbose, TEXT(" - after parameter optimiser"));
+					UE_LOG(LogMutableCore, Verbose, TEXT("(int) %s : %ld"), TEXT("ast size"), int64(ASTOp::CountNodes(roots)));
+
+                    // All kind of optimisations that depend on the meaning of each operation
+                    UE_LOG(LogMutableCore, Verbose, TEXT(" - semantic optimiser"));
+                    modified |= SemanticOptimiserAST(
+                        roots, m_options->GetPrivate()->m_optimisationOptions );
+					UE_LOG(LogMutableCore, Verbose, TEXT("(int) %s : %ld"), TEXT("ast size"), int64(ASTOp::CountNodes(roots)));
+					//ASTOp::LogHistogram(roots);
+
+                    UE_LOG(LogMutableCore, Verbose, TEXT(" - sink optimiser"));
+                    modified |=
+                        SinkOptimiserAST( roots, m_options->GetPrivate()->m_optimisationOptions );
+					UE_LOG(LogMutableCore, Verbose, TEXT("(int) %s : %ld"), TEXT("ast size"), int64(ASTOp::CountNodes(roots)));
+					//ASTOp::LogHistogram(roots);
+
+                    // Image size operations are treated separately
+                    UE_LOG(LogMutableCore, Verbose, TEXT(" - size optimiser"));
+                    modified |= SizeOptimiserAST( roots );
+					//UE_LOG(LogMutableCore, Verbose, TEXT("(int) %s : %ld"), TEXT("ast size"), int64(ASTOp::CountNodes(roots)));
+
+                    //LogicOptimiser log;
+                    //modified |= log.Apply( program, s );
+                }
+
+                vector<Ptr<ASTOp>> roots;
+                roots.push_back(m_states[s].root);
+
+                UE_LOG(LogMutableCore, Verbose, TEXT(" - duplicated data remover"));
+                DuplicatedDataRemoverAST( roots );
+				//UE_LOG(LogMutableCore, Verbose, TEXT("(int) %s : %ld"), TEXT("ast size"), int64(ASTOp::CountNodes(roots)));
+
+                UE_LOG(LogMutableCore, Verbose, TEXT(" - duplicated code remover"));
+                DuplicatedCodeRemoverAST( roots );
+				//UE_LOG(LogMutableCore, Verbose, TEXT("(int) %s : %ld"), TEXT("ast size"), int64(ASTOp::CountNodes(roots)));
+
+                m_states[s].root = roots.front();
+            }
+        }
+
+        // Mark the instructions that don't depend on runtime parameters to be cached. This is
+        // necessary at this stage before GPU optimisation.
+        {
+            vector<Ptr<ASTOp>> roots;
+            for(const auto& s:m_states)
+            {
+                roots.push_back(s.root);
+            }
+
+            AccumulateAllImageFormatsOpAST opFormats;
+            opFormats.Run(roots);
+
+            for ( auto& s: m_states )
+            {
+                {
+                    UE_LOG(LogMutableCore, Verbose, TEXT(" - state cache"));
+                    MUTABLE_CPUPROFILER_SCOPE(StateCache);
+                    StateCacheDetectorAST c( &s );
+                }
+
+                {
+                    UE_LOG(LogMutableCore, Verbose, TEXT(" - state cache format"));
+                    MUTABLE_CPUPROFILER_SCOPE(StateCacheFormat);
+                    StateCacheFormatOptimiserAST f( s, opFormats );
+                }
+            }
+        }
+
+        // If we are targetting GPU build, gpuize the code
+//        m_pModelReport->GetPrivate()->m_states.resize( pModel->GetStateCount() );
+
+//        for ( int s=0; s<pModel->GetStateCount(); ++s )
+//        {
+//            if ( m_gpuPlatformProps.type!=GPU_NONE
+//                 &&
+//                 ( m_states[s].optimisationFlags.m_gpu.m_internal || m_states[s].optimisationFlags.m_gpu.m_external ) )
+//            {
+//                MUTABLE_CPUPROFILER_SCOPE(GPUTranslator);
+//                GPUTranslator( program,
+//                               s,
+//                               m_gpuPlatformProps,
+//                               m_states[s].optimisationFlags.m_gpu,
+//                               m_pModelReport );
+//            }
+//        }
+
+        //program.Check();
+
+        // Reoptimise because of state cache reformats and gpuization
+        {
+            MUTABLE_CPUPROFILER_SCOPE(Reoptimise);
+            bool modified = true;
+            int numIterations = 0;
+            while (modified && (!m_optimizeIterationsMax || m_optimizeIterationsLeft>0 || !numIterations ))
+            {
+                vector<Ptr<ASTOp>> roots;
+                for(const auto& s:m_states)
+                {
+                    roots.push_back(s.root);
+                }
+
+                ++numIterations;
+                --m_optimizeIterationsLeft;
+                UE_LOG(LogMutableCore, Verbose, TEXT("State reoptimise iteration %d, max %d, left %d"),
+                        numIterations, m_optimizeIterationsMax, m_optimizeIterationsLeft);
+
+                modified = false;
+
+                UE_LOG(LogMutableCore, Verbose, TEXT(" - semantic optimiser"));
+                modified |=
+                    SemanticOptimiserAST( roots, m_options->GetPrivate()->m_optimisationOptions );
+				UE_LOG(LogMutableCore, Verbose, TEXT("(int) %s : %ld"), TEXT("ast size"), int64(ASTOp::CountNodes(roots)));
+
+                // Image size operations are treated separately
+                UE_LOG(LogMutableCore, Verbose, TEXT(" - size optimiser"));
+                modified |= SizeOptimiserAST( roots );
+				UE_LOG(LogMutableCore, Verbose, TEXT("(int) %s : %ld"), TEXT("ast size"), int64(ASTOp::CountNodes(roots)));
+			}
+
+            for(auto& s:m_states)
+            {
+                UE_LOG(LogMutableCore, Verbose, TEXT(" - constant optimiser"));
+                ConstantGeneratorAST( m_options->GetPrivate(), s.root, pTaskManager );
+            }
+
+            vector<Ptr<ASTOp>> roots;
+            for(const auto& s:m_states)
+            {
+                roots.push_back(s.root);
+            }
+			UE_LOG(LogMutableCore, Verbose, TEXT("(int) %s : %ld"), TEXT("ast size"), int64(ASTOp::CountNodes(roots)));
+
+            UE_LOG(LogMutableCore, Verbose, TEXT(" - duplicated data remover"));
+            DuplicatedDataRemoverAST( roots );
+			//UE_LOG(LogMutableCore, Verbose, TEXT("(int) %s : %ld"), TEXT("ast size"), int64(ASTOp::CountNodes(roots)));
+
+            UE_LOG(LogMutableCore, Verbose, TEXT(" - duplicated code remover"));
+            DuplicatedCodeRemoverAST( roots );
+			//UE_LOG(LogMutableCore, Verbose, TEXT("(int) %s : %ld"), TEXT("ast size"), int64(ASTOp::CountNodes(roots)));
+		}
+
+        // Gather all the current roots
+        vector<Ptr<ASTOp>> roots;
+        for(const auto& s:m_states)
+        {
+            roots.push_back(s.root);
+        }
+
+        // Optimise the data formats
+        {
+            MUTABLE_CPUPROFILER_SCOPE(DataFormats);
+
+            DataOptimiseAST( m_options->GetPrivate()->m_imageCompressionQuality, roots,
+                             m_options->GetPrivate()->m_optimisationOptions );
+
+            // After optimising the data formats, we may remove more constants
+            DuplicatedDataRemoverAST( roots );
+            DuplicatedCodeRemoverAST( roots );
+
+            // Update the marks for the instructions that don't depend on runtime parameters to be
+            // cached.
+            for (auto& s:m_states)
+            {
+                StateCacheDetectorAST c( &s );
+            }
+        }
+
+    }
+
+
+}
+
