[mutable] Move Mutable plugin from the restricted folder to the experimental plugins folder.

diff --git a/Engine/Plugins/Experimental/Mutable/Source/MutableTools/Private/CodeOptimiser_Semantic.cpp b/Engine/Plugins/Experimental/Mutable/Source/MutableTools/Private/CodeOptimiser_Semantic.cpp
new file mode 100644
--- /dev/null
+++ b/Engine/Plugins/Experimental/Mutable/Source/MutableTools/Private/CodeOptimiser_Semantic.cpp
@@ -0,0 +1,2671 @@
+// Copyright Epic Games, Inc. All Rights Reserved.
+
+#include "MutableTrace.h"
+#include "MutableTools/Private/CodeOptimiser.h"
+#include "MutableTools/Private/CodeGenerator.h"
+#include "MutableTools/Private/AST.h"
+#include "MutableTools/Private/ASTOpMeshFormat.h"
+#include "MutableTools/Private/ASTOpImagePixelFormat.h"
+#include "MutableTools/Private/ASTOpImageMipmap.h"
+#include "MutableTools/Private/ASTOpImagePatch.h"
+#include "MutableTools/Private/ASTOpImageCompose.h"
+#include "MutableTools/Private/ASTOpImageMultiLayer.h"
+#include "MutableTools/Private/ASTOpConditional.h"
+#include "MutableTools/Private/ASTOpConstantBool.h"
+
+#include "MutableRuntime/Private/ModelPrivate.h"
+#include "MutableRuntime/Private/SystemPrivate.h"
+#include "MutableRuntime/Private/Operations.h"
+#include "MutableRuntime/Private/ImagePrivate.h"
+#include "MutableRuntime/Private/ImageRLE.h"
+#include "MutableRuntime/Private/OpMeshRemoveChart.h"
+
+namespace mu
+{
+
+    //---------------------------------------------------------------------------------------------
+    //! Set al the non-null sources of an image swizzle operation to the given value
+    //---------------------------------------------------------------------------------------------
+    void ReplaceAllSources( Ptr<ASTOpFixed>& op, Ptr<ASTOp>& value )
+    {
+        check( op->GetOpType() == OP_TYPE::IM_SWIZZLE );
+        for ( int c=0; c<MUTABLE_OP_MAX_SWIZZLE_CHANNELS; ++c )
+        {
+            if ( op->op.args.ImageSwizzle.sources[c] )
+            {
+                op->SetChild( op->op.args.ImageSwizzle.sources[c], value);
+            }
+        }
+    }
+
+
+    //---------------------------------------------------------------------------------------------
+    //! Get an uncompressed version of the image.
+    //! If the image is not compressed, it si returned unmodified.
+    //---------------------------------------------------------------------------------------------
+    ImagePtrConst GetUncompressed( ImagePtrConst pSource )
+    {
+        ImagePtrConst pResult = pSource;
+
+        if ( pSource->GetFormat()== EImageFormat::IF_L_UBYTE_RLE )
+        {
+            ImagePtr pNew = new Image( pSource->GetSizeX(), pSource->GetSizeY(),
+                                 pSource->GetLODCount(),
+				EImageFormat::IF_L_UBYTE );
+            UncompressRLE_L( pSource.get(), pNew.get() );
+            pResult = pNew;
+        }
+        else if ( pSource->GetFormat()== EImageFormat::IF_L_UBIT_RLE )
+        {
+            ImagePtr pNew = new Image( pSource->GetSizeX(), pSource->GetSizeY(),
+                                 pSource->GetLODCount(),
+				EImageFormat::IF_L_UBYTE );
+            UncompressRLE_L1( pSource.get(), pNew.get() );
+            pResult = pNew;
+        }
+
+        return pResult;
+    }
+
+
+    //---------------------------------------------------------------------------------------------
+    //---------------------------------------------------------------------------------------------
+    //---------------------------------------------------------------------------------------------
+    bool SemanticOptimiserAST(
+            ASTOpList& roots,
+            const MODEL_OPTIMIZATION_OPTIONS& optimisationOptions
+            )
+    {
+        MUTABLE_CPUPROFILER_SCOPE(SemanticOptimiserAST);
+
+        bool modified = false;
+
+        // TODO: isn't top down better suited?
+        ASTOp::Traverse_BottomUp_Unique( roots, [&](Ptr<ASTOp>& n)
+        {
+            auto o = n->OptimiseSemantic(optimisationOptions);
+
+            // If the returned value is null it means no change.
+            if (o && o!=n)
+            {
+                modified = true;
+                ASTOp::Replace(n,o);
+            }
+        });
+
+        return modified;
+    }
+
+
+    //---------------------------------------------------------------------------------------------
+    Ptr<ASTOp> ASTOpFixed::OptimiseSemantic( const MODEL_OPTIMIZATION_OPTIONS& options ) const
+    {
+        Ptr<ASTOp> at;
+
+        OP_TYPE type = GetOpType();
+        switch ( type )
+        {
+
+        case OP_TYPE::BO_NOT:
+        {
+            auto sourceAt = children[op.args.BoolNot.source].child();
+            if  ( sourceAt->GetOpType() == OP_TYPE::BO_CONSTANT )
+            {
+                Ptr<ASTOpConstantBool> newAt = new ASTOpConstantBool();
+                newAt->value = !newAt->value;
+                at=newAt;
+            }
+            break;
+        }
+
+        case OP_TYPE::BO_AND:
+        {
+            bool changed = false;
+            auto aAt = children[op.args.BoolBinary.a].child();
+            auto bAt = children[op.args.BoolBinary.b].child();
+            if ( !aAt )
+            {
+                at = bAt;
+                changed = true;
+            }
+            else if ( !bAt )
+            {
+                at = aAt;
+                changed = true;
+            }
+            else if ( aAt->GetOpType() == OP_TYPE::BO_CONSTANT )
+            {
+                if ( dynamic_cast<const ASTOpConstantBool*>(aAt.get())->value )
+                {
+                    at = bAt;
+                    changed = true;
+                }
+                else
+                {
+                    at = aAt;
+                    changed = true;
+                }
+            }
+            else if ( bAt->GetOpType() == OP_TYPE::BO_CONSTANT )
+            {
+                if ( dynamic_cast<const ASTOpConstantBool*>(bAt.get())->value )
+                {
+                    at = aAt;
+                    changed = true;
+                }
+                else
+                {
+                    at = bAt;
+                    changed = true;
+                }
+            }
+
+            // Common cases of repeated branch in children
+            else if ( aAt->GetOpType() == type )
+            {
+                auto typedA = dynamic_cast<const ASTOpFixed*>(aAt.get());
+                if ( typedA->children[typedA->op.args.BoolBinary.a].child()==bAt
+                     ||
+                     typedA->children[typedA->op.args.BoolBinary.b].child()==bAt )
+                {
+                    at = aAt;
+                    changed = true;
+                }
+            }
+            else if ( bAt->GetOpType() == type )
+            {
+                auto typedB = dynamic_cast<const ASTOpFixed*>(bAt.get());
+                if ( typedB->children[typedB->op.args.BoolBinary.b].child()==aAt
+                     ||
+                     typedB->children[typedB->op.args.BoolBinary.b].child()==bAt )
+                {
+                    at = bAt;
+                    changed = true;
+                }
+            }
+
+            else if (aAt == bAt || *aAt==*bAt)
+            {
+                at = aAt;
+                changed = true;
+            }
+
+            // if it became null, it means true (neutral AND argument)
+            if (changed && !at)
+            {
+                Ptr<ASTOpConstantBool> newAt = new ASTOpConstantBool();
+                newAt->value = true;
+                at = newAt;
+            }
+            break;
+        }
+
+        case OP_TYPE::BO_OR:
+        {
+            bool changed = false;
+            auto aAt = children[op.args.BoolBinary.a].child();
+            auto bAt = children[op.args.BoolBinary.b].child();
+            if ( !aAt )
+            {
+                at = bAt;
+                changed = true;
+            }
+            else if ( !bAt )
+            {
+                at = aAt;
+                changed = true;
+            }
+            else if ( aAt->GetOpType() == OP_TYPE::BO_CONSTANT )
+            {
+                if ( dynamic_cast<const ASTOpConstantBool*>(aAt.get())->value )
+                {
+                    at = aAt;
+                    changed = true;
+                }
+                else
+                {
+                    at = bAt;
+                    changed = true;
+                }
+            }
+            else if ( bAt->GetOpType() == OP_TYPE::BO_CONSTANT )
+            {
+                if ( dynamic_cast<const ASTOpConstantBool*>(bAt.get())->value )
+                {
+                    at = bAt;
+                    changed = true;
+                }
+                else
+                {
+                    at = aAt;
+                    changed = true;
+                }
+            }
+
+            // Common cases of repeated branch in children
+            else if ( aAt->GetOpType() == type )
+            {
+                auto typedA = dynamic_cast<const ASTOpFixed*>(aAt.get());
+                if ( typedA->children[typedA->op.args.BoolBinary.a].child()==bAt
+                     ||
+                     typedA->children[typedA->op.args.BoolBinary.b].child()==bAt )
+                {
+                    at = aAt;
+                    changed = true;
+                }
+            }
+            else if ( bAt->GetOpType() == type )
+            {
+                auto typedB = dynamic_cast<const ASTOpFixed*>(bAt.get());
+                if ( typedB->children[typedB->op.args.BoolBinary.b].child()==aAt
+                     ||
+                     typedB->children[typedB->op.args.BoolBinary.b].child()==bAt )
+                {
+                    at = bAt;
+                    changed = true;
+                }
+            }
+
+            else if (aAt == bAt || *aAt==*bAt)
+            {
+                at = aAt;
+                changed = true;
+            }
+
+            // if it became null, it means false (neutral OR argument)
+            if (changed && !at)
+            {
+                Ptr<ASTOpConstantBool> newAt = new ASTOpConstantBool();
+                newAt->value = false;
+                at = newAt;
+            }
+
+            break;
+        }
+
+
+        //-------------------------------------------------------------------------------------
+        case OP_TYPE::IM_PIXELFORMAT:
+        {
+            // All in the sink optimiser?
+            break;
+        }
+
+
+        case OP_TYPE::IM_SWIZZLE:
+        {
+            // If children channels are also swizzle ops, recurse them
+            {
+                Ptr<ASTOpFixed> sat;
+
+                for ( int c=0; c<MUTABLE_OP_MAX_SWIZZLE_CHANNELS; ++c )
+                {
+                    auto candidate = children[op.args.ImageSwizzle.sources[c]].child();
+
+                    // Swizzle
+                    if (candidate && candidate->GetOpType()==OP_TYPE::IM_SWIZZLE)
+                    {
+                        if (!sat)
+                        {
+                            sat = mu::Clone<ASTOpFixed>(this);
+                        }
+                        ASTOpFixed* typedCandidate = dynamic_cast<ASTOpFixed*>(candidate.get());
+                        int candidateChannel = op.args.ImageSwizzle.sourceChannels[c];
+
+                        sat->SetChild( sat->op.args.ImageSwizzle.sources[c],
+                                       typedCandidate->children[typedCandidate->op.args.ImageSwizzle.sources[candidateChannel]] );
+                        sat->op.args.ImageSwizzle.sourceChannels[c] =  typedCandidate->op.args.ImageSwizzle.sourceChannels[candidateChannel];
+                    }
+
+                    // Format
+                    if (candidate && candidate->GetOpType()==OP_TYPE::IM_PIXELFORMAT)
+                    {
+                        // We can remove the format if its source is already an uncompressed format
+                        ASTOpImagePixelFormat* typedCandidate = dynamic_cast<ASTOpImagePixelFormat*>(candidate.get());
+                        Ptr<ASTOp> formatSource = typedCandidate->Source.child();
+
+                        if (formatSource)
+                        {
+                            auto desc = formatSource->GetImageDesc();
+                            if ( desc.m_format!= EImageFormat::IF_NONE && !IsCompressedFormat( desc.m_format ) )
+                            {
+                                if (!sat)
+                                {
+                                    sat = mu::Clone<ASTOpFixed>(this);
+                                }
+                                sat->SetChild( sat->op.args.ImageSwizzle.sources[c], formatSource );
+                            }
+                        }
+                    }
+
+                }
+
+                if (sat)
+                {
+                    at = sat;
+                }
+            }
+
+            break;
+        }
+
+
+        //-------------------------------------------------------------------------------------
+        case OP_TYPE::IM_LAYER:
+        {
+            auto baseAt = children[op.args.ImageLayer.base].child();
+
+            // Plain masks optimization
+            auto maskAt = children[op.args.ImageLayer.mask].child();
+            if ( maskAt )
+            {
+                vec4<float> colour;
+                if ( maskAt->IsImagePlainConstant( colour ))
+                {
+                    if ( colour.xyz().AlmostNull() )
+                    {
+                        // If the mask is black, we can skip the entire operation
+                        at = children[op.args.ImageLayer.base].child();
+                    }
+                    else if ( colour.xyz().AlmostEqual( vec3<float>(1,1,1) ) )
+                    {
+                        // If the mask is white, we can remove it
+                        auto nop = mu::Clone<ASTOpFixed>(this);
+                        nop->op.args.ImageLayer.mask = 0;
+                        at = nop;
+                    }
+                }
+            }
+
+            // Introduce crop if mask is constant and smaller than the base
+            FImageRect sourceMaskUsage;
+            FImageDesc maskDesc;
+            bool validUsageRect = false;
+            if (!at && maskAt)
+            {
+                //MUTABLE_CPUPROFILER_SCOPE(EvaluateAreasForCrop);
+
+                validUsageRect = maskAt->GetNonBlackRect(sourceMaskUsage);
+                if (validUsageRect)
+                {
+					check(sourceMaskUsage.size[0] > 0);
+					check(sourceMaskUsage.size[1] > 0);
+					
+					GetImageDescContext context;
+                    maskDesc = maskAt->GetImageDesc( false, &context );
+                }
+            }
+
+            if (!at && maskAt && validUsageRect)
+            {
+                // Adjust for compressed blocks (4), and some extra mips (2 more mips, which is 4)
+                constexpr int blockSize = 4 * 4;
+
+                FImageRect maskUsage;
+                maskUsage.min[0] = (sourceMaskUsage.min[0]/blockSize)*blockSize;
+                maskUsage.min[1] = (sourceMaskUsage.min[1]/blockSize)*blockSize;
+                vec2<uint16_t> minOffset = sourceMaskUsage.min - maskUsage.min;
+                maskUsage.size[0] = ((sourceMaskUsage.size[0]+minOffset[0]+blockSize-1)/blockSize)*blockSize;
+                maskUsage.size[1] = ((sourceMaskUsage.size[1]+minOffset[1]+blockSize-1)/blockSize)*blockSize;
+
+                // Is it worth?
+                float ratio = float(maskUsage.size[0]*maskUsage.size[1])
+                        / float(maskDesc.m_size[0]*maskDesc.m_size[1]);
+                float acceptableCropRatio = options.m_acceptableCropRatio;
+                if (ratio<acceptableCropRatio)
+                {
+					check(maskUsage.size[0] > 0);
+					check(maskUsage.size[1] > 0);
+
+                    Ptr<ASTOpFixed> cropMask = new ASTOpFixed();
+                    cropMask->op.type = OP_TYPE::IM_CROP;
+                    cropMask->SetChild( cropMask->op.args.ImageCrop.source, children[op.args.ImageLayer.mask].child() );
+                    cropMask->op.args.ImageCrop.minX = maskUsage.min[0];
+                    cropMask->op.args.ImageCrop.minY = maskUsage.min[1];
+                    cropMask->op.args.ImageCrop.sizeX = maskUsage.size[0];
+                    cropMask->op.args.ImageCrop.sizeY = maskUsage.size[1];
+
+                    Ptr<ASTOpFixed> cropBlended = new ASTOpFixed();
+                    cropBlended->op.type = OP_TYPE::IM_CROP;
+                    cropBlended->SetChild( cropBlended->op.args.ImageCrop.source, children[op.args.ImageLayer.blended].child() );
+                    cropBlended->op.args.ImageCrop.minX = maskUsage.min[0];
+                    cropBlended->op.args.ImageCrop.minY = maskUsage.min[1];
+                    cropBlended->op.args.ImageCrop.sizeX = maskUsage.size[0];
+                    cropBlended->op.args.ImageCrop.sizeY = maskUsage.size[1];
+
+                    Ptr<ASTOpFixed> cropBase = new ASTOpFixed();
+                    cropBase->op.type = OP_TYPE::IM_CROP;
+                    cropBase->SetChild( cropBase->op.args.ImageCrop.source, children[op.args.ImageLayer.base].child() );
+                    cropBase->op.args.ImageCrop.minX = maskUsage.min[0];
+                    cropBase->op.args.ImageCrop.minY = maskUsage.min[1];
+                    cropBase->op.args.ImageCrop.sizeX = maskUsage.size[0];
+                    cropBase->op.args.ImageCrop.sizeY = maskUsage.size[1];
+
+                    Ptr<ASTOpFixed> newLayer = new ASTOpFixed();
+                    newLayer->op.type = op.type;
+					newLayer->op.args.ImageLayer.blendType = op.args.ImageLayer.blendType;
+					newLayer->op.args.ImageLayer.flags = op.args.ImageLayer.flags;
+					newLayer->SetChild( newLayer->op.args.ImageLayer.base, cropBase );
+                    newLayer->SetChild( newLayer->op.args.ImageLayer.blended, cropBlended );
+                    newLayer->SetChild( newLayer->op.args.ImageLayer.mask, cropMask );
+
+                    Ptr<ASTOpImagePatch> patch = new ASTOpImagePatch();
+                    patch->base = baseAt;
+                    patch->patch = newLayer;
+                    patch->location = maskUsage.min;
+                    at = patch;
+                }
+            }
+
+            break;
+        }
+
+        default:
+            break;
+        }
+
+        return at;
+    }
+
+
+    //---------------------------------------------------------------------------------------------
+    Ptr<ASTOp> ASTOpFixed::OptimiseSwizzle( const MODEL_OPTIMIZATION_OPTIONS& options ) const
+    {
+        MUTABLE_CPUPROFILER_SCOPE(OptimiseSwizzleAST);
+
+        //! Basic optimisation first
+        Ptr<ASTOp> at = OptimiseSemantic(options);
+        if (at)
+        {
+            if (at->GetOpType()==OP_TYPE::IM_SWIZZLE)
+            {
+                at = dynamic_cast<ASTOpFixed*>(at.get())->OptimiseSwizzle( options );
+            }
+            return at;
+        }
+
+        // If all sources are the same, we can sink the instruction
+        bool same = true;
+		int RepeatedChild = 0;
+        Ptr<ASTOp> channelSourceAt;
+        for ( int c=0; c<MUTABLE_OP_MAX_SWIZZLE_CHANNELS; ++c )
+        {
+            auto candidate = children[op.args.ImageSwizzle.sources[c]].child();
+            if ( candidate )
+            {
+                if (!channelSourceAt)
+                {
+                    channelSourceAt=candidate;
+                }
+                else
+                {
+                    same = channelSourceAt==candidate;
+					++RepeatedChild;
+				}
+            }
+        }
+
+        if (same && channelSourceAt)
+        {
+            // The instruction can be sunk
+            OP_TYPE sourceType = channelSourceAt->GetOpType();
+            switch ( sourceType )
+            {
+
+			case OP_TYPE::IM_PLAINCOLOUR:
+			{
+				auto NewPlain = mu::Clone<ASTOpFixed>(channelSourceAt);
+				Ptr<ASTOpFixed> NewSwizzle = new ASTOpFixed;
+				NewSwizzle->op.type = OP_TYPE::CO_SWIZZLE;
+				for (int i = 0; i < MUTABLE_OP_MAX_SWIZZLE_CHANNELS; ++i)
+				{
+					NewSwizzle->SetChild(NewSwizzle->op.args.ColourSwizzle.sources[i], NewPlain->children[NewPlain->op.args.ImagePlainColour.colour]);
+					NewSwizzle->op.args.ColourSwizzle.sourceChannels[i] = op.args.ImageSwizzle.sourceChannels[i];
+				}
+				NewPlain->SetChild(NewPlain->op.args.ImagePlainColour.colour, NewSwizzle);
+				at = NewPlain;
+				break;
+			}
+
+            case OP_TYPE::IM_SWITCH:
+            {
+                // Move the swizzle down all the paths
+                auto nop = mu::Clone<ASTOpSwitch>(channelSourceAt);
+
+                if (nop->def)
+                {
+                    auto defOp = mu::Clone<ASTOpFixed>(this);
+                    ReplaceAllSources( defOp, nop->def.child() );
+                    nop->def = defOp;
+                }
+
+                // We need to copy the options because we change them
+                for ( size_t v=0; v<nop->cases.size(); ++v )
+                {
+                    if ( nop->cases[v].branch )
+                    {
+                        auto bOp = mu::Clone<ASTOpFixed>(this);
+                        ReplaceAllSources( bOp, nop->cases[v].branch.child() );
+                        nop->cases[v].branch = bOp;
+                    }
+                }
+
+                at = nop;
+                break;
+            }
+
+            case OP_TYPE::IM_CONDITIONAL:
+            {
+                // We move the swizzle down the two paths
+                auto nop = mu::Clone<ASTOpConditional>(channelSourceAt);
+
+                auto aOp = mu::Clone<ASTOpFixed>(this);
+                ReplaceAllSources( aOp, nop->yes.child() );
+                nop->yes = aOp;
+
+                auto bOp = mu::Clone<ASTOpFixed>(this);
+                ReplaceAllSources( bOp, nop->no.child() );
+                nop->no = bOp;
+
+                at = nop;
+                break;
+            }
+
+            case OP_TYPE::IM_LAYER:
+            {
+                // We move the swizzle down the two paths
+                auto nop = mu::Clone<ASTOpFixed>(channelSourceAt);
+
+                auto aOp = mu::Clone<ASTOpFixed>(this);
+                ReplaceAllSources( aOp, nop->children[nop->op.args.ImageLayer.base].child() );
+                nop->SetChild( nop->op.args.ImageLayer.base, aOp );
+
+                auto bOp = mu::Clone<ASTOpFixed>(this);
+                ReplaceAllSources( bOp, nop->children[nop->op.args.ImageLayer.blended].child() );
+                nop->SetChild( nop->op.args.ImageLayer.blended, bOp );
+
+                at = nop;
+                break;
+            }
+
+            case OP_TYPE::IM_LAYERCOLOUR:
+            {
+                // We move the swizzle down the base paths
+                auto nop = mu::Clone<ASTOpFixed>(channelSourceAt);
+
+                auto NewSwizzle = mu::Clone<ASTOpFixed>(this);
+                ReplaceAllSources(NewSwizzle, nop->children[nop->op.args.ImageLayerColour.base].child() );
+                nop->SetChild( nop->op.args.ImageLayerColour.base, NewSwizzle);
+
+                // We need to swizzle the colour too
+                Ptr<ASTOpFixed> cOp = new ASTOpFixed;
+                cOp->op.type = OP_TYPE::CO_SWIZZLE;
+                for ( int i=0; i<MUTABLE_OP_MAX_SWIZZLE_CHANNELS; ++i)
+                {
+                    cOp->SetChild( cOp->op.args.ColourSwizzle.sources[i], nop->children[nop->op.args.ImageLayerColour.colour] );
+                    cOp->op.args.ColourSwizzle.sourceChannels[i] = nop->op.args.ImageSwizzle.sourceChannels[i];
+                }
+                nop->SetChild( nop->op.args.ImageLayerColour.colour, cOp );
+
+                at = nop;
+                break;
+            }
+
+			case OP_TYPE::IM_DISPLACE:
+			{
+				auto NewDisplace = mu::Clone<ASTOpFixed>(channelSourceAt);
+				auto NewSwizzle = mu::Clone<ASTOpFixed>(this);
+				ReplaceAllSources(NewSwizzle, NewDisplace->children[NewDisplace->op.args.ImageDisplace.source].child());
+				NewDisplace->SetChild(NewDisplace->op.args.ImageDisplace.source, NewSwizzle);
+				at = NewDisplace;
+				break;
+			}
+
+			case OP_TYPE::IM_RASTERMESH:
+			{
+				auto NewRaster = mu::Clone<ASTOpFixed>(channelSourceAt);
+				auto NewSwizzle = mu::Clone<ASTOpFixed>(this);
+				ReplaceAllSources(NewSwizzle, NewRaster->children[NewRaster->op.args.ImageRasterMesh.image].child());
+				NewRaster->SetChild(NewRaster->op.args.ImageRasterMesh.image, NewSwizzle);
+				at = NewRaster;
+				break;
+			}
+
+            default:
+                same = false;
+                break;
+            }
+        }
+
+        if (!same)
+        {
+            // Maybe we can still sink the instruction in some cases
+//            OP::ADDRESS sourceAt = program.m_code[at].args.ImageSwizzle.sources[0];
+//            OP_TYPE sourceType = (OP_TYPE)program.m_code[sourceAt].type;
+//            switch ( sourceType )
+//            {
+//            case OP_TYPE::IM_COMPOSE:
+//            {
+//                bool canSink = true;
+//                OP::ADDRESS layout = program.m_code[sourceAt].args.ImageCompose.layout;
+//                uint32_t block = program.m_code[sourceAt].args.ImageCompose.blockIndex;
+//                for ( int c=1; c<MUTABLE_OP_MAX_SWIZZLE_CHANNELS; ++c )
+//                {
+//                    OP::ADDRESS candidate = program.m_code[at].args.ImageSwizzle.sources[c];
+//                    if ( candidate )
+//                    {
+//                        if ( program.m_code[candidate].type==sourceType )
+//                        {
+//                            OP::ADDRESS clayout =
+//                                    program.m_code[candidate].args.ImageCompose.layout;
+//                            uint32_t cblock =
+//                                    program.m_code[candidate].args.ImageCompose.blockIndex;
+//                            canSink = canSink
+//                                    && ( layout == clayout )
+//                                    && ( block == cblock );
+//                        }
+//                        else
+//                        {
+//                            canSink = false;
+//                        }
+//                    }
+//                }
+
+//                if ( canSink )
+//                {
+//                    m_modified = true;
+
+//                    OP op = program.m_code[sourceAt];
+
+//                    OP swizzleBaseOp = program.m_code[at];
+//                    OP swizzleBlockOp = program.m_code[at];
+
+//                    for ( int c=0; c<MUTABLE_OP_MAX_SWIZZLE_CHANNELS; ++c )
+//                    {
+//                        OP::ADDRESS source = program.m_code[at].args.ImageSwizzle.sources[c];
+//                        if ( source )
+//                        {
+//                            swizzleBaseOp.args.ImageSwizzle.sources[c] =
+//                                    program.m_code[ source ].args.ImageCompose.base;
+
+//                            swizzleBlockOp.args.ImageSwizzle.sources[c] =
+//                                    program.m_code[ source ].args.ImageCompose.blockImage;
+//                        }
+//                    }
+//                    op.args.ImageCompose.base = program.AddOp( swizzleBaseOp );
+//                    op.args.ImageCompose.blockImage = program.AddOp( swizzleBlockOp );
+
+//                    at = program.AddOp( op );
+//                }
+//                break;
+//            }
+
+//            case OP_TYPE::IM_CONDITIONAL:
+//            {
+//                // We can always sink into conditionals
+//                m_modified = true;
+
+//                OP op = program.m_code[sourceAt];
+
+//                OP swizzleYesOp = program.m_code[at];
+//                OP swizzleNoOp = program.m_code[at];
+//                OP::ADDRESS source0 = 0;
+//                for ( int c=0; c<MUTABLE_OP_MAX_SWIZZLE_CHANNELS; ++c )
+//                {
+//                    OP::ADDRESS source = program.m_code[at].args.ImageSwizzle.sources[c];
+//                    if ( c==0 )
+//                    {
+//                        source0 = source;
+//                    }
+
+//                    if ( source
+//                        &&
+//                        ( (c==0) || ( program.m_code[source0].args.Conditional.condition == program.m_code[source].args.Conditional.condition ) )
+//                        )
+//                    {
+//                        swizzleYesOp.args.ImageSwizzle.sources[c] =
+//                                program.m_code[ source ].args.Conditional.yes;
+
+//                        swizzleNoOp.args.ImageSwizzle.sources[c] =
+//                                program.m_code[ source ].args.Conditional.no;
+//                    }
+//                }
+//                op.args.Conditional.yes = program.AddOp( swizzleYesOp );
+//                op.args.Conditional.no = program.AddOp( swizzleNoOp );
+
+//                at = program.AddOp( op );
+//                break;
+//            }
+
+
+//            case OP_TYPE::IM_SWITCH:
+//            {
+//                bool canSink = true;
+//                OP::ADDRESS variable = program.m_code[sourceAt].args.Switch.variable;
+
+//                // If at least 3 channels can be combined, do it. We may be duplicating some data
+//                // But it greatly optimises speed in some cases like the constant alpha masks of
+//                // the face colour interpolation.
+//                // TODO: Make this dependent on compilation options.
+//                //for ( int c=1; c<MUTABLE_OP_MAX_SWIZZLE_CHANNELS; ++c )
+//                for ( int c=1; c<3; ++c )
+//                {
+//                    OP::ADDRESS candidate = program.m_code[at].args.ImageSwizzle.sources[c];
+//                    if ( candidate )
+//                    {
+//                        if ( program.m_code[candidate].type==sourceType )
+//                        {
+//                            OP::ADDRESS cvariable =
+//                                    program.m_code[candidate].args.Switch.variable;
+//                            canSink = canSink && ( variable == cvariable );
+//                        }
+//                        else
+//                        {
+//                            canSink = false;
+//                        }
+//                    }
+//                }
+
+//                if ( canSink )
+//                {
+//                    // TODO since data-switch
+////                    m_modified = true;
+
+////                    OP op = program.m_code[sourceAt];
+
+////                    if ( op.args.Switch.def )
+////                    {
+////                        OP swizzleDefOp = program.m_code[at];
+////                        for ( int c=0; c<MUTABLE_OP_MAX_SWIZZLE_CHANNELS; ++c )
+////                        {
+////                            OP::ADDRESS source = program.m_code[at].args.ImageSwizzle.sources[c];
+////                            if ( source )
+////                            {
+////                                if ( program.m_code[source].type==sourceType )
+////                                {
+////                                    swizzleDefOp.args.ImageSwizzle.sources[c] =
+////                                            program.m_code[ source ].args.Switch.def;
+////                                }
+////                                else
+////                                {
+////                                    swizzleDefOp.args.ImageSwizzle.sources[c] = source;
+////                                }
+////                            }
+////                        }
+////                        op.args.Switch.def = program.AddOp( swizzleDefOp );
+////                    }
+
+////                    for ( int o=0; o<MUTABLE_OP_MAX_SWITCH_OPTIONS; ++o )
+////                    {
+////                        if ( op.args.Switch.values[o] )
+////                        {
+////                            OP swizzleOptOp = program.m_code[at];
+////                            for ( int c=0; c<MUTABLE_OP_MAX_SWIZZLE_CHANNELS; ++c )
+////                            {
+////                                OP::ADDRESS source = program.m_code[at].args.ImageSwizzle.sources[c];
+////                                if ( source )
+////                                {
+////                                    if ( program.m_code[source].type==sourceType )
+////                                    {
+////                                        swizzleOptOp.args.ImageSwizzle.sources[c] =
+////                                                program.m_code[ source ].args.Switch.values[o];
+////                                    }
+////                                    else
+////                                    {
+////                                        swizzleOptOp.args.ImageSwizzle.sources[c] = source;
+////                                    }
+////                                }
+////                            }
+////                            op.args.Switch.values[o] = program.AddOp( swizzleOptOp );
+////                        }
+////                    }
+
+////                    at = program.AddOp( op );
+//                }
+//                break;
+//            }
+
+
+//            case OP_TYPE::IM_INTERPOLATE:
+//            {
+//                bool canSink = true;
+//                OP::ADDRESS factor = program.m_code[sourceAt].args.ImageInterpolate.factor;
+
+//                // It is worth to sink 3 channels
+//                // TODO: Make it depend on compilation parameters
+//                //for ( int c=1; c<MUTABLE_OP_MAX_SWIZZLE_CHANNELS; ++c )
+//                for ( int c=1; c<3; ++c )
+//                {
+//                    OP::ADDRESS candidate = program.m_code[at].args.ImageSwizzle.sources[c];
+//                    if ( candidate )
+//                    {
+//                        if ( program.m_code[candidate].type==sourceType )
+//                        {
+//                            OP::ADDRESS cfactor =
+//                                    program.m_code[candidate].args.ImageInterpolate.factor;
+//                            canSink = canSink && ( factor == cfactor );
+//                        }
+//                        else
+//                        {
+//                            canSink = false;
+//                        }
+//                    }
+//                }
+
+//                if ( canSink )
+//                {
+//                    m_modified = true;
+
+//                    OP op = program.m_code[sourceAt];
+
+//                    for ( int o=0; o<MUTABLE_OP_MAX_INTERPOLATE_COUNT; ++o )
+//                    {
+//                        if ( op.args.ImageInterpolate.targets[o] )
+//                        {
+//                            OP swizzleOptOp = program.m_code[at];
+//                            for ( int c=0; c<MUTABLE_OP_MAX_SWIZZLE_CHANNELS; ++c )
+//                            {
+//                                OP::ADDRESS source = program.m_code[at].args.ImageSwizzle.sources[c];
+
+//                                if ( source
+//                                     && program.m_code[source].type==OP_TYPE::IM_INTERPOLATE
+//                                     && program.m_code[source].args.ImageInterpolate.factor==factor
+//                                     )
+//                                {
+//                                    swizzleOptOp.args.ImageSwizzle.sources[c] =
+//                                            program.m_code[ source ].args.ImageInterpolate.targets[o];
+//                                }
+//                            }
+//                            op.args.ImageInterpolate.targets[o] = program.AddOp( swizzleOptOp );
+//                        }
+//                    }
+
+//                    at = program.AddOp( op );
+//                }
+//                break;
+//            }
+
+
+//            case OP_TYPE::IM_INTERPOLATE3:
+//            {
+//                // We can sink if all channels are interpolates for the same factors
+//                OP::ADDRESS sourceAt0 = program.m_code[at].args.ImageSwizzle.sources[0];
+//                OP::ADDRESS factor10 = program.m_code[sourceAt0].args.ImageInterpolate3.factor1;
+//                OP::ADDRESS factor20 = program.m_code[sourceAt0].args.ImageInterpolate3.factor2;
+
+//                OP::ADDRESS sourceAt1 = program.m_code[at].args.ImageSwizzle.sources[1];
+//                OP::ADDRESS factor11 = program.m_code[sourceAt1].args.ImageInterpolate3.factor1;
+//                OP::ADDRESS factor21 = program.m_code[sourceAt1].args.ImageInterpolate3.factor2;
+
+//                OP::ADDRESS sourceAt2 = program.m_code[at].args.ImageSwizzle.sources[2];
+//                OP::ADDRESS factor12 = program.m_code[sourceAt2].args.ImageInterpolate3.factor1;
+//                OP::ADDRESS factor22 = program.m_code[sourceAt2].args.ImageInterpolate3.factor2;
+
+//                bool canSink = true;
+//                canSink &= program.m_code[sourceAt0].type == program.m_code[sourceAt1].type;
+//                canSink &= program.m_code[sourceAt0].type == program.m_code[sourceAt2].type;
+//                canSink &= factor10 == factor11;
+//                canSink &= factor10 == factor12;
+//                canSink &= factor20 == factor21;
+//                canSink &= factor20 == factor22;
+
+//                if ( canSink )
+//                {
+//                    m_modified = true;
+
+//                    OP op = program.m_code[sourceAt0];
+
+//                    // Target 0
+//                    {
+//                        OP swizzleOptOp = program.m_code[at];
+//                        for ( int c=0; c<MUTABLE_OP_MAX_SWIZZLE_CHANNELS; ++c )
+//                        {
+//                            OP::ADDRESS source = program.m_code[at].args.ImageSwizzle.sources[c];
+//                            if ( source
+//                                 && program.m_code[source].type==OP_TYPE::IM_INTERPOLATE3
+//                                 && program.m_code[source].args.ImageInterpolate3.factor1==factor10
+//                                 && program.m_code[source].args.ImageInterpolate3.factor2==factor20
+//                                 )
+//                            {
+//                                swizzleOptOp.args.ImageSwizzle.sources[c] =
+//                                        program.m_code[ source ].args.ImageInterpolate3.target0;
+//                            }
+//                        }
+//                        op.args.ImageInterpolate3.target0 = program.AddOp( swizzleOptOp );
+//                    }
+
+//                    // Target 1
+//                    {
+//                        OP swizzleOptOp = program.m_code[at];
+//                        for ( int c=0; c<MUTABLE_OP_MAX_SWIZZLE_CHANNELS; ++c )
+//                        {
+//                            OP::ADDRESS source = program.m_code[at].args.ImageSwizzle.sources[c];
+//                            if ( source
+//                                 && program.m_code[source].type==OP_TYPE::IM_INTERPOLATE3
+//                                 && program.m_code[source].args.ImageInterpolate3.factor1==factor10
+//                                 && program.m_code[source].args.ImageInterpolate3.factor2==factor20
+//                                 )
+//                            {
+//                                swizzleOptOp.args.ImageSwizzle.sources[c] =
+//                                        program.m_code[ source ].args.ImageInterpolate3.target1;
+//                            }
+//                        }
+//                        op.args.ImageInterpolate3.target1 = program.AddOp( swizzleOptOp );
+//                    }
+
+//                    // Target 2
+//                    {
+//                        OP swizzleOptOp = program.m_code[at];
+//                        for ( int c=0; c<MUTABLE_OP_MAX_SWIZZLE_CHANNELS; ++c )
+//                        {
+//                            OP::ADDRESS source = program.m_code[at].args.ImageSwizzle.sources[c];
+//                            if ( source
+//                                 && program.m_code[source].type==OP_TYPE::IM_INTERPOLATE3
+//                                 && program.m_code[source].args.ImageInterpolate3.factor1==factor10
+//                                 && program.m_code[source].args.ImageInterpolate3.factor2==factor20
+//                                 )
+//                            {
+//                                swizzleOptOp.args.ImageSwizzle.sources[c] =
+//                                        program.m_code[ source ].args.ImageInterpolate3.target2;
+//                            }
+//                        }
+//                        op.args.ImageInterpolate3.target2 = program.AddOp( swizzleOptOp );
+//                    }
+
+//                    at = program.AddOp( op );
+//                }
+//                break;
+//            }
+
+//            default:
+//                break;
+//            }
+        }
+
+        return at;
+    }
+
+
+    //---------------------------------------------------------------------------------------------
+    //---------------------------------------------------------------------------------------------
+    //---------------------------------------------------------------------------------------------
+    bool SinkOptimiserAST
+    (
+            ASTOpList& roots,
+            const MODEL_OPTIMIZATION_OPTIONS& optimisationOptions
+    )
+    {
+        MUTABLE_CPUPROFILER_SCOPE(SinkOptimiserAST);
+
+        bool modified = false;
+
+		OPTIMIZE_SINK_CONTEXT context;
+
+        ASTOp::Traverse_TopDown_Unique_Imprecise( roots, [&](Ptr<ASTOp>& n)
+        {
+            auto o = n->OptimiseSink(optimisationOptions, context);
+            if (o && n!=o)
+            {
+                modified = true;
+                ASTOp::Replace(n,o);
+            }
+
+            return true;
+        });
+
+        return modified;
+    }
+
+
+	//---------------------------------------------------------------------------------------------
+	//---------------------------------------------------------------------------------------------
+	//---------------------------------------------------------------------------------------------
+    Ptr<ASTOp> Sink_ImageCropAST::Apply( const ASTOp* root )
+    {
+        m_root = root;
+        m_oldToNew.clear();
+		m_newOps.clear();
+
+        check(root->GetOpType()==OP_TYPE::IM_CROP);
+
+        auto typedRoot = dynamic_cast<const ASTOpFixed*>(root);
+        m_initialSource = typedRoot->children[typedRoot->op.args.ImageCrop.source].child();
+        Ptr<ASTOp> newSource = Visit( m_initialSource, typedRoot );
+
+        // If there is any change, it is the new root.
+        if (newSource!=m_initialSource)
+        {
+            return newSource;
+        }
+
+        return nullptr;
+    }
+
+
+	//---------------------------------------------------------------------------------------------
+	Ptr<ASTOp> Sink_ImageCropAST::Visit( Ptr<ASTOp> at, const ASTOpFixed* currentCropOp )
+    {
+        if (!at) return nullptr;
+
+        // Newly created?
+        if (std::find(m_newOps.begin(), m_newOps.end(), at )!=m_newOps.end())
+        {
+            return at;
+        }
+
+        // Already visited?
+        auto cacheIt = m_oldToNew[currentCropOp].find(at);
+        if (cacheIt!=m_oldToNew[currentCropOp].end())
+        {
+            return cacheIt->second;
+        }
+
+        bool skipSinking=false;
+        Ptr<ASTOp> newAt = at;
+        switch ( at->GetOpType() )
+        {
+
+        case OP_TYPE::IM_CONDITIONAL:
+        {
+            // We move down the two paths
+            Ptr<ASTOpConditional> newOp = mu::Clone<ASTOpConditional>(at);
+            newOp->yes = Visit(newOp->yes.child(), currentCropOp);
+            newOp->no = Visit(newOp->no.child(), currentCropOp);
+            newAt = newOp;
+            break;
+        }
+
+        case OP_TYPE::IM_SWITCH:
+        {
+            // We move down all the paths
+            Ptr<ASTOpSwitch> newOp = mu::Clone<ASTOpSwitch>(at);
+            newOp->def = Visit(newOp->def.child(), currentCropOp);
+            for( auto& c:newOp->cases )
+            {
+                c.branch = Visit(c.branch.child(), currentCropOp);
+            }
+            newAt = newOp;
+            break;
+        }
+
+        case OP_TYPE::IM_PIXELFORMAT:
+        {
+            Ptr<ASTOpImagePixelFormat> nop = mu::Clone<ASTOpImagePixelFormat>(at);
+            nop->Source = Visit(nop->Source.child(), currentCropOp);
+            newAt = nop;
+            break;
+        }
+
+        case OP_TYPE::IM_PATCH:
+        {
+            const ASTOpImagePatch* typedPatch = dynamic_cast<const ASTOpImagePatch*>(at.get());
+
+            auto rectOp = typedPatch->patch.child();
+            ASTOp::GetImageDescContext context;
+            auto patchDesc = rectOp->GetImageDesc( false, &context );
+            box<vec2<int16_t>> patchBox;
+            patchBox.min = typedPatch->location;
+            patchBox.size = patchDesc.m_size;
+
+            box<vec2<int16_t>> cropBox;
+            cropBox.min[0] = currentCropOp->op.args.ImageCrop.minX;
+            cropBox.min[1] = currentCropOp->op.args.ImageCrop.minY;
+            cropBox.size[0] = currentCropOp->op.args.ImageCrop.sizeX;
+            cropBox.size[1] = currentCropOp->op.args.ImageCrop.sizeY;
+
+            if (!patchBox.IntersectsExclusive(cropBox))
+            {
+                // We can ignore the patch
+                newAt = Visit( typedPatch->base.child(), currentCropOp );
+            }
+            else
+            {
+                // Crop the base with the full crop, and the patch with the intersected part,
+                // adapting the patch origin
+                auto newOp = mu::Clone<ASTOpImagePatch>(at);
+                newOp->base = Visit(newOp->base.child(), currentCropOp);
+
+                box<vec2<int16_t>> ibox = patchBox.Intersect(cropBox);
+                check(ibox.size[0]>0 && ibox.size[1]>0);
+
+                Ptr<ASTOpFixed> patchCropOp = mu::Clone<ASTOpFixed>(currentCropOp);
+                patchCropOp->op.args.ImageCrop.minX = ibox.min[0] - patchBox.min[0];
+                patchCropOp->op.args.ImageCrop.minY = ibox.min[1] - patchBox.min[1];
+                patchCropOp->op.args.ImageCrop.sizeX = ibox.size[0];
+                patchCropOp->op.args.ImageCrop.sizeY = ibox.size[1];
+                newOp->patch = Visit(newOp->patch.child(), patchCropOp.get());
+
+                newOp->location[0] = ibox.min[0] - cropBox.min[0];
+                newOp->location[1] = ibox.min[1] - cropBox.min[1];
+                newAt = newOp;
+            }
+
+            break;
+        }
+
+        case OP_TYPE::IM_CROP:
+        {
+            // We can combine the two crops into a possibliy smaller crop
+            auto childCrop = dynamic_cast<const ASTOpFixed*>(at.get());
+
+            box<vec2<int16_t>> childCropBox;
+            childCropBox.min[0] = childCrop->op.args.ImageCrop.minX;
+            childCropBox.min[1] = childCrop->op.args.ImageCrop.minY;
+            childCropBox.size[0] = childCrop->op.args.ImageCrop.sizeX;
+            childCropBox.size[1] = childCrop->op.args.ImageCrop.sizeY;
+
+            box<vec2<int16_t>> cropBox;
+            cropBox.min[0] = currentCropOp->op.args.ImageCrop.minX;
+            cropBox.min[1] = currentCropOp->op.args.ImageCrop.minY;
+            cropBox.size[0] = currentCropOp->op.args.ImageCrop.sizeX;
+            cropBox.size[1] = currentCropOp->op.args.ImageCrop.sizeY;
+
+            // Compose the crops: in the final image the child crop is applied first and the
+            // current ctop is applied to the result. So the final crop box would be:
+            box<vec2<int16_t>> ibox;
+            ibox.min = childCropBox.min + cropBox.min;
+            ibox.size = vec2<int16_t>::min(cropBox.size, childCropBox.size);
+            check(cropBox.min.AllSmallerOrEqualThan(childCropBox.size));
+            check((cropBox.min + cropBox.size)
+                                .AllSmallerOrEqualThan(childCropBox.size));
+            check((ibox.min + ibox.size)
+                                .AllSmallerOrEqualThan(childCropBox.min + childCropBox.size));
+
+            // This happens more often that one would think
+            if (ibox==childCropBox)
+            {
+                // the parent crop is not necessary
+                skipSinking = true;
+            }
+            else if (ibox == cropBox)
+            {
+                // The child crop is not necessary
+                auto childSource =
+                    childCrop->children[childCrop->op.args.ImageCrop.source].child();
+                newAt = Visit(childSource, currentCropOp);
+            }
+            else
+            {
+                // combine into one crop
+                Ptr<ASTOpFixed> newCropOp = mu::Clone<ASTOpFixed>(currentCropOp);
+                newCropOp->op.args.ImageCrop.minX = ibox.min[0];
+                newCropOp->op.args.ImageCrop.minY = ibox.min[1];
+                newCropOp->op.args.ImageCrop.sizeX = ibox.size[0];
+                newCropOp->op.args.ImageCrop.sizeY = ibox.size[1];
+
+                auto childSource =
+                    childCrop->children[childCrop->op.args.ImageCrop.source].child();
+                newAt = Visit(childSource, newCropOp.get());
+            }
+            break;
+        }
+
+        case OP_TYPE::IM_LAYER:
+        {
+            // We move the op down the arguments
+            auto nop = mu::Clone<ASTOpFixed>(at);
+
+            auto aOp = nop->children[nop->op.args.ImageLayer.base].child();
+            nop->SetChild(nop->op.args.ImageLayer.base, Visit(aOp, currentCropOp));
+
+            auto bOp = nop->children[nop->op.args.ImageLayer.blended].child();
+            nop->SetChild(nop->op.args.ImageLayer.blended, Visit(bOp, currentCropOp));
+
+            auto mOp = nop->children[nop->op.args.ImageLayer.mask].child();
+            nop->SetChild(nop->op.args.ImageLayer.mask, Visit(mOp, currentCropOp));
+
+            newAt = nop;
+            break;
+        }
+
+        case OP_TYPE::IM_LAYERCOLOUR:
+        {
+            // We move the op down the arguments
+            auto nop = mu::Clone<ASTOpFixed>(at);
+
+            auto aOp = nop->children[nop->op.args.ImageLayerColour.base].child();
+            nop->SetChild(nop->op.args.ImageLayerColour.base, Visit(aOp, currentCropOp));
+
+            auto mOp = nop->children[nop->op.args.ImageLayerColour.mask].child();
+            nop->SetChild(nop->op.args.ImageLayerColour.mask, Visit(mOp, currentCropOp));
+
+            newAt = nop;
+            break;
+        }
+
+        case OP_TYPE::IM_INTERPOLATE:
+        {
+            // Move the op  down all the paths
+            auto newOp = mu::Clone<ASTOpFixed>(at);
+
+            for (int v = 0; v < MUTABLE_OP_MAX_INTERPOLATE_COUNT; ++v)
+            {
+                auto child =
+                    newOp->children[newOp->op.args.ImageInterpolate.targets[v]].child();
+                auto bOp = Visit(child, currentCropOp);
+                newOp->SetChild(newOp->op.args.ImageInterpolate.targets[v], bOp);
+            }
+
+            newAt = newOp;
+            break;
+        }
+
+
+        case OP_TYPE::IM_INTERPOLATE3:
+        {
+            // We move the op down all the paths
+            auto newOp = mu::Clone<ASTOpFixed>(at);
+
+            auto top0 = newOp->children[newOp->op.args.ImageInterpolate3.target0].child();
+            newOp->SetChild(newOp->op.args.ImageInterpolate3.target0,
+                            Visit(top0, currentCropOp));
+
+            auto top1 = newOp->children[newOp->op.args.ImageInterpolate3.target1].child();
+            newOp->SetChild(newOp->op.args.ImageInterpolate3.target1,
+                            Visit(top1, currentCropOp));
+
+            auto top2 = newOp->children[newOp->op.args.ImageInterpolate3.target2].child();
+            newOp->SetChild(newOp->op.args.ImageInterpolate3.target2,
+                            Visit(top2, currentCropOp));
+
+            newAt = newOp;
+            break;
+        }
+
+        default:
+            break;
+        }
+
+        // end on line, replace with crop
+        if (at==newAt && at!=m_initialSource && !skipSinking)
+        {
+            Ptr<ASTOpFixed> newOp = mu::Clone<ASTOpFixed>(currentCropOp);
+            check(newOp->GetOpType()==OP_TYPE::IM_CROP);
+
+            newOp->SetChild( newOp->op.args.ImageCrop.source,at);
+
+            newAt = newOp;
+        }
+
+        m_oldToNew[currentCropOp][at] = newAt;
+
+        return newAt;
+    }
+
+
+	//---------------------------------------------------------------------------------------------
+	//---------------------------------------------------------------------------------------------
+	//---------------------------------------------------------------------------------------------
+	Ptr<ASTOp> ASTOpFixed::OptimiseSink( const MODEL_OPTIMIZATION_OPTIONS& options, OPTIMIZE_SINK_CONTEXT& context ) const
+    {
+        Ptr<ASTOp> at;
+
+        OP_TYPE type = GetOpType();
+        switch ( type )
+        {
+
+        case OP_TYPE::ME_PROJECT:
+        {
+            auto sourceAt = children[op.args.MeshProject.mesh].child();
+            if (!sourceAt) { break; }
+
+            auto projectorAt = children[op.args.MeshProject.projector].child();
+            if (!projectorAt) { break; }
+
+            OP_TYPE sourceType = sourceAt->GetOpType();
+            switch ( sourceType )
+            {
+
+            case OP_TYPE::ME_CONDITIONAL:
+            {
+                if (projectorAt->GetOpType()==OP_TYPE::PR_CONSTANT)
+                {
+                    // We move the project down the two paths
+                    auto nop = mu::Clone<ASTOpConditional>(sourceAt);
+
+                    auto aOp = mu::Clone<ASTOpFixed>(this);
+                    aOp->SetChild( aOp->op.args.MeshProject.mesh, nop->yes );
+                    nop->yes = aOp;
+
+                    auto bOp = mu::Clone<ASTOpFixed>(this);
+                    bOp->SetChild( bOp->op.args.MeshProject.mesh, nop->no);
+                    nop->no = bOp;
+
+                    at = nop;
+                }
+                break;
+            }
+
+            case OP_TYPE::ME_SWITCH:
+            {
+                if (projectorAt->GetOpType()==OP_TYPE::PR_CONSTANT)
+                {
+                    // Move the format down all the paths
+                    auto nop = mu::Clone<ASTOpSwitch>(sourceAt);
+
+                    if (nop->def)
+                    {
+                        auto defOp = mu::Clone<ASTOpFixed>(this);
+                        defOp->SetChild( defOp->op.args.MeshProject.mesh, nop->def );
+                        nop->def = defOp;
+                    }
+
+                    // We need to copy the options because we change them
+                    for ( size_t v=0; v<nop->cases.size(); ++v )
+                    {
+                        if ( nop->cases[v].branch )
+                        {
+                            auto bOp = mu::Clone<ASTOpFixed>(this);
+                            bOp->SetChild( bOp->op.args.MeshProject.mesh, nop->cases[v].branch );
+                            nop->cases[v].branch = bOp;
+                        }
+                    }
+
+                    at = nop;
+                }
+                break;
+            }
+
+            default:
+                break;
+            }
+
+            break;
+        }
+
+
+        //-----------------------------------------------------------------------------------------
+        //-----------------------------------------------------------------------------------------
+        //-----------------------------------------------------------------------------------------
+        case OP_TYPE::LA_REMOVEBLOCKS:
+        {
+            auto meshAt = children[op.args.LayoutRemoveBlocks.mesh].child();
+
+            OP_TYPE meshType = meshAt->GetOpType();
+            switch ( meshType )
+            {
+
+            case OP_TYPE::ME_SWITCH:
+            {
+                // Move the remove-blocks down all the paths
+                Ptr<ASTOpSwitch> nop = mu::Clone<ASTOpSwitch>(meshAt);
+                nop->type = OP_TYPE::LA_SWITCH;
+
+                if (nop->def)
+                {
+                    auto defOp = mu::Clone<ASTOpFixed>(this);
+                    defOp->SetChild( defOp->op.args.LayoutRemoveBlocks.mesh, nop->def );
+                    nop->def = defOp;
+                }
+
+                for ( auto& o: nop->cases )
+                {
+                    if ( o.branch )
+                    {
+                        auto bOp = mu::Clone<ASTOpFixed>(this);
+                        bOp->SetChild( bOp->op.args.LayoutRemoveBlocks.mesh, o.branch );
+                        o.branch = bOp;
+                    }
+                }
+
+                at = nop;
+                break;
+            }
+
+            case OP_TYPE::ME_MORPH2:
+            {
+                // Move the remove-blocks down the base
+                auto typedMeshAt = dynamic_cast<const ASTOpFixed*>(meshAt.get());
+                auto nop = mu::Clone<ASTOpFixed>(this);
+                nop->SetChild( nop->op.args.LayoutRemoveBlocks.mesh, typedMeshAt->children[ typedMeshAt->op.args.MeshMorph2.base] );
+                at = nop;
+                break;
+            }
+
+            case OP_TYPE::ME_INTERPOLATE:
+            {
+                // Assume that the UVs will not be interpolated across blocks!
+                auto typedMeshAt = dynamic_cast<const ASTOpFixed*>(meshAt.get());
+                auto nop = mu::Clone<ASTOpFixed>(this);
+                nop->SetChild( nop->op.args.LayoutRemoveBlocks.mesh, typedMeshAt->children[ typedMeshAt->op.args.MeshInterpolate.base] );
+                at = nop;
+                break;
+            }
+
+            default:
+                break;
+            }
+
+            break;
+        }
+
+
+        //-----------------------------------------------------------------------------------------
+        //-----------------------------------------------------------------------------------------
+        //-----------------------------------------------------------------------------------------
+        case OP_TYPE::IM_RASTERMESH:
+        {
+			Ptr<ASTOp> OriginalAt = at;
+			Ptr<ASTOp> sourceAt = children[op.args.ImageRasterMesh.mesh].child();
+			Ptr<ASTOp> imageAt = children[op.args.ImageRasterMesh.image].child();
+
+            OP_TYPE sourceType = sourceAt->GetOpType();
+            switch ( sourceType )
+            {
+
+            case OP_TYPE::ME_PROJECT:
+            {
+                // If we are rastering just the UV layout (to create a mask) we don't care about
+                // mesh project operations, which modify only the positions.
+                // This optimisation helps with states removing fake dependencies on projector
+                // parameters that may be runtime.
+                if (!imageAt)
+                {
+                    // We remove the project from the raster children
+                    // \todo: maybe this clone is not necessary
+					const ASTOpFixed* typedSource = dynamic_cast<const ASTOpFixed*>(sourceAt.get());
+                    Ptr<ASTOpFixed> nop = mu::Clone<ASTOpFixed>(this);
+                    nop->SetChild( nop->op.args.ImageRasterMesh.mesh,
+                                   typedSource->children[typedSource->op.args.MeshProject.mesh]);
+                    at = nop;
+                }
+                break;
+            }
+
+            case OP_TYPE::ME_INTERPOLATE:
+            {
+                auto typedSource = dynamic_cast<const ASTOpFixed*>(sourceAt.get());
+				Ptr<ASTOpFixed> rasterOp = mu::Clone<ASTOpFixed>(this);
+                rasterOp->SetChild( rasterOp->op.args.ImageRasterMesh.mesh,
+                        typedSource->children[typedSource->op.args.MeshInterpolate.base] );
+                at = rasterOp;
+                break;
+            }
+
+            case OP_TYPE::ME_MORPH2:
+            {
+                auto typedSource = dynamic_cast<const ASTOpFixed*>(sourceAt.get());
+				Ptr<ASTOpFixed> rasterOp = mu::Clone<ASTOpFixed>(this);
+                rasterOp->SetChild( rasterOp->op.args.ImageRasterMesh.mesh,
+                        typedSource->children[typedSource->op.args.MeshMorph2.base] );
+                at = rasterOp;
+                break;
+            }
+
+            case OP_TYPE::ME_CONDITIONAL:
+            {
+                auto nop = mu::Clone<ASTOpConditional>(sourceAt.get());
+                nop->type = OP_TYPE::IM_CONDITIONAL;
+
+				Ptr<ASTOpFixed> aOp = mu::Clone<ASTOpFixed>(this);
+                aOp->SetChild( aOp->op.args.ImageRasterMesh.mesh, nop->yes );
+                nop->yes = aOp;
+
+				Ptr<ASTOpFixed> bOp = mu::Clone<ASTOpFixed>(this);
+                bOp->SetChild( bOp->op.args.ImageRasterMesh.mesh, nop->no );
+                nop->no = bOp;
+
+                at = nop;
+                break;
+            }
+
+            case OP_TYPE::ME_SWITCH:
+            {
+                // Make an image for every path
+                auto nop = mu::Clone<ASTOpSwitch>(sourceAt.get());
+                nop->type = OP_TYPE::IM_SWITCH;
+
+                if (nop->def)
+                {
+					Ptr<ASTOpFixed> defOp = mu::Clone<ASTOpFixed>(this);
+                    defOp->SetChild( defOp->op.args.ImageRasterMesh.mesh, nop->def );
+                    nop->def = defOp;
+                }
+
+                // We need to copy the options because we change them
+                for ( size_t o=0; o<nop->cases.size(); ++o )
+                {
+                    if ( nop->cases[o].branch )
+                    {
+						Ptr<ASTOpFixed> bOp = mu::Clone<ASTOpFixed>(this);
+                        bOp->SetChild( bOp->op.args.ImageRasterMesh.mesh, nop->cases[o].branch );
+                        nop->cases[o].branch = bOp;
+                    }
+                }
+
+                at = nop;
+                break;
+            }
+
+            default:
+                break;
+            }
+
+			// If we didn't optimize the mesh child, try to optimize the image child.
+			if (OriginalAt == at && imageAt)
+			{
+				OP_TYPE imageType = imageAt->GetOpType();
+				switch (imageType)
+				{
+
+				// TODO: Implement for image conditionals.
+				//case OP_TYPE::ME_CONDITIONAL:
+				//{
+				//	auto nop = mu::Clone<ASTOpConditional>(sourceAt.get());
+				//	nop->type = OP_TYPE::IM_CONDITIONAL;
+
+				//	Ptr<ASTOpFixed> aOp = mu::Clone<ASTOpFixed>(this);
+				//	aOp->SetChild(aOp->op.args.ImageRasterMesh.mesh, nop->yes);
+				//	nop->yes = aOp;
+
+				//	Ptr<ASTOpFixed> bOp = mu::Clone<ASTOpFixed>(this);
+				//	bOp->SetChild(bOp->op.args.ImageRasterMesh.mesh, nop->no);
+				//	nop->no = bOp;
+
+				//	at = nop;
+				//	break;
+				//}
+
+				case OP_TYPE::IM_SWITCH:
+				{
+					// TODO: Do this only if the projector is constant?
+					
+					// Make a project for every path
+					auto nop = mu::Clone<ASTOpSwitch>(imageAt.get());
+
+					if (nop->def)
+					{
+						Ptr<ASTOpFixed> defOp = mu::Clone<ASTOpFixed>(this);
+						defOp->SetChild(defOp->op.args.ImageRasterMesh.image, nop->def);
+						nop->def = defOp;
+					}
+
+					// We need to copy the options because we change them
+					for (size_t o = 0; o < nop->cases.size(); ++o)
+					{
+						if (nop->cases[o].branch)
+						{
+							Ptr<ASTOpFixed> bOp = mu::Clone<ASTOpFixed>(this);
+							bOp->SetChild(bOp->op.args.ImageRasterMesh.image, nop->cases[o].branch);
+							nop->cases[o].branch = bOp;
+						}
+					}
+
+					at = nop;
+					break;
+				}
+
+				default:
+					break;
+				}
+			}
+
+            break;
+        }
+
+        //-------------------------------------------------------------------------------------
+        case OP_TYPE::IM_SWIZZLE:
+        {
+            at = OptimiseSwizzle( options );
+            break;
+        }
+
+        //-------------------------------------------------------------------------------------
+        case OP_TYPE::IM_LAYER:
+        {
+            // Layer effects may be worth sinking down switches and conditionals, to be able
+            // to apply extra optimisations
+            auto baseAt = children[op.args.ImageLayer.base].child();
+            auto blendAt = children[op.args.ImageLayer.blended].child();
+            auto maskAt = children[op.args.ImageLayer.mask].child();
+
+            // Promote conditions from the base
+            OP_TYPE baseType = baseAt->GetOpType();
+            switch ( baseType )
+            {
+            // Seems to cause operation explosion in optimizer in bandit model.
+            // moved to generic sink in the default.
+//            case OP_TYPE::IM_CONDITIONAL:
+//            {
+//                m_modified = true;
+
+//                OP op = program.m_code[baseAt];
+
+//                OP aOp = program.m_code[at];
+//                aOp.args.ImageLayer.base = program.m_code[baseAt].args.Conditional.yes;
+//                op.args.Conditional.yes = program.AddOp( aOp );
+
+//                OP bOp = program.m_code[at];
+//                bOp.args.ImageLayer.base = program.m_code[baseAt].args.Conditional.no;
+//                op.args.Conditional.no = program.AddOp( bOp );
+
+//                at = program.AddOp( op );
+//                break;
+//            }
+
+            case OP_TYPE::IM_SWITCH:
+            {
+                // Disabled:
+                // It seems to cause data explosion in optimizer in some models. Because
+                // all switch branches become unique constants
+
+//                // See if the blended has an identical switch, to optimise it too
+//                auto baseSwitch = dynamic_cast<const ASTOpSwitch*>( baseAt.get() );
+//                auto blendedSwitch = dynamic_cast<const ASTOpSwitch*>( blendAt.get() );
+//                auto maskSwitch = dynamic_cast<const ASTOpSwitch*>( maskAt.get() );
+
+//                bool blendedToo = baseSwitch->IsCompatibleWith( blendedSwitch );
+//                bool maskToo = baseSwitch->IsCompatibleWith( maskSwitch );
+
+//                // Move the layer operation down all the paths
+//                auto nop = mu::Clone<ASTOpSwitch>(baseSwitch);
+
+//                if (nop->def)
+//                {
+//                    auto defOp = mu::Clone<ASTOpFixed>(this);
+//                    defOp->SetChild( defOp->op.args.ImageLayer.base, nop->def );
+//                    if (blendedToo)
+//                    {
+//                        defOp->SetChild( defOp->op.args.ImageLayer.blended, blendedSwitch->def );
+//                    }
+//                    if (maskToo)
+//                    {
+//                        defOp->SetChild( defOp->op.args.ImageLayer.mask, maskSwitch->def );
+//                    }
+//                    nop->def = defOp;
+//                }
+
+//                for ( size_t v=0; v<nop->cases.size(); ++v )
+//                {
+//                    if ( nop->cases[v].branch )
+//                    {
+//                        auto bOp = mu::Clone<ASTOpFixed>(this);
+//                        bOp->SetChild( bOp->op.args.ImageLayer.base, nop->cases[v].branch );
+
+//                        if (blendedToo)
+//                        {
+//                            bOp->SetChild( bOp->op.args.ImageLayer.blended, blendedSwitch->FindBranch( nop->cases[v].condition ) );
+//                        }
+
+//                        if (maskToo)
+//                        {
+//                            bOp->SetChild( bOp->op.args.ImageLayer.mask, maskSwitch->FindBranch( nop->cases[v].condition ) );
+//                        }
+
+//                        nop->cases[v].branch = bOp;
+//                    }
+//                }
+
+//                at = nop;
+                break;
+            }
+
+            default:
+            {
+                // Try generic base sink
+//                OP templateOp = program.m_code[at];
+
+//                CloneNeutralPartialTreeVisitor neutral;
+//                OP_TYPE supportedOps[] = {
+//                    OP_TYPE::IM_CONDITIONAL,
+//                    OP_TYPE::NONE
+//                };
+//                auto newRoot = neutral.Apply
+//                         (
+//                            program.m_code[at].args.ImageLayer.base, program,
+//                            DT_IMAGE, supportedOps,
+//                            &templateOp, &templateOp.args.ImageLayer.base);
+
+//                // If there is a change
+//                if (newRoot!=program.m_code[at].args.ImageLayer.base)
+//                {
+//                    m_modified = true;
+//                    at = newRoot;
+//                }
+
+                break;
+            }
+
+            }
+
+            break;
+        }
+
+        //-----------------------------------------------------------------------------------------
+        //-----------------------------------------------------------------------------------------
+        //-----------------------------------------------------------------------------------------
+        case OP_TYPE::IM_CROP:
+        {
+            auto sourceAt = children[op.args.ImageCrop.source].child();
+
+            switch ( sourceAt->GetOpType() )
+            {
+            // In case we have other operations with special optimisation rules.
+            case OP_TYPE::NONE:
+                break;
+
+            default:
+            {
+                at = context.ImageCropSinker.Apply(this);
+
+                break;
+            } // default
+
+            } // switch
+
+            break;
+        }
+
+
+        //-----------------------------------------------------------------------------------------
+        //-----------------------------------------------------------------------------------------
+        //-----------------------------------------------------------------------------------------
+        case OP_TYPE::IM_MAKEGROWMAP:
+        {
+            auto sourceAt = children[op.args.ImageMakeGrowMap.mask].child();
+
+            switch ( sourceAt->GetOpType() )
+            {
+
+            case OP_TYPE::IM_CONDITIONAL:
+            {
+                // We move the format down the two paths
+                auto nop = mu::Clone<ASTOpConditional>(sourceAt);
+
+                auto aOp = mu::Clone<ASTOpFixed>(this);
+                aOp->SetChild( aOp->op.args.ImageMakeGrowMap.mask, nop->yes );
+                nop->yes = aOp;
+
+                auto bOp = mu::Clone<ASTOpFixed>(this);
+                bOp->SetChild( bOp->op.args.ImageMakeGrowMap.mask, nop->no);
+                nop->no = bOp;
+
+                at = nop;
+                break;
+            }
+
+            case OP_TYPE::IM_SWITCH:
+            {
+                // Move the format down all the paths
+                auto nop = mu::Clone<ASTOpSwitch>(sourceAt);
+
+                if (nop->def)
+                {
+                    auto defOp = mu::Clone<ASTOpFixed>(this);
+                    defOp->SetChild( defOp->op.args.ImageMakeGrowMap.mask, nop->def );
+                    nop->def = defOp;
+                }
+
+                // We need to copy the options because we change them
+                for ( size_t v=0; v<nop->cases.size(); ++v )
+                {
+                    if ( nop->cases[v].branch )
+                    {
+                        auto bOp = mu::Clone<ASTOpFixed>(this);
+                        bOp->SetChild( bOp->op.args.ImageMakeGrowMap.mask, nop->cases[v].branch );
+                        nop->cases[v].branch = bOp;
+                    }
+                }
+
+                at = nop;
+                break;
+            }
+
+            default:
+                break;
+            }
+
+            break;
+        }
+
+
+
+        //-----------------------------------------------------------------------------------------
+        //-----------------------------------------------------------------------------------------
+        //-----------------------------------------------------------------------------------------
+        case OP_TYPE::IM_DISPLACE:
+        {
+			Ptr<ASTOp> OriginalAt = at;
+			Ptr<ASTOp> sourceAt = children[op.args.ImageDisplace.source].child();
+			Ptr<ASTOp> displaceMapAt = children[op.args.ImageDisplace.displacementMap].child();
+
+            switch ( sourceAt->GetOpType() )
+            {
+
+            case OP_TYPE::IM_CONDITIONAL:
+            {
+                if (displaceMapAt->GetOpType()==OP_TYPE::IM_CONDITIONAL)
+                {
+                    auto typedSource = dynamic_cast<const ASTOpConditional*>(sourceAt.get());
+                    auto typedDisplacementMap = dynamic_cast<const ASTOpConditional*>(displaceMapAt.get());
+
+                    if (typedSource->condition==typedDisplacementMap->condition)
+                    {
+                        auto nop = mu::Clone<ASTOpConditional>(sourceAt);
+
+                        auto aOp = mu::Clone<ASTOpFixed>(this);
+                        aOp->SetChild( aOp->op.args.ImageDisplace.source, typedSource->yes );
+                        aOp->SetChild( aOp->op.args.ImageDisplace.displacementMap, typedDisplacementMap->yes );
+                        nop->yes = aOp;
+
+                        auto bOp = mu::Clone<ASTOpFixed>(this);
+                        bOp->SetChild( bOp->op.args.ImageDisplace.source, typedSource->no);
+                        bOp->SetChild( bOp->op.args.ImageDisplace.displacementMap, typedDisplacementMap->no);
+                        nop->no = bOp;
+
+                        at = nop;
+                    }
+                }
+                break;
+            }
+
+            case OP_TYPE::IM_SWITCH:
+            {
+                if (displaceMapAt->GetOpType()==OP_TYPE::IM_SWITCH)
+                {
+                    auto typedSource = dynamic_cast<const ASTOpSwitch*>(sourceAt.get());
+                    auto typedDisplacementMap = dynamic_cast<const ASTOpSwitch*>(displaceMapAt.get());
+
+                    if (typedSource->IsCompatibleWith(typedDisplacementMap))
+                    {
+                        // Move the format down all the paths
+                        auto nop = mu::Clone<ASTOpSwitch>(sourceAt);
+
+                        if (nop->def)
+                        {
+							Ptr<ASTOpFixed> defOp = mu::Clone<ASTOpFixed>(this);
+                            defOp->SetChild( defOp->op.args.ImageDisplace.source, typedSource->def );
+                            defOp->SetChild( defOp->op.args.ImageDisplace.displacementMap, typedDisplacementMap->def );
+                            nop->def = defOp;
+                        }
+
+                        // We need to copy the options because we change them
+                        for ( size_t v=0; v<nop->cases.size(); ++v )
+                        {
+                            if ( nop->cases[v].branch )
+                            {
+								Ptr<ASTOpFixed> bOp = mu::Clone<ASTOpFixed>(this);
+                                bOp->SetChild( bOp->op.args.ImageDisplace.source, typedSource->cases[v].branch );
+                                bOp->SetChild( bOp->op.args.ImageDisplace.displacementMap, typedDisplacementMap->FindBranch(typedSource->cases[v].condition) );
+                                nop->cases[v].branch = bOp;
+                            }
+                        }
+
+                        at = nop;
+                    }
+                }
+
+				// If we didn't optimize already, try to simply sink the source.
+				if (OriginalAt == at)
+				{
+					// Make a project for every path
+					auto nop = mu::Clone<ASTOpSwitch>(sourceAt.get());
+
+					if (nop->def)
+					{
+						Ptr<ASTOpFixed> defOp = mu::Clone<ASTOpFixed>(this);
+						defOp->SetChild(defOp->op.args.ImageDisplace.source, nop->def);
+						nop->def = defOp;
+					}
+
+					// We need to copy the options because we change them
+					for (size_t o = 0; o < nop->cases.size(); ++o)
+					{
+						if (nop->cases[o].branch)
+						{
+							Ptr<ASTOpFixed> bOp = mu::Clone<ASTOpFixed>(this);
+							bOp->SetChild(bOp->op.args.ImageDisplace.source, nop->cases[o].branch);
+							nop->cases[o].branch = bOp;
+						}
+					}
+
+					at = nop;
+				}
+
+                break;
+            }
+
+            default:
+                break;
+            }
+
+            break;
+        }
+
+
+        default:
+            break;
+        }
+
+        return at;
+    }
+
+    //---------------------------------------------------------------------------------------------
+    //---------------------------------------------------------------------------------------------
+    //---------------------------------------------------------------------------------------------
+    bool SizeOptimiserAST( ASTOpList& roots )
+    {
+        MUTABLE_CPUPROFILER_SCOPE(SizeOptimiser);
+
+        bool modified = false;
+
+        // TODO: isn't top down better suited?
+        ASTOp::Traverse_BottomUp_Unique( roots, [&](Ptr<ASTOp>& n)
+        {
+            auto o = n->OptimiseSize();
+            if (o && o!=n)
+            {
+                modified = true;
+                ASTOp::Replace(n,o);
+            }
+        });
+
+        return modified;
+    }
+
+
+    //---------------------------------------------------------------------------------------------
+    //!
+    //---------------------------------------------------------------------------------------------
+    class Sink_ImageResizeRelAST
+    {
+    public:
+
+        Ptr<ASTOp> Apply( const ASTOp* root )
+        {
+            m_root = root;
+            m_oldToNew.clear();
+
+            check(root->GetOpType()==OP_TYPE::IM_RESIZEREL);
+
+            auto typedRoot = dynamic_cast<const ASTOpFixed*>(root);
+            m_initialSource = typedRoot->children[typedRoot->op.args.ImageResizeRel.source].child();
+            Ptr<ASTOp> newSource = Visit( m_initialSource, typedRoot );
+
+            // If there is any change, it is the new root.
+            if (newSource!=m_initialSource)
+            {
+                return newSource;
+            }
+
+            return nullptr;
+        }
+
+    protected:
+
+        const ASTOp* m_root;
+        Ptr<ASTOp> m_initialSource;
+        //! For each operation we sink, the map from old instructions to new instructions.
+        std::unordered_map<Ptr<const ASTOpFixed>,std::unordered_map<Ptr<ASTOp>,Ptr<ASTOp>>> m_oldToNew;
+        vector<Ptr<ASTOp>> m_newOps;
+
+        Ptr<ASTOp> Visit( Ptr<ASTOp> at, const ASTOpFixed* currentSinkingOp )
+        {
+            if (!at) return nullptr;
+
+            // Newly created?
+            if (std::find(m_newOps.begin(), m_newOps.end(), at )!=m_newOps.end())
+            {
+                return at;
+            }
+
+            // Already visited?
+            auto cacheIt = m_oldToNew[currentSinkingOp].find(at);
+            if (cacheIt!=m_oldToNew[currentSinkingOp].end())
+            {
+                return cacheIt->second;
+            }
+
+            float scaleX = currentSinkingOp->op.args.ImageResizeRel.factor[0];
+            float scaleY = currentSinkingOp->op.args.ImageResizeRel.factor[1];
+
+            Ptr<ASTOp> newAt = at;
+            switch ( at->GetOpType() )
+            {
+
+            case OP_TYPE::IM_CONDITIONAL:
+            {
+                // We move the mask creation down the two paths
+                auto newOp = mu::Clone<ASTOpConditional>(at);
+                newOp->yes = Visit(newOp->yes.child(), currentSinkingOp);
+                newOp->no = Visit(newOp->no.child(), currentSinkingOp);
+                newAt = newOp;
+                break;
+            }
+
+            case OP_TYPE::IM_PIXELFORMAT:
+            {
+                // We move the mask creation down format.
+                // \todo: only if shrinking?
+                Ptr<ASTOpImagePixelFormat> newOp = mu::Clone<ASTOpImagePixelFormat>(at);
+                newOp->Source = Visit( newOp->Source.child(), currentSinkingOp);
+                newAt = newOp;
+                break;
+            }
+
+            case OP_TYPE::IM_SWITCH:
+            {
+                // We move the mask creation down all the paths
+                auto newOp = mu::Clone<ASTOpSwitch>(at);
+                newOp->def = Visit(newOp->def.child(), currentSinkingOp);
+                for( auto& c:newOp->cases )
+                {
+                    c.branch = Visit(c.branch.child(), currentSinkingOp);
+                }
+                newAt = newOp;
+                break;
+            }
+
+            case OP_TYPE::IM_SWIZZLE:
+            {
+                auto newOp = mu::Clone<ASTOpFixed>(at);
+                for( int s=0; s<MUTABLE_OP_MAX_SWIZZLE_CHANNELS; ++s )
+                {
+                    auto channelOp = newOp->children[newOp->op.args.ImageSwizzle.sources[s]].child();
+                    if (channelOp)
+                    {
+                        newOp->SetChild( newOp->op.args.ImageSwizzle.sources[s], Visit( channelOp, currentSinkingOp) );
+                    }
+                }
+                newAt = newOp;
+                break;
+            }
+
+            case OP_TYPE::IM_COMPOSE:
+            {
+                // We can only optimise if the layout grid blocks size in pixels is
+                // still an integer after relative scale
+                bool acceptable = false;
+                {
+                    auto typedAt = dynamic_cast<const ASTOpImageCompose*>(at.get());
+                    auto originalBaseOp = typedAt->Base.child();
+
+                    // \todo: recursion-proof cache?
+                    int layoutBlockPixelsX = 0;
+                    int layoutBlockPixelsY = 0;
+                    originalBaseOp->GetLayoutBlockSize( &layoutBlockPixelsX, &layoutBlockPixelsY );
+
+                    int scaledLayoutBlockPixelsX = int(layoutBlockPixelsX*scaleX);
+                    int scaledLayoutBlockPixelsY = int(layoutBlockPixelsY*scaleY);
+                    int unscaledLayoutBlockPixelsX = int(scaledLayoutBlockPixelsX/scaleX);
+                    int unscaledLayoutBlockPixelsY = int(scaledLayoutBlockPixelsY/scaleY);
+                    acceptable =
+                            ( layoutBlockPixelsX!=0 && layoutBlockPixelsY!=0 )
+                            &&
+                            ( layoutBlockPixelsX == unscaledLayoutBlockPixelsX )
+                            &&
+                            ( layoutBlockPixelsY == unscaledLayoutBlockPixelsY );
+                }
+
+                if (acceptable)
+                {
+                    auto newOp = mu::Clone<ASTOpImageCompose>(at);
+
+                    auto baseOp = newOp->Base.child();
+                    newOp->Base = Visit( baseOp, currentSinkingOp);
+
+                    auto blockOp = newOp->BlockImage.child();
+                    newOp->BlockImage = Visit( blockOp, currentSinkingOp);
+
+                    newAt = newOp;
+                }
+
+                break;
+            }
+
+            case OP_TYPE::IM_PATCH:
+            {
+                auto newOp = mu::Clone<ASTOpImagePatch>(at);
+
+                newOp->base = Visit( newOp->base.child(), currentSinkingOp );
+                newOp->patch = Visit( newOp->patch.child(), currentSinkingOp);
+
+                // todo: review if this is always correct, or we need some "divisible" check
+                newOp->location[0] = uint16_t( newOp->location[0] *  scaleX );
+                newOp->location[1] = uint16_t( newOp->location[1] *  scaleY );
+
+                newAt = newOp;
+
+                break;
+            }
+
+            case OP_TYPE::IM_MIPMAP:
+            {
+                auto newOp = mu::Clone<ASTOpImageMipmap>(at.get());
+                auto baseOp = newOp->Source.child();
+                newOp->Source = Visit( baseOp, currentSinkingOp );
+                newAt = newOp;
+                break;
+            }
+
+            case OP_TYPE::IM_INTERPOLATE:
+            {
+                auto newOp = mu::Clone<ASTOpFixed>(at);
+
+                for ( int v=0; v<MUTABLE_OP_MAX_INTERPOLATE_COUNT; ++v )
+                {
+                    auto child = newOp->children[ newOp->op.args.ImageInterpolate.targets[v] ].child();
+                    auto bOp = Visit( child, currentSinkingOp );
+                    newOp->SetChild( newOp->op.args.ImageInterpolate.targets[v], bOp );
+                }
+
+                newAt = newOp;
+                break;
+            }
+
+
+            case OP_TYPE::IM_INTERPOLATE3:
+            {
+                auto newOp = mu::Clone<ASTOpFixed>(at);
+
+                auto top0 = newOp->children[ newOp->op.args.ImageInterpolate3.target0 ].child();
+                newOp->SetChild( newOp->op.args.ImageInterpolate3.target0, Visit(top0, currentSinkingOp) );
+
+                auto top1 = newOp->children[ newOp->op.args.ImageInterpolate3.target1 ].child();
+                newOp->SetChild( newOp->op.args.ImageInterpolate3.target1, Visit(top1, currentSinkingOp) );
+
+                auto top2 = newOp->children[ newOp->op.args.ImageInterpolate3.target2 ].child();
+                newOp->SetChild( newOp->op.args.ImageInterpolate3.target2, Visit(top2, currentSinkingOp) );
+
+                newAt = newOp;
+                break;
+            }
+
+            case OP_TYPE::IM_MULTILAYER:
+            {
+                auto nop = mu::Clone<ASTOpImageMultiLayer>(at);
+                nop->base = Visit( nop->base.child(), currentSinkingOp );
+                nop->mask = Visit( nop->mask.child(), currentSinkingOp );
+                nop->blend = Visit( nop->blend.child(), currentSinkingOp );
+                newAt = nop;
+                break;
+            }
+
+            case OP_TYPE::IM_LAYER:
+            {
+                auto nop = mu::Clone<ASTOpFixed>(at);
+
+                auto aOp = nop->children[ nop->op.args.ImageLayer.base ].child();
+                nop->SetChild( nop->op.args.ImageLayer.base, Visit( aOp, currentSinkingOp ) );
+
+                auto bOp = nop->children[ nop->op.args.ImageLayer.blended ].child();
+                nop->SetChild( nop->op.args.ImageLayer.blended, Visit( bOp, currentSinkingOp ) );
+
+                auto mOp = nop->children[ nop->op.args.ImageLayer.mask ].child();
+                nop->SetChild( nop->op.args.ImageLayer.mask, Visit( mOp, currentSinkingOp ) );
+
+                newAt = nop;
+                break;
+            }
+
+            case OP_TYPE::IM_LAYERCOLOUR:
+            {
+                auto nop = mu::Clone<ASTOpFixed>(at);
+
+                auto aOp = nop->children[ nop->op.args.ImageLayerColour.base ].child();
+                nop->SetChild( nop->op.args.ImageLayerColour.base, Visit( aOp, currentSinkingOp ) );
+
+                auto mOp = nop->children[ nop->op.args.ImageLayerColour.mask ].child();
+                nop->SetChild( nop->op.args.ImageLayer.mask, Visit( mOp, currentSinkingOp ) );
+
+                newAt = nop;
+                break;
+            }
+
+            case OP_TYPE::IM_DISPLACE:
+            {
+                auto nop = mu::Clone<ASTOpFixed>(at);
+
+                auto sourceOp = nop->children[ nop->op.args.ImageDisplace.source ].child();
+                nop->SetChild( nop->op.args.ImageDisplace.source, Visit(sourceOp, currentSinkingOp ) );
+
+                auto mapOp = nop->children[ nop->op.args.ImageDisplace.displacementMap ].child();
+                nop->SetChild( nop->op.args.ImageDisplace.displacementMap, Visit(mapOp, currentSinkingOp ) );
+
+                // Make sure we don't scale a constant displacement map, which is very wrong.
+                if ( mapOp->GetOpType()==OP_TYPE::IM_MAKEGROWMAP )
+                {
+                    newAt = nop;
+                }
+                else
+                {
+                    // We cannot resize an already calculated displacement map.
+                }
+
+                break;
+            }
+
+            case OP_TYPE::IM_MAKEGROWMAP:
+            {
+                auto nop = mu::Clone<ASTOpFixed>(at);
+                auto maskOp = nop->children[ nop->op.args.ImageMakeGrowMap.mask ].child();
+                nop->SetChild( nop->op.args.ImageMakeGrowMap.mask, Visit( maskOp, currentSinkingOp ) );
+                newAt = nop;
+                break;
+            }
+
+            case OP_TYPE::IM_RASTERMESH:
+            {
+                auto nop = mu::Clone<ASTOpFixed>(at);
+                auto maskOp = nop->children[ nop->op.args.ImageRasterMesh.mask ].child();
+                nop->SetChild( nop->op.args.ImageRasterMesh.mask, Visit( maskOp, currentSinkingOp ) );
+
+				// Resize the image to project as well, assuming that since the target has a different resolution
+				// it make sense for the source image to have a similar resize.
+				// Actually, don't do it because the LODBias will be applied separetely at graph generation time.
+				//auto imageOp = nop->children[nop->op.args.ImageRasterMesh.image].child();
+				//nop->SetChild(nop->op.args.ImageRasterMesh.image, Visit(imageOp, currentSinkingOp));
+
+                nop->op.args.ImageRasterMesh.sizeX = uint16_t( nop->op.args.ImageRasterMesh.sizeX * scaleX + 0.5f );
+                nop->op.args.ImageRasterMesh.sizeY = uint16_t( nop->op.args.ImageRasterMesh.sizeY * scaleY + 0.5f );
+                newAt = nop;
+                break;
+            }
+
+            default:
+                break;
+            }
+
+            // end on line, replace with sinking op
+            if (at==newAt && at!=m_initialSource)
+            {
+                Ptr<ASTOpFixed> newOp = mu::Clone<ASTOpFixed>(currentSinkingOp);
+                check(newOp->GetOpType()==OP_TYPE::IM_RESIZEREL);
+
+                newOp->SetChild( newOp->op.args.ImageResizeRel.source,at);
+
+                newAt = newOp;
+            }
+
+            m_oldToNew[currentSinkingOp][at] = newAt;
+
+            return newAt;
+        }
+    };
+
+
+    //---------------------------------------------------------------------------------------------
+    Ptr<ASTOp> ASTOpFixed::OptimiseSize() const
+    {
+        Ptr<ASTOp> at;
+
+        OP_TYPE type = GetOpType();
+        switch ( type )
+        {
+
+        //-------------------------------------------------------------------------------------
+        case OP_TYPE::IM_RESIZE:
+        {
+            auto sourceAt = children[op.args.ImageResize.source].child();
+
+            // The instruction can be sunk
+            OP_TYPE sourceType = sourceAt->GetOpType();
+            switch ( sourceType )
+            {
+            case OP_TYPE::IM_RESIZE:
+            {
+                // Keep top resize
+                auto sourceOp = dynamic_cast<const ASTOpFixed*>(sourceAt.get());
+
+                Ptr<ASTOpFixed> newOp = mu::Clone<ASTOpFixed>(this);
+                newOp->SetChild( newOp->op.args.ImageResize.source, sourceOp->children[sourceOp->op.args.ImageResize.source]);
+
+                at = newOp;
+                break;
+            }
+
+            case OP_TYPE::IM_PLAINCOLOUR:
+            {
+                // Set the size in the children and remove resize
+                auto sourceOp = mu::Clone<ASTOpFixed>(sourceAt.get());
+                sourceOp->op.args.ImagePlainColour.size[0] = op.args.ImageResize.size[0];
+                sourceOp->op.args.ImagePlainColour.size[1] = op.args.ImageResize.size[1];
+                at = sourceOp;
+                break;
+            }
+
+            case OP_TYPE::IM_CONDITIONAL:
+            {
+                // We move the resize down the two paths
+                auto newOp = mu::Clone<ASTOpConditional>(sourceAt);
+
+                Ptr<ASTOpFixed> aOp = mu::Clone<ASTOpFixed>(this);
+                aOp->SetChild( aOp->op.args.ImageResize.source,newOp->yes);
+                newOp->yes = aOp;
+
+                Ptr<ASTOpFixed> bOp = mu::Clone<ASTOpFixed>(this);
+                bOp->SetChild( bOp->op.args.ImageResize.source,newOp->no);
+                newOp->no = bOp;
+
+                at = newOp;
+                break;
+            }
+
+            case OP_TYPE::IM_SWITCH:
+            {
+                // Move the resize down all the paths
+                Ptr<ASTOpSwitch> newOp = mu::Clone<ASTOpSwitch>(sourceAt);
+
+                if (newOp->def)
+                {
+                    Ptr<ASTOpFixed> defOp = mu::Clone<ASTOpFixed>(this);
+                    defOp->SetChild( defOp->op.args.ImageResize.source, newOp->def);
+                    newOp->def = defOp;
+                }
+
+                for ( auto& cas:newOp->cases )
+                {
+                    if ( cas.branch )
+                    {
+                        Ptr<ASTOpFixed> bOp = mu::Clone<ASTOpFixed>(this);
+                        bOp->SetChild( bOp->op.args.ImageResize.source, cas.branch);
+                        cas.branch = bOp;
+                    }
+                }
+
+                at = newOp;
+                break;
+            }
+
+            case OP_TYPE::IM_COMPOSE:
+            {
+                Ptr<ASTOpImageCompose> newOp = mu::Clone<ASTOpImageCompose>(sourceAt);
+
+                auto aOp = mu::Clone<ASTOpFixed>(this);
+                aOp->SetChild( aOp->op.args.ImageResize.source, newOp->Base.child());
+                newOp->Base = aOp;
+
+                auto bOp = mu::Clone<ASTOpFixed>(this);
+                bOp->SetChild( bOp->op.args.ImageResize.source, newOp->BlockImage.child());
+                newOp->BlockImage = bOp;
+
+                if ( newOp->Mask )
+                {
+                    auto maskOp = mu::Clone<ASTOpFixed>(this);
+                    maskOp->SetChild( maskOp->op.args.ImageResize.source, newOp->Mask.child());
+                    newOp->Mask = maskOp;
+                }
+
+                at = newOp;
+                break;
+            }
+
+            case OP_TYPE::IM_RASTERMESH:
+            {
+                Ptr<ASTOpFixed> newOp = mu::Clone<ASTOpFixed>(sourceAt);
+
+                //if ( newOp->op.args.ImageRasterMesh.sizeX != op.args.ImageResize.size[0]
+                //     ||
+                //     newOp->op.args.ImageRasterMesh.sizeY != op.args.ImageResize.size[1] )
+                {
+                    newOp->op.args.ImageRasterMesh.sizeX = op.args.ImageResize.size[0];
+                    newOp->op.args.ImageRasterMesh.sizeY = op.args.ImageResize.size[1];
+
+					if (newOp->op.args.ImageRasterMesh.mask)
+					{
+						Ptr<ASTOpFixed> mop = mu::Clone<ASTOpFixed>(this);
+						mop->SetChild(mop->op.args.ImageResize.source, newOp->children[newOp->op.args.ImageRasterMesh.mask].child());
+						newOp->SetChild(newOp->op.args.ImageRasterMesh.mask, mop);
+					}
+
+					// Don't apply absolute resizes to the image to raster: it could even enlarge it.
+					// This should only be scaled with relative resizes, which come from LOD biases, etc.
+					//if (newOp->op.args.ImageRasterMesh.image)
+					//{
+					//	Ptr<ASTOpFixed> mop = mu::Clone<ASTOpFixed>(this);
+					//	mop->SetChild(mop->op.args.ImageResize.source, newOp->children[newOp->op.args.ImageRasterMesh.image].child());
+					//	newOp->SetChild(newOp->op.args.ImageRasterMesh.image, mop);
+					//}
+				}
+
+                at = newOp;
+                break;
+            }
+
+            case OP_TYPE::IM_INTERPOLATE:
+            {
+                Ptr<ASTOpFixed> newOp = mu::Clone<ASTOpFixed>(sourceAt);
+
+                for ( int s=0; s<MUTABLE_OP_MAX_INTERPOLATE_COUNT; ++s )
+                {
+                    auto targetAt = newOp->children[ newOp->op.args.ImageInterpolate.targets[s] ].child();
+                    if ( targetAt )
+                    {
+                        Ptr<ASTOpFixed> sourceOp = mu::Clone<ASTOpFixed>(this);
+                        sourceOp->SetChild( sourceOp->op.args.ImageResize.source, targetAt);
+                        newOp->SetChild( newOp->op.args.ImageInterpolate.targets[s], sourceOp );
+                    }
+                }
+
+                at = newOp;
+                break;
+            }
+
+            case OP_TYPE::IM_INTERPOLATE3:
+            {
+                Ptr<ASTOpFixed> newOp = mu::Clone<ASTOpFixed>(sourceAt);
+
+                {
+                    Ptr<ASTOpFixed> sourceOp = mu::Clone<ASTOpFixed>(this);
+                    auto targetAt = newOp->children[ newOp->op.args.ImageInterpolate3.target0 ].child();
+                    sourceOp->SetChild( sourceOp->op.args.ImageResize.source, targetAt);
+                    newOp->SetChild( newOp->op.args.ImageInterpolate3.target0, sourceOp );
+                }
+
+                {
+                    Ptr<ASTOpFixed> sourceOp = mu::Clone<ASTOpFixed>(this);
+                    auto targetAt = newOp->children[ newOp->op.args.ImageInterpolate3.target1 ].child();
+                    sourceOp->SetChild( sourceOp->op.args.ImageResize.source, targetAt);
+                    newOp->SetChild( newOp->op.args.ImageInterpolate3.target1, sourceOp );
+                }
+
+                {
+                    Ptr<ASTOpFixed> sourceOp = mu::Clone<ASTOpFixed>(this);
+                    auto targetAt = newOp->children[ newOp->op.args.ImageInterpolate3.target2 ].child();
+                    sourceOp->SetChild( sourceOp->op.args.ImageResize.source, targetAt);
+                    newOp->SetChild( newOp->op.args.ImageInterpolate3.target2, sourceOp );
+                }
+
+                at = newOp;
+                break;
+            }
+
+            case OP_TYPE::IM_LAYER:
+            {
+                Ptr<ASTOpFixed> newOp = mu::Clone<ASTOpFixed>(sourceAt);
+
+                Ptr<ASTOpFixed> baseOp = mu::Clone<ASTOpFixed>(this);
+                baseOp->SetChild( baseOp->op.args.ImageResize.source,
+                        newOp->children[newOp->op.args.ImageLayer.base] );
+                newOp->SetChild( newOp->op.args.ImageLayer.base, baseOp );
+
+                Ptr<ASTOpFixed> blendOp = mu::Clone<ASTOpFixed>(this);
+                blendOp->SetChild( blendOp->op.args.ImageResize.source,
+                        newOp->children[newOp->op.args.ImageLayer.blended] );
+                newOp->SetChild( newOp->op.args.ImageLayer.blended, blendOp );
+
+                auto maskAt = newOp->children[newOp->op.args.ImageLayer.mask].child();
+                if (maskAt)
+                {
+                    Ptr<ASTOpFixed> maskOp = mu::Clone<ASTOpFixed>(this);
+                    maskOp->SetChild( maskOp->op.args.ImageResize.source, maskAt );
+                    newOp->SetChild( newOp->op.args.ImageLayer.mask, maskOp );
+                }
+
+                at = newOp;
+                break;
+            }
+
+            case OP_TYPE::IM_LAYERCOLOUR:
+            {
+                Ptr<ASTOpFixed> newOp = mu::Clone<ASTOpFixed>(sourceAt);
+
+                Ptr<ASTOpFixed> baseOp = mu::Clone<ASTOpFixed>(this);
+                baseOp->SetChild( baseOp->op.args.ImageResize.source,
+                        newOp->children[newOp->op.args.ImageLayerColour.base] );
+                newOp->SetChild( newOp->op.args.ImageLayerColour.base, baseOp );
+
+                auto maskAt = newOp->children[newOp->op.args.ImageLayerColour.mask].child();
+                if (maskAt)
+                {
+                    Ptr<ASTOpFixed> maskOp = mu::Clone<ASTOpFixed>(this);
+                    maskOp->SetChild( maskOp->op.args.ImageResize.source, maskAt );
+                    newOp->SetChild( newOp->op.args.ImageLayerColour.mask, maskOp );
+                }
+
+                at = newOp;
+                break;
+            }
+
+			case OP_TYPE::IM_DISPLACE:
+			{
+				// In the size optimization phase we can optimize the resize with the displace
+				// because the constants have not been collapsed yet.
+				// We will still check it and sink the size directly below the IM_MAKEGROWMAP op
+				Ptr<ASTOpFixed> SourceTyped = dynamic_cast<ASTOpFixed*>(sourceAt.get());
+				check(SourceTyped);
+				Ptr<ASTOp> OriginalDisplacementMapOp = SourceTyped->children[SourceTyped->op.args.ImageDisplace.displacementMap].m_child;
+				if (OriginalDisplacementMapOp->GetOpType() == OP_TYPE::IM_MAKEGROWMAP)
+				{
+					Ptr<ASTOpFixed> newOp = mu::Clone<ASTOpFixed>(sourceAt);
+
+					Ptr<ASTOpFixed> baseOp = mu::Clone<ASTOpFixed>(this);
+					baseOp->SetChild(baseOp->op.args.ImageResize.source,
+						newOp->children[newOp->op.args.ImageDisplace.source]);
+					newOp->SetChild(newOp->op.args.ImageDisplace.source, baseOp);
+
+					// Clone the map op and replace its children
+					Ptr<ASTOpFixed> mapOp = mu::Clone<ASTOpFixed>(OriginalDisplacementMapOp);
+					newOp->SetChild(newOp->op.args.ImageDisplace.displacementMap, mapOp);
+
+					Ptr<ASTOpFixed> mapSourceOp = mu::Clone<ASTOpFixed>(this);
+					mapSourceOp->SetChild(mapSourceOp->op.args.ImageResize.source,
+						mapOp->children[mapOp->op.args.ImageMakeGrowMap.mask]);
+					mapOp->SetChild(mapOp->op.args.ImageDisplace.displacementMap, mapSourceOp);
+
+					at = newOp;
+				}
+				break;
+			}
+
+            default:
+                break;
+            }
+
+            break;
+        }
+
+
+
+        //-------------------------------------------------------------------------------------
+        case OP_TYPE::IM_RESIZEREL:
+        {
+            auto sourceAt = children[op.args.ImageResizeRel.source].child();
+
+            // The instruction can be sunk
+            OP_TYPE sourceType = sourceAt->GetOpType();
+            switch ( sourceType )
+            {
+
+                case OP_TYPE::IM_BLANKLAYOUT:
+                {
+                    auto newOp = mu::Clone<ASTOpFixed>(sourceAt);
+
+                    newOp->op.args.ImageBlankLayout.blockSize[0] =
+                            uint16_t( newOp->op.args.ImageBlankLayout.blockSize[0]
+                                    * op.args.ImageResizeRel.factor[0]
+                                    + 0.5f );
+                    newOp->op.args.ImageBlankLayout.blockSize[1] =
+                            uint16_t( newOp->op.args.ImageBlankLayout.blockSize[1]
+                                    * op.args.ImageResizeRel.factor[1]
+                                    + 0.5f );
+                    at = newOp;
+                    break;
+                }
+
+
+				// Don't combine. ResizeRel sometimes can resize more children than Resize can do. (see RasterMesh)
+//                case OP_TYPE::IM_RESIZE:
+//                {
+//                    OP op = program.m_code[sourceAt];
+//                    op.args.ImageResize.size[0] =
+//                            int16_t( op.args.ImageResize.size[0]
+//                                   * program.m_code[at].args.ImageResizeRel.factor[0] );
+//                    op.args.ImageResize.size[1] =
+//                            int16_t( op.args.ImageResize.size[1]
+//                                   * program.m_code[at].args.ImageResizeRel.factor[1] );
+//                    return op;
+//                    break;
+//                }
+
+
+                default:
+                {
+                    Sink_ImageResizeRelAST sinker;
+                    at = sinker.Apply(this);
+
+                    break;
+                }
+
+                }
+
+                break;
+            }
+
+
+            //-------------------------------------------------------------------------------------
+            case OP_TYPE::IM_RESIZELIKE:
+            {
+                Ptr<ImageSizeExpression> sourceSize =
+                        children[op.args.ImageResizeLike.source]->GetImageSizeExpression();
+                Ptr<ImageSizeExpression> sizeSourceSize =
+                        children[op.args.ImageResizeLike.sizeSource]->GetImageSizeExpression();
+
+                if ( *sourceSize == *sizeSourceSize )
+                {
+                    at = children[op.args.ImageResizeLike.source].child();
+                }
+
+                else if ( sizeSourceSize->type == ImageSizeExpression::ISET_CONSTANT )
+                {
+                    Ptr<ASTOpFixed> newAt = new ASTOpFixed;
+                    newAt->op.type = OP_TYPE::IM_RESIZE;
+                    newAt->SetChild( newAt->op.args.ImageResize.source, children[op.args.ImageResizeLike.source] );
+                    newAt->op.args.ImageResize.size[0] = sizeSourceSize->size[0];
+                    newAt->op.args.ImageResize.size[1] = sizeSourceSize->size[1];
+                    at = newAt;
+                }
+                else if ( sizeSourceSize->type == ImageSizeExpression::ISET_LAYOUTFACTOR )
+                {
+                    // TODO
+                    // Skip intermediate ops until the layout
+//                    if ( program.m_code[sizeSourceAt].type != OP_TYPE::IM_BLANKLAYOUT )
+//                    {
+//                        m_modified = true;
+
+//                        OP blackLayoutOp;
+//                        blackLayoutOp.type = OP_TYPE::IM_BLANKLAYOUT;
+//                        blackLayoutOp.args.ImageBlankLayout.layout = sizeSourceSize->layout;
+//                        blackLayoutOp.args.ImageBlankLayout.blockSize[0] = sizeSourceSize->factor[0];
+//                        blackLayoutOp.args.ImageBlankLayout.blockSize[1] = sizeSourceSize->factor[1];
+//                        blackLayoutOp.args.ImageBlankLayout.format = IF_L_UBYTE;
+
+//                        OP opResize = program.m_code[at];
+//                        opResize.args.ImageResizeLike.sizeSource = program.AddOp( blackLayoutOp );
+//                        at = program.AddOp( opResize );
+//                    }
+                }
+
+                break;
+            }
+
+        default:
+            break;
+
+        }
+
+        return at;
+    }
+
+}
+
