[mutable] Move Mutable plugin from the restricted folder to the experimental plugins folder.

diff --git a/Engine/Plugins/Experimental/Mutable/Source/MutableTools/Private/Compiler.cpp b/Engine/Plugins/Experimental/Mutable/Source/MutableTools/Private/Compiler.cpp
new file mode 100644
--- /dev/null
+++ b/Engine/Plugins/Experimental/Mutable/Source/MutableTools/Private/Compiler.cpp
@@ -0,0 +1,509 @@
+// Copyright Epic Games, Inc. All Rights Reserved.
+
+
+#include "Compiler.h"
+
+#include "CompilerPrivate.h"
+#include "CodeGenerator.h"
+#include "CodeOptimiser.h"
+#include "DataPacker.h"
+#include "ErrorLogPrivate.h"
+#include "Hash/CityHash.h"
+
+#include "Node.h"
+#include "NodeImage.h"
+#include "NodeMesh.h"
+#include "NodeObject.h"
+
+#include "Operations.h"
+#include "ModelPrivate.h"
+#include "SystemPrivate.h"
+
+#include "MutableRuntime/Public/Model.h"
+#include "MutableRuntime/Public/Instance.h"
+
+#include "Config.h"
+#include "MutableTrace.h"
+
+
+namespace mu
+{
+
+    const char* CompilerOptions::GetTextureLayoutStrategyName( CompilerOptions::TextureLayoutStrategy s )
+    {
+        static const char* s_textureLayoutStrategyName[ 2 ] =
+        {
+            "Unscaled Pack",
+            "No Packing",
+        };
+
+        return s_textureLayoutStrategyName[(int)s];
+    }
+
+
+    //---------------------------------------------------------------------------------------------
+    //---------------------------------------------------------------------------------------------
+    //---------------------------------------------------------------------------------------------
+    CompilerOptions::CompilerOptions()
+    {
+        m_pD = new Private();
+        m_pD->m_log = false;
+    }
+
+
+    //---------------------------------------------------------------------------------------------
+    CompilerOptions::~CompilerOptions()
+    {
+        check( m_pD );
+        delete m_pD;
+        m_pD = 0;
+    }
+
+
+    //---------------------------------------------------------------------------------------------
+    CompilerOptions::Private* CompilerOptions::GetPrivate() const
+    {
+        return m_pD;
+    }
+
+
+    //---------------------------------------------------------------------------------------------
+    void CompilerOptions::SetLogEnabled( bool enabled )
+    {
+        m_pD->m_log = enabled;
+    }
+
+
+    //---------------------------------------------------------------------------------------------
+    void CompilerOptions::SetOptimisationEnabled( bool enabled )
+    {
+        m_pD->m_optimisationOptions.m_enabled = enabled;
+        if (enabled)
+        {
+            m_pD->m_optimisationOptions.m_constReduction = true;
+        }
+    }
+
+
+    //---------------------------------------------------------------------------------------------
+    void CompilerOptions::SetConstReductionEnabled( bool constReduction )
+    {
+        m_pD->m_optimisationOptions.m_constReduction = constReduction;
+    }
+
+
+	//---------------------------------------------------------------------------------------------
+	void CompilerOptions::SetEnablePartialOptimisation(bool bEnabled)
+	{
+		m_pD->m_enablePartialOptimise = bEnabled;
+	}
+
+
+	//---------------------------------------------------------------------------------------------
+	void CompilerOptions::SetEnableConcurrency(bool bEnabled)
+	{
+		m_pD->m_enableConcurrency = bEnabled;
+	}
+
+
+    //---------------------------------------------------------------------------------------------
+    void CompilerOptions::SetUseDiskCache( bool enabled )
+    {
+        m_pD->m_optimisationOptions.m_useDiskCache = enabled;
+    }
+
+
+    //---------------------------------------------------------------------------------------------
+    void CompilerOptions::SetOptimisationMaxIteration( int maxIterations )
+    {
+        m_pD->m_optimisationOptions.m_maxOptimisationLoopCount = maxIterations;
+    }
+
+
+    //---------------------------------------------------------------------------------------------
+    void CompilerOptions::SetIgnoreStates( bool ignore )
+    {
+        m_pD->m_ignoreStates = ignore;
+    }
+
+
+    //---------------------------------------------------------------------------------------------
+    void CompilerOptions::SetTextureLayoutStrategy( TextureLayoutStrategy strategy )
+    {
+        m_pD->m_textureLayoutStrategy = strategy;
+    }
+
+
+    //---------------------------------------------------------------------------------------------
+    void CompilerOptions::SetImageCompressionQuality( int quality )
+    {
+        m_pD->m_imageCompressionQuality = quality;
+    }
+
+
+	//---------------------------------------------------------------------------------------------
+	void CompilerOptions::SetDataPackingStrategy(int32 MinRomSize, int32 MinTextureResidentMipCount)
+	{
+		m_pD->MinRomSize = MinRomSize;
+		m_pD->MinTextureResidentMipCount = MinTextureResidentMipCount;
+	}
+
+
+    //---------------------------------------------------------------------------------------------
+    //---------------------------------------------------------------------------------------------
+    //---------------------------------------------------------------------------------------------
+    Compiler::Compiler( CompilerOptionsPtr options )
+    {
+        m_pD = new Private();
+        m_pD->m_options = options;
+        if (!m_pD->m_options)
+        {
+            m_pD->m_options = new CompilerOptions;
+        }
+    }
+
+
+    //---------------------------------------------------------------------------------------------
+    Compiler::~Compiler()
+    {
+        check( m_pD );
+        delete m_pD;
+        m_pD = 0;
+    }
+
+
+    //---------------------------------------------------------------------------------------------
+    ModelPtr Compiler::Compile( const Ptr<Node>& pNode )
+    {
+        MUTABLE_CPUPROFILER_SCOPE(Compile);
+
+        TaskManager* pTaskManager = new TaskManager(m_pD->m_options->GetPrivate()->m_enableConcurrency);
+
+        vector< STATE_COMPILATION_DATA > states;
+        Ptr<ErrorLog> genErrorLog;
+        {
+            CodeGenerator gen( m_pD->m_options->GetPrivate() );
+
+            gen.GenerateRoot( pNode, pTaskManager );
+
+            check( !gen.m_states.empty() );
+
+            for ( const auto& s: gen.m_states )
+            {
+                STATE_COMPILATION_DATA data;
+                data.nodeState = s.first;
+                data.root = s.second;
+                data.state.m_name = s.first.m_name;
+                data.optimisationFlags = s.first.m_optimisation;
+                states.push_back( data );
+            }
+
+            genErrorLog = gen.m_pErrorLog;
+        }
+
+        vector<Ptr<ASTOp>> roots;
+        for( const auto& s: states)
+        {
+            roots.push_back(s.root);
+        }
+
+        // Slow AST code verification for debugging.
+        //ASTOp::FullAssert(roots);
+
+        // Optimize the generated code
+        {
+            CodeOptimiser optimiser( m_pD->m_options, states );
+            optimiser.OptimiseAST( pTaskManager );
+        }
+
+
+        // Link the program and generate state data.
+        ModelPtr pResult = new Model();
+        auto& program = pResult->GetPrivate()->m_program;
+        for( auto& s: states )
+        {
+			FLinkerOptions LinkerOptions;
+			LinkerOptions.MinTextureResidentMipCount = m_pD->m_options->GetPrivate()->MinTextureResidentMipCount;
+
+            ASTOp::FullLink(s.root,program, &LinkerOptions);
+            if (s.root)
+            {
+                s.state.m_root = s.root->linkedAddress;
+            }
+            else
+            {
+                s.state.m_root = 0;
+            }
+        }
+
+        // Set the runtime parameter indices.
+        for( auto& s: states )
+        {
+            for ( size_t p=0; p<s.nodeState.m_runtimeParams.size(); ++p )
+            {
+                int paramIndex = -1;
+                for ( size_t i=0; paramIndex<0 && i<program.m_parameters.size(); ++i )
+                {
+                    if ( program.m_parameters[i].m_name
+                         ==
+                         s.nodeState.m_runtimeParams[p] )
+                    {
+                        paramIndex = (int)i;
+                    }
+                }
+
+                if (paramIndex>=0)
+                {
+                    s.state.m_runtimeParameters.push_back( paramIndex );
+                }
+                else
+                {
+                    char temp[256];
+                    mutable_snprintf( temp, 256,
+                                      "The state [%s] refers to a parameter [%s] "
+                                      "that has not been found in the model. This error can be "
+                                      "safely dismissed in case of partial compilation.",
+                                      s.nodeState.m_name.c_str(),
+                                      s.nodeState.m_runtimeParams[p].c_str()
+                                      );
+                    m_pD->m_pErrorLog->GetPrivate()->Add
+                            ( temp, ELMT_WARNING, pNode->GetBasePrivate()->m_errorContext );
+                }
+            }
+
+            // Generate the mask of update cache ops
+            for ( const auto& a: s.m_updateCache )
+            {
+                s.state.m_updateCache.emplace_back(a->linkedAddress);
+            }
+
+            // Sort the update cache addresses for performance and determinism
+            std::sort( s.state.m_updateCache.begin(), s.state.m_updateCache.end() );
+
+            // Generate the mask of dynamic resources
+            for ( const auto& a: s.m_dynamicResources )
+            {
+                uint64_t relevantMask = 0;
+                for ( const auto& b: a.second )
+                {
+                    // Find the index in the model parameter list
+                    int paramIndex = -1;
+                    for ( size_t i=0; paramIndex<0 && i<program.m_parameters.size(); ++i )
+                    {
+                        if ( program.m_parameters[i].m_name == b )
+                        {
+                            paramIndex = (int)i;
+                        }
+                    }
+                    check(paramIndex>=0);
+
+                    // Find the position in the state data vector.
+                    auto it = std::find( s.state.m_runtimeParameters.cbegin(),
+                                         s.state.m_runtimeParameters.cend(),
+                                         paramIndex );
+
+                    if ( it!=s.state.m_runtimeParameters.end() )
+                    {
+                        size_t pos = it - s.state.m_runtimeParameters.begin();
+                        relevantMask |= uint64_t(1) << pos;
+                    }
+                }
+
+                // TODO: this shouldn't happen but it seems to happen. Investigate.
+                // Maybe something with the difference of precision between the relevant parameters
+                // in operation subtrees.
+                //check(relevantMask!=0);
+                if (relevantMask!=0)
+                {
+                    s.state.m_dynamicResources.emplace_back( a.first->linkedAddress, relevantMask );
+                }
+            }            
+
+            // Sort for performance and determinism
+            std::sort( s.state.m_dynamicResources.begin(), s.state.m_dynamicResources.end() );
+
+            program.m_states.push_back(s.state);
+        }
+
+        //
+        if ( m_pD->m_options->GetPrivate()->m_log )
+        {
+			UE_LOG(LogMutableCore, Log, TEXT("----------------------------------------------------------------------\n") );
+			UE_LOG(LogMutableCore, Log, TEXT("non-optimised size : %d\n"), program.m_opAddress.size() );
+			UE_LOG(LogMutableCore, Log, TEXT("----------------------------------------------------------------------\n") );
+            Log( pResult );
+			UE_LOG(LogMutableCore, Log, TEXT("----------------------------------------------------------------------\n") );
+        }
+
+		UE_LOG(LogMutableCore, Verbose, TEXT("(int) %s : %ld"), TEXT("program size"), int64(program.m_opAddress.size()));
+
+        // Optimize the generated code
+        {        
+//            CodeOptimiser optimiser( m_pD->m_gpuPlatformProps,
+//                                     m_pD->m_optimisationOptions,
+//                                     states );
+//            optimiser.Optimise( pResult.get() );
+
+//            m_pD->m_pModelReport = optimiser.m_pModelReport;
+        }
+
+        // Merge the log in the right order
+        genErrorLog->Merge( m_pD->m_pErrorLog.get() );
+        m_pD->m_pErrorLog = genErrorLog.get();
+
+		// Pack data
+		int32 MinimumBytesPerRom = 1024; // \TODO: compilation parameter
+		m_pD->GenerateRoms(pResult.get(),MinimumBytesPerRom);
+
+		UE_LOG(LogMutableCore, Verbose, TEXT("(int) %s : %ld"), TEXT("program size"), int64(program.m_opAddress.size()));
+
+        delete pTaskManager;
+
+        // Finish logging
+        if ( m_pD->m_options->GetPrivate()->m_log )
+        {
+            Log( pResult );
+			UE_LOG(LogMutableCore, Log, TEXT("----------------------------------------------------------------------") );
+            m_pD->m_pErrorLog->Log();
+			UE_LOG(LogMutableCore, Log, TEXT("----------------------------------------------------------------------") );
+        }
+
+        return pResult;
+    }
+
+
+    //---------------------------------------------------------------------------------------------
+    ErrorLogPtrConst Compiler::GetLog() const
+    {
+        return m_pD->m_pErrorLog;
+    }
+
+
+    //---------------------------------------------------------------------------------------------
+    ModelReportPtrConst Compiler::GetModelReport() const
+    {
+        return m_pD->m_pModelReport;
+    }
+
+
+	//---------------------------------------------------------------------------------------------
+	namespace
+	{
+		inline void EnsureUniqueRomId(uint32& RomId, const mu::PROGRAM& program)
+		{
+			while (true)
+			{
+				bool bUnique = true;
+				for (const FRomData& Rom : program.m_roms)
+				{
+					if (Rom.Id == RomId)
+					{
+						bUnique = false;
+						break;
+					}
+				}
+
+				if (bUnique)
+				{
+					break;
+				}
+
+				// Generate a new Id. It is not going to be stable accross build, which means it may hurt content patching
+				// a little bit, but it shouldn't happen often.
+				RomId++;
+			}
+		}
+	}
+
+
+	//---------------------------------------------------------------------------------------------
+	void Compiler::Private::GenerateRoms(Model* p, int32 MinRomSize)
+	{
+		LLM_SCOPE_BYNAME(TEXT("MutableRuntime"));
+		MUTABLE_CPUPROFILER_SCOPE(Mutable_GenerateRoms);
+
+		mu::PROGRAM& program = p->GetPrivate()->m_program;
+
+		int32 NumRoms = 0;
+		int32 NumEmbedded = 0;
+
+		// Serialise roms in a separate files if possible
+		if (MinRomSize >= 0)
+		{
+			// Save images and unload from memory
+			for (int32 ResourceIndex = 0; ResourceIndex < program.m_constantImageLODs.Num(); ++ResourceIndex)
+			{
+				TPair<int32, mu::ImagePtrConst>& ResData = program.m_constantImageLODs[ResourceIndex];
+
+				// This shouldn't have been serialised with rom support before.
+				check(ResData.Key < 0);
+
+				// Serialize to memory, to find out final size of this rom
+				OutputMemoryStream MemStream(1024 * 1024);
+				OutputArchive MemoryArch(&MemStream);
+				Image::Serialise(ResData.Value.get(), MemoryArch);
+
+				// If the resource uses less memory than the threshold, don't save it in a separate rom.
+				if (MemStream.GetBufferSize() < MinRomSize)
+				{
+					NumEmbedded++;
+					continue;
+				}
+
+				NumRoms++;
+
+				FRomData RomData;
+				RomData.ResourceType = DT_IMAGE;
+				RomData.ResourceIndex = ResourceIndex;
+				RomData.Size = MemStream.GetBufferSize();
+
+				// Ensure that the Id is unique
+				RomData.Id = CityHash32(static_cast<const char*>(MemStream.GetBuffer()), MemStream.GetBufferSize());
+				EnsureUniqueRomId(RomData.Id, program);
+
+				int32 RomIndex = program.m_roms.Add(RomData);
+				ResData.Key = RomIndex;
+			}
+
+			// Save meshes and unload from memory
+			for (int32 ResourceIndex = 0; ResourceIndex < program.m_constantMeshes.Num(); ++ResourceIndex)
+			{
+				TPair<int32, mu::MeshPtrConst>& ResData = program.m_constantMeshes[ResourceIndex];
+
+				// This shouldn't have been serialised with rom support before.
+				check(ResData.Key < 0);
+
+				// Serialize to memory, to find out final size of this rom
+				OutputMemoryStream MemStream(1024 * 1024);
+				OutputArchive MemoryArch(&MemStream);
+				Mesh::Serialise(ResData.Value.get(), MemoryArch);
+
+				// If the resource uses less memory than the threshold, don't save it in a separate rom.
+				if (MemStream.GetBufferSize() < MinRomSize)
+				{
+					NumEmbedded++;
+					continue;
+				}
+
+				NumRoms++;
+
+				FRomData RomData;
+				RomData.ResourceType = DT_MESH;
+				RomData.ResourceIndex = ResourceIndex;
+				RomData.Size = MemStream.GetBufferSize();
+
+				// Ensure that the Id is unique
+				RomData.Id = CityHash32(static_cast<const char*>(MemStream.GetBuffer()), MemStream.GetBufferSize());
+				EnsureUniqueRomId(RomData.Id, program);
+
+				int32 RomIndex = program.m_roms.Add(RomData);
+				ResData.Key = RomIndex;
+			}
+		}
+
+		UE_LOG(LogMutableCore, Log, TEXT("Generated roms for model with %d embedded constants and %d in roms."), NumEmbedded, NumRoms);
+	}
+
+
+}
+
diff --git a/Engine/Plugins/Experimental/Mutable/Source/MutableTools/Private/CompilerPrivate.h b/Engine/Plugins/Experimental/Mutable/Source/MutableTools/Private/CompilerPrivate.h
new file mode 100644
--- /dev/null
+++ b/Engine/Plugins/Experimental/Mutable/Source/MutableTools/Private/CompilerPrivate.h
@@ -0,0 +1,163 @@
+// Copyright Epic Games, Inc. All Rights Reserved.
+
+#pragma once
+
+#include "Compiler.h"
+
+#include "AST.h"
+#include "ErrorLogPrivate.h"
+#include "ModelReport.h"
+
+#include "Operations.h"
+#include "MemoryPrivate.h"
+#include "NodeObjectPrivate.h"
+
+
+namespace mu
+{
+
+    //!
+    struct STATE_OPTIMIZATION_OPTIONS
+    {
+        STATE_OPTIMIZATION_OPTIONS()
+        {
+			m_firstLOD = 0;
+            m_onlyFirstLOD = false;
+            m_avoidRuntimeCompression = false;
+        }
+
+		uint8 m_firstLOD;
+		bool m_onlyFirstLOD;
+        bool m_avoidRuntimeCompression;
+
+        void Serialise( OutputArchive& arch ) const
+        {
+            const int32_t ver = 2;
+            arch << ver;
+
+			arch << m_firstLOD;
+			arch << m_onlyFirstLOD;
+            arch << m_avoidRuntimeCompression;
+        }
+
+
+        void Unserialise( InputArchive& arch )
+        {
+            int32_t ver = 0;
+            arch >> ver;
+			check(ver <= 2);
+
+			if (ver >= 2)
+			{
+				arch >> m_firstLOD;
+			}
+			else
+			{
+				m_firstLOD = 0;
+			}
+
+            arch >> m_onlyFirstLOD;
+            arch >> m_avoidRuntimeCompression;
+        }
+    };
+
+
+    //!
+    class CompilerOptions::Private : public Base
+    {
+    public:
+
+        //! Detailed optimization options
+        MODEL_OPTIMIZATION_OPTIONS m_optimisationOptions;
+
+        bool m_enablePartialOptimise = false;
+		bool m_enableConcurrency = false;
+        bool m_ignoreStates = false;
+        CompilerOptions::TextureLayoutStrategy m_textureLayoutStrategy = CompilerOptions::TextureLayoutStrategy::Pack;
+
+		int MinRomSize = 3;
+		int MinTextureResidentMipCount = 3;
+
+        int m_imageCompressionQuality = 0;
+
+        bool m_log = false;
+
+    };
+
+
+    //! Information about an object state in the source data
+    struct OBJECT_STATE
+    {
+        //! Name used to identify the state from the code and user interface.
+        string m_name;
+
+        //! GPU Optimisation options
+        STATE_OPTIMIZATION_OPTIONS m_optimisation;
+
+        //! List of names of the runtime parameters in this state
+        vector<string> m_runtimeParams;
+
+        void Serialise( OutputArchive& arch ) const
+        {
+            const int32_t ver = 5;
+            arch << ver;
+
+            arch << m_name;
+            arch << m_optimisation;
+            arch << m_runtimeParams;
+        }
+
+
+        void Unserialise( InputArchive& arch )
+        {
+            int32_t ver = 0;
+            arch >> ver;
+            check( ver==5 );
+
+            arch >> m_name;
+            arch >> m_optimisation;
+            arch >> m_runtimeParams;
+        }
+    };
+
+
+    //!
+    struct STATE_COMPILATION_DATA
+    {
+        OBJECT_STATE nodeState;
+        Ptr<ASTOp> root;
+        PROGRAM::STATE state;
+        STATE_OPTIMIZATION_OPTIONS optimisationFlags;
+
+        //! List of instructions that need to be cached to efficiently update this state
+        vector<Ptr<ASTOp>> m_updateCache;
+
+        //! List of root instructions for the dynamic resources that depend on the runtime
+        //! parameters of this state.
+        vector< std::pair<Ptr<ASTOp>, std::unordered_set<string> > > m_dynamicResources;
+    };
+
+
+    //!
+    class Compiler::Private : public Base
+    {
+    public:
+
+        Private()
+        {
+            m_pErrorLog = new ErrorLog();
+        }
+
+        ErrorLogPtr m_pErrorLog;
+        ModelReportPtr m_pModelReport;
+
+        //! Detailed options
+        CompilerOptionsPtr m_options;
+
+
+		//!
+		void GenerateRoms(Model* p, int32 MinRomSize);
+
+    };
+
+}
diff --git a/Engine/Plugins/Experimental/Mutable/Source/MutableTools/Private/DataPacker.cpp b/Engine/Plugins/Experimental/Mutable/Source/MutableTools/Private/DataPacker.cpp
new file mode 100644
--- /dev/null
+++ b/Engine/Plugins/Experimental/Mutable/Source/MutableTools/Private/DataPacker.cpp
@@ -0,0 +1,769 @@
+// Copyright Epic Games, Inc. All Rights Reserved.
+
+#include "DataPacker.h"
+#include "ErrorLogPrivate.h"
+#include "Streams.h"
+
+#include "CodeVisitor.h"
+#include "ModelPrivate.h"
+#include "SystemPrivate.h"
+#include "Operations.h"
+#include "OpImagePixelFormat.h"
+
+#include "MutableTrace.h"
+#include "MutableTools/Private/ASTOpImageCompose.h"
+#include "MutableTools/Private/ASTOpConstantResource.h"
+#include "MutableTools/Private/ASTOpImageMultiLayer.h"
+#include "MutableTools/Private/ASTOpMeshRemoveMask.h"
+#include "MutableTools/Private/ASTOpInstanceAdd.h"
+#include "MutableTools/Private/ASTOpMeshExtractLayoutBlocks.h"
+
+#include <algorithm>
+#include <cstdint>
+#include <random>
+
+
+namespace mu
+{
+
+    //---------------------------------------------------------------------------------------------
+    //---------------------------------------------------------------------------------------------
+    //---------------------------------------------------------------------------------------------
+    SubtreeSearchConstantVisitor::SubtreeSearchConstantVisitor( PROGRAM& program,
+                                                                OP::ADDRESS constant,
+                                                                OP_TYPE optype )
+        : m_program( program ),
+          m_constant( constant ),
+          m_opType( optype )
+    {
+        m_visited.resize( program.m_opAddress.size() );
+        if ( m_visited.size() )
+        {
+            memset( &m_visited[0], 0, m_visited.size() );
+        }
+    }
+
+
+    //---------------------------------------------------------------------------------------------
+    bool SubtreeSearchConstantVisitor::Run( OP::ADDRESS root )
+    {
+        // too spammy
+        // \todo: called too many times!
+        // MUTABLE_CPUPROFILER_SCOPE(SubtreeSearchConstantVisitor);
+
+        bool found = false;
+
+        vector< std::pair<bool,int> > pending;
+        pending.reserve(m_program.m_opAddress.size()/4);
+
+        pending.push_back( std::make_pair(false,root) );
+
+        while( pending.size() )
+        {
+            std::pair<bool,int> item = pending.back();
+            pending.pop_back();
+
+            if (item.first)
+            {
+                // Item indicating we finished with all the children of a parent.
+
+                // Propagate the "found" state upwards
+                m_visited[item.second]=1;
+
+                //OP cop = m_program.m_code[item.second];
+                ForEachReference( m_program, (OP::ADDRESS)item.second, [&](OP::ADDRESS ref)
+                {
+                    if (ref && m_visited[ref]==2)
+                    {
+                        m_visited[item.second]=2;
+                    }
+                });
+            }
+            else
+            {
+                if ( !m_visited[item.second] )
+                {
+                    OP_TYPE thisOpType = m_program.GetOpType(item.second);
+                    if ( m_opType == thisOpType )
+                    {
+                        auto args = m_program.GetOpArgs<OP::ResourceConstantArgs>(item.second);
+                        if ( m_constant == args.value )
+                        {
+                            found = true;
+                            m_visited[item.second] = 2;
+                        }
+                    }
+
+                    if (!found)
+                    {
+                        pending.push_back( std::make_pair(true,item.second) );
+
+                        //OP cop = m_program.m_code[item.second];
+                        ForEachReference( m_program, item.second, [&](OP::ADDRESS ref)
+                        {
+                            if (ref && !m_visited[ref])
+                            {
+                                pending.push_back( std::make_pair(false,ref) );
+                            }
+                        });
+                    }
+                }
+                else
+                {
+                    found = (m_visited[item.second]==2);
+                }
+            }
+        }
+
+        return found;
+    }
+
+
+
+    //---------------------------------------------------------------------------------------------
+    //! Get all the parameters that affect the constant.
+    //! "Affect" means that the constant may be used or not depending on the parameter.
+    //---------------------------------------------------------------------------------------------
+    class GatherParametersVisitor : public UniqueConstCodeVisitorIterative< vector<int> >
+    {
+    public:
+
+        GatherParametersVisitor( PROGRAM& program, OP::ADDRESS constant, OP_TYPE opType )
+            : m_constSearch( program, constant, opType )
+        {
+            MUTABLE_CPUPROFILER_SCOPE(GatherParametersVisitor);
+
+            vector<int> currentParams;
+            currentParams.resize( program.m_parameters.size(), 0 );
+            SetDefaultState(currentParams);
+
+            // TODO: we can optimize by precalculating what ops have the required optype below
+            m_opType = opType;
+            m_constant = constant;
+
+            m_allParams.resize( program.m_parameters.size(), 0 );
+            FullTraverse( program );
+
+            for (size_t i=0; i<m_allParams.size(); ++i )
+            {
+                if ( m_allParams[i] )
+                {
+                    m_sortedParams.push_back( int(i) );
+                }
+            }
+        }
+
+
+        virtual bool Visit( OP::ADDRESS at, PROGRAM& program )
+        {
+            bool recurse = true;
+
+            if ( m_constSearch.Run( at ) )
+            {
+                switch ( program.GetOpType(at) )
+                {
+                case OP_TYPE::NU_CONDITIONAL:
+                case OP_TYPE::SC_CONDITIONAL:
+                case OP_TYPE::CO_CONDITIONAL:
+                case OP_TYPE::IM_CONDITIONAL:
+                case OP_TYPE::ME_CONDITIONAL:
+                case OP_TYPE::LA_CONDITIONAL:
+                case OP_TYPE::IN_CONDITIONAL:
+                {
+                    auto args = program.GetOpArgs<OP::ConditionalArgs>(at);
+
+                    // If the constant is present in only one of the 2 branches
+                    bool foundYes = m_constSearch.Run( args.yes  );
+                    bool foundNo = m_constSearch.Run( args.no );
+
+                    if ( foundYes!=foundNo )
+                    {
+                        m_conditionVisitor.Run( args.condition, program );
+
+                        vector<int> currentParams = GetCurrentState();
+                        for (size_t p=0; p<m_conditionVisitor.m_params.size(); ++p )
+                        {
+                            currentParams[ m_conditionVisitor.m_params[p] ]++;
+                        }
+
+                        // Set state for child recursion
+                        SetCurrentState( currentParams );
+                    }
+
+                    break;
+                }
+
+                case OP_TYPE::NU_SWITCH:
+                case OP_TYPE::SC_SWITCH:
+                case OP_TYPE::CO_SWITCH:
+                case OP_TYPE::IM_SWITCH:
+                case OP_TYPE::ME_SWITCH:
+                case OP_TYPE::LA_SWITCH:
+                {
+					const uint8_t* data = program.GetOpArgsPointer(at);
+					
+					OP::ADDRESS VarAddress;
+					FMemory::Memcpy( &VarAddress, data, sizeof(OP::ADDRESS) );
+
+                    m_conditionVisitor.Run( VarAddress, program );
+
+                    vector<int> currentParams = GetCurrentState();
+                    for (size_t p=0; p<m_conditionVisitor.m_params.size(); ++p )
+                    {
+                        currentParams[ m_conditionVisitor.m_params[p] ]++;
+                    }
+
+                    // Set state for child recursion
+                    SetCurrentState( currentParams );
+
+                    break;
+                }
+
+                default:
+                {
+                    // If we hit the constant we are analysing
+                    if ( m_opType == OP_TYPE::ME_CONSTANT
+                         &&
+                         program.GetOpType(at) == m_opType )
+                    {
+                        auto args = program.GetOpArgs<OP::MeshConstantArgs>(at);
+                        if ( args.value == m_constant )
+                        {
+                            // Accumulate the currently relevant parameters
+                            const vector<int>& currentParams = GetCurrentState();
+                            for (size_t i=0; i<currentParams.size(); ++i )
+                            {
+                                m_allParams[i] += currentParams[i];
+                            }
+                        }
+                    }
+                    else if ( program.GetOpType(at) == m_opType )
+                    {
+                        auto args = program.GetOpArgs<OP::ResourceConstantArgs>(at);
+                        if ( args.value == m_constant )
+                        {
+                            // Accumulate the currently relevant parameters
+                            const vector<int>& currentParams = GetCurrentState();
+                            for (size_t i=0; i<currentParams.size(); ++i )
+                            {
+                                m_allParams[i] += currentParams[i];
+                            }
+                        }
+                    }
+
+                    break;
+                }
+
+                }
+
+            }
+
+            return recurse;
+        }
+
+    public:
+
+        //! Output result
+        vector<int> m_sortedParams;
+
+    private:
+
+        OP::ADDRESS m_constant;
+        OP_TYPE m_opType;
+
+        vector<int> m_allParams;
+
+        SubtreeParametersVisitor m_conditionVisitor;
+        SubtreeSearchConstantVisitor m_constSearch;
+
+    };
+
+
+    //---------------------------------------------------------------------------------------------
+    //---------------------------------------------------------------------------------------------
+    //---------------------------------------------------------------------------------------------
+    class AccumulateImageFormatsAST : public Visitor_TopDown_Unique_Const< vector<bool> >
+    {
+    public:
+
+        void Run( const ASTOpList& roots )
+        {
+            MUTABLE_CPUPROFILER_SCOPE(AccumulateImageFormatsAST);
+
+            // TODO: Use statically sized vectors (std::array)
+            //vector<bool> initial( IF_COUNT, true );
+            //std::fill( m_supportedFormats.begin(), m_supportedFormats.end(), initial );
+
+            vector<bool> defaultState(size_t(EImageFormat::IF_COUNT), false);
+
+            Traverse( roots, defaultState );
+        }
+
+
+        bool Visit( const Ptr<ASTOp>& node ) override
+        {
+            bool recurse = true;
+
+            const vector<bool>& currentFormats = GetCurrentState();
+            bool allFalse = currentFormats == vector<bool>(size_t(EImageFormat::IF_COUNT), false );
+
+            // Can we use the cache?
+            if (allFalse)
+            {
+                if (m_visited.count(node))
+                {
+                    return false;
+                }
+
+                m_visited.insert(node);
+            }
+
+            switch ( node->GetOpType() )
+            {
+            case OP_TYPE::IM_CONSTANT:
+            {
+                // Remove unsupported formats
+                auto op = dynamic_cast<const ASTOpConstantResource*>(node.get());
+                if (!m_supportedFormats.count(op))
+                {
+                    vector<bool> initial(size_t(EImageFormat::IF_COUNT), true );
+                    m_supportedFormats.insert( std::make_pair(op, std::move(initial)) );
+                }
+
+                for ( unsigned f=0; f< unsigned(EImageFormat::IF_COUNT); ++f )
+                {
+                    if ( !currentFormats[f] )
+                    {
+                        m_supportedFormats[op][f] = false;
+                    }
+                }
+                recurse = false;
+                break;
+            }
+
+            case OP_TYPE::IM_SWITCH:
+            case OP_TYPE::IM_CONDITIONAL:
+                // Switches and conditionals don't change the supported formats
+                break;
+
+            case OP_TYPE::IM_COMPOSE:
+            {
+                recurse = false;
+				const ASTOpImageCompose* op = dynamic_cast<const ASTOpImageCompose*>(node.get());
+
+                vector<bool> newState = vector<bool>(size_t(EImageFormat::IF_COUNT), false);
+                RecurseWithState( op->Layout.child(), newState );
+                RecurseWithState( op->Base.child(), newState );
+                RecurseWithState( op->BlockImage.child(), newState );
+
+                if ( op->Mask )
+                {
+                    newState[(size_t)EImageFormat::IF_L_UBIT_RLE] = true;
+                    RecurseWithState( op->Mask.child(), newState );
+                }
+                break;
+            }
+
+            case OP_TYPE::IM_LAYERCOLOUR:
+            {
+                recurse = false;
+				const ASTOpFixed* op = dynamic_cast<const ASTOpFixed*>(node.get());
+                vector<bool> newState = vector<bool>(size_t(EImageFormat::IF_COUNT), false);
+                RecurseWithState( op->children[op->op.args.ImageLayerColour.base].child(), newState );
+                RecurseWithState( op->children[op->op.args.ImageLayerColour.colour].child(), newState );
+
+                if ( op->children[op->op.args.ImageLayerColour.mask] )
+                {
+                    newState[(size_t)EImageFormat::IF_L_UBYTE] = true;
+                    newState[(size_t)EImageFormat::IF_L_UBYTE_RLE] = true;
+
+                    RecurseWithState( op->children[op->op.args.ImageLayerColour.mask].child(), newState );
+                }
+                break;
+            }
+
+            case OP_TYPE::IM_LAYER:
+            {
+                recurse = false;
+				const ASTOpFixed* op = dynamic_cast<const ASTOpFixed*>(node.get());
+                vector<bool> newState = vector<bool>(size_t(EImageFormat::IF_COUNT), false);
+                RecurseWithState( op->children[op->op.args.ImageLayer.base].child(), newState );
+                RecurseWithState( op->children[op->op.args.ImageLayer.blended].child(), newState );
+
+                if (op->children[op->op.args.ImageLayer.mask])
+                {
+                    newState[(size_t)EImageFormat::IF_L_UBYTE] = true;
+                    newState[(size_t)EImageFormat::IF_L_UBYTE_RLE] = true;
+
+                    RecurseWithState( op->children[op->op.args.ImageLayer.mask].child(), newState );
+                }
+                break;
+            }
+
+            case OP_TYPE::IM_MULTILAYER:
+            {
+                recurse = false;
+				const ASTOpImageMultiLayer* op = dynamic_cast<const ASTOpImageMultiLayer*>(node.get());
+                vector<bool> newState = vector<bool>(size_t(EImageFormat::IF_COUNT), false);
+                RecurseWithState( op->base.child(), newState );
+                RecurseWithState( op->blend.child(), newState );
+
+                if (op->mask)
+                {
+                    newState[(size_t)EImageFormat::IF_L_UBYTE] = true;
+                    newState[(size_t)EImageFormat::IF_L_UBYTE_RLE] = true;
+
+                    RecurseWithState( op->mask.child(), newState );
+                }
+                break;
+            }
+
+            case OP_TYPE::IM_DISPLACE:
+            {
+                recurse = false;
+				const ASTOpFixed* op = dynamic_cast<const ASTOpFixed*>(node.get());
+                vector<bool> newState = vector<bool>(size_t(EImageFormat::IF_COUNT), false);
+                RecurseWithState( op->children[op->op.args.ImageDisplace.source].child(), newState );
+
+                newState[(size_t)EImageFormat::IF_L_UBYTE ] = true;
+                newState[(size_t)EImageFormat::IF_L_UBYTE_RLE ] = true;
+
+                RecurseWithState( op->children[op->op.args.ImageDisplace.displacementMap].child(), newState );
+                break;
+            }
+
+            default:
+            {
+                //m_currentFormats.push_back(vector<bool>(IF_COUNT, false));
+                //Recurse(at, program);
+                //m_currentFormats.pop_back();
+
+                vector<bool> newState = vector<bool>(size_t(EImageFormat::IF_COUNT), false);
+                if (currentFormats != newState)
+                {
+                    RecurseWithState(node, newState);
+                    recurse = false;
+                }
+                else
+                {
+                    recurse = true;
+                }
+                break;
+            }
+
+            }
+
+            return recurse;
+        }
+
+    public:
+
+        //! Result of this visitor:
+        //! Formats known to be supported by every constant image.
+        std::unordered_map< Ptr<const ASTOpConstantResource>, vector<bool> > m_supportedFormats;
+
+    private:
+
+        //! Cache. Only valid is current formats are all false.
+        std::unordered_set<Ptr<ASTOp>> m_visited;
+
+    };
+
+
+
+
+    //---------------------------------------------------------------------------------------------
+    //---------------------------------------------------------------------------------------------
+    //---------------------------------------------------------------------------------------------
+    class AccumulateMeshChannelUsageAST : public Visitor_TopDown_Unique_Const< uint64_t >
+    {
+    public:
+
+        void Run( const ASTOpList& roots )
+        {
+            MUTABLE_CPUPROFILER_SCOPE(AccumulateMeshChannelUsageAST);
+
+            // Sanity check in case we add more semantics
+            static_assert(MBS_COUNT<sizeof(uint64_t)*8, "Too many mesh buffer semantics." );
+
+            // Default state: we need everything except internal semantics
+            uint64_t defaultState = 0xffffffffffffffff;
+            defaultState ^= (UINT64_C(1)<<MBS_LAYOUTBLOCK);
+            defaultState ^= (UINT64_C(1)<<MBS_CHART);
+            defaultState ^= (UINT64_C(1)<<MBS_VERTEXINDEX);
+
+            Traverse(roots,defaultState);
+        }
+
+
+        bool Visit( const Ptr<ASTOp>& node ) override
+        {
+            bool recurse = true;
+
+            uint64_t currentSemantics = GetCurrentState();
+
+            switch ( node->GetOpType() )
+            {
+
+            case OP_TYPE::ME_CONSTANT:
+            {
+                // Accumulate necessary semantics
+				const ASTOpConstantResource* op = dynamic_cast<const ASTOpConstantResource*>(node.get());
+                uint64_t currentFlags = 0;
+                if (m_requiredSemantics.count(op))
+                {
+                    currentFlags = m_requiredSemantics[op];
+                }
+                else
+                {
+                    m_requiredSemantics.insert(std::make_pair(op,currentFlags));
+                }
+
+                currentFlags |= currentSemantics;
+                m_requiredSemantics[op] = currentFlags;
+                recurse = false;
+                break;
+            }
+
+            // TODO: These could probably optimise something
+            //case OP_TYPE::IM_RASTERMESH: break;
+
+            case OP_TYPE::ME_DIFFERENCE:
+            {
+                recurse = false;
+
+				const ASTOpFixed* op = dynamic_cast<const ASTOpFixed*>(node.get());
+
+                uint64_t newState = currentSemantics;
+                newState |= (UINT64_C(1)<<MBS_VERTEXINDEX);
+                RecurseWithState( op->children[op->op.args.MeshDifference.base].child(), newState );
+
+                RecurseWithState( op->children[op->op.args.MeshDifference.target].child(), currentSemantics );
+                break;
+             }
+
+            case OP_TYPE::ME_REMOVEMASK:
+            {
+                recurse = false;
+
+				const ASTOpMeshRemoveMask* op = dynamic_cast<const ASTOpMeshRemoveMask*>(node.get());
+
+                uint64_t newState = currentSemantics;
+                newState |= (UINT64_C(1)<<MBS_VERTEXINDEX);
+                RecurseWithState( op->source.child(), newState );
+                for( auto& r: op->removes )
+                {
+                    RecurseWithState( r.second.child(), newState );
+                }
+                break;
+             }
+
+            case OP_TYPE::ME_MORPH2:
+            {
+                recurse = false;
+
+				const ASTOpFixed* op = dynamic_cast<const ASTOpFixed*>(node.get());
+
+                uint64_t newState = currentSemantics;
+                newState |= (UINT64_C(1)<<MBS_VERTEXINDEX);
+                RecurseWithState( op->children[op->op.args.MeshMorph2.base].child(), newState );
+                for ( int o=0; o<MUTABLE_OP_MAX_MORPH2_TARGETS; ++o )
+                {
+                    if (op->children[op->op.args.MeshMorph2.targets[o]])
+                    {
+                        RecurseWithState( op->children[op->op.args.MeshMorph2.targets[o]].child(), newState );
+                    }
+                }
+                break;
+             }
+
+
+            case OP_TYPE::ME_APPLYLAYOUT:
+            {
+                recurse = false;
+
+				const ASTOpFixed* op = dynamic_cast<const ASTOpFixed*>(node.get());
+
+                uint64_t newState = currentSemantics;
+                newState |= (UINT64_C(1)<<MBS_LAYOUTBLOCK);
+                RecurseWithState( op->children[op->op.args.MeshApplyLayout.mesh].child(), newState );
+
+                RecurseWithState( op->children[op->op.args.MeshApplyLayout.layout].child(), currentSemantics );
+                break;
+            }
+
+            case OP_TYPE::ME_EXTRACTLAYOUTBLOCK:
+            {
+                recurse = false;
+
+				const ASTOpMeshExtractLayoutBlocks* op = dynamic_cast<const ASTOpMeshExtractLayoutBlocks*>(node.get());
+
+                // todo: check if we really need all of them
+                uint64_t newState = currentSemantics;
+                newState |= (UINT64_C(1)<<MBS_LAYOUTBLOCK);
+                newState |= (UINT64_C(1)<<MBS_CHART);
+                newState |= (UINT64_C(1)<<MBS_VERTEXINDEX);
+
+                RecurseWithState( op->source.child(), newState );
+                break;
+            }
+
+            case OP_TYPE::ME_EXTRACTFACEGROUP:
+            {
+                recurse = false;
+
+				const ASTOpFixed* op = dynamic_cast<const ASTOpFixed*>(node.get());
+
+                // todo: check if we really need all of them
+                uint64_t newState = currentSemantics;
+                newState |= (UINT64_C(1)<<MBS_LAYOUTBLOCK);
+                newState |= (UINT64_C(1)<<MBS_CHART);
+                newState |= (UINT64_C(1)<<MBS_VERTEXINDEX);
+
+                RecurseWithState( op->children[op->op.args.MeshExtractFaceGroup.source].child(), newState );
+                break;
+            }
+
+            case OP_TYPE::IN_ADDMESH:
+            {
+                recurse = false;
+
+				const ASTOpInstanceAdd* op = dynamic_cast<const ASTOpInstanceAdd*>(node.get());
+
+                RecurseWithState( op->instance.child(), currentSemantics );
+
+                uint64_t newState = GetDefaultState();
+                RecurseWithState( op->value.child(), newState );
+                break;
+            }
+
+            default:
+                // Unhandled op, we may need everything? Recurse with current state?
+                //uint64_t newState = 0xffffffffffffffff;
+                break;
+
+            }
+
+            return recurse;
+        }
+
+    public:
+
+        //! Result of this visitor:
+        //! Used mesh channel semantics for each constant mesh
+        std::unordered_map< Ptr<const ASTOpConstantResource>, uint64_t > m_requiredSemantics;
+
+    };
+
+
+    //---------------------------------------------------------------------------------------------
+    //---------------------------------------------------------------------------------------------
+    //---------------------------------------------------------------------------------------------
+
+    // Todo: move to its own file
+    inline void MeshRemoveUnusedBufferSemantics( Mesh* pMesh, uint64_t usedSemantics )
+    {
+        // right now we only remove entire buffers if no channel is used
+        // TODO: remove from inside the buffer?
+        for (int v=0; v<pMesh->GetVertexBuffers().GetBufferCount(); )
+        {
+            bool used = false;
+            for (int c=0; !used && c<pMesh->GetVertexBuffers().GetBufferChannelCount(v); ++c)
+            {
+                MESH_BUFFER_SEMANTIC semantic;
+                pMesh->GetVertexBuffers().GetChannel(v,c,&semantic,
+                                                      nullptr, nullptr, nullptr, nullptr);
+                used = (( (UINT64_C(1)<<semantic) ) & usedSemantics) != 0;
+            }
+
+            if (!used)
+            {
+                TArray<MESH_BUFFER>& buffers = pMesh->GetVertexBuffers().m_buffers;
+                buffers.RemoveAt(v);
+            }
+            else
+            {
+                ++v;
+            }
+        }
+
+        // TODO: hack, if we don't need layouts, remove them.
+        {
+            uint64_t layoutSemantics = 0;
+            layoutSemantics |= (UINT64_C(1)<<MBS_LAYOUTBLOCK);
+            layoutSemantics |= (UINT64_C(1)<<MBS_CHART);
+
+            if ( (usedSemantics&layoutSemantics) == 0)
+            {
+                pMesh->m_layouts.Empty();
+            }
+        }
+    }
+
+
+    //---------------------------------------------------------------------------------------------
+    void DataOptimiseAST( int imageCompressionQuality, ASTOpList& roots,
+                          const MODEL_OPTIMIZATION_OPTIONS& options )
+    {
+        // Images
+        AccumulateImageFormatsAST accFormat;
+        accFormat.Run( roots );
+
+		// TEMP: Disabled for RLE bug
+        // See if we can convert some constants to more efficient formats
+        ASTOp::Traverse_BottomUp_Unique_NonReentrant( roots, [&](Ptr<ASTOp>& n)
+        {
+            if (n->GetOpType()==OP_TYPE::IM_CONSTANT)
+            {
+				ASTOpConstantResource* typed = dynamic_cast<ASTOpConstantResource*>(n.get());
+                Ptr<const Image> pOld = static_cast<const Image*>(typed->GetValue().get());
+                if ( accFormat.m_supportedFormats[typed][(size_t)EImageFormat::IF_L_UBIT_RLE] )
+                {
+                    ImagePtr pNew =
+                        ImagePixelFormat( imageCompressionQuality, pOld.get(), EImageFormat::IF_L_UBIT_RLE );
+
+                    // Only replace if the compression was worth!
+                    size_t oldSize = pOld->GetDataSize();
+                    size_t newSize = pNew->GetDataSize();
+                    if (float(oldSize) > float(newSize) * options.m_minRLECompressionGain)
+                    {
+                        typed->SetValue(pNew, options.m_useDiskCache);
+                    }
+                }
+                else if ( accFormat.m_supportedFormats[typed][(size_t)EImageFormat::IF_L_UBYTE_RLE] )
+                {
+                    ImagePtr pNew =
+                        ImagePixelFormat( imageCompressionQuality, pOld.get(), EImageFormat::IF_L_UBYTE_RLE );
+
+                    // Only replace if the compression was worth!
+                    size_t oldSize = pOld->GetDataSize();
+                    size_t newSize = pNew->GetDataSize();
+                    if (float(oldSize) > float(newSize) * options.m_minRLECompressionGain)
+                    {
+                        typed->SetValue(pNew, options.m_useDiskCache);
+                    }
+                }
+            }
+        });
+
+
+        // Meshes
+        AccumulateMeshChannelUsageAST meshSemanticsVisitor;
+        meshSemanticsVisitor.Run( roots );
+
+        // See if we can remove some buffers from the constants
+        ASTOp::Traverse_BottomUp_Unique_NonReentrant( roots, [&](Ptr<ASTOp>& n)
+        {
+            if (n->GetOpType()==OP_TYPE::ME_CONSTANT)
+            {
+				ASTOpConstantResource* typed = dynamic_cast<ASTOpConstantResource*>(n.get());
+                Ptr<Mesh> pMesh = static_cast<const Mesh*>(typed->GetValue().get())->Clone();
+                MeshRemoveUnusedBufferSemantics( pMesh.get(), meshSemanticsVisitor.m_requiredSemantics[typed]);
+                typed->SetValue(pMesh, options.m_useDiskCache);
+            }
+        });
+
+    }
+
+
+}
diff --git a/Engine/Plugins/Experimental/Mutable/Source/MutableTools/Private/DataPacker.h b/Engine/Plugins/Experimental/Mutable/Source/MutableTools/Private/DataPacker.h
new file mode 100644
--- /dev/null
+++ b/Engine/Plugins/Experimental/Mutable/Source/MutableTools/Private/DataPacker.h
@@ -0,0 +1,49 @@
+// Copyright Epic Games, Inc. All Rights Reserved.
+
+#pragma once
+
+#include "Compiler.h"
+#include "AST.h"
+#include "MemoryPrivate.h"
+#include "ModelPrivate.h"
+#include "CodeVisitor.h"
+#include "RefCounted.h"
+
+
+#define MUTABLE_HASH_SEED					((uint32_t)0xcadababa)
+
+
+namespace mu
+{
+
+    //---------------------------------------------------------------------------------------------
+    //! Convert constant data to different formats, based on their usage
+    //---------------------------------------------------------------------------------------------
+    extern void DataOptimiseAST( int imageCompressionQuality, ASTOpList& roots,
+                                 const MODEL_OPTIMIZATION_OPTIONS& );
+
+    //---------------------------------------------------------------------------------------------
+    //! Find the given constant in a subtree
+    //---------------------------------------------------------------------------------------------
+    class SubtreeSearchConstantVisitor
+    {
+    public:
+
+        SubtreeSearchConstantVisitor( PROGRAM& program, OP::ADDRESS constant, OP_TYPE optype );
+
+        bool Run( OP::ADDRESS root );
+
+
+    private:
+
+        PROGRAM& m_program;
+        OP::ADDRESS m_constant;
+        OP_TYPE m_opType;
+
+        //! 0: not visited
+        //! 1: visited and not found
+        //! 2: visited and found
+        vector<uint8_t> m_visited;
+    };
+
+}
diff --git a/Engine/Plugins/Experimental/Mutable/Source/MutableTools/Private/ErrorLog.cpp b/Engine/Plugins/Experimental/Mutable/Source/MutableTools/Private/ErrorLog.cpp
new file mode 100644
--- /dev/null
+++ b/Engine/Plugins/Experimental/Mutable/Source/MutableTools/Private/ErrorLog.cpp
@@ -0,0 +1,826 @@
+// Copyright Epic Games, Inc. All Rights Reserved.
+
+
+#include "ErrorLog.h"
+
+#include "ErrorLogPrivate.h"
+#include "ModelPrivate.h"
+#include "MemoryPrivate.h"
+
+namespace mu
+{
+
+
+	//---------------------------------------------------------------------------------------------
+	ErrorLog::ErrorLog()
+	{
+		m_pD = new Private();
+	}
+
+
+	//---------------------------------------------------------------------------------------------
+	ErrorLog::~ErrorLog()
+	{
+        check( m_pD );
+		delete m_pD;
+		m_pD = 0;
+	}
+
+
+	//---------------------------------------------------------------------------------------------
+	ErrorLog::Private* ErrorLog::GetPrivate() const
+	{
+		return m_pD;
+	}
+
+
+	//---------------------------------------------------------------------------------------------
+    int ErrorLog::GetMessageCount( ErrorLogMessageType ) const
+	{
+        // \todo: by type argument
+		return (int)m_pD->m_messages.size();
+	}
+
+
+	//---------------------------------------------------------------------------------------------
+	const char* ErrorLog::GetMessageText( int index ) const
+	{
+		const char* strResult = "";
+
+		if ( index >=0 && index<(int)m_pD->m_messages.size() )
+		{
+			strResult = m_pD->m_messages[index].m_text.c_str();
+		}
+
+		return strResult;
+	}
+
+
+	//---------------------------------------------------------------------------------------------
+	const void* ErrorLog::GetMessageContext( int index ) const
+	{
+		const void* result = 0;
+
+		if ( index >=0 && index<(int)m_pD->m_messages.size() )
+		{
+			result = m_pD->m_messages[index].m_context;
+		}
+
+		return result;
+	}
+
+
+	//---------------------------------------------------------------------------------------------
+	ErrorLogMessageType ErrorLog::GetMessageType( int index ) const
+	{
+		ErrorLogMessageType result = ELMT_NONE;
+
+		if ( index >=0 && index<(int)m_pD->m_messages.size() )
+		{
+			result = m_pD->m_messages[index].m_type;
+		}
+
+		return result;
+	}
+
+	//---------------------------------------------------------------------------------------------
+    ErrorLogMessageAttachedDataView ErrorLog::GetMessageAttachedData( int index ) const
+	{
+        ErrorLogMessageAttachedDataView result;
+
+		if ( index >=0 && index<(int)m_pD->m_messages.size() )
+		{
+			const Private::MSG& message = m_pD->m_messages[index];
+            
+            if ( message.m_data ) 
+            {
+                result.m_unassignedUVs = message.m_data->m_unassignedUVs.data();
+			    result.m_unassignedUVsSize = message.m_data->m_unassignedUVs.size();
+            }
+		}
+
+		return result;
+	}
+
+	//---------------------------------------------------------------------------------------------
+	void ErrorLog::Private::Add( const char* strMessage, ErrorLogMessageType type,
+								 const void* context )
+	{
+		m_messages.push_back( MSG() );
+		m_messages.back().m_type = type;
+		m_messages.back().m_text = strMessage;
+		m_messages.back().m_context = context;
+	}
+
+	//---------------------------------------------------------------------------------------------
+	void ErrorLog::Private::Add( const char* strMessage, 
+                                 const ErrorLogMessageAttachedDataView& dataView,
+                                 ErrorLogMessageType type, const void* context )
+	{
+		m_messages.push_back( MSG() );
+		m_messages.back().m_type = type;
+		m_messages.back().m_text = strMessage;
+		m_messages.back().m_context = context;
+		m_messages.back().m_data = std::make_shared< DATA >();
+
+        if ( dataView.m_unassignedUVs && dataView.m_unassignedUVsSize > 0 )
+        {
+            m_messages.back().m_data->m_unassignedUVs = vector<float>(
+                dataView.m_unassignedUVs, dataView.m_unassignedUVs + dataView.m_unassignedUVsSize);
+        }
+	}
+	//---------------------------------------------------------------------------------------------
+	void ErrorLog::Log() const
+	{
+		UE_LOG(LogMutableCore, Log, TEXT(" Error Log :\n"));
+
+		for ( vector<Private::MSG>::const_iterator it = m_pD->m_messages.begin()
+			; it != m_pD->m_messages.end()
+			; ++it )
+		{
+			const Private::MSG& msg = *it;
+
+			switch ( msg.m_type )
+			{
+			case ELMT_ERROR: 	UE_LOG(LogMutableCore, Log, TEXT("  ERR  %s\n"), msg.m_text.c_str()); break;
+			case ELMT_WARNING: 	UE_LOG(LogMutableCore, Log, TEXT("  WRN  %s\n"), msg.m_text.c_str()); break;
+			case ELMT_INFO: 	UE_LOG(LogMutableCore, Log, TEXT("  INF  %s\n"), msg.m_text.c_str()); break;
+			default: 			UE_LOG(LogMutableCore, Log, TEXT("  NON  %s\n"), msg.m_text.c_str()); break;
+			}
+		}
+	}
+
+
+	//---------------------------------------------------------------------------------------------
+	void ErrorLog::Merge( const ErrorLog* pOther )
+	{
+		m_pD->m_messages.insert
+			(
+				m_pD->m_messages.end(),
+				pOther->GetPrivate()->m_messages.begin(),
+				pOther->GetPrivate()->m_messages.end()
+			);
+	}
+
+
+	//---------------------------------------------------------------------------------------------
+	//---------------------------------------------------------------------------------------------
+	//---------------------------------------------------------------------------------------------
+    
+    // clang-format off
+
+    const TCHAR* s_opNames[ size_t(OP_TYPE::COUNT) ] =
+	{
+		TEXT("NONE             "),
+
+		TEXT("BO_CONSTANT      "),
+        TEXT("NU_CONSTANT      "),
+        TEXT("SC_CONSTANT      "),
+        TEXT("CO_CONSTANT      "),
+		TEXT("IM_CONSTANT      "),
+		TEXT("ME_CONSTANT      "),
+		TEXT("LA_CONSTANT      "),
+		TEXT("PR_CONSTANT      "),
+		TEXT("ST_CONSTANT      "),
+
+		TEXT("BO_PARAMETER     "),
+		TEXT("NU_PARAMETER     "),
+		TEXT("SC_PARAMETER     "),
+		TEXT("CO_PARAMETER     "),
+		TEXT("PR_PARAMETER     "),
+		TEXT("IM_PARAMETER     "),
+		TEXT("ST_PARAMETER     "),
+		
+		TEXT("NU_CONDITIONAL   "),
+		TEXT("SC_CONDITIONAL   "),
+		TEXT("CO_CONDITIONAL   "),
+		TEXT("IM_CONDITIONAL   "),
+		TEXT("ME_CONDITIONAL   "),
+		TEXT("LA_CONDITIONAL   "),
+		TEXT("IN_CONDITIONAL   "),
+		
+		TEXT("NU_SWITCH        "),
+		TEXT("SC_SWITCH        "),
+		TEXT("CO_SWITCH        "),
+		TEXT("IM_SWITCH        "),
+		TEXT("ME_SWITCH        "),
+		TEXT("LA_SWITCH        "),
+		TEXT("IN_SWITCH        "),
+
+		TEXT("BO_LESS          "),
+		TEXT("BO_EQUAL_SC_CONST"),
+		TEXT("BO_AND           "),
+		TEXT("BO_OR            "),
+		TEXT("BO_NOT           "),
+		
+		TEXT("SC_MULTIPLYADD   "),
+		TEXT("SC_ARITHMETIC    "),
+		TEXT("SC_CURVE         "),
+		
+		TEXT("CO_SAMPLEIMAGE   "),
+		TEXT("CO_SWIZZLE       "),
+		TEXT("CO_IMAGESIZE     "),
+		TEXT("CO_LAYOUTBLOCKTR "),
+		TEXT("CO_FROMSCALARS   "),
+		TEXT("CO_ARITHMETIC    "),
+
+		TEXT("IM_LAYER         "),
+		TEXT("IM_LAYERCOLOUR   "),
+		TEXT("IM_PIXELFORMAT   "),
+		TEXT("IM_MIPMAP        "),
+		TEXT("IM_RESIZE        "),
+		TEXT("IM_RESIZELIKE    "),
+		TEXT("IM_RESIZEREL     "),
+		TEXT("IM_BLANKLAYOUT   "),
+		TEXT("IM_COMPOSE       "),
+		TEXT("IM_DIFFERENCE    "),
+		TEXT("IM_INTERPOLATE   "),
+		TEXT("IM_INTERPOLATE3  "),
+		TEXT("IM_SATURATE      "),
+		TEXT("IM_LUMINANCE     "),
+		TEXT("IM_SWIZZLE       "),
+		TEXT("IM_SELECTCOLOUR  "),
+		TEXT("IM_COLOURMAP     "),
+		TEXT("IM_GRADIENT      "),
+		TEXT("IM_BINARISE      "),
+		TEXT("IM_PLAINCOLOUR   "),
+		TEXT("IM_GPU           "),
+		TEXT("IM_CROP          "),
+		TEXT("IM_PATCH         "),
+		TEXT("IM_RASTERMESH    "),
+		TEXT("IM_MAKEGROWMAP   "),
+		TEXT("IM_DISPLACE      "),
+		TEXT("IM_MULTILAYER    "),
+		TEXT("IM_INVERT        "),
+		TEXT("IM_NORMAL_COMPO  "),
+		TEXT("IM_TRANSFORM     "),
+
+		TEXT("ME_APPLYLAYOUT   "),
+		TEXT("ME_DIFFERENCE    "),
+		TEXT("ME_MORPH2        "),
+		TEXT("ME_MERGE         "),
+		TEXT("ME_INTERPOLATE   "),
+		TEXT("ME_MASKCLIPMESH  "),
+		TEXT("ME_MASKDIFF      "),
+        TEXT("ME_SUBTRACT      "),
+        TEXT("ME_REMOVEMASK    "),
+        TEXT("ME_FORMAT        "),
+        TEXT("ME_EXTRACTLABLOCK"),
+        TEXT("ME_EXTRACTFACEGRP"),
+		TEXT("ME_TRANSFORM     "),
+		TEXT("ME_CLIPMORPHPLANE"),
+        TEXT("ME_CLIPWITHMESH  "),
+        TEXT("ME_SETSKELETON   "),
+        TEXT("ME_PROJECT       "),
+        TEXT("ME_APPLYPOSE     "),
+        TEXT("ME_REMAPINDICES  "),
+		TEXT("ME_GEOMETRYOP	   "),
+		TEXT("ME_BINDSHAPE	   "),
+		TEXT("ME_APPLYSHAPE	   "),
+		TEXT("ME_CLIPDEFORM	   "),
+		TEXT("ME_MORPHRESHAPE  "),
+
+		TEXT("IN_ADDMESH       "),
+		TEXT("IN_ADDIMAGE      "),
+		TEXT("IN_ADDVECTOR     "),
+		TEXT("IN_ADDSCALAR     "),
+		TEXT("IN_ADDSTRING     "),
+        TEXT("IN_ADDSURFACE    "),
+        TEXT("IN_ADDCOMPONENT  "),
+        TEXT("IN_ADDLOD        "),
+
+		TEXT("LA_PACK          "),
+		TEXT("LA_MERGE         "),
+        TEXT("LA_REMOVEBLOCKS  ")
+	};
+
+    // clang-format on
+
+    //---------------------------------------------------------------------------------------------
+    const TCHAR* GetOpName( OP_TYPE type )
+    {
+        //check( type>=0 && type<(int)OP_TYPE::COUNT );
+        return s_opNames[(int)type];
+    }
+
+
+	//---------------------------------------------------------------------------------------------
+    extern FString GetOpDesc( const PROGRAM& program, OP::ADDRESS at )
+	{
+		char temp[1024];
+		FMemory::Memzero( temp, 1024 );
+
+		size_t done = 0;
+
+        OP_TYPE type = program.GetOpType(at);
+        switch (type)
+		{
+		case OP_TYPE::BO_CONSTANT:
+        {
+            auto args = program.GetOpArgs<OP::BoolConstantArgs>(at);
+			mutable_snprintf( temp, 1024,
+                              "bool : %s\t", args.value?"true":"false" );
+			break;
+        }
+        case OP_TYPE::NU_CONSTANT:
+        {
+            auto args = program.GetOpArgs<OP::IntConstantArgs>(at);
+            mutable_snprintf( temp, 1024,
+                              "int : %d\t", args.value );
+            break;
+        }
+        case OP_TYPE::SC_CONSTANT:
+        {
+            auto args = program.GetOpArgs<OP::ScalarConstantArgs>(at);
+            mutable_snprintf( temp, 1024,
+                              "float : %f\t", args.value );
+            break;
+        }
+        case OP_TYPE::CO_CONSTANT:
+        {
+            auto args = program.GetOpArgs<OP::ColourConstantArgs>(at);
+            mutable_snprintf( temp, 1024,
+							  "colour : %f %f %f\t",
+                              args.value[0],
+                              args.value[1],
+                              args.value[2] );
+			break;
+        }
+        case OP_TYPE::IM_CONSTANT:
+        case OP_TYPE::LA_CONSTANT:
+        {
+            auto args = program.GetOpArgs<OP::ResourceConstantArgs>(at);
+            mutable_snprintf( temp, 1024,
+                              "res  : %d\t", (int)args.value );
+            break;
+        }
+        case OP_TYPE::ME_CONSTANT:
+        {
+            auto args = program.GetOpArgs<OP::MeshConstantArgs>(at);
+            mutable_snprintf( temp, 1024,
+                              "res  : %d\t", (int)args.value );
+            break;
+        }
+        case OP_TYPE::BO_PARAMETER:
+        case OP_TYPE::NU_PARAMETER:
+        case OP_TYPE::SC_PARAMETER:
+        case OP_TYPE::CO_PARAMETER:
+        case OP_TYPE::PR_PARAMETER:
+        case OP_TYPE::IM_PARAMETER:
+        {
+            auto args = program.GetOpArgs<OP::ParameterArgs>(at);
+            mutable_snprintf( temp, 1024,
+                              "variable : %2d", (int)args.variable );
+			break;
+        }
+        case OP_TYPE::NU_CONDITIONAL:
+        case OP_TYPE::SC_CONDITIONAL:
+        case OP_TYPE::CO_CONDITIONAL:
+		case OP_TYPE::IM_CONDITIONAL:
+		case OP_TYPE::ME_CONDITIONAL:
+		case OP_TYPE::LA_CONDITIONAL:
+		case OP_TYPE::IN_CONDITIONAL:
+        {
+            auto args = program.GetOpArgs<OP::ConditionalArgs>(at);
+            mutable_snprintf( temp, 1024,
+							  "yes : %3d\tno : %3d\t",
+                              (int)args.yes,
+                              (int)args.no );
+			break;
+        }
+
+        case OP_TYPE::NU_SWITCH:
+        case OP_TYPE::SC_SWITCH:
+        case OP_TYPE::CO_SWITCH:
+        case OP_TYPE::IM_SWITCH:
+		case OP_TYPE::ME_SWITCH:
+        case OP_TYPE::LA_SWITCH:
+        case OP_TYPE::IN_SWITCH:
+        {	
+			const uint8_t* data = program.GetOpArgsPointer(at);
+		
+			OP::ADDRESS VarAddress;
+			FMemory::Memcpy( &VarAddress, data, sizeof(OP::ADDRESS) );
+			data += sizeof(OP::ADDRESS);
+			
+			OP::ADDRESS DefAddress;
+			FMemory::Memcpy( &DefAddress, data, sizeof(OP::ADDRESS));
+			data += sizeof(OP::ADDRESS);
+	
+			uint32_t CaseCount;
+			FMemory::Memcpy( &CaseCount, data,  sizeof(uint32_t));
+			data += sizeof(uint32_t);
+
+            done = mutable_snprintf( temp, 1024, "variable : %3d", (int)VarAddress );
+
+            for ( uint32_t C = 0; C < CaseCount; ++C )
+            {
+				int32_t Condition;
+				FMemory::Memcpy( &Condition, data, sizeof(int32_t));
+				data += sizeof(int32_t);
+				
+				OP::ADDRESS At;
+				FMemory::Memcpy( &At, data, sizeof(OP::ADDRESS));
+				data += sizeof(OP::ADDRESS);
+
+                done += mutable_snprintf( temp+done, 1024-done, "\t[ %3d : %3d]", (int)Condition, (int)At);
+			}
+
+			done += mutable_snprintf( temp+done, 1024-done, "\tdefault : %3d", (int)DefAddress );
+			break;
+        }
+
+		//-----------------------------------------------------------------------------------------
+		case OP_TYPE::IM_RESIZE:
+        {
+            auto args = program.GetOpArgs<OP::ImageResizeArgs>(at);
+            mutable_snprintf( temp, 1024,
+							  "base : %3d\tsize  : %4d %4d",
+                              (int)args.source,
+                              (int)args.size[0],
+                              (int)args.size[1] );
+			break;
+        }
+		case OP_TYPE::IM_RESIZELIKE:
+        {
+            auto args = program.GetOpArgs<OP::ImageResizeLikeArgs>(at);
+            mutable_snprintf( temp, 1024,
+							  "base : %3d\tsize  : %3d",
+                              (int)args.source,
+                              (int)args.sizeSource );
+			break;
+        }
+        case OP_TYPE::IM_PIXELFORMAT:
+        {
+            auto args = program.GetOpArgs<OP::ImagePixelFormatArgs>(at);
+            mutable_snprintf( temp, 1024,
+                              "base : %3d format : %2d  format-if-alpha : %2d",
+                              (int)args.source,
+                              (int)args.format,
+                              (int)args.formatIfAlpha );
+            break;
+        }
+        case OP_TYPE::IM_MIPMAP:
+        {
+            auto args = program.GetOpArgs<OP::ImageMipmapArgs>(at);
+            mutable_snprintf( temp, 1024,
+                              "base : %3d\tlevels : %2d\tblockLevels : %2d",
+                              (int)args.source,
+                              (int)args.levels,
+                              (int)args.blockLevels );
+            break;
+        }
+		case OP_TYPE::IM_BLANKLAYOUT:
+        {
+            auto args = program.GetOpArgs<OP::ImageBlankLayoutArgs>(at);
+            mutable_snprintf( temp, 1024,
+							  "block : %4d %4d  fmt  : %2d",
+                              (int)args.blockSize[0],
+                              (int)args.blockSize[1],
+                              (int)args.format );
+			break;
+        }
+		case OP_TYPE::IM_COMPOSE:
+        {
+            auto args = program.GetOpArgs<OP::ImageComposeArgs>(at);
+            mutable_snprintf( temp, 1024,
+							  "base : %3d\tblock image : %3d\tblock index : %3d\tlayout: %3d\tmask: %3d",
+                              (int)args.base,
+                              (int)args.blockImage,
+                              (int)args.blockIndex,
+                              (int)args.layout,
+                              (int)args.mask );
+			break;
+        }
+		case OP_TYPE::IM_INTERPOLATE:
+        {
+            auto args = program.GetOpArgs<OP::ImageInterpolateArgs>(at);
+            done += mutable_snprintf( temp, 1024,
+                                      "factor : %3d", args.factor);
+			for (int t=0; t<MUTABLE_OP_MAX_INTERPOLATE_COUNT; ++t)
+			{
+				done += mutable_snprintf( temp+done, 1024-done,
+                                          "\tt%d : %3d", t, args.targets[t] );
+			}
+			break;
+        }
+		case OP_TYPE::IM_SWIZZLE:
+        {
+            auto args = program.GetOpArgs<OP::ImageSwizzleArgs>(at);
+            done += mutable_snprintf( temp, 1024,
+                              "format : %3d", args.format);
+			for (int t=0; t<MUTABLE_OP_MAX_SWIZZLE_CHANNELS; ++t)
+			{
+				done += mutable_snprintf( temp+done, 1024-done,
+							  "\tsource %d, channel %3d",
+                          args.sources[t],
+                          args.sourceChannels[t] );
+			}
+			break;
+        }
+		case OP_TYPE::IM_SELECTCOLOUR:
+        {
+            auto args = program.GetOpArgs<OP::ImageSelectColourArgs>(at);
+            mutable_snprintf( temp, 1024,
+							  "base : %3d\tcolour  : %3d",
+                    (int)args.base,
+                    (int)args.colour );
+			break;
+        }
+		case OP_TYPE::IM_COLOURMAP:
+        {
+            auto args = program.GetOpArgs<OP::ImageColourMapArgs>(at);
+            mutable_snprintf( temp, 1024,
+							  "base : %3d\tmask  : %3d\tmap  : %3d",
+                      (int)args.base,
+                      (int)args.mask,
+                      (int)args.map );
+			break;
+        }
+		case OP_TYPE::IM_GRADIENT:
+        {
+            auto args = program.GetOpArgs<OP::ImageGradientArgs>(at);
+            mutable_snprintf( temp, 1024,
+							  "colour0 : %3d\tcolour1  : %3d\tsize  : %3d x %3d",
+                      (int)args.colour0,
+                      (int)args.colour1,
+                      (int)args.size[0],
+                      (int)args.size[1]
+					  );
+			break;
+        }
+		case OP_TYPE::IM_GPU:
+		{
+            auto args = program.GetOpArgs<OP::ImageGPUArgs>(at);
+            mutable_snprintf( temp, 1024,
+							  "program : %3d",
+                              (int)args.program
+							  );
+			break;
+		}
+
+
+		//-----------------------------------------------------------------------------------------
+//        case OP_TYPE::ME_REMOVEMASK:
+//        {
+//            auto args = program.GetOpArgs<OP::MeshRemoveMaskArgs>(at);
+//            mutable_snprintf( temp, 1024,
+//                              "source : %3d    mask : %3d",
+//                    (int)args.source,
+//                    (int)args.mask );
+//            break;
+//        }
+        case OP_TYPE::ME_MERGE:
+        {
+            auto args = program.GetOpArgs<OP::MeshMergeArgs>(at);
+            mutable_snprintf( temp, 1024,
+                              "base : %3d    added : %3d",
+                    (int)args.base,
+                    (int)args.added );
+			break;
+        }
+		case OP_TYPE::ME_INTERPOLATE:
+        {
+            auto args = program.GetOpArgs<OP::MeshInterpolateArgs>(at);
+            done += mutable_snprintf( temp, 1024,
+							  "factor : %3d\tbase : %3d",
+                              args.factor,
+                              args.base);
+			for (int t=0; t<MUTABLE_OP_MAX_INTERPOLATE_COUNT-1; ++t)
+			{
+				done += mutable_snprintf( temp+done, 1024-done,
+                              "\tt%d : %3d", t, args.targets[t] );
+			}
+			break;
+        }
+
+		//-----------------------------------------------------------------------------------------
+		default:
+			//check( false );
+			break;
+		}
+
+		temp[1023] = 0;
+		return temp;
+	}
+
+
+	//---------------------------------------------------------------------------------------------
+    void LogOp( const PROGRAM& prog, int at )
+	{
+        OP_TYPE type = prog.GetOpType(at);
+		UE_LOG(LogMutableCore, Log, TEXT("  @%3d  %s %s"), (int)at, s_opNames[(int)type], *GetOpDesc(prog, at));
+	}
+
+
+	//---------------------------------------------------------------------------------------------
+	void Log( ModelPtrConst pModel )
+	{
+		const PROGRAM& prog = pModel->GetPrivate()->m_program;
+
+		for ( vector<PROGRAM::STATE>::const_iterator it = prog.m_states.begin()
+			; it != prog.m_states.end()
+			; ++it )
+		{
+			UE_LOG(LogMutableCore, Log, TEXT( "state name : % s\n      root : %d" ),
+					  *FString(it->m_name.c_str()),
+					  (int)it->m_root );
+		}
+
+		UE_LOG(LogMutableCore, Log, TEXT("code : ") );
+
+//		int at = 0;
+//		for ( vector<OP>::const_iterator it = prog.m_code.begin()
+//			; it != prog.m_code.end()
+//			; ++it )
+//		{
+//            LogOp( prog, at, *it );
+//			++at;
+//		}
+
+		UE_LOG(LogMutableCore, Log, TEXT("parameters : ") );
+		for ( vector<PARAMETER_DESC>::const_iterator it = prog.m_parameters.begin()
+			; it != prog.m_parameters.end()
+			; ++it )
+		{
+			UE_LOG(LogMutableCore, Log, TEXT("    [%s] : %2d"), *FString(it->m_name.c_str()), (int)it->m_type  );
+		}
+
+
+//		LogDebug( "constant images : \n" );
+//		for ( vector<ImagePtr>::const_iterator it = prog.m_constantImages.begin()
+//			; it != prog.m_constantImages.end()
+//			; ++it )
+//		{
+//			LogDebug( "    image : %2d", (int)(*it)->GetFormat() );
+//			LogDebug( "\t %4d x %4d\n", (int)(*it)->GetSizeX(), (int)(*it)->GetSizeY() );
+//		}
+
+		UE_LOG(LogMutableCore, Log, TEXT("constant meshes : ") );
+        for ( const TPair<int32, mu::MeshPtrConst>& at: prog.m_constantMeshes )
+        {
+            const Mesh* m = at.Value.get();
+			UE_LOG(LogMutableCore, Log, TEXT("    mesh : %6d vertices \t %6d indices "),
+                    (int)m->GetVertexCount(), (int)m->GetIndexCount());
+        }
+
+		UE_LOG(LogMutableCore, Log, TEXT("constant layouts : ") );
+        for ( vector<LayoutPtrConst>::const_iterator it = prog.m_constantLayouts.begin()
+			; it != prog.m_constantLayouts.end()
+			; ++it )
+		{
+            auto pLayout = *it;
+			UE_LOG(LogMutableCore, Log, TEXT("    layout : %3d x %3d"),
+					pLayout->m_size[0],
+					pLayout->m_size[1]
+					);
+			for ( int b=0; b<pLayout->GetBlockCount(); ++b )
+			{
+				UE_LOG(LogMutableCore, Log, TEXT("        block : %3d x %3d"),
+						pLayout->m_blocks[b].m_size[0],
+						pLayout->m_blocks[b].m_size[1]
+						);
+			}
+		}
+
+		UE_LOG(LogMutableCore, Log, TEXT("constant skeletons : ") );
+        for ( auto it = prog.m_constantSkeletons.begin()
+            ; it != prog.m_constantSkeletons.end()
+            ; ++it )
+        {
+            auto pSkeleton = *it;
+			UE_LOG(LogMutableCore, Log, TEXT("    skeleton : %3d bones"), pSkeleton->GetBoneCount() );
+            for ( int b=0; b<pSkeleton->GetBoneCount(); ++b )
+            {
+				UE_LOG(LogMutableCore, Log, TEXT("        bone : %s"), *FString(pSkeleton->GetBoneName(b)) );
+            }
+        }
+
+		UE_LOG(LogMutableCore, Log, TEXT("constant strings : ") );
+		for ( vector<string>::const_iterator it = prog.m_constantStrings.begin()
+			; it != prog.m_constantStrings.end()
+			; ++it )
+		{
+			UE_LOG(LogMutableCore, Log, TEXT("    string : [%s]"), *FString(it->c_str()) );
+		}
+
+	}
+
+
+    //---------------------------------------------------------------------------------------------
+    void RecurseLogOp( int indent, int at, PROGRAM& prog, vector<uint8_t>& done )
+    {
+        OP_TYPE type = prog.GetOpType(at);
+
+        if (done[at])
+        {
+			UE_LOG(LogMutableCore, Log, TEXT("%*s@%3d  %s [repeated]"), indent, " ", (int)at, s_opNames[(int)type] );
+        }
+        else
+        {
+            done[at] = 1;
+
+            string desc = "";//GetOpDesc(prog,op);
+			UE_LOG(LogMutableCore, Log, TEXT("%*s@%3d  %s %s"), indent, " ", (int)at, s_opNames[(int)type], desc.c_str() );
+
+            ForEachReference( prog, at, [&]( OP::ADDRESS ref )
+            {
+                if (ref)
+                {
+                    RecurseLogOp( indent+2, ref, prog, done );
+                }
+            });
+        }
+    }
+
+
+    //---------------------------------------------------------------------------------------------
+    void LogSubtree( PROGRAM& prog, OP::ADDRESS root )
+    {
+        for ( vector<PROGRAM::STATE>::const_iterator it = prog.m_states.begin()
+            ; it != prog.m_states.end()
+            ; ++it )
+        {
+			UE_LOG(LogMutableCore, Log, TEXT("state name : % s\n      root : %d"),
+                      it->m_name.c_str(),
+                      (int)it->m_root );
+        }
+
+		UE_LOG(LogMutableCore, Log, TEXT("code : ") );
+
+        vector<uint8_t> done( prog.m_opAddress.size(), 0 );
+        RecurseLogOp( 0, root, prog, done );
+
+		UE_LOG(LogMutableCore, Log, TEXT("parameters : ") );
+        for ( vector<PARAMETER_DESC>::const_iterator it = prog.m_parameters.begin()
+            ; it != prog.m_parameters.end()
+            ; ++it )
+        {
+			UE_LOG(LogMutableCore, Log, TEXT("    [%s] : %2d"), it->m_name.c_str(), (int)it->m_type  );
+        }
+
+
+//		LogDebug( "constant images : \n" );
+//		for ( vector<ImagePtr>::const_iterator it = prog.m_constantImages.begin()
+//			; it != prog.m_constantImages.end()
+//			; ++it )
+//		{
+//			LogDebug( "    image : %2d", (int)(*it)->GetFormat() );
+//			LogDebug( "\t %4d x %4d\n", (int)(*it)->GetSizeX(), (int)(*it)->GetSizeY() );
+//		}
+
+		UE_LOG(LogMutableCore, Log, TEXT("constant meshes : ") );
+        int meshIndex = 0;
+        for ( const auto & at: prog.m_constantMeshes )
+        {
+            const Mesh* m = at.Value.get();
+			UE_LOG(LogMutableCore, Log, TEXT("    %4d mesh : %6d vertices \t %6d indices "),
+                      ++meshIndex,
+                      (int)m->GetVertexCount(), (int)m->GetIndexCount() );
+        }
+
+		UE_LOG(LogMutableCore, Log, TEXT("constant layouts : ") );
+        for ( vector<LayoutPtrConst>::const_iterator it = prog.m_constantLayouts.begin()
+            ; it != prog.m_constantLayouts.end()
+            ; ++it )
+        {
+            auto pLayout = *it;
+			UE_LOG(LogMutableCore, Log, TEXT("    layout : %3d x %3d"),
+                    pLayout->m_size[0],
+                    pLayout->m_size[1]
+                    );
+            for ( int b=0; b<pLayout->GetBlockCount(); ++b )
+            {
+				UE_LOG(LogMutableCore, Log, TEXT("        block : %3d x %3d"),
+                        pLayout->m_blocks[b].m_size[0],
+                        pLayout->m_blocks[b].m_size[1]
+                        );
+            }
+        }
+
+		UE_LOG(LogMutableCore, Log, TEXT("constant skeletons : ") );
+        for ( auto it = prog.m_constantSkeletons.begin()
+            ; it != prog.m_constantSkeletons.end()
+            ; ++it )
+        {
+            auto pSkeleton = *it;
+			UE_LOG(LogMutableCore, Log, TEXT("    skeleton : %3d bones"), pSkeleton->GetBoneCount() );
+            for ( int b=0; b<pSkeleton->GetBoneCount(); ++b )
+            {
+				UE_LOG(LogMutableCore, Log, TEXT("        bone : %s"), pSkeleton->GetBoneName(b) );
+            }
+        }
+
+		UE_LOG(LogMutableCore, Log, TEXT("constant strings : ") );
+        for ( vector<string>::const_iterator it = prog.m_constantStrings.begin()
+            ; it != prog.m_constantStrings.end()
+            ; ++it )
+        {
+			UE_LOG(LogMutableCore, Log, TEXT("    string : [%s]"), it->c_str() );
+        }
+
+    }
+
+
+
+}
diff --git a/Engine/Plugins/Experimental/Mutable/Source/MutableTools/Private/ErrorLogPrivate.h b/Engine/Plugins/Experimental/Mutable/Source/MutableTools/Private/ErrorLogPrivate.h
new file mode 100644
--- /dev/null
+++ b/Engine/Plugins/Experimental/Mutable/Source/MutableTools/Private/ErrorLogPrivate.h
@@ -0,0 +1,73 @@
+// Copyright Epic Games, Inc. All Rights Reserved.
+
+
+#pragma once
+
+#include "MutableTools/Public/ErrorLog.h"
+
+#include "MutableRuntime/Private/MemoryPrivate.h"
+#include "MutableRuntime/Private/ModelPrivate.h"
+#include "MutableRuntime/Private/Operations.h"
+
+#include <memory>
+
+namespace mu
+{
+	MUTABLETOOLS_API extern const TCHAR* s_opNames[(int)OP_TYPE::COUNT];
+
+	class ErrorLog::Private : public Base
+	{
+	public:
+
+		Private()
+		{
+		}
+
+        struct DATA
+        {
+            vector< float > m_unassignedUVs;
+        };
+
+		struct MSG
+		{
+			MSG()
+			{
+                m_type = ELMT_NONE;
+				m_context = 0;
+			}
+
+			ErrorLogMessageType m_type;
+			string m_text;
+            std::shared_ptr< DATA > m_data;
+			const void* m_context;
+		};
+
+		vector< MSG > m_messages;
+
+
+		//-----------------------------------------------------------------------------------------
+
+		//!
+		void Add( const char* strMessage, ErrorLogMessageType type, const void* context );
+
+        //!
+        void Add( const char* strMessage, const ErrorLogMessageAttachedDataView& data, 
+                  ErrorLogMessageType type, const void* context );
+	};
+
+
+	//---------------------------------------------------------------------------------------------
+	class Model;
+	typedef Ptr<Model> ModelPtr;
+	typedef Ptr<const Model> ModelPtrConst;
+
+	extern void Log( ModelPtrConst pModel );
+    extern const TCHAR* GetOpName( OP_TYPE type );
+    extern FString GetOpDesc( const PROGRAM& program, OP::ADDRESS at );
+    extern void LogSubtree( PROGRAM& program, OP::ADDRESS root );
+
+    // Log in a graphviz/dot format
+	MUTABLETOOLS_API extern void LogSubtreeGraph( PROGRAM& program, OP::ADDRESS root );
+
+}
+
diff --git a/Engine/Plugins/Experimental/Mutable/Source/MutableTools/Private/ImageDescGenerator.cpp b/Engine/Plugins/Experimental/Mutable/Source/MutableTools/Private/ImageDescGenerator.cpp
new file mode 100644
--- /dev/null
+++ b/Engine/Plugins/Experimental/Mutable/Source/MutableTools/Private/ImageDescGenerator.cpp
@@ -0,0 +1,579 @@
+// Copyright Epic Games, Inc. All Rights Reserved.
+
+
+#include "ImageDescGenerator.h"
+
+#include "NodeImagePrivate.h"
+#include "NodeImageConstantPrivate.h"
+#include "NodeImageDifferencePrivate.h"
+#include "NodeImageInterpolatePrivate.h"
+#include "NodeImageInterpolate3Private.h"
+#include "NodeImageSaturatePrivate.h"
+#include "NodeImageLuminancePrivate.h"
+#include "NodeImageTablePrivate.h"
+#include "NodeImageSwizzlePrivate.h"
+#include "NodeImageFormatPrivate.h"
+#include "NodeImageSelectColourPrivate.h"
+#include "NodeImageColourMapPrivate.h"
+#include "NodeImageGradientPrivate.h"
+#include "NodeImageBinarisePrivate.h"
+#include "NodeImageLayerPrivate.h"
+#include "NodeImageLayerColourPrivate.h"
+#include "NodeImageMultiLayerPrivate.h"
+#include "NodeImageResizePrivate.h"
+#include "NodeImagePlainColourPrivate.h"
+#include "NodeImageSwitchPrivate.h"
+#include "NodeImageVariationPrivate.h"
+#include "NodeImageProjectPrivate.h"
+#include "NodeImageMipmapPrivate.h"
+#include "NodeImageConditionalPrivate.h"
+#include "NodeImageParameterPrivate.h"
+#include "NodeImageInvertPrivate.h"
+#include "NodeImageNormalCompositePrivate.h"
+#include "NodeImageTransformPrivate.h"
+
+#include "TablePrivate.h"
+
+#include "MeshPrivate.h"
+#include "ImagePrivate.h"
+#include "Platform.h"
+
+
+
+namespace mu
+{
+
+
+    //---------------------------------------------------------------------------------------------
+    void ImageDescGenerator::Generate( const Node::Private& node )
+    {
+        VisitedMap::const_iterator it = m_compiled.find(&node);
+        if ( it != m_compiled.end() )
+        {
+            m_desc = it->second;
+        }
+        else
+        {
+            // The result will be stored in m_size
+            node.Accept( *this );
+            m_compiled[ &node ] = m_desc;
+        }
+    }
+
+
+    //---------------------------------------------------------------------------------------------
+    Ptr<ASTOp> ImageDescGenerator::Visit( const NodeImagePlainColour::Private& node )
+    {
+        m_desc = MUTABLE_MISSING_IMAGE_DESC;
+        m_desc.m_format = EImageFormat::IF_RGB_UBYTE;
+        m_desc.m_size[0] = (int16_t)node.m_sizeX;
+        m_desc.m_size[1] = (int16_t)node.m_sizeY;
+
+        return 0;
+    }
+
+
+    //---------------------------------------------------------------------------------------------
+    Ptr<ASTOp> ImageDescGenerator::Visit( const NodeImageConstant::Private& node )
+    {
+        m_desc = MUTABLE_MISSING_IMAGE_DESC;
+
+        Ptr<const Image> pImage;
+        if (node.m_pProxy) pImage = node.m_pProxy->Get();
+
+        if (pImage)
+        {
+            m_desc.m_size[0]= pImage->GetSizeX();
+            m_desc.m_size[1] = pImage->GetSizeY();
+            m_desc.m_format = pImage->GetFormat();
+            m_desc.m_lods = (uint8_t)pImage->GetLODCount();
+        }
+
+        return 0;
+    }
+
+
+    //---------------------------------------------------------------------------------------------
+    Ptr<ASTOp> ImageDescGenerator::Visit( const NodeImageParameter::Private& )
+    {
+        // We cannot know anything about the image at code generation time
+        // \todo. Provide a template?
+        m_desc = FImageDesc();
+        return 0;
+    }
+
+
+    //---------------------------------------------------------------------------------------------
+    Ptr<ASTOp> ImageDescGenerator::Visit( const NodeImageTable::Private& node )
+    {
+        m_desc = MUTABLE_MISSING_IMAGE_DESC;
+
+        // Verify that the table column is the right type
+        int colIndex = node.m_pTable->FindColumn( node.m_columnName.c_str() );
+        if ( colIndex<0 )
+        {
+            check( false );
+            return 0;
+        }
+
+        Ptr<const Image> pImage;
+
+        std::size_t i = 0;
+        while ( !pImage && i<node.m_pTable->GetPrivate()->m_rows.size() )
+        {
+            pImage = node.m_pTable->GetPrivate()->m_rows[i].m_values[ colIndex ].m_pProxyImage->Get();
+            ++i;
+        }
+
+        if (pImage)
+        {
+            m_desc.m_size[0] = pImage->GetSizeX();
+            m_desc.m_size[1] = pImage->GetSizeY();
+            m_desc.m_format = pImage->GetFormat();
+            m_desc.m_lods = (uint8_t)pImage->GetLODCount();
+        }
+
+        return 0;
+    }
+
+
+    //---------------------------------------------------------------------------------------------
+    Ptr<ASTOp> ImageDescGenerator::Visit( const NodeImageResize::Private& node )
+    {
+        m_desc = MUTABLE_MISSING_IMAGE_DESC;
+
+        if (node.m_pBase)
+        {
+            m_desc = MUTABLE_MISSING_IMAGE_DESC;
+            Generate( *node.m_pBase->GetBasePrivate() );
+
+            if ( node.m_relative )
+            {
+                m_desc.m_size[0] = uint16_t( m_desc.m_size[0] * node.m_sizeX );
+                m_desc.m_size[1] = uint16_t( m_desc.m_size[1] * node.m_sizeY );
+            }
+            else
+            {
+                m_desc.m_size[0] = uint16_t( node.m_sizeX );
+                m_desc.m_size[1] = uint16_t( node.m_sizeY );
+            }
+        }
+
+        return 0;
+    }
+
+
+    //---------------------------------------------------------------------------------------------
+    Ptr<ASTOp> ImageDescGenerator::Visit( const NodeImageSaturate::Private& node )
+    {
+        m_desc = MUTABLE_MISSING_IMAGE_DESC;
+
+        if ( NodeImage* pSource = node.m_pSource.get() )
+        {
+            Generate( *pSource->GetBasePrivate() );
+        }
+
+        return 0;
+    }
+
+
+    //---------------------------------------------------------------------------------------------
+    Ptr<ASTOp> ImageDescGenerator::Visit( const NodeImageMipmap::Private& node )
+    {
+        m_desc = MUTABLE_MISSING_IMAGE_DESC;
+
+        if ( NodeImage* pSource = node.m_pSource.get() )
+        {
+            Generate( *pSource->GetBasePrivate() );
+
+            int mipLevel = std::max(
+                        (int)ceilf( logf( (float)m_desc.m_size[0] )/logf(2.0f) ),
+                            (int)ceilf( logf( (float)m_desc.m_size[1] )/logf(2.0f) )
+                            );
+
+            m_desc.m_lods = std::max( m_desc.m_lods, (uint8_t)mipLevel );
+        }
+
+        return 0;
+    }
+
+
+    //---------------------------------------------------------------------------------------------
+    Ptr<ASTOp> ImageDescGenerator::Visit( const NodeImageLuminance::Private& node )
+    {
+        m_desc = MUTABLE_MISSING_IMAGE_DESC;
+
+        if ( NodeImage* pSource = node.m_pSource.get() )
+        {
+            Generate( *pSource->GetBasePrivate() );
+        }
+
+        m_desc.m_format = EImageFormat::IF_L_UBYTE;
+
+        return 0;
+    }
+
+
+    //---------------------------------------------------------------------------------------------
+    Ptr<ASTOp> ImageDescGenerator::Visit( const NodeImageInterpolate::Private& node )
+    {
+        m_desc = MUTABLE_MISSING_IMAGE_DESC;
+
+        bool found = false;
+        for ( std::size_t t=0; !found && t<node.m_targets.size(); ++t )
+        {
+            if ( NodeImage* pB = node.m_targets[t].get() )
+            {
+                found = true;
+                Generate( *pB->GetBasePrivate() );
+            }
+        }
+
+        m_desc.m_format = EImageFormat::IF_RGB_UBYTE;
+
+        return 0;
+    }
+
+
+    //---------------------------------------------------------------------------------------------
+    Ptr<ASTOp> ImageDescGenerator::Visit( const NodeImageInterpolate3::Private& node )
+    {
+        m_desc = MUTABLE_MISSING_IMAGE_DESC;
+
+        if ( node.m_pTarget0 )
+        {
+            Generate( *node.m_pTarget0->GetBasePrivate() );
+        }
+        else if ( node.m_pTarget1 )
+        {
+            Generate( *node.m_pTarget1->GetBasePrivate() );
+        }
+        else if ( node.m_pTarget2 )
+        {
+            Generate( *node.m_pTarget2->GetBasePrivate() );
+        }
+
+        m_desc.m_format = EImageFormat::IF_RGB_UBYTE;
+
+        return 0;
+    }
+
+
+	//---------------------------------------------------------------------------------------------
+    Ptr<ASTOp> ImageDescGenerator::Visit(const NodeImageSwizzle::Private& node)
+	{
+		m_desc = MUTABLE_MISSING_IMAGE_DESC;
+
+		// The base image size has higher priority
+		bool found = false;
+		for (std::size_t t = 0; !found && t<node.m_sources.size(); ++t)
+		{
+			if (NodeImage* pB = node.m_sources[t].get())
+			{
+				found = true;
+				Generate(*pB->GetBasePrivate());
+			}
+		}
+
+		m_desc.m_format = node.m_format;
+
+		return 0;
+	}
+
+
+	//---------------------------------------------------------------------------------------------
+    Ptr<ASTOp> ImageDescGenerator::Visit(const NodeImageFormat::Private& node)
+	{
+		m_desc = MUTABLE_MISSING_IMAGE_DESC;
+
+		if (node.m_source)
+		{
+			Generate(*node.m_source->GetBasePrivate());
+		}
+
+        // Necessary since dxt1 reports 4 channels just in case
+        if (m_desc.m_format!=node.m_format)
+        {
+            if (GetImageFormatData(m_desc.m_format).m_channels>3
+                    &&
+                    node.m_formatIfAlpha!= EImageFormat::IF_NONE
+                    )
+            {
+                m_desc.m_format = node.m_formatIfAlpha;
+            }
+            else
+            {
+                m_desc.m_format = node.m_format;
+            }
+        }
+
+		return 0;
+	}
+
+
+    //---------------------------------------------------------------------------------------------
+    Ptr<ASTOp> ImageDescGenerator::Visit( const NodeImageDifference::Private& node )
+    {
+        m_desc = MUTABLE_MISSING_IMAGE_DESC;
+
+        // The first image size has higher priority
+        if ( NodeImage* pA = node.m_pA.get() )
+        {
+            Generate( *pA->GetBasePrivate() );
+        }
+        else if ( NodeImage* pB = node.m_pB.get() )
+        {
+            Generate( *pB->GetBasePrivate() );
+        }
+
+        m_desc.m_format = EImageFormat::IF_L_UBYTE;
+
+        return 0;
+    }
+
+
+    //---------------------------------------------------------------------------------------------
+    Ptr<ASTOp> ImageDescGenerator::Visit( const NodeImageBinarise::Private& node )
+    {
+        m_desc = MUTABLE_MISSING_IMAGE_DESC;
+
+        if ( NodeImage* pA = node.m_pBase.get() )
+        {
+            Generate( *pA->GetBasePrivate() );
+        }
+
+        m_desc.m_format = EImageFormat::IF_L_UBYTE;
+
+        return 0;
+    }
+
+
+    //---------------------------------------------------------------------------------------------
+    Ptr<ASTOp> ImageDescGenerator::Visit( const NodeImageSelectColour::Private& node )
+    {
+        m_desc = MUTABLE_MISSING_IMAGE_DESC;
+
+        if ( NodeImage* pBase = node.m_pSource.get() )
+        {
+            Generate( *pBase->GetBasePrivate() );
+        }
+
+        m_desc.m_format = EImageFormat::IF_L_UBYTE;
+
+        return 0;
+    }
+
+
+    //---------------------------------------------------------------------------------------------
+    Ptr<ASTOp> ImageDescGenerator::Visit( const NodeImageColourMap::Private& node )
+    {
+        m_desc = MUTABLE_MISSING_IMAGE_DESC;
+
+        if ( NodeImage* pBase = node.m_pBase.get() )
+        {
+            Generate( *pBase->GetBasePrivate() );
+        }
+
+        return 0;
+    }
+
+
+    //---------------------------------------------------------------------------------------------
+    Ptr<ASTOp> ImageDescGenerator::Visit( const NodeImageGradient::Private& node )
+    {
+        m_desc.m_format = EImageFormat::IF_RGB_UBYTE;
+        m_desc.m_size[0] = (uint16_t)node.m_size[0];
+        m_desc.m_size[1] = (uint16_t)node.m_size[1];
+        m_desc.m_lods = 1;
+        return 0;
+    }
+
+
+    //---------------------------------------------------------------------------------------------
+    Ptr<ASTOp> ImageDescGenerator::Visit( const NodeImageLayer::Private& node )
+    {
+        m_desc = MUTABLE_MISSING_IMAGE_DESC;
+
+        // The base image size has higher priority
+        if ( NodeImage* pBase = node.m_pBase.get() )
+        {
+            Generate( *pBase->GetBasePrivate() );
+        }
+        else if ( NodeImage* pMask = node.m_pMask.get() )
+        {
+            Generate( *pMask->GetBasePrivate() );
+        }
+
+        return 0;
+    }
+
+
+    //---------------------------------------------------------------------------------------------
+    Ptr<ASTOp> ImageDescGenerator::Visit( const NodeImageLayerColour::Private& node )
+    {
+        m_desc = MUTABLE_MISSING_IMAGE_DESC;
+
+        // The base image size has higher priority
+        if ( NodeImage* pBase = node.m_pBase.get() )
+        {
+            Generate( *pBase->GetBasePrivate() );
+        }
+        else if ( NodeImage* pMask = node.m_pMask.get() )
+        {
+            Generate( *pMask->GetBasePrivate() );
+        }
+
+        m_desc.m_format = EImageFormat::IF_RGB_UBYTE;
+
+        return 0;
+    }
+
+    //---------------------------------------------------------------------------------------------
+    Ptr<ASTOp> ImageDescGenerator::Visit( const NodeImageMultiLayer::Private& node )
+    {
+        m_desc = MUTABLE_MISSING_IMAGE_DESC;
+
+        // The base image size has higher priority
+        if ( NodeImage* pBase = node.m_pBase.get() )
+        {
+            Generate( *pBase->GetBasePrivate() );
+        }
+        else if ( NodeImage* pMask = node.m_pMask.get() )
+        {
+            Generate( *pMask->GetBasePrivate() );
+        }
+
+        return 0;
+    }
+
+    //---------------------------------------------------------------------------------------------
+    Ptr<ASTOp> ImageDescGenerator::Visit( const NodeImageNormalComposite::Private& node )
+    {
+        m_desc = MUTABLE_MISSING_IMAGE_DESC;
+
+        // The base image size has higher priority
+        if ( NodeImage* pBase = node.m_pBase.get() )
+        {
+            Generate( *pBase->GetBasePrivate() );
+        }
+
+        return 0;
+    }
+
+    //---------------------------------------------------------------------------------------------
+    Ptr<ASTOp> ImageDescGenerator::Visit( const NodeImageSwitch::Private& node )
+    {
+        m_desc = MUTABLE_MISSING_IMAGE_DESC;
+
+        if ( node.m_options.size() > 0 && node.m_options[0] )
+        {
+            Generate( *node.m_options[0]->GetBasePrivate() );
+        }
+
+        return 0;
+    }
+
+
+    //---------------------------------------------------------------------------------------------
+    Ptr<ASTOp> ImageDescGenerator::Visit( const NodeImageVariation::Private& node )
+    {
+        m_desc = MUTABLE_MISSING_IMAGE_DESC;
+
+        if ( node.m_defaultImage )
+        {
+            Generate( *node.m_defaultImage->GetBasePrivate() );
+        }
+        else if ( node.m_variations.size() > 0 && node.m_variations[0].m_image )
+        {
+            Generate( *node.m_variations[0].m_image->GetBasePrivate() );
+        }
+
+        return 0;
+    }
+
+
+    //---------------------------------------------------------------------------------------------
+    Ptr<ASTOp> ImageDescGenerator::Visit( const NodeImageConditional::Private& node )
+    {
+        m_desc = MUTABLE_MISSING_IMAGE_DESC;
+
+        if ( node.m_true )
+        {
+            Generate( *node.m_true->GetBasePrivate() );
+        }
+        else if ( node.m_false )
+        {
+            Generate( *node.m_false->GetBasePrivate() );
+        }
+
+        return 0;
+    }
+
+
+    //---------------------------------------------------------------------------------------------
+    Ptr<ASTOp> ImageDescGenerator::Visit( const NodeImageProject::Private& node )
+    {
+        m_desc = MUTABLE_MISSING_IMAGE_DESC;
+
+		EImageFormat format = EImageFormat::IF_RGB_UBYTE;
+
+        // Format from the image to project
+        if ( node.m_pImage )
+        {
+            Generate( *node.m_pImage->GetBasePrivate() );
+            format = m_desc.m_format;
+        }
+
+		// Size override?
+		if (node.m_imageSize.X > 0 && node.m_imageSize.Y > 0)
+		{
+			m_desc.m_size[0] = uint16_t(node.m_imageSize[0]);
+			m_desc.m_size[1] = uint16_t(node.m_imageSize[1]);
+		}
+
+        // Size from the mask
+        else if ( node.m_pMask )
+        {
+            Generate( *node.m_pMask->GetBasePrivate() );
+        }
+
+        else
+        {
+			// \TODO: Warning?
+            m_desc.m_size[0] = 256;
+            m_desc.m_size[1] = 256;
+        }
+
+        m_desc.m_format = format;
+
+        return 0;
+    }
+
+
+	//---------------------------------------------------------------------------------------------
+	Ptr<ASTOp> ImageDescGenerator::Visit(const NodeImageInvert::Private& node)
+	{
+		m_desc = MUTABLE_MISSING_IMAGE_DESC;
+
+		if (NodeImage* pA = node.m_pBase.get())
+		{
+			Generate(*pA->GetBasePrivate());
+		}
+
+		return 0;
+	}
+
+    //---------------------------------------------------------------------------------------------
+    Ptr<ASTOp> ImageDescGenerator::Visit( const NodeImageTransform::Private& node )
+    {
+        m_desc = MUTABLE_MISSING_IMAGE_DESC;
+
+        if ( NodeImage* pBase = node.m_pBase.get() )
+        {
+            Generate( *pBase->GetBasePrivate() );
+        }
+
+        return 0;
+    }
+}
+
