[mutable] Move Mutable plugin from the restricted folder to the experimental plugins folder.

diff --git a/Engine/Plugins/Experimental/Mutable/Source/MutableTools/Private/ImageDescGenerator.h b/Engine/Plugins/Experimental/Mutable/Source/MutableTools/Private/ImageDescGenerator.h
new file mode 100644
--- /dev/null
+++ b/Engine/Plugins/Experimental/Mutable/Source/MutableTools/Private/ImageDescGenerator.h
@@ -0,0 +1,126 @@
+// Copyright Epic Games, Inc. All Rights Reserved.
+
+#pragma once
+
+#include "Visitor.h"
+
+#include "AST.h"
+#include "NodeImage.h"
+#include "NodeImageConstant.h"
+#include "NodeImageDifference.h"
+#include "NodeImageInterpolate.h"
+#include "NodeImageInterpolate3.h"
+#include "NodeImageSaturate.h"
+#include "NodeImageLuminance.h"
+#include "NodeImageLayer.h"
+#include "NodeImageLayerColour.h"
+#include "NodeImageMultiLayer.h"
+#include "NodeImageTable.h"
+#include "NodeImageSwizzle.h"
+#include "NodeImageFormat.h"
+#include "NodeImageSelectColour.h"
+#include "NodeImageColourMap.h"
+#include "NodeImageGradient.h"
+#include "NodeImageBinarise.h"
+#include "NodeImageResize.h"
+#include "NodeImagePlainColour.h"
+#include "NodeImageSwitch.h"
+#include "NodeImageVariation.h"
+#include "NodeImageProject.h"
+#include "NodeImageMipmap.h"
+#include "NodeImageConditional.h"
+#include "NodeImageParameter.h"
+#include "NodeImageInvert.h"
+#include "NodeImageNormalComposite.h"
+#include "NodeImageTransform.h"
+
+#include "ErrorLogPrivate.h"
+
+#include "Operations.h"
+#include "ImagePrivate.h"
+
+
+#define MUTABLE_MISSING_IMAGE_DESC 	FImageDesc( FImageSize( 16, 16 ), EImageFormat::IF_RGB_UBYTE, 1 )
+
+
+namespace mu
+{
+
+    //---------------------------------------------------------------------------------------------
+    //! Generate the description of an image without generating its code.
+    //---------------------------------------------------------------------------------------------
+    class ImageDescGenerator : public Base,
+                               public BaseVisitor,
+
+                               // Standard operations from models
+                               public Visitor<NodeImageConstant::Private, Ptr<ASTOp>, true>,
+                               public Visitor<NodeImageDifference::Private, Ptr<ASTOp>, true>,
+                               public Visitor<NodeImageInterpolate::Private, Ptr<ASTOp>, true>,
+                               public Visitor<NodeImageInterpolate3::Private, Ptr<ASTOp>, true>,
+                               public Visitor<NodeImageSaturate::Private, Ptr<ASTOp>, true>,
+                               public Visitor<NodeImageLuminance::Private, Ptr<ASTOp>, true>,
+                               public Visitor<NodeImageLayer::Private, Ptr<ASTOp>, true>,
+                               public Visitor<NodeImageLayerColour::Private, Ptr<ASTOp>, true>,
+                               public Visitor<NodeImageTable::Private, Ptr<ASTOp>, true>,
+                               public Visitor<NodeImageSwizzle::Private, Ptr<ASTOp>, true>,
+                               public Visitor<NodeImageFormat::Private, Ptr<ASTOp>, true>,
+                               public Visitor<NodeImageSelectColour::Private, Ptr<ASTOp>, true>,
+                               public Visitor<NodeImageColourMap::Private, Ptr<ASTOp>, true>,
+                               public Visitor<NodeImageGradient::Private, Ptr<ASTOp>, true>,
+                               public Visitor<NodeImageBinarise::Private, Ptr<ASTOp>, true>,
+                               public Visitor<NodeImageResize::Private, Ptr<ASTOp>, true>,
+                               public Visitor<NodeImagePlainColour::Private, Ptr<ASTOp>, true>,
+                               public Visitor<NodeImageSwitch::Private, Ptr<ASTOp>, true>,
+                               public Visitor<NodeImageVariation::Private, Ptr<ASTOp>, true>,
+                               public Visitor<NodeImageProject::Private, Ptr<ASTOp>, true>,
+                               public Visitor<NodeImageMipmap::Private, Ptr<ASTOp>, true>,
+                               public Visitor<NodeImageConditional::Private, Ptr<ASTOp>, true>,
+                               public Visitor<NodeImageParameter::Private, Ptr<ASTOp>, true>,
+                               public Visitor<NodeImageMultiLayer::Private, Ptr<ASTOp>, true>,
+                               public Visitor<NodeImageInvert::Private, Ptr<ASTOp>, true>,
+                               public Visitor<NodeImageNormalComposite::Private, Ptr<ASTOp>, true>,
+                               public Visitor<NodeImageTransform::Private, Ptr<ASTOp>, true>
+    {
+    public:
+
+        void Generate( const Node::Private& );
+
+        Ptr<ASTOp> Visit( const NodeImageConstant::Private& );
+        Ptr<ASTOp> Visit( const NodeImageDifference::Private& );
+        Ptr<ASTOp> Visit( const NodeImageInterpolate::Private& );
+        Ptr<ASTOp> Visit( const NodeImageInterpolate3::Private& );
+        Ptr<ASTOp> Visit( const NodeImageSaturate::Private& );
+        Ptr<ASTOp> Visit( const NodeImageLuminance::Private& );
+        Ptr<ASTOp> Visit( const NodeImageLayer::Private& );
+        Ptr<ASTOp> Visit( const NodeImageLayerColour::Private& );
+        Ptr<ASTOp> Visit( const NodeImageTable::Private& );
+        Ptr<ASTOp> Visit( const NodeImageSwizzle::Private&);
+        Ptr<ASTOp> Visit( const NodeImageFormat::Private&);
+        Ptr<ASTOp> Visit( const NodeImageSelectColour::Private& );
+        Ptr<ASTOp> Visit( const NodeImageColourMap::Private& );
+        Ptr<ASTOp> Visit( const NodeImageGradient::Private& );
+        Ptr<ASTOp> Visit( const NodeImageBinarise::Private& );
+        Ptr<ASTOp> Visit( const NodeImageResize::Private& );
+        Ptr<ASTOp> Visit( const NodeImagePlainColour::Private& );
+        Ptr<ASTOp> Visit( const NodeImageSwitch::Private& );
+        Ptr<ASTOp> Visit( const NodeImageVariation::Private& );
+        Ptr<ASTOp> Visit( const NodeImageProject::Private& );
+        Ptr<ASTOp> Visit( const NodeImageMipmap::Private& );
+        Ptr<ASTOp> Visit( const NodeImageConditional::Private& );
+        Ptr<ASTOp> Visit( const NodeImageParameter::Private& );
+        Ptr<ASTOp> Visit( const NodeImageMultiLayer::Private& );
+        Ptr<ASTOp> Visit( const NodeImageInvert::Private& );
+        Ptr<ASTOp> Visit( const NodeImageNormalComposite::Private& );
+        Ptr<ASTOp> Visit( const NodeImageTransform::Private& );
+
+    public:
+
+        typedef map<const Node::Private*,FImageDesc> VisitedMap;
+        VisitedMap m_compiled;
+
+        //! Result
+        FImageDesc m_desc;
+
+    };
+
+}
diff --git a/Engine/Plugins/Experimental/Mutable/Source/MutableTools/Private/ModelReport.cpp b/Engine/Plugins/Experimental/Mutable/Source/MutableTools/Private/ModelReport.cpp
new file mode 100644
--- /dev/null
+++ b/Engine/Plugins/Experimental/Mutable/Source/MutableTools/Private/ModelReport.cpp
@@ -0,0 +1,293 @@
+// Copyright Epic Games, Inc. All Rights Reserved.
+
+
+#include "ModelReport.h"
+#include "ModelReportPrivate.h"
+
+
+
+namespace mu
+{
+
+	//---------------------------------------------------------------------------------------------
+	//---------------------------------------------------------------------------------------------
+	//---------------------------------------------------------------------------------------------
+	ModelReport::ModelReport()
+	{
+		m_pD = new Private();
+	}
+
+
+	//---------------------------------------------------------------------------------------------
+	ModelReport::~ModelReport()
+	{
+        check( m_pD );
+		delete m_pD;
+		m_pD = 0;
+	}
+
+
+	//---------------------------------------------------------------------------------------------
+	ModelReport::Private* ModelReport::GetPrivate() const
+	{
+		return m_pD;
+	}
+
+
+	//---------------------------------------------------------------------------------------------
+	int ModelReport::GetStateCount() const
+	{
+		return (int)m_pD->m_states.size();
+	}
+
+
+	//---------------------------------------------------------------------------------------------
+	int ModelReport::GetLODCount( int state ) const
+	{
+		int res = 0;
+
+		if ( state>=0 && state<(int)m_pD->m_states.size() )
+		{
+			res = (int)m_pD->m_states[state].m_lods.size();
+		}
+		else
+		{
+			check(false);
+		}
+
+		return res;
+	}
+
+
+	//---------------------------------------------------------------------------------------------
+	int ModelReport::GetComponentCount( int state, int lod ) const
+	{
+		int res = 0;
+
+		if ( state>=0 && state<(int)m_pD->m_states.size() )
+		{
+			res = (int)m_pD->m_states[state].m_lods[lod].m_components.size();
+		}
+		else
+		{
+			check(false);
+		}
+
+		return res;
+	}
+
+
+	//---------------------------------------------------------------------------------------------
+	const char* ModelReport::GetComponentName( int state, int lod, int comp ) const
+	{
+		const char* strRes = 0;
+
+		if ( state>=0 && state<(int)m_pD->m_states.size()
+			 && comp>=0 && comp<(int)m_pD->m_states[state].m_lods[lod].m_components.size()
+			 )
+		{
+			strRes = m_pD->m_states[state].m_lods[lod].m_components[comp].m_name.c_str();
+		}
+		else
+		{
+			check(false);
+		}
+
+		return strRes;
+	}
+
+
+	//---------------------------------------------------------------------------------------------
+	int ModelReport::GetComponentImageCount( int state, int lod, int comp ) const
+	{
+		int res = 0;
+
+		if ( state>=0 && state<(int)m_pD->m_states.size()
+			 && comp>=0 && comp<(int)m_pD->m_states[state].m_lods[lod].m_components.size()
+			 )
+		{
+			res = (int)m_pD->m_states[state].m_lods[lod].m_components[comp].m_images.size();
+		}
+		else
+		{
+			check(false);
+		}
+
+		return res;
+	}
+
+
+	//---------------------------------------------------------------------------------------------
+	const char* ModelReport::GetImageName( int state, int lod, int comp, int image ) const
+	{
+		const char* res = 0;
+
+		if ( state>=0 && state<(int)m_pD->m_states.size()
+			 && comp>=0 && comp<(int)m_pD->m_states[state].m_lods[lod].m_components.size()
+			 && image>=0 && image<(int)m_pD->m_states[state].m_lods[lod].m_components[comp].m_images.size() )
+		{
+			res = m_pD->m_states[state].m_lods[lod].m_components[comp].m_images[image].m_name.c_str();
+		}
+		else
+		{
+			check(false);
+		}
+
+		return res;
+	}
+
+
+	//---------------------------------------------------------------------------------------------
+	const char* ModelReport::GetImageFragmentCode( int state, int lod, int comp, int image ) const
+	{
+		const char* res = 0;
+
+		if ( state>=0 && state<(int)m_pD->m_states.size()
+			 && comp>=0 && comp<(int)m_pD->m_states[state].m_lods[lod].m_components.size()
+			 && image>=0 && image<(int)m_pD->m_states[state].m_lods[lod].m_components[comp].m_images.size() )
+		{
+			res = m_pD->m_states[state].m_lods[lod].m_components[comp].m_images[image].m_fragmentCode.c_str();
+		}
+		else
+		{
+			check(false);
+		}
+
+		return res;
+	}
+
+
+	//---------------------------------------------------------------------------------------------
+	int ModelReport::GetFragmentSourceImageCount( int state, int lod, int comp, int image ) const
+	{
+		int result = 0;
+
+		if ( state>=0 && state<(int)m_pD->m_states.size()
+			 && comp>=0 && comp<(int)m_pD->m_states[state].m_lods[lod].m_components.size()
+			 && image>=0 && image<(int)m_pD->m_states[state].m_lods[lod].m_components[comp].m_images.size() )
+		{
+			result = (int)m_pD->m_states[state].m_lods[lod].m_components[comp].m_images[image].m_sourceImages.size();
+		}
+		else
+		{
+			check(false);
+		}
+
+		return result;
+	}
+
+
+    //---------------------------------------------------------------------------------------------
+    const char* ModelReport::GetFragmentSourceImageName( int state, int lod, int comp,
+                                                                   int image,
+                                                                   int source ) const
+    {
+        const char* strRes = 0;
+
+        if ( state>=0 && state<(int)m_pD->m_states.size()
+             && comp>=0 && comp<(int)m_pD->m_states[state].m_lods[lod].m_components.size()
+             && image>=0 && image<(int)m_pD->m_states[state].m_lods[lod].m_components[comp].m_images.size()
+             && source>=0
+             && source<(int)m_pD->m_states[state].m_lods[lod].m_components[comp].m_images[image].m_sourceImages.size() )
+        {
+            strRes = m_pD->m_states[state].m_lods[lod].m_components[comp].m_images[image].m_sourceImages[source].c_str();
+        }
+		else
+		{
+			check(false);
+		}
+
+        return strRes;
+    }
+
+
+	//---------------------------------------------------------------------------------------------
+	int ModelReport::GetFragmentSourceVectorCount( int state, int lod, int comp, int image ) const
+	{
+		int result = 0;
+
+		if ( state>=0 && state<(int)m_pD->m_states.size()
+			 && comp>=0 && comp<(int)m_pD->m_states[state].m_lods[lod].m_components.size()
+			 && image>=0 && image<(int)m_pD->m_states[state].m_lods[lod].m_components[comp].m_images.size() )
+		{
+			result = (int)m_pD->m_states[state].m_lods[lod].m_components[comp].m_images[image].m_sourceVectors.size();
+		}
+		else
+		{
+			check(false);
+		}
+
+		return result;
+	}
+
+
+	//---------------------------------------------------------------------------------------------
+	const char* ModelReport::GetFragmentSourceVectorName( int state, int lod, int comp,
+																	int image,
+																	int source ) const
+	{
+		const char* strRes = 0;
+
+		if ( state>=0 && state<(int)m_pD->m_states.size()
+			 && comp>=0 && comp<(int)m_pD->m_states[state].m_lods[lod].m_components.size()
+			 && image>=0 && image<(int)m_pD->m_states[state].m_lods[lod].m_components[comp].m_images.size()
+			 && source>=0
+			 && source<(int)m_pD->m_states[state].m_lods[lod].m_components[comp].m_images[image].m_sourceVectors.size() )
+		{
+			strRes = m_pD->m_states[state].m_lods[lod].m_components[comp].m_images[image].m_sourceVectors[source].c_str();
+		}
+		else
+		{
+			check(false);
+		}
+
+		return strRes;
+	}
+
+
+	//---------------------------------------------------------------------------------------------
+	int ModelReport::GetFragmentSourceScalarCount( int state, int lod, int comp,
+																	  int image ) const
+	{
+		int result = 0;
+
+		if ( state>=0 && state<(int)m_pD->m_states.size()
+			 && comp>=0 && comp<(int)m_pD->m_states[state].m_lods[lod].m_components.size()
+			 && image>=0 && image<(int)m_pD->m_states[state].m_lods[lod].m_components[comp].m_images.size() )
+		{
+			result = (int)m_pD->m_states[state].m_lods[lod].m_components[comp].m_images[image].m_sourceScalars.size();
+		}
+		else
+		{
+			check(false);
+		}
+
+		return result;
+	}
+
+
+	//---------------------------------------------------------------------------------------------
+	const char* ModelReport::GetFragmentSourceScalarName( int state, int lod, int comp,
+																	int image,
+																	int source ) const
+	{
+		const char* strRes = 0;
+
+		if ( state>=0 && state<(int)m_pD->m_states.size()
+			 && comp>=0 && comp<(int)m_pD->m_states[state].m_lods[lod].m_components.size()
+			 && image>=0 && image<(int)m_pD->m_states[state].m_lods[lod].m_components[comp].m_images.size()
+			 && source>=0
+			 && source<(int)m_pD->m_states[state].m_lods[lod].m_components[comp].m_images[image].m_sourceScalars.size() )
+		{
+			strRes = m_pD->m_states[state].m_lods[lod].m_components[comp].m_images[image].m_sourceScalars[source].c_str();
+		}
+		else
+		{
+			check(false);
+		}
+
+		return strRes;
+	}
+
+}
+
diff --git a/Engine/Plugins/Experimental/Mutable/Source/MutableTools/Private/ModelReportPrivate.h b/Engine/Plugins/Experimental/Mutable/Source/MutableTools/Private/ModelReportPrivate.h
new file mode 100644
--- /dev/null
+++ b/Engine/Plugins/Experimental/Mutable/Source/MutableTools/Private/ModelReportPrivate.h
@@ -0,0 +1,52 @@
+// Copyright Epic Games, Inc. All Rights Reserved.
+
+#pragma once
+
+#include "ModelReport.h"
+#include "SerialisationPrivate.h"
+
+
+namespace mu
+{
+
+	//!
+	struct REPORT_STATE
+	{
+		struct IMAGE
+		{
+			string m_name;
+			string m_fragmentCode;
+			vector<string> m_sourceImages;
+			vector<string> m_sourceVectors;
+			vector<string> m_sourceScalars;
+		};
+
+		struct COMPONENT
+		{
+			string m_name;
+			vector<IMAGE> m_images;
+		};
+
+		struct LOD
+		{
+			vector<COMPONENT> m_components;
+		};
+
+		vector<LOD> m_lods;
+	};
+
+
+	//!
+	class ModelReport::Private : public Base
+	{
+	public:
+
+		Private()
+		{
+		}
+
+		vector<REPORT_STATE> m_states;
+
+	};
+
+}
diff --git a/Engine/Plugins/Experimental/Mutable/Source/MutableTools/Private/Node.cpp b/Engine/Plugins/Experimental/Mutable/Source/MutableTools/Private/Node.cpp
new file mode 100644
--- /dev/null
+++ b/Engine/Plugins/Experimental/Mutable/Source/MutableTools/Private/Node.cpp
@@ -0,0 +1,222 @@
+// Copyright Epic Games, Inc. All Rights Reserved.
+
+
+#include "Node.h"
+#include "NodePrivate.h"
+
+#include "NodeColour.h"
+#include "NodeComponent.h"
+#include "NodeImage.h"
+#include "NodeLayout.h"
+#include "NodeLOD.h"
+#include "NodeMesh.h"
+#include "NodeObject.h"
+#include "NodePatchImage.h"
+#include "NodePatchMesh.h"
+#include "NodeScalar.h"
+#include "NodeString.h"
+#include "NodeProjector.h"
+#include "NodeSurface.h"
+#include "NodeModifier.h"
+#include "NodeRange.h"
+#include "NodeBool.h"
+
+
+namespace mu
+{
+
+	//---------------------------------------------------------------------------------------------
+	// Static initialisation
+	//---------------------------------------------------------------------------------------------
+	static NODE_TYPE s_nodeType = 		NODE_TYPE( "Node", 0 );
+
+
+	//---------------------------------------------------------------------------------------------
+	NODE_TYPE::NODE_TYPE()
+	{
+		m_strName = "";
+		m_pParent = 0;
+	}
+
+	//---------------------------------------------------------------------------------------------
+	NODE_TYPE::NODE_TYPE( const char* strName, const NODE_TYPE* pParent )
+	{
+        m_strName = strName;
+		m_pParent = pParent;
+	}
+
+
+
+	//---------------------------------------------------------------------------------------------
+	//---------------------------------------------------------------------------------------------
+	//---------------------------------------------------------------------------------------------
+	void Node::Serialise( const Node* pNode, OutputArchive& arch )
+	{
+        uint32_t ver = 0;
+		arch << ver;
+
+		// Warning: cannot be replaced with the wrapping virtual call because some subclasses are abstract
+		arch << uint32_t(pNode->Type);
+
+		#define SERIALISE_CHILDREN( C, ID ) C::Serialise(static_cast<const C*>(pNode), arch );
+
+		switch (pNode->Type)
+		{
+		case EType::Colour:		SERIALISE_CHILDREN(NodeColour,arch); break;
+		case EType::Component:  SERIALISE_CHILDREN(NodeComponent,arch); break;
+		case EType::Image:		SERIALISE_CHILDREN(NodeImage,arch); break;
+		case EType::Layout:		SERIALISE_CHILDREN(NodeLayout,arch); break;
+		case EType::LOD:		SERIALISE_CHILDREN(NodeLOD,arch); break;
+		case EType::Mesh:		SERIALISE_CHILDREN(NodeMesh,arch); break;
+		case EType::Object:		SERIALISE_CHILDREN(NodeObject,arch); break;
+		case EType::PatchImage: SERIALISE_CHILDREN(NodePatchImage,arch); break;
+		case EType::Scalar:		SERIALISE_CHILDREN(NodeScalar,arch); break;
+		case EType::PatchMesh:	SERIALISE_CHILDREN(NodePatchMesh, arch); break;
+		case EType::Projector:	SERIALISE_CHILDREN(NodeProjector,arch); break;
+		case EType::Surface:	SERIALISE_CHILDREN(NodeSurface,arch); break;
+		case EType::Modifier:	SERIALISE_CHILDREN(NodeModifier,arch); break;
+		case EType::Range:		SERIALISE_CHILDREN(NodeRange,arch); break;
+		case EType::String:		SERIALISE_CHILDREN(NodeString, arch); break;
+		case EType::Bool:		SERIALISE_CHILDREN(NodeBool, arch); break;
+		default: check(false);
+		}
+
+		#undef SERIALISE_CHILDREN
+	}
+
+
+	//---------------------------------------------------------------------------------------------
+	NodePtr Node::StaticUnserialise( InputArchive& arch )
+	{
+        uint32_t ver;
+		arch >> ver;
+		check( ver == 0 );
+
+        uint32_t id;
+		arch >> id;
+
+		switch (EType(id))
+		{
+		case EType::Colour:		return NodeColour::StaticUnserialise( arch ); break;
+		case EType::Component:  return NodeComponent::StaticUnserialise( arch ); break;
+		case EType::Image:		return NodeImage::StaticUnserialise( arch ); break;
+		case EType::Layout:		return NodeLayout::StaticUnserialise( arch ); break;
+		case EType::LOD:		return NodeLOD::StaticUnserialise( arch ); break;
+		case EType::Mesh:		return NodeMesh::StaticUnserialise( arch ); break;
+		case EType::Object:		return NodeObject::StaticUnserialise( arch ); break;
+		case EType::PatchImage: return NodePatchImage::StaticUnserialise( arch ); break;
+		case EType::Scalar:		return NodeScalar::StaticUnserialise( arch ); break;
+		case EType::PatchMesh:	return NodePatchMesh::StaticUnserialise( arch ); break;
+		case EType::Projector:	return NodeProjector::StaticUnserialise( arch ); break;
+		case EType::Surface:	return NodeSurface::StaticUnserialise( arch ); break;
+		case EType::Modifier:	return NodeModifier::StaticUnserialise( arch ); break;
+		case EType::Range:		return NodeRange::StaticUnserialise( arch ); break;
+		case EType::String:		return NodeString::StaticUnserialise( arch ); break;
+		case EType::Bool:		return NodeBool::StaticUnserialise( arch ); break;
+		default : check(false);
+		}
+
+		return 0;
+	}
+
+
+
+	//---------------------------------------------------------------------------------------------
+	const NODE_TYPE* Node::GetType() const
+	{
+		return GetStaticType();
+	}
+
+
+	//---------------------------------------------------------------------------------------------
+	const NODE_TYPE* Node::GetStaticType()
+	{
+		return &s_nodeType;
+	}
+
+
+	//---------------------------------------------------------------------------------------------
+	void Node::SetMessageContext( const void* context )
+	{
+		GetBasePrivate()->m_errorContext = context;
+	}
+
+
+	//---------------------------------------------------------------------------------------------
+	void Node::Private::CloneChildren( const Node* pThis, NodePtr pTarget, NodeMapPtr pMap )
+	{
+		// Clone the children
+		for ( int i=0; i<pThis->GetInputCount(); ++i )
+		{
+			NodePtr pSource = pThis->GetInputNode(i);
+			NodePtr pCloned;
+			if (pSource)
+			{
+				pCloned = pMap->Get( pSource.get() );
+				if (!pCloned)
+				{
+					pCloned = pSource->Clone( pMap );
+				}
+			}
+
+			pTarget->SetInputNode( i, pCloned.get() );
+		}
+	}
+
+
+	//---------------------------------------------------------------------------------------------
+	//---------------------------------------------------------------------------------------------
+	//---------------------------------------------------------------------------------------------
+	NodeMap::NodeMap()
+	{
+		m_pD = new Private();
+	}
+
+
+	//---------------------------------------------------------------------------------------------
+	NodeMap::~NodeMap()
+	{
+        check( m_pD );
+		delete m_pD;
+		m_pD = 0;
+	}
+
+
+	//---------------------------------------------------------------------------------------------
+	NodeMap::Private* NodeMap::GetPrivate() const
+	{
+		return m_pD;
+	}
+
+
+	//---------------------------------------------------------------------------------------------
+	int NodeMap::GetSize( ) const
+	{
+		return (int)m_pD->m_map.size();
+	}
+
+
+	//---------------------------------------------------------------------------------------------
+	void NodeMap::Add( const void* key, NodePtr value )
+	{
+		m_pD->m_map[ key ] = value;
+	}
+
+
+	//---------------------------------------------------------------------------------------------
+	NodePtr NodeMap::Get( const void* key ) const
+	{
+		NodePtr pResult;
+
+		map<const void*,NodePtr>::const_iterator it = m_pD->m_map.find(key);
+		if ( it != m_pD->m_map.end() )
+		{
+			pResult = it->second;
+		}
+
+		return pResult;
+	}
+
+}
+
+
diff --git a/Engine/Plugins/Experimental/Mutable/Source/MutableTools/Private/NodeBool.cpp b/Engine/Plugins/Experimental/Mutable/Source/MutableTools/Private/NodeBool.cpp
new file mode 100644
--- /dev/null
+++ b/Engine/Plugins/Experimental/Mutable/Source/MutableTools/Private/NodeBool.cpp
@@ -0,0 +1,68 @@
+// Copyright Epic Games, Inc. All Rights Reserved.
+
+
+#include "NodeBool.h"
+#include "NodePrivate.h"
+
+
+namespace mu
+{
+
+	//---------------------------------------------------------------------------------------------
+	// Static initialisation
+	//---------------------------------------------------------------------------------------------
+	static NODE_TYPE s_nodeBoolType = NODE_TYPE( "NodeBool", Node::GetStaticType() );
+
+
+	//---------------------------------------------------------------------------------------------
+	void NodeBool::Serialise( const NodeBool* p, OutputArchive& arch )
+	{
+        uint32_t ver = 0;
+		arch << ver;
+
+		arch << uint32_t(p->Type);
+		p->SerialiseWrapper(arch);
+	}
+
+        
+	//---------------------------------------------------------------------------------------------
+	NodeBoolPtr NodeBool::StaticUnserialise( InputArchive& arch )
+	{
+        uint32_t ver;
+		arch >> ver;
+		check( ver == 0 );
+
+        uint32_t id;
+		arch >> id;
+
+		switch (id)
+		{
+		case 0 :  return NodeBoolConstant::StaticUnserialise( arch ); break;
+		case 1 :  return NodeBoolParameter::StaticUnserialise( arch ); break;
+		case 2 :  return NodeBoolIsNull::StaticUnserialise( arch ); break;
+		case 3 :  return NodeBoolNot::StaticUnserialise( arch ); break;
+		case 4 :  return NodeBoolAnd::StaticUnserialise( arch ); break;
+		default : check(false);
+		}
+
+		return 0;
+	}
+
+
+	//---------------------------------------------------------------------------------------------
+	const NODE_TYPE* NodeBool::GetType() const
+	{
+		return GetStaticType();
+	}
+
+
+	//---------------------------------------------------------------------------------------------
+	const NODE_TYPE* NodeBool::GetStaticType()
+	{
+		return &s_nodeBoolType;
+	}
+
+
+}
+
+
diff --git a/Engine/Plugins/Experimental/Mutable/Source/MutableTools/Private/NodeBoolAnd.cpp b/Engine/Plugins/Experimental/Mutable/Source/MutableTools/Private/NodeBoolAnd.cpp
new file mode 100644
--- /dev/null
+++ b/Engine/Plugins/Experimental/Mutable/Source/MutableTools/Private/NodeBoolAnd.cpp
@@ -0,0 +1,97 @@
+// Copyright Epic Games, Inc. All Rights Reserved.
+
+
+#include "MutableTools/Private/NodeBoolPrivate.h"
+
+#include "MutableRuntime/Private/MutableMath.h"
+#include "MutableRuntime/Private/MemoryPrivate.h"
+
+
+namespace mu
+{
+
+	//---------------------------------------------------------------------------------------------
+	// Static initialisation
+	//---------------------------------------------------------------------------------------------
+	NODE_TYPE NodeBoolAnd::Private::s_type =
+			NODE_TYPE( "BoolParameter", NodeBool::GetStaticType() );
+
+
+	//---------------------------------------------------------------------------------------------
+	//!
+	//---------------------------------------------------------------------------------------------
+
+	MUTABLE_IMPLEMENT_NODE( NodeBoolAnd, EType::And, Node, Node::EType::Bool);
+
+
+	//---------------------------------------------------------------------------------------------
+	// Node Interface
+	//---------------------------------------------------------------------------------------------
+	int NodeBoolAnd::GetInputCount() const
+	{
+		return 2;
+	}
+
+
+	//---------------------------------------------------------------------------------------------
+	Node* NodeBoolAnd::GetInputNode( int i ) const
+	{
+		check( i>=0 && i<GetInputCount() );
+		if ( i == 0)
+		{
+			return m_pD->m_pA.get();
+		}
+		else
+		{
+			return m_pD->m_pB.get();
+		}
+	}
+
+
+	//---------------------------------------------------------------------------------------------
+	void NodeBoolAnd::SetInputNode( int i, NodePtr p )
+	{
+		check( i>=0 && i<GetInputCount() );
+		if ( i == 0)
+		{
+			m_pD->m_pA = dynamic_cast<NodeBool*>( p.get() );
+		}
+		else
+		{
+			m_pD->m_pB = dynamic_cast<NodeBool*>( p.get() );
+		}
+	}
+
+
+	//---------------------------------------------------------------------------------------------
+	// Own Interface
+	//---------------------------------------------------------------------------------------------
+	NodeBoolPtr NodeBoolAnd::GetA() const
+	{
+		return m_pD->m_pA;
+	}
+
+
+	//---------------------------------------------------------------------------------------------
+	void NodeBoolAnd::SetA( NodeBoolPtr p )
+	{
+		m_pD->m_pA = p;
+	}
+
+
+	//---------------------------------------------------------------------------------------------
+	NodeBoolPtr NodeBoolAnd::GetB() const
+	{
+		return m_pD->m_pB;
+	}
+
+
+	//---------------------------------------------------------------------------------------------
+	void NodeBoolAnd::SetB( NodeBoolPtr p )
+	{
+		m_pD->m_pB = p;
+	}
+
+}
+
+
diff --git a/Engine/Plugins/Experimental/Mutable/Source/MutableTools/Private/NodeBoolConstant.cpp b/Engine/Plugins/Experimental/Mutable/Source/MutableTools/Private/NodeBoolConstant.cpp
new file mode 100644
--- /dev/null
+++ b/Engine/Plugins/Experimental/Mutable/Source/MutableTools/Private/NodeBoolConstant.cpp
@@ -0,0 +1,70 @@
+// Copyright Epic Games, Inc. All Rights Reserved.
+
+
+#include "NodeBoolPrivate.h"
+
+#include "MutableMath.h"
+
+
+#define NODE_INPUT_COUNT 	0
+
+
+namespace mu
+{
+
+	//---------------------------------------------------------------------------------------------
+	// Static initialisation
+	//---------------------------------------------------------------------------------------------
+	NODE_TYPE NodeBoolConstant::Private::s_type =
+			NODE_TYPE( "BoolConstant", NodeBool::GetStaticType() );
+
+
+	//---------------------------------------------------------------------------------------------
+	//!
+	//---------------------------------------------------------------------------------------------
+
+	MUTABLE_IMPLEMENT_NODE( NodeBoolConstant, EType::Constant, Node, Node::EType::Bool);
+
+
+	//---------------------------------------------------------------------------------------------
+	// Node Interface
+	//---------------------------------------------------------------------------------------------
+	int NodeBoolConstant::GetInputCount() const
+	{
+		return NODE_INPUT_COUNT;
+	}
+
+
+	//---------------------------------------------------------------------------------------------
+    Node* NodeBoolConstant::GetInputNode( int ) const
+	{
+		check( false );
+		return 0;
+	}
+
+
+	//---------------------------------------------------------------------------------------------
+    void NodeBoolConstant::SetInputNode( int, NodePtr )
+	{
+		check( false );
+	}
+
+
+	//---------------------------------------------------------------------------------------------
+	// Own Interface
+	//---------------------------------------------------------------------------------------------
+	bool NodeBoolConstant::GetValue() const
+	{
+		return m_pD->m_value;
+	}
+
+	//---------------------------------------------------------------------------------------------
+	void NodeBoolConstant::SetValue( bool v )
+	{
+		m_pD->m_value = v;
+	}
+
+
+}
+
+#undef NODE_INPUT_COUNT
diff --git a/Engine/Plugins/Experimental/Mutable/Source/MutableTools/Private/NodeBoolIsNull.cpp b/Engine/Plugins/Experimental/Mutable/Source/MutableTools/Private/NodeBoolIsNull.cpp
new file mode 100644
--- /dev/null
+++ b/Engine/Plugins/Experimental/Mutable/Source/MutableTools/Private/NodeBoolIsNull.cpp
@@ -0,0 +1,51 @@
+// Copyright Epic Games, Inc. All Rights Reserved.
+
+
+#include "NodeBoolPrivate.h"
+
+#include "MutableMath.h"
+#include "MemoryPrivate.h"
+
+
+namespace mu
+{
+
+	//---------------------------------------------------------------------------------------------
+	// Static initialisation
+	//---------------------------------------------------------------------------------------------
+	NODE_TYPE NodeBoolIsNull::Private::s_type =
+			NODE_TYPE( "BoolIsNull", NodeBool::GetStaticType() );
+
+
+	//---------------------------------------------------------------------------------------------
+	//!
+	//---------------------------------------------------------------------------------------------
+
+	MUTABLE_IMPLEMENT_NODE( NodeBoolIsNull, EType::IsNull, Node, Node::EType::Bool);
+
+
+	//---------------------------------------------------------------------------------------------
+	// Node Interface
+	//---------------------------------------------------------------------------------------------
+	int NodeBoolIsNull::GetInputCount() const
+	{
+		return 1;
+	}
+
+
+	//---------------------------------------------------------------------------------------------
+    Node* NodeBoolIsNull::GetInputNode( int ) const
+	{
+		return m_pD->m_pSource.get();
+	}
+
+
+	//---------------------------------------------------------------------------------------------
+    void NodeBoolIsNull::SetInputNode( int, NodePtr p )
+	{
+		m_pD->m_pSource = p;
+	}
+
+}
+
+
diff --git a/Engine/Plugins/Experimental/Mutable/Source/MutableTools/Private/NodeBoolNot.cpp b/Engine/Plugins/Experimental/Mutable/Source/MutableTools/Private/NodeBoolNot.cpp
new file mode 100644
--- /dev/null
+++ b/Engine/Plugins/Experimental/Mutable/Source/MutableTools/Private/NodeBoolNot.cpp
@@ -0,0 +1,71 @@
+// Copyright Epic Games, Inc. All Rights Reserved.
+
+
+#include "NodeBoolPrivate.h"
+
+#include "MutableMath.h"
+#include "MemoryPrivate.h"
+
+
+namespace mu
+{
+
+	//---------------------------------------------------------------------------------------------
+	// Static initialisation
+	//---------------------------------------------------------------------------------------------
+	NODE_TYPE NodeBoolNot::Private::s_type =
+			NODE_TYPE( "BoolParameter", NodeBool::GetStaticType() );
+
+
+	//---------------------------------------------------------------------------------------------
+	//!
+	//---------------------------------------------------------------------------------------------
+
+	MUTABLE_IMPLEMENT_NODE( NodeBoolNot, EType::Not, Node, Node::EType::Bool);
+
+
+	//---------------------------------------------------------------------------------------------
+	// Node Interface
+	//---------------------------------------------------------------------------------------------
+	int NodeBoolNot::GetInputCount() const
+	{
+		return 1;
+	}
+
+
+	//---------------------------------------------------------------------------------------------
+	Node* NodeBoolNot::GetInputNode( int i ) const
+	{
+		check( i>=0 && i<GetInputCount() );
+        (void)i;
+        return m_pD->m_pSource.get();
+	}
+
+
+	//---------------------------------------------------------------------------------------------
+	void NodeBoolNot::SetInputNode( int i, NodePtr p)
+	{
+		check( i>=0 && i<GetInputCount() );
+        (void)i;
+        m_pD->m_pSource = dynamic_cast<NodeBool*>( p.get() );
+	}
+
+
+	//---------------------------------------------------------------------------------------------
+	// Own Interface
+	//---------------------------------------------------------------------------------------------
+	NodeBoolPtr NodeBoolNot::GetInput() const
+	{
+		return m_pD->m_pSource;
+	}
+
+
+	//---------------------------------------------------------------------------------------------
+	void NodeBoolNot::SetInput( NodeBoolPtr p )
+	{
+		m_pD->m_pSource = p;
+	}
+
+}
+
+
diff --git a/Engine/Plugins/Experimental/Mutable/Source/MutableTools/Private/NodeBoolParameter.cpp b/Engine/Plugins/Experimental/Mutable/Source/MutableTools/Private/NodeBoolParameter.cpp
new file mode 100644
--- /dev/null
+++ b/Engine/Plugins/Experimental/Mutable/Source/MutableTools/Private/NodeBoolParameter.cpp
@@ -0,0 +1,116 @@
+// Copyright Epic Games, Inc. All Rights Reserved.
+
+
+#include "NodeBoolPrivate.h"
+
+#include "MutableMath.h"
+#include "MemoryPrivate.h"
+
+
+namespace mu
+{
+
+	//---------------------------------------------------------------------------------------------
+	// Static initialisation
+	//---------------------------------------------------------------------------------------------
+	NODE_TYPE NodeBoolParameter::Private::s_type =
+			NODE_TYPE( "BoolParameter", NodeBool::GetStaticType() );
+
+
+	//---------------------------------------------------------------------------------------------
+	//!
+	//---------------------------------------------------------------------------------------------
+
+	MUTABLE_IMPLEMENT_NODE( NodeBoolParameter, EType::Parameter, Node, Node::EType::Bool);
+
+
+	//---------------------------------------------------------------------------------------------
+	// Node Interface
+	//---------------------------------------------------------------------------------------------
+	int NodeBoolParameter::GetInputCount() const
+	{
+        return int( m_pD->m_ranges.size() );
+    }
+
+
+	//---------------------------------------------------------------------------------------------
+    Node* NodeBoolParameter::GetInputNode( int i ) const
+	{
+        check( i<GetInputCount() );
+        if (i<GetInputCount())
+        {
+            return m_pD->m_ranges[i].get();
+        }
+        return nullptr;
+	}
+
+
+	//---------------------------------------------------------------------------------------------
+    void NodeBoolParameter::SetInputNode( int i , NodePtr n )
+	{
+        check( i<GetInputCount() );
+        if (i<GetInputCount())
+        {
+            m_pD->m_ranges[i] = dynamic_cast<NodeRange*>(n.get());
+        }
+    }
+
+
+	//---------------------------------------------------------------------------------------------
+	// Own Interface
+	//---------------------------------------------------------------------------------------------
+	const char* NodeBoolParameter::GetName() const
+	{
+		return m_pD->m_name.c_str();
+	}
+
+
+	//---------------------------------------------------------------------------------------------
+	void NodeBoolParameter::SetName( const char* strName )
+	{
+		if ( strName )
+		{
+			m_pD->m_name = strName;
+		}
+		else
+		{
+			m_pD->m_name = "";
+		}
+	}
+
+
+	//---------------------------------------------------------------------------------------------
+	bool NodeBoolParameter::GetDefaultValue() const
+	{
+		return m_pD->m_defaultValue;
+	}
+
+
+	//---------------------------------------------------------------------------------------------
+	void NodeBoolParameter::SetDefaultValue( bool v )
+	{
+		m_pD->m_defaultValue = v;
+	}
+
+
+    //---------------------------------------------------------------------------------------------
+    void NodeBoolParameter::SetRangeCount( int i )
+    {
+        check(i>=0);
+        m_pD->m_ranges.resize(i);
+    }
+
+
+    //---------------------------------------------------------------------------------------------
+    void NodeBoolParameter::SetRange( int i, NodeRangePtr pRange )
+    {
+        check( i>=0 && i<int(m_pD->m_ranges.size()) );
+        if ( i>=0 && i<int(m_pD->m_ranges.size()) )
+        {
+            m_pD->m_ranges[i] = pRange;
+        }
+    }
+
+}
+
+
diff --git a/Engine/Plugins/Experimental/Mutable/Source/MutableTools/Private/NodeBoolPrivate.h b/Engine/Plugins/Experimental/Mutable/Source/MutableTools/Private/NodeBoolPrivate.h
new file mode 100644
--- /dev/null
+++ b/Engine/Plugins/Experimental/Mutable/Source/MutableTools/Private/NodeBoolPrivate.h
@@ -0,0 +1,199 @@
+// Copyright Epic Games, Inc. All Rights Reserved.
+
+#pragma once
+
+#include "NodePrivate.h"
+
+#include "NodeBool.h"
+#include "NodeImage.h"
+#include "NodeRange.h"
+#include "AST.h"
+
+
+namespace mu
+{
+
+
+	class NodeBool::Private : public Node::Private
+	{
+	};
+
+
+	class NodeBoolConstant::Private : public NodeBool::Private
+	{
+	public:
+
+		MUTABLE_DEFINE_CONST_VISITABLE();
+
+	public:
+
+		static NODE_TYPE s_type;
+
+		bool m_value;
+
+		//!
+		void Serialise( OutputArchive& arch ) const
+		{
+            uint32_t ver = 0;
+			arch << ver;
+
+			arch << m_value;
+		}
+
+		//!
+		void Unserialise( InputArchive& arch )
+		{
+            uint32_t ver;
+			arch >> ver;
+			check(ver==0);
+
+			arch >> m_value;
+		}
+	};
+
+
+	class NodeBoolParameter::Private : public NodeBool::Private
+	{
+	public:
+
+		MUTABLE_DEFINE_CONST_VISITABLE();
+
+	public:
+
+		static NODE_TYPE s_type;
+
+		bool m_defaultValue;
+		string m_name;
+		string m_uid;
+
+        vector<Ptr<NodeRange>> m_ranges;
+
+		//!
+		void Serialise( OutputArchive& arch ) const
+		{
+            uint32_t ver = 2;
+			arch << ver;
+
+			arch << m_defaultValue;
+			arch << m_name;
+			arch << m_uid;
+            arch << m_ranges;
+        }
+
+		//!
+		void Unserialise( InputArchive& arch )
+		{
+            uint32_t ver;
+			arch >> ver;
+            check(ver==2);
+
+			arch >> m_defaultValue;
+			arch >> m_name;
+            arch >> m_uid;
+            arch >> m_ranges;
+        }
+	};
+
+
+	class NodeBoolIsNull::Private : public NodeBool::Private
+	{
+	public:
+
+		MUTABLE_DEFINE_CONST_VISITABLE();
+
+	public:
+
+		static NODE_TYPE s_type;
+
+		NodePtr m_pSource;
+
+		//!
+		void Serialise( OutputArchive& arch ) const
+		{
+            uint32_t ver = 0;
+			arch << ver;
+
+			arch << m_pSource;
+		}
+
+		//!
+		void Unserialise( InputArchive& arch )
+		{
+            uint32_t ver;
+			arch >> ver;
+            check(ver<=0);
+
+			arch >> m_pSource;
+		}
+	};
+
+
+	class NodeBoolNot::Private : public NodeBool::Private
+	{
+	public:
+
+		MUTABLE_DEFINE_CONST_VISITABLE();
+
+	public:
+
+		static NODE_TYPE s_type;
+
+		NodeBoolPtr m_pSource;
+
+		//!
+		void Serialise( OutputArchive& arch ) const
+		{
+            uint32_t ver = 0;
+			arch << ver;
+
+			arch << m_pSource;
+		}
+
+		//!
+		void Unserialise( InputArchive& arch )
+		{
+            uint32_t ver;
+			arch >> ver;
+            check(ver<=0);
+
+			arch >> m_pSource;
+		}
+	};
+
+
+	class NodeBoolAnd::Private : public NodeBool::Private
+	{
+	public:
+
+		MUTABLE_DEFINE_CONST_VISITABLE();
+
+	public:
+
+		static NODE_TYPE s_type;
+
+		NodeBoolPtr m_pA;
+		NodeBoolPtr m_pB;
+
+		//!
+		void Serialise( OutputArchive& arch ) const
+		{
+            uint32_t ver = 0;
+			arch << ver;
+
+			arch << m_pA;
+			arch << m_pB;
+		}
+
+		//!
+		void Unserialise( InputArchive& arch )
+		{
+            uint32_t ver;
+			arch >> ver;
+            check(ver<=0);
+
+			arch >> m_pA;
+			arch >> m_pB;
+		}
+	};
+
+}
diff --git a/Engine/Plugins/Experimental/Mutable/Source/MutableTools/Private/NodeColour.cpp b/Engine/Plugins/Experimental/Mutable/Source/MutableTools/Private/NodeColour.cpp
new file mode 100644
--- /dev/null
+++ b/Engine/Plugins/Experimental/Mutable/Source/MutableTools/Private/NodeColour.cpp
@@ -0,0 +1,81 @@
+// Copyright Epic Games, Inc. All Rights Reserved.
+
+
+#include "NodeColour.h"
+#include "NodePrivate.h"
+
+#include "NodeColourConstant.h"
+#include "NodeColourParameter.h"
+#include "NodeColourSampleImage.h"
+#include "NodeColourTable.h"
+#include "NodeColourFromScalars.h"
+#include "NodeColourArithmeticOperation.h"
+#include "NodeColourSwitch.h"
+#include "NodeColourVariation.h"
+
+
+namespace mu
+{
+
+	//---------------------------------------------------------------------------------------------
+	// Static initialisation
+	//---------------------------------------------------------------------------------------------
+	static NODE_TYPE s_nodeColourType = NODE_TYPE( "NodeColour", Node::GetStaticType() );
+
+
+	//---------------------------------------------------------------------------------------------
+	void NodeColour::Serialise( const NodeColour* p, OutputArchive& arch )
+	{
+        uint32_t ver = 0;
+		arch << ver;
+
+		arch << uint32_t(p->Type);
+		p->SerialiseWrapper(arch);
+    }
+
+        
+	//---------------------------------------------------------------------------------------------
+	NodeColourPtr NodeColour::StaticUnserialise( InputArchive& arch )
+	{
+        uint32_t ver;
+		arch >> ver;
+		check( ver == 0 );
+
+        uint32_t id;
+		arch >> id;
+
+		switch (id)
+		{
+		case 0 :  return NodeColourConstant::StaticUnserialise( arch ); break;
+		case 1 :  return NodeColourParameter::StaticUnserialise( arch ); break;
+		case 2 :  return NodeColourSampleImage::StaticUnserialise( arch ); break;
+		case 3 :  return NodeColourTable::StaticUnserialise( arch ); break;
+		//case 4 :  return NodeColourImageSize::StaticUnserialise( arch ); break; // deprecated
+		case 5 :  return NodeColourFromScalars::StaticUnserialise( arch ); break;
+        case 6 :  return NodeColourArithmeticOperation::StaticUnserialise( arch ); break;
+        case 7 :  return NodeColourSwitch::StaticUnserialise( arch ); break;
+        case 8 :  return NodeColourVariation::StaticUnserialise( arch ); break;
+        default : check(false);
+		}
+
+		return nullptr;
+	}
+
+
+	//---------------------------------------------------------------------------------------------
+	const NODE_TYPE* NodeColour::GetType() const
+	{
+		return GetStaticType();
+	}
+
+
+	//---------------------------------------------------------------------------------------------
+	const NODE_TYPE* NodeColour::GetStaticType()
+	{
+		return &s_nodeColourType;
+	}
+
+
+}
+
+
diff --git a/Engine/Plugins/Experimental/Mutable/Source/MutableTools/Private/NodeColourArithmeticOperation.cpp b/Engine/Plugins/Experimental/Mutable/Source/MutableTools/Private/NodeColourArithmeticOperation.cpp
new file mode 100644
--- /dev/null
+++ b/Engine/Plugins/Experimental/Mutable/Source/MutableTools/Private/NodeColourArithmeticOperation.cpp
@@ -0,0 +1,124 @@
+// Copyright Epic Games, Inc. All Rights Reserved.
+
+
+#include "NodeColourArithmeticOperation.h"
+#include "NodeColourArithmeticOperationPrivate.h"
+
+#include "NodeScalar.h"
+
+#include "ImagePrivate.h"
+
+
+#define NODE_INPUT_COUNT 	2
+
+
+namespace mu
+{
+
+	//---------------------------------------------------------------------------------------------
+	// Static initialisation
+	//---------------------------------------------------------------------------------------------
+	NODE_TYPE NodeColourArithmeticOperation::Private::s_type =
+			NODE_TYPE( "ColourArithmenticOperation", NodeColour::GetStaticType() );
+
+	//---------------------------------------------------------------------------------------------
+	//!
+	//---------------------------------------------------------------------------------------------
+
+	MUTABLE_IMPLEMENT_NODE( NodeColourArithmeticOperation, EType::ArithmeticOperation, Node, Node::EType::Colour)
+
+
+    //---------------------------------------------------------------------------------------------
+    const char* NodeColourArithmeticOperation::s_opTypeName[] =
+    {
+        "Add",
+        "Subtract",
+        "Multiply",
+        "Divide"
+    };
+
+
+	//---------------------------------------------------------------------------------------------
+	// Node Interface
+	//---------------------------------------------------------------------------------------------
+	int NodeColourArithmeticOperation::GetInputCount() const
+	{
+		return NODE_INPUT_COUNT;
+	}
+
+
+	//---------------------------------------------------------------------------------------------
+	Node* NodeColourArithmeticOperation::GetInputNode( int i ) const
+	{
+		check( i>=0 && i<NODE_INPUT_COUNT );
+
+		Node* pResult = 0;
+
+		switch (i)
+		{
+		case 0: pResult = m_pD->m_pA.get(); break;
+		case 1: pResult = m_pD->m_pB.get(); break;
+		}
+
+		return pResult;
+	}
+
+
+	//---------------------------------------------------------------------------------------------
+	void NodeColourArithmeticOperation::SetInputNode( int i, NodePtr pNode )
+	{
+		check( i>=0 && i<NODE_INPUT_COUNT );
+
+		switch (i)
+		{
+		case 0: m_pD->m_pA = dynamic_cast<NodeColour*>(pNode.get()); break;
+		case 1: m_pD->m_pB = dynamic_cast<NodeColour*>(pNode.get()); break;
+		}
+	}
+
+
+	//---------------------------------------------------------------------------------------------
+	// Own Interface
+	//---------------------------------------------------------------------------------------------
+	NodeColourArithmeticOperation::OPERATION NodeColourArithmeticOperation::GetOperation() const
+	{
+		return m_pD->m_operation;
+	}
+
+
+	//---------------------------------------------------------------------------------------------
+	void NodeColourArithmeticOperation::SetOperation(NodeColourArithmeticOperation::OPERATION o)
+	{
+		m_pD->m_operation = o;
+	}
+
+
+	NodeColourPtr NodeColourArithmeticOperation::GetA() const
+	{
+		return m_pD->m_pA.get();
+	}
+
+
+	//---------------------------------------------------------------------------------------------
+	void NodeColourArithmeticOperation::SetA(NodeColourPtr pNode)
+	{
+		m_pD->m_pA = pNode;
+	}
+
+
+	//---------------------------------------------------------------------------------------------
+	NodeColourPtr NodeColourArithmeticOperation::GetB() const
+	{
+		return m_pD->m_pB.get();
+	}
+
+
+	//---------------------------------------------------------------------------------------------
+	void NodeColourArithmeticOperation::SetB( NodeColourPtr pNode )
+	{
+		m_pD->m_pB = pNode;
+	}
+
+}
+
+#undef NODE_INPUT_COUNT
diff --git a/Engine/Plugins/Experimental/Mutable/Source/MutableTools/Private/NodeColourArithmeticOperationPrivate.h b/Engine/Plugins/Experimental/Mutable/Source/MutableTools/Private/NodeColourArithmeticOperationPrivate.h
new file mode 100644
--- /dev/null
+++ b/Engine/Plugins/Experimental/Mutable/Source/MutableTools/Private/NodeColourArithmeticOperationPrivate.h
@@ -0,0 +1,55 @@
+// Copyright Epic Games, Inc. All Rights Reserved.
+
+#pragma once
+
+#include "NodeColourPrivate.h"
+#include "NodeColourArithmeticOperation.h"
+
+#include "MemoryPrivate.h"
+
+
+namespace mu
+{
+
+	MUTABLE_DEFINE_ENUM_SERIALISABLE(NodeColourArithmeticOperation::OPERATION)
+
+	class NodeColourArithmeticOperation::Private : public NodeColour::Private
+	{
+	public:
+
+		MUTABLE_DEFINE_CONST_VISITABLE()
+
+	public:
+
+		static NODE_TYPE s_type;
+
+		NodeColourArithmeticOperation::OPERATION m_operation;
+		NodeColourPtr m_pA;
+		NodeColourPtr m_pB;
+
+		//!
+		void Serialise( OutputArchive& arch ) const
+		{
+            uint32_t ver = 0;
+			arch << ver;
+
+			arch << m_operation;
+			arch << m_pA;
+			arch << m_pB;
+		}
+
+		//!
+		void Unserialise( InputArchive& arch )
+		{
+            uint32_t ver;
+			arch >> ver;
+			check(ver==0);
+
+			arch >> m_operation;
+			arch >> m_pA;
+			arch >> m_pB;
+		}
+	};
+
+
+}
diff --git a/Engine/Plugins/Experimental/Mutable/Source/MutableTools/Private/NodeColourConstant.cpp b/Engine/Plugins/Experimental/Mutable/Source/MutableTools/Private/NodeColourConstant.cpp
new file mode 100644
--- /dev/null
+++ b/Engine/Plugins/Experimental/Mutable/Source/MutableTools/Private/NodeColourConstant.cpp
@@ -0,0 +1,85 @@
+// Copyright Epic Games, Inc. All Rights Reserved.
+
+
+#include "NodeColourConstant.h"
+#include "NodeColourConstantPrivate.h"
+
+#include "MutableMath.h"
+
+
+#define NODE_INPUT_COUNT 	0
+
+
+namespace mu
+{
+
+	//---------------------------------------------------------------------------------------------
+	// Static initialisation
+	//---------------------------------------------------------------------------------------------
+	NODE_TYPE NodeColourConstant::Private::s_type =
+			NODE_TYPE( "ColourConstant", NodeColour::GetStaticType() );
+
+
+	//---------------------------------------------------------------------------------------------
+	//!
+	//---------------------------------------------------------------------------------------------
+
+	MUTABLE_IMPLEMENT_NODE( NodeColourConstant, EType::Constant, Node, Node::EType::Colour)
+
+
+	//---------------------------------------------------------------------------------------------
+	// Node Interface
+	//---------------------------------------------------------------------------------------------
+	int NodeColourConstant::GetInputCount() const
+	{
+		return NODE_INPUT_COUNT;
+	}
+
+
+	//---------------------------------------------------------------------------------------------
+    Node* NodeColourConstant::GetInputNode( int ) const
+	{
+		check( false );
+		return 0;
+	}
+
+
+	//---------------------------------------------------------------------------------------------
+    void NodeColourConstant::SetInputNode( int, NodePtr )
+	{
+		check( false );
+	}
+
+
+	//---------------------------------------------------------------------------------------------
+	// Own Interface
+	//---------------------------------------------------------------------------------------------
+	void NodeColourConstant::GetValue( float* pR, float *pG, float* pB ) const
+	{
+		if (pR)
+		{
+			*pR = m_pD->m_value[0];
+		}
+
+		if (pG)
+		{
+			*pG = m_pD->m_value[1];
+		}
+
+		if (pB)
+		{
+			*pB = m_pD->m_value[2];
+		}
+	}
+
+
+	//---------------------------------------------------------------------------------------------
+	void NodeColourConstant::SetValue( float r, float g, float b )
+	{
+		m_pD->m_value = vec3<float>(r,g,b);
+	}
+
+
+}
+
+#undef NODE_INPUT_COUNT
diff --git a/Engine/Plugins/Experimental/Mutable/Source/MutableTools/Private/NodeColourConstantPrivate.h b/Engine/Plugins/Experimental/Mutable/Source/MutableTools/Private/NodeColourConstantPrivate.h
new file mode 100644
--- /dev/null
+++ b/Engine/Plugins/Experimental/Mutable/Source/MutableTools/Private/NodeColourConstantPrivate.h
@@ -0,0 +1,40 @@
+// Copyright Epic Games, Inc. All Rights Reserved.
+
+#pragma once
+
+#include "NodeColourConstant.h"
+#include "NodeColourPrivate.h"
+
+#include "MutableMath.h"
+
+
+namespace mu
+{
+
+	class NodeColourConstant::Private : public NodeColour::Private
+	{
+	public:
+
+		MUTABLE_DEFINE_CONST_VISITABLE()
+
+	public:
+
+		static NODE_TYPE s_type;
+
+		vec3<float> m_value;
+
+		//!
+		void Serialise( OutputArchive& arch ) const
+		{
+			arch << m_value;
+		}
+
+		//!
+		void Unserialise( InputArchive& arch )
+		{
+			arch >> m_value;
+		}
+
+	};
+
+}
diff --git a/Engine/Plugins/Experimental/Mutable/Source/MutableTools/Private/NodeColourFromScalars.cpp b/Engine/Plugins/Experimental/Mutable/Source/MutableTools/Private/NodeColourFromScalars.cpp
new file mode 100644
--- /dev/null
+++ b/Engine/Plugins/Experimental/Mutable/Source/MutableTools/Private/NodeColourFromScalars.cpp
@@ -0,0 +1,134 @@
+// Copyright Epic Games, Inc. All Rights Reserved.
+
+
+#include "NodeColourFromScalars.h"
+#include "NodeColourFromScalarsPrivate.h"
+
+#include "NodeScalar.h"
+
+#include "ImagePrivate.h"
+
+
+#define NODE_INPUT_COUNT 	4
+
+
+namespace mu
+{
+
+	//---------------------------------------------------------------------------------------------
+	// Static initialisation
+	//---------------------------------------------------------------------------------------------
+	NODE_TYPE NodeColourFromScalars::Private::s_type =
+			NODE_TYPE( "ColourFromScalars", NodeColour::GetStaticType() );
+
+
+	//---------------------------------------------------------------------------------------------
+	//!
+	//---------------------------------------------------------------------------------------------
+
+	MUTABLE_IMPLEMENT_NODE( NodeColourFromScalars, EType::FromScalars, Node, Node::EType::Colour)
+
+
+	//---------------------------------------------------------------------------------------------
+	// Node Interface
+	//---------------------------------------------------------------------------------------------
+	int NodeColourFromScalars::GetInputCount() const
+	{
+		return NODE_INPUT_COUNT;
+	}
+
+
+	//---------------------------------------------------------------------------------------------
+	Node* NodeColourFromScalars::GetInputNode( int i ) const
+	{
+		check( i>=0 && i<NODE_INPUT_COUNT );
+
+		Node* pResult = 0;
+
+		switch (i)
+		{
+		case 0: pResult = m_pD->m_pX.get(); break;
+		case 1: pResult = m_pD->m_pY.get(); break;
+		case 2: pResult = m_pD->m_pZ.get(); break;
+		case 3: pResult = m_pD->m_pW.get(); break;
+		}
+
+		return pResult;
+	}
+
+
+	//---------------------------------------------------------------------------------------------
+	void NodeColourFromScalars::SetInputNode( int i, NodePtr pNode )
+	{
+		check( i>=0 && i<NODE_INPUT_COUNT );
+
+		switch (i)
+		{
+		case 0: m_pD->m_pX = dynamic_cast<NodeScalar*>(pNode.get()); break;
+		case 1: m_pD->m_pY = dynamic_cast<NodeScalar*>(pNode.get()); break;
+		case 2: m_pD->m_pZ = dynamic_cast<NodeScalar*>(pNode.get()); break;
+		case 3: m_pD->m_pW = dynamic_cast<NodeScalar*>(pNode.get()); break;
+		}
+	}
+
+
+	//---------------------------------------------------------------------------------------------
+	// Own Interface
+	//---------------------------------------------------------------------------------------------
+	NodeScalarPtr NodeColourFromScalars::GetX() const
+	{
+		return m_pD->m_pX.get();
+	}
+
+
+	//---------------------------------------------------------------------------------------------
+	void NodeColourFromScalars::SetX( NodeScalarPtr pNode )
+	{
+		m_pD->m_pX = pNode;
+	}
+
+
+	//---------------------------------------------------------------------------------------------
+	NodeScalarPtr NodeColourFromScalars::GetY() const
+	{
+		return m_pD->m_pY.get();
+	}
+
+
+	//---------------------------------------------------------------------------------------------
+	void NodeColourFromScalars::SetY( NodeScalarPtr pNode )
+	{
+		m_pD->m_pY = pNode;
+	}
+
+
+	//---------------------------------------------------------------------------------------------
+	NodeScalarPtr NodeColourFromScalars::GetZ() const
+	{
+		return m_pD->m_pZ.get();
+	}
+
+
+	//---------------------------------------------------------------------------------------------
+	void NodeColourFromScalars::SetZ( NodeScalarPtr pNode )
+	{
+		m_pD->m_pZ = pNode;
+	}
+
+
+	//---------------------------------------------------------------------------------------------
+	NodeScalarPtr NodeColourFromScalars::GetW() const
+	{
+		return m_pD->m_pW.get();
+	}
+
+
+	//---------------------------------------------------------------------------------------------
+	void NodeColourFromScalars::SetW( NodeScalarPtr pNode )
+	{
+		m_pD->m_pW = pNode;
+	}
+
+}
+
+#undef NODE_INPUT_COUNT
diff --git a/Engine/Plugins/Experimental/Mutable/Source/MutableTools/Private/NodeColourFromScalarsPrivate.h b/Engine/Plugins/Experimental/Mutable/Source/MutableTools/Private/NodeColourFromScalarsPrivate.h
new file mode 100644
--- /dev/null
+++ b/Engine/Plugins/Experimental/Mutable/Source/MutableTools/Private/NodeColourFromScalarsPrivate.h
@@ -0,0 +1,60 @@
+// Copyright Epic Games, Inc. All Rights Reserved.
+
+#pragma once
+
+#include "NodeColourPrivate.h"
+#include "NodeColourFromScalars.h"
+#include "NodeScalar.h"
+#include "NodeImage.h"
+#include "AST.h"
+
+#include "MemoryPrivate.h"
+
+
+namespace mu
+{
+
+
+	class NodeColourFromScalars::Private : public NodeColour::Private
+	{
+	public:
+
+		MUTABLE_DEFINE_CONST_VISITABLE()
+
+	public:
+
+		static NODE_TYPE s_type;
+
+		NodeScalarPtr m_pX;
+		NodeScalarPtr m_pY;
+		NodeScalarPtr m_pZ;
+		NodeScalarPtr m_pW;
+
+		//!
+		void Serialise( OutputArchive& arch ) const
+		{
+            uint32_t ver = 0;
+			arch << ver;
+
+			arch << m_pX;
+			arch << m_pY;
+			arch << m_pZ;
+			arch << m_pW;
+		}
+
+		//!
+		void Unserialise( InputArchive& arch )
+		{
+            uint32_t ver;
+			arch >> ver;
+			check(ver==0);
+
+			arch >> m_pX;
+			arch >> m_pY;
+			arch >> m_pZ;
+			arch >> m_pW;
+		}
+	};
+
+
+}
diff --git a/Engine/Plugins/Experimental/Mutable/Source/MutableTools/Private/NodeColourParameter.cpp b/Engine/Plugins/Experimental/Mutable/Source/MutableTools/Private/NodeColourParameter.cpp
new file mode 100644
--- /dev/null
+++ b/Engine/Plugins/Experimental/Mutable/Source/MutableTools/Private/NodeColourParameter.cpp
@@ -0,0 +1,139 @@
+// Copyright Epic Games, Inc. All Rights Reserved.
+
+
+#include "NodeColourParameter.h"
+#include "NodeColourParameterPrivate.h"
+
+#include "MutableMath.h"
+#include "MemoryPrivate.h"
+
+
+namespace mu
+{
+
+	//---------------------------------------------------------------------------------------------
+	// Static initialisation
+	//---------------------------------------------------------------------------------------------
+	NODE_TYPE NodeColourParameter::Private::s_type =
+			NODE_TYPE( "ColourParameter", NodeColour::GetStaticType() );
+
+
+	//---------------------------------------------------------------------------------------------
+	//!
+	//---------------------------------------------------------------------------------------------
+
+	MUTABLE_IMPLEMENT_NODE( NodeColourParameter, EType::Parameter, Node, Node::EType::Colour)
+
+
+	//---------------------------------------------------------------------------------------------
+	// Node Interface
+	//---------------------------------------------------------------------------------------------
+	int NodeColourParameter::GetInputCount() const
+	{
+        return int( m_pD->m_ranges.size() );
+    }
+
+
+	//---------------------------------------------------------------------------------------------
+    Node* NodeColourParameter::GetInputNode( int i ) const
+	{
+        check( i<GetInputCount() );
+        if (i<GetInputCount())
+        {
+            return m_pD->m_ranges[i].get();
+        }
+        return nullptr;
+    }
+
+
+	//---------------------------------------------------------------------------------------------
+    void NodeColourParameter::SetInputNode( int i, NodePtr n )
+	{
+        check( i<GetInputCount() );
+        if (i<GetInputCount())
+        {
+            m_pD->m_ranges[i] = dynamic_cast<NodeRange*>(n.get());
+        }
+    }
+
+
+	//---------------------------------------------------------------------------------------------
+	// Own Interface
+	//---------------------------------------------------------------------------------------------
+	const char* NodeColourParameter::GetName() const
+	{
+		return m_pD->m_name.c_str();
+	}
+
+
+	//---------------------------------------------------------------------------------------------
+	void NodeColourParameter::SetName( const char* strName )
+	{
+		if ( strName )
+		{
+			m_pD->m_name = strName;
+		}
+		else
+		{
+			m_pD->m_name = "";
+		}
+	}
+
+
+	const char* NodeColourParameter::GetUid() const
+	{
+		return m_pD->m_uid.c_str();
+	}
+
+
+	//---------------------------------------------------------------------------------------------
+	void NodeColourParameter::SetUid( const char* strUid )
+	{
+		if ( strUid )
+		{
+			m_pD->m_uid = strUid;
+		}
+		else
+		{
+			m_pD->m_uid = "";
+		}
+	}
+
+
+	//---------------------------------------------------------------------------------------------
+	void NodeColourParameter::GetDefaultValue( float* pR, float* pG, float* pB ) const
+	{
+		if (*pR) *pR = m_pD->m_defaultValue[0];
+		if (*pG) *pG = m_pD->m_defaultValue[1];
+		if (*pB) *pB = m_pD->m_defaultValue[2];
+	}
+
+
+	//---------------------------------------------------------------------------------------------
+	void NodeColourParameter::SetDefaultValue( float r, float g, float b )
+	{
+		m_pD->m_defaultValue = vec3<float>( r, g, b );
+	}
+
+
+    //---------------------------------------------------------------------------------------------
+    void NodeColourParameter::SetRangeCount( int i )
+    {
+        check(i>=0);
+        m_pD->m_ranges.resize(i);
+    }
+
+
+    //---------------------------------------------------------------------------------------------
+    void NodeColourParameter::SetRange( int i, NodeRangePtr pRange )
+    {
+        check( i>=0 && i<int(m_pD->m_ranges.size()) );
+        if ( i>=0 && i<int(m_pD->m_ranges.size()) )
+        {
+            m_pD->m_ranges[i] = pRange;
+        }
+    }
+
+}
+
+
diff --git a/Engine/Plugins/Experimental/Mutable/Source/MutableTools/Private/NodeColourParameterPrivate.h b/Engine/Plugins/Experimental/Mutable/Source/MutableTools/Private/NodeColourParameterPrivate.h
new file mode 100644
--- /dev/null
+++ b/Engine/Plugins/Experimental/Mutable/Source/MutableTools/Private/NodeColourParameterPrivate.h
@@ -0,0 +1,57 @@
+// Copyright Epic Games, Inc. All Rights Reserved.
+
+#pragma once
+
+#include "NodeColourParameter.h"
+#include "NodeColourPrivate.h"
+
+#include "NodeRange.h"
+#include "MutableMath.h"
+
+
+namespace mu
+{
+
+	class NodeColourParameter::Private : public NodeColour::Private
+	{
+	public:
+
+		MUTABLE_DEFINE_CONST_VISITABLE()
+
+	public:
+
+		static NODE_TYPE s_type;
+
+		vec3<float> m_defaultValue;
+		string m_name;
+		string m_uid;
+
+        vector<Ptr<NodeRange>> m_ranges;
+
+		//!
+		void Serialise( OutputArchive& arch ) const
+		{
+            uint32_t ver = 2;
+			arch << ver;
+
+			arch << m_defaultValue;
+			arch << m_name;
+            arch << m_uid;
+            arch << m_ranges;
+        }
+
+		//!
+		void Unserialise( InputArchive& arch )
+		{
+            uint32_t ver;
+			arch >> ver;
+            check(ver==2);
+
+			arch >> m_defaultValue;
+			arch >> m_name;
+            arch >> m_uid;
+            arch >> m_ranges;
+        }
+	};
+
+}
diff --git a/Engine/Plugins/Experimental/Mutable/Source/MutableTools/Private/NodeColourPrivate.h b/Engine/Plugins/Experimental/Mutable/Source/MutableTools/Private/NodeColourPrivate.h
new file mode 100644
--- /dev/null
+++ b/Engine/Plugins/Experimental/Mutable/Source/MutableTools/Private/NodeColourPrivate.h
@@ -0,0 +1,18 @@
+// Copyright Epic Games, Inc. All Rights Reserved.
+
+#pragma once
+
+#include "NodePrivate.h"
+#include "NodeColour.h"
+
+
+namespace mu
+{
+
+
+	class NodeColour::Private : public Node::Private
+	{
+	};
+
+
+}
diff --git a/Engine/Plugins/Experimental/Mutable/Source/MutableTools/Private/NodeColourSampleImage.cpp b/Engine/Plugins/Experimental/Mutable/Source/MutableTools/Private/NodeColourSampleImage.cpp
new file mode 100644
--- /dev/null
+++ b/Engine/Plugins/Experimental/Mutable/Source/MutableTools/Private/NodeColourSampleImage.cpp
@@ -0,0 +1,121 @@
+// Copyright Epic Games, Inc. All Rights Reserved.
+
+
+#include "NodeColourSampleImage.h"
+#include "NodeColourSampleImagePrivate.h"
+
+#include "NodeScalar.h"
+
+#include "ImagePrivate.h"
+
+
+#define NODE_INPUT_COUNT 	3
+
+
+namespace mu
+{
+
+	//---------------------------------------------------------------------------------------------
+	// Static initialisation
+	//---------------------------------------------------------------------------------------------
+	NODE_TYPE NodeColourSampleImage::Private::s_type =
+			NODE_TYPE( "ColourSampleImage", NodeColour::GetStaticType() );
+
+
+	//---------------------------------------------------------------------------------------------
+	//!
+	//---------------------------------------------------------------------------------------------
+
+	MUTABLE_IMPLEMENT_NODE( NodeColourSampleImage, EType::SampleImage, Node, Node::EType::Colour)
+
+
+	//---------------------------------------------------------------------------------------------
+	// Node Interface
+	//---------------------------------------------------------------------------------------------
+	int NodeColourSampleImage::GetInputCount() const
+	{
+		return NODE_INPUT_COUNT;
+	}
+
+
+	//---------------------------------------------------------------------------------------------
+	Node* NodeColourSampleImage::GetInputNode( int i ) const
+	{
+		check( i>=0 && i<NODE_INPUT_COUNT );
+
+		Node* pResult = 0;
+
+		switch (i)
+		{
+		case 0: pResult = m_pD->m_pImage.get(); break;
+		case 1: pResult = m_pD->m_pX.get(); break;
+		case 2: pResult = m_pD->m_pY.get(); break;
+		}
+
+		return pResult;
+	}
+
+
+	//---------------------------------------------------------------------------------------------
+	void NodeColourSampleImage::SetInputNode( int i, NodePtr pNode )
+	{
+		check( i>=0 && i<NODE_INPUT_COUNT );
+
+		switch (i)
+		{
+		case 0: m_pD->m_pImage = dynamic_cast<NodeImage*>(pNode.get()); break;
+		case 1: m_pD->m_pX = dynamic_cast<NodeScalar*>(pNode.get()); break;
+		case 2: m_pD->m_pY = dynamic_cast<NodeScalar*>(pNode.get()); break;
+		}
+	}
+
+
+	//---------------------------------------------------------------------------------------------
+	// Own Interface
+	//---------------------------------------------------------------------------------------------
+	NodeScalarPtr NodeColourSampleImage::GetX() const
+	{
+		return m_pD->m_pX.get();
+	}
+
+
+	//---------------------------------------------------------------------------------------------
+	void NodeColourSampleImage::SetX( NodeScalarPtr pNode )
+	{
+		m_pD->m_pX = pNode;
+	}
+
+
+	//---------------------------------------------------------------------------------------------
+	NodeScalarPtr NodeColourSampleImage::GetY() const
+	{
+		return m_pD->m_pY.get();
+	}
+
+
+	//---------------------------------------------------------------------------------------------
+	void NodeColourSampleImage::SetY( NodeScalarPtr pNode )
+	{
+		m_pD->m_pY = pNode;
+	}
+
+
+	//---------------------------------------------------------------------------------------------
+	NodeImagePtr NodeColourSampleImage::GetImage() const
+	{
+		return m_pD->m_pImage.get();
+	}
+
+
+	//---------------------------------------------------------------------------------------------
+	void NodeColourSampleImage::SetImage( NodeImagePtr pNode )
+	{
+		m_pD->m_pImage = pNode;
+	}
+
+
+
+}
+
+#undef NODE_INPUT_COUNT
+
diff --git a/Engine/Plugins/Experimental/Mutable/Source/MutableTools/Private/NodeColourSampleImagePrivate.h b/Engine/Plugins/Experimental/Mutable/Source/MutableTools/Private/NodeColourSampleImagePrivate.h
new file mode 100644
--- /dev/null
+++ b/Engine/Plugins/Experimental/Mutable/Source/MutableTools/Private/NodeColourSampleImagePrivate.h
@@ -0,0 +1,57 @@
+// Copyright Epic Games, Inc. All Rights Reserved.
+
+#pragma once
+
+#include "NodeColourPrivate.h"
+#include "NodeColourSampleImage.h"
+#include "NodeScalar.h"
+#include "NodeImage.h"
+#include "AST.h"
+
+#include "MemoryPrivate.h"
+
+
+namespace mu
+{
+
+
+	class NodeColourSampleImage::Private : public NodeColour::Private
+	{
+	public:
+
+		MUTABLE_DEFINE_CONST_VISITABLE()
+
+	public:
+
+		static NODE_TYPE s_type;
+
+		NodeImagePtr m_pImage;
+		NodeScalarPtr m_pX;
+		NodeScalarPtr m_pY;
+
+		//!
+		void Serialise( OutputArchive& arch ) const
+		{
+            uint32_t ver = 0;
+			arch << ver;
+
+			arch << m_pImage;
+			arch << m_pX;
+			arch << m_pY;
+		}
+
+		//!
+		void Unserialise( InputArchive& arch )
+		{
+            uint32_t ver;
+			arch >> ver;
+			check(ver==0);
+
+			arch >> m_pImage;
+			arch >> m_pX;
+			arch >> m_pY;
+		}
+	};
+
+
+}
diff --git a/Engine/Plugins/Experimental/Mutable/Source/MutableTools/Private/NodeColourSwitch.cpp b/Engine/Plugins/Experimental/Mutable/Source/MutableTools/Private/NodeColourSwitch.cpp
new file mode 100644
--- /dev/null
+++ b/Engine/Plugins/Experimental/Mutable/Source/MutableTools/Private/NodeColourSwitch.cpp
@@ -0,0 +1,118 @@
+// Copyright Epic Games, Inc. All Rights Reserved.
+
+
+#include "NodeColourSwitch.h"
+#include "NodeColourSwitchPrivate.h"
+
+#include "NodeScalar.h"
+
+
+namespace mu
+{
+
+	//---------------------------------------------------------------------------------------------
+	// Static initialisation
+	//---------------------------------------------------------------------------------------------
+    NODE_TYPE NodeColourSwitch::Private::s_type =
+            NODE_TYPE( "ColourSwitch", NodeColour::GetStaticType() );
+
+
+	//---------------------------------------------------------------------------------------------
+	//!
+	//---------------------------------------------------------------------------------------------
+
+    MUTABLE_IMPLEMENT_NODE( NodeColourSwitch, EType::Switch, Node, Node::EType::Colour)
+
+
+	//---------------------------------------------------------------------------------------------
+	// Node Interface
+	//---------------------------------------------------------------------------------------------
+    int NodeColourSwitch::GetInputCount() const
+	{
+		return 1 + (int)m_pD->m_options.size();
+	}
+
+
+	//---------------------------------------------------------------------------------------------
+    Node* NodeColourSwitch::GetInputNode( int i ) const
+	{
+		check( i>=0 && i<GetInputCount() );
+
+		Node* pResult = 0;
+
+		switch (i)
+		{
+		case 0:
+			pResult = m_pD->m_pParameter.get();
+			break;
+
+		default:
+			pResult = m_pD->m_options[i-1].get();
+			break;
+		}
+
+		return pResult;
+	}
+
+
+	//---------------------------------------------------------------------------------------------
+    void NodeColourSwitch::SetInputNode( int i, NodePtr pNode )
+	{
+		check( i>=0 && i<GetInputCount() );
+
+		switch (i)
+		{
+		case 0:
+			m_pD->m_pParameter = dynamic_cast<NodeScalar*>(pNode.get());
+			break;
+
+		default:
+            m_pD->m_options[i-1] = dynamic_cast<NodeColour*>(pNode.get());
+			break;
+		}
+	}
+
+
+	//---------------------------------------------------------------------------------------------
+	// Own Interface
+	//---------------------------------------------------------------------------------------------
+    NodeScalarPtr NodeColourSwitch::GetParameter() const
+	{
+		return m_pD->m_pParameter.get();
+	}
+
+
+	//---------------------------------------------------------------------------------------------
+    void NodeColourSwitch::SetParameter( NodeScalarPtr pNode )
+	{
+		m_pD->m_pParameter = pNode;
+	}
+
+
+	//---------------------------------------------------------------------------------------------
+    void NodeColourSwitch::SetOptionCount( int t )
+	{
+		m_pD->m_options.resize(t);
+	}
+
+
+	//---------------------------------------------------------------------------------------------
+    NodeColourPtr NodeColourSwitch::GetOption( int t ) const
+	{
+		check( t>=0 && t<(int)m_pD->m_options.size() );
+		return m_pD->m_options[t].get();
+	}
+
+
+	//---------------------------------------------------------------------------------------------
+    void NodeColourSwitch::SetOption( int t, NodeColourPtr pNode )
+	{
+		check( t>=0 && t<(int)m_pD->m_options.size() );
+		m_pD->m_options[t] = pNode;
+	}
+
+
+
+}
+
+
diff --git a/Engine/Plugins/Experimental/Mutable/Source/MutableTools/Private/NodeColourSwitchPrivate.h b/Engine/Plugins/Experimental/Mutable/Source/MutableTools/Private/NodeColourSwitchPrivate.h
new file mode 100644
--- /dev/null
+++ b/Engine/Plugins/Experimental/Mutable/Source/MutableTools/Private/NodeColourSwitchPrivate.h
@@ -0,0 +1,52 @@
+// Copyright Epic Games, Inc. All Rights Reserved.
+
+#pragma once
+
+#include "NodeColourSwitch.h"
+#include "NodeColourPrivate.h"
+#include "NodeScalar.h"
+#include "AST.h"
+
+
+namespace mu
+{
+
+
+    class NodeColourSwitch::Private : public NodeColour::Private
+	{
+	public:
+
+		MUTABLE_DEFINE_CONST_VISITABLE()
+
+	public:
+
+		static NODE_TYPE s_type;
+
+		NodeScalarPtr m_pParameter;
+        vector<NodeColourPtr> m_options;
+
+		//!
+		void Serialise( OutputArchive& arch ) const
+		{
+            uint32_t ver = 0;
+			arch << ver;
+
+			arch << m_pParameter;
+			arch << m_options;
+		}
+
+		//!
+		void Unserialise( InputArchive& arch )
+		{
+            uint32_t ver;
+			arch >> ver;
+			check(ver==0);
+
+			arch >> m_pParameter;
+			arch >> m_options;
+		}
+	};
+
+
+}
+
diff --git a/Engine/Plugins/Experimental/Mutable/Source/MutableTools/Private/NodeColourTable.cpp b/Engine/Plugins/Experimental/Mutable/Source/MutableTools/Private/NodeColourTable.cpp
new file mode 100644
--- /dev/null
+++ b/Engine/Plugins/Experimental/Mutable/Source/MutableTools/Private/NodeColourTable.cpp
@@ -0,0 +1,100 @@
+// Copyright Epic Games, Inc. All Rights Reserved.
+
+
+#include "NodeColourTable.h"
+#include "NodeColourTablePrivate.h"
+
+
+
+namespace mu
+{
+
+
+	//---------------------------------------------------------------------------------------------
+	// Static initialisation
+	//---------------------------------------------------------------------------------------------
+	NODE_TYPE NodeColourTable::Private::s_type =
+			NODE_TYPE( "TableColour", NodeColour::GetStaticType() );
+
+
+	//---------------------------------------------------------------------------------------------
+	//!
+	//---------------------------------------------------------------------------------------------
+
+	MUTABLE_IMPLEMENT_NODE( NodeColourTable, EType::Table, Node, Node::EType::Colour)
+
+
+	//---------------------------------------------------------------------------------------------
+	// Node Interface
+	//---------------------------------------------------------------------------------------------
+	int NodeColourTable::GetInputCount() const
+	{
+		return 0;
+	}
+
+
+	//---------------------------------------------------------------------------------------------
+	Node* NodeColourTable::GetInputNode( int i ) const
+	{
+		check( i >=0 && i < 0 );
+        (void)i;
+        return 0;
+	}
+
+
+	//---------------------------------------------------------------------------------------------
+    void NodeColourTable::SetInputNode( int i, NodePtr )
+	{
+		check( i >=0 && i < 0 );
+        (void)i;
+        //m_pD->m_pObject = dynamic_cast<NodeObject*>( pNode.get() );
+	}
+
+
+	//---------------------------------------------------------------------------------------------
+	// Own Interface
+	//---------------------------------------------------------------------------------------------
+	const char* NodeColourTable::GetColumn() const
+	{
+		return m_pD->m_columnName.c_str();
+	}
+
+
+	//---------------------------------------------------------------------------------------------
+	void NodeColourTable::SetColumn( const char* strName )
+	{
+		if (strName)
+		{
+			m_pD->m_columnName = strName;
+		}
+		else
+		{
+			m_pD->m_columnName = "";
+		}
+	}
+
+
+	//---------------------------------------------------------------------------------------------
+	void NodeColourTable::SetParameterName( const char* strName )
+	{
+		m_pD->m_parameterName = strName;
+	}
+
+
+	//---------------------------------------------------------------------------------------------
+	void NodeColourTable::SetTable( TablePtr pTable )
+	{
+		m_pD->m_pTable = pTable;
+	}
+
+
+	//---------------------------------------------------------------------------------------------
+	TablePtr NodeColourTable::GetTable() const
+	{
+		return m_pD->m_pTable;
+	}
+
+
+}
+
+
diff --git a/Engine/Plugins/Experimental/Mutable/Source/MutableTools/Private/NodeColourTablePrivate.h b/Engine/Plugins/Experimental/Mutable/Source/MutableTools/Private/NodeColourTablePrivate.h
new file mode 100644
--- /dev/null
+++ b/Engine/Plugins/Experimental/Mutable/Source/MutableTools/Private/NodeColourTablePrivate.h
@@ -0,0 +1,58 @@
+// Copyright Epic Games, Inc. All Rights Reserved.
+
+
+#pragma once
+
+#include "NodePrivate.h"
+#include "NodeColourTable.h"
+#include "Table.h"
+#include "AST.h"
+
+#include "MemoryPrivate.h"
+
+
+namespace mu
+{
+
+
+	class NodeColourTable::Private : public Node::Private
+	{
+	public:
+
+		MUTABLE_DEFINE_CONST_VISITABLE()
+
+	public:
+
+		static NODE_TYPE s_type;
+
+		string m_parameterName;
+		TablePtr m_pTable;
+		string m_columnName;
+
+		//!
+		void Serialise( OutputArchive& arch ) const
+		{
+            uint32_t ver = 0;
+			arch << ver;
+
+			arch << m_parameterName;
+			arch << m_pTable;
+			arch << m_columnName;
+		}
+
+		//!
+		void Unserialise( InputArchive& arch )
+		{
+            uint32_t ver;
+			arch >> ver;
+            check(ver==0);
+
+			arch >> m_parameterName;
+			arch >> m_pTable;
+			arch >> m_columnName;
+		}
+
+	};
+
+}
+
diff --git a/Engine/Plugins/Experimental/Mutable/Source/MutableTools/Private/NodeColourVariation.cpp b/Engine/Plugins/Experimental/Mutable/Source/MutableTools/Private/NodeColourVariation.cpp
new file mode 100644
--- /dev/null
+++ b/Engine/Plugins/Experimental/Mutable/Source/MutableTools/Private/NodeColourVariation.cpp
@@ -0,0 +1,115 @@
+// Copyright Epic Games, Inc. All Rights Reserved.
+
+
+#include "NodeColourVariation.h"
+#include "NodeColourVariationPrivate.h"
+
+
+namespace mu
+{
+
+    //---------------------------------------------------------------------------------------------
+    // Static initialisation
+    //---------------------------------------------------------------------------------------------
+    NODE_TYPE NodeColourVariation::Private::s_type =
+        NODE_TYPE( "ColourVariation", NodeColour::GetStaticType() );
+
+
+    //---------------------------------------------------------------------------------------------
+    //---------------------------------------------------------------------------------------------
+    //---------------------------------------------------------------------------------------------
+
+    MUTABLE_IMPLEMENT_NODE( NodeColourVariation, EType::Variation, Node, Node::EType::Colour)
+
+
+    //---------------------------------------------------------------------------------------------
+    int NodeColourVariation::GetInputCount() const { return 1 + int( m_pD->m_variations.size() ); }
+
+
+    //---------------------------------------------------------------------------------------------
+    Node* NodeColourVariation::GetInputNode( int i ) const
+    {
+        check( i >= 0 && i < GetInputCount() );
+
+        if ( i == 0 )
+        {
+            return m_pD->m_defaultColour.get();
+        }
+        i -= 1;
+
+        if ( i < int( m_pD->m_variations.size() ) )
+        {
+            return m_pD->m_variations[i].m_colour.get();
+        }
+        i -= int( m_pD->m_variations.size() );
+
+        return nullptr;
+    }
+
+
+    //---------------------------------------------------------------------------------------------
+    void NodeColourVariation::SetInputNode( int i, NodePtr pNode )
+    {
+        check( i >= 0 && i < GetInputCount() );
+
+        if ( i == 0 )
+        {
+            m_pD->m_defaultColour = dynamic_cast<NodeColour*>( pNode.get() );
+            return;
+        }
+
+        i -= 1;
+        if ( i < int( m_pD->m_variations.size() ) )
+        {
+
+            m_pD->m_variations[i].m_colour = dynamic_cast<NodeColour*>( pNode.get() );
+            return;
+        }
+        i -= (int)m_pD->m_variations.size();
+    }
+
+
+    //---------------------------------------------------------------------------------------------
+    // Own Interface
+    //---------------------------------------------------------------------------------------------
+    void NodeColourVariation::SetDefaultColour( NodeColour* p ) { m_pD->m_defaultColour = p; }
+
+
+    //---------------------------------------------------------------------------------------------
+    int NodeColourVariation::GetVariationCount() const { return int( m_pD->m_variations.size() ); }
+
+
+    //---------------------------------------------------------------------------------------------
+    void NodeColourVariation::SetVariationCount( int num )
+    {
+        check( num >= 0 );
+        m_pD->m_variations.resize( num );
+    }
+
+    //---------------------------------------------------------------------------------------------
+    void NodeColourVariation::SetVariationTag( int index, const char* strTag )
+    {
+        check( index >= 0 && index < (int)m_pD->m_variations.size() );
+        check( strTag );
+
+        if ( strTag )
+        {
+            m_pD->m_variations[index].m_tag = strTag;
+        }
+        else
+        {
+            m_pD->m_variations[index].m_tag = "";
+        }
+    }
+
+
+    //---------------------------------------------------------------------------------------------
+    void NodeColourVariation::SetVariationColour( int index, NodeColour* pNode )
+    {
+        check( index >= 0 && index < (int)m_pD->m_variations.size() );
+
+        m_pD->m_variations[index].m_colour = pNode;
+    }
+
+
+} // namespace mu
diff --git a/Engine/Plugins/Experimental/Mutable/Source/MutableTools/Private/NodeColourVariationPrivate.h b/Engine/Plugins/Experimental/Mutable/Source/MutableTools/Private/NodeColourVariationPrivate.h
new file mode 100644
--- /dev/null
+++ b/Engine/Plugins/Experimental/Mutable/Source/MutableTools/Private/NodeColourVariationPrivate.h
@@ -0,0 +1,80 @@
+// Copyright Epic Games, Inc. All Rights Reserved.
+
+#pragma once
+
+#include "NodePrivate.h"
+#include "NodeColourVariation.h"
+
+#include "MemoryPrivate.h"
+
+
+namespace mu
+{
+
+    //---------------------------------------------------------------------------------------------
+    //!
+    //---------------------------------------------------------------------------------------------
+    class NodeColourVariation::Private : public Node::Private
+    {
+    public:
+        MUTABLE_DEFINE_CONST_VISITABLE()
+
+    public:
+        Private() {}
+
+        static NODE_TYPE s_type;
+
+        NodeColourPtr m_defaultColour;
+
+        struct VARIATION
+        {
+            NodeColourPtr m_colour;
+            string m_tag;
+
+            //!
+            void Serialise( OutputArchive& arch ) const
+            {
+                uint32_t ver = 1;
+                arch << ver;
+
+                arch << m_tag;
+                arch << m_colour;
+            }
+
+            void Unserialise( InputArchive& arch )
+            {
+                uint32_t ver = 0;
+                arch >> ver;
+                check( ver == 1 );
+
+                arch >> m_tag;
+                arch >> m_colour;
+            }
+        };
+
+        vector<VARIATION> m_variations;
+
+        //!
+        void Serialise( OutputArchive& arch ) const
+        {
+            uint32_t ver = 1;
+            arch << ver;
+
+            arch << m_defaultColour;
+            arch << m_variations;
+        }
+
+        //!
+        void Unserialise( InputArchive& arch )
+        {
+            uint32_t ver;
+            arch >> ver;
+			check(ver == 1);
+
+            arch >> m_defaultColour;
+            arch >> m_variations;
+        }
+    };
+
+} // namespace mu
+
diff --git a/Engine/Plugins/Experimental/Mutable/Source/MutableTools/Private/NodeComponent.cpp b/Engine/Plugins/Experimental/Mutable/Source/MutableTools/Private/NodeComponent.cpp
new file mode 100644
--- /dev/null
+++ b/Engine/Plugins/Experimental/Mutable/Source/MutableTools/Private/NodeComponent.cpp
@@ -0,0 +1,75 @@
+// Copyright Epic Games, Inc. All Rights Reserved.
+
+
+#include "NodeComponent.h"
+#include "NodeComponentPrivate.h"
+
+#include "NodeMesh.h"
+#include "NodeImage.h"
+
+#include "NodeComponentNew.h"
+#include "NodeComponentEdit.h"
+
+
+namespace mu
+{
+
+
+	//---------------------------------------------------------------------------------------------
+	// Static initialisation
+	//---------------------------------------------------------------------------------------------
+	static NODE_TYPE s_nodeComponentType =
+			NODE_TYPE( "NodeComponent", Node::GetStaticType() );
+
+
+	//---------------------------------------------------------------------------------------------
+	//---------------------------------------------------------------------------------------------
+	//---------------------------------------------------------------------------------------------
+	const NODE_TYPE* NodeComponent::GetType() const
+	{
+		return GetStaticType();
+	}
+
+
+	//---------------------------------------------------------------------------------------------
+	const NODE_TYPE* NodeComponent::GetStaticType()
+	{
+		return &s_nodeComponentType;
+	}
+
+
+	//---------------------------------------------------------------------------------------------
+	void NodeComponent::Serialise( const NodeComponent* p, OutputArchive& arch )
+	{
+        uint32_t ver = 0;
+		arch << ver;
+
+		arch << uint32_t(p->Type);
+		p->SerialiseWrapper(arch);
+	}
+
+        
+	//---------------------------------------------------------------------------------------------
+	NodeComponentPtr NodeComponent::StaticUnserialise( InputArchive& arch )
+	{
+        uint32_t ver;
+		arch >> ver;
+		check( ver == 0 );
+
+        uint32_t id;
+		arch >> id;
+
+		switch (id)
+		{
+		case 0 :  return NodeComponentNew::StaticUnserialise( arch ); break;
+		case 1 :  return NodeComponentEdit::StaticUnserialise( arch ); break;
+//		case 2 :  return NodeSelectComponent::StaticUnserialise( arch ); break;
+		default : check(false);
+		}
+
+		return 0;
+	}
+
+}
+
+
diff --git a/Engine/Plugins/Experimental/Mutable/Source/MutableTools/Private/NodeComponentEdit.cpp b/Engine/Plugins/Experimental/Mutable/Source/MutableTools/Private/NodeComponentEdit.cpp
new file mode 100644
--- /dev/null
+++ b/Engine/Plugins/Experimental/Mutable/Source/MutableTools/Private/NodeComponentEdit.cpp
@@ -0,0 +1,115 @@
+// Copyright Epic Games, Inc. All Rights Reserved.
+
+
+#include "NodeComponentEdit.h"
+#include "NodeComponentEditPrivate.h"
+
+#include "NodeSurface.h"
+
+
+namespace mu
+{
+
+
+	//---------------------------------------------------------------------------------------------
+	// Static initialisation
+	//---------------------------------------------------------------------------------------------
+	NODE_TYPE NodeComponentEdit::Private::s_type =
+			NODE_TYPE( "EditComponent", NodeComponent::GetStaticType() );
+
+
+	//---------------------------------------------------------------------------------------------
+	//---------------------------------------------------------------------------------------------
+	//---------------------------------------------------------------------------------------------
+
+    MUTABLE_IMPLEMENT_NODE( NodeComponentEdit, EType::Edit, Node, Node::EType::Component)
+
+
+    //---------------------------------------------------------------------------------------------
+    // Node Interface
+    //---------------------------------------------------------------------------------------------
+    int NodeComponentEdit::GetInputCount() const
+    {
+        return (int) m_pD->m_surfaces.size();
+    }
+
+
+    //---------------------------------------------------------------------------------------------
+    Node* NodeComponentEdit::GetInputNode( int i ) const
+    {
+        check( i >=0 && i < GetInputCount() );
+
+        NodePtr pResult;
+
+        if ( i<(int)m_pD->m_surfaces.size() )
+        {
+            pResult = m_pD->m_surfaces[i];
+        }
+
+        return pResult.get();
+    }
+
+
+    //---------------------------------------------------------------------------------------------
+    void NodeComponentEdit::SetInputNode( int i, NodePtr pNode )
+    {
+        check( i >=0 && i < GetInputCount() );
+
+        if ( i<(int)m_pD->m_surfaces.size() )
+        {
+            m_pD->m_surfaces[ i ] = dynamic_cast<NodeSurface*>(pNode.get());
+        }
+    }
+
+
+	//---------------------------------------------------------------------------------------------
+	// Own Interface
+	//---------------------------------------------------------------------------------------------
+    void NodeComponentEdit::SetParent( NodeComponent* p )
+	{
+		m_pD->m_pParent = p;
+	}
+
+
+	//---------------------------------------------------------------------------------------------
+    NodeComponent* NodeComponentEdit::GetParent() const
+	{
+        return m_pD->m_pParent.get();
+	}
+
+
+    //---------------------------------------------------------------------------------------------
+    int NodeComponentEdit::GetSurfaceCount() const
+    {
+        return (int)m_pD->m_surfaces.size();
+    }
+
+
+    //---------------------------------------------------------------------------------------------
+    void NodeComponentEdit::SetSurfaceCount( int num )
+    {
+        check( num >=0 );
+        m_pD->m_surfaces.resize( num );
+    }
+
+
+    //---------------------------------------------------------------------------------------------
+    NodeSurface* NodeComponentEdit::GetSurface( int index ) const
+    {
+        check( index >=0 && index < (int)m_pD->m_surfaces.size() );
+
+        return m_pD->m_surfaces[ index ].get();
+    }
+
+
+    //---------------------------------------------------------------------------------------------
+    void NodeComponentEdit::SetSurface( int index, NodeSurface* pNode )
+    {
+        check( index >=0 && index < (int)m_pD->m_surfaces.size() );
+
+        m_pD->m_surfaces[ index ] = pNode;
+    }
+
+}
+
+
diff --git a/Engine/Plugins/Experimental/Mutable/Source/MutableTools/Private/NodeComponentEditPrivate.h b/Engine/Plugins/Experimental/Mutable/Source/MutableTools/Private/NodeComponentEditPrivate.h
new file mode 100644
--- /dev/null
+++ b/Engine/Plugins/Experimental/Mutable/Source/MutableTools/Private/NodeComponentEditPrivate.h
@@ -0,0 +1,80 @@
+// Copyright Epic Games, Inc. All Rights Reserved.
+
+
+#pragma once
+
+#include "NodeComponentPrivate.h"
+#include "NodeComponentEdit.h"
+#include "NodePatchImagePrivate.h"
+
+#include "NodeMesh.h"
+#include "NodePatchMesh.h"
+#include "NodeImage.h"
+
+#include "MemoryPrivate.h"
+
+
+namespace mu
+{
+
+
+	//---------------------------------------------------------------------------------------------
+	//!
+	//---------------------------------------------------------------------------------------------
+	class NodeComponentEdit::Private : public NodeComponent::Private
+	{
+	public:
+
+		MUTABLE_DEFINE_CONST_VISITABLE()
+
+	public:
+
+		Private()
+		{
+		}
+
+		static NODE_TYPE s_type;
+
+        NodeComponentPtr m_pParent;
+
+        vector<NodeSurfacePtr> m_surfaces;
+
+		//!
+		void Serialise( OutputArchive& arch ) const
+		{
+            uint32_t ver = 2;
+			arch << ver;
+
+			arch << m_pParent;
+            arch << m_surfaces;
+		}
+
+		//!
+		void Unserialise( InputArchive& arch )
+		{
+            uint32_t ver;
+			arch >> ver;
+			check(ver == 2);
+
+            arch >> m_pParent;
+            arch >> m_surfaces;
+		}
+
+
+		// NodeComponent::Private interface
+        const NodeComponentNew::Private* GetParentComponentNew() const override
+		{
+			const NodeComponentNew::Private* parent = nullptr;
+			if (m_pParent)
+			{
+				NodeComponent::Private* ParentPrivate = dynamic_cast<NodeComponent::Private*>(m_pParent->GetBasePrivate());
+				parent = ParentPrivate->GetParentComponentNew();
+			}
+
+			check(parent);
+			return parent;
+		}
+	};
+
+}
+
diff --git a/Engine/Plugins/Experimental/Mutable/Source/MutableTools/Private/NodeComponentNew.cpp b/Engine/Plugins/Experimental/Mutable/Source/MutableTools/Private/NodeComponentNew.cpp
new file mode 100644
--- /dev/null
+++ b/Engine/Plugins/Experimental/Mutable/Source/MutableTools/Private/NodeComponentNew.cpp
@@ -0,0 +1,138 @@
+// Copyright Epic Games, Inc. All Rights Reserved.
+
+
+#include "NodeComponentNew.h"
+#include "NodeComponentNewPrivate.h"
+
+#include "NodeSurface.h"
+
+
+namespace mu
+{
+
+
+	//---------------------------------------------------------------------------------------------
+	// Static initialisation
+	//---------------------------------------------------------------------------------------------
+	NODE_TYPE NodeComponentNew::Private::s_type =
+			NODE_TYPE( "NewComponent", NodeComponent::GetStaticType() );
+
+
+	//---------------------------------------------------------------------------------------------
+	//---------------------------------------------------------------------------------------------
+	//---------------------------------------------------------------------------------------------
+
+    MUTABLE_IMPLEMENT_NODE( NodeComponentNew, EType::New, Node, Node::EType::Component)
+
+
+	//---------------------------------------------------------------------------------------------
+	// Node Interface
+	//---------------------------------------------------------------------------------------------
+	int NodeComponentNew::GetInputCount() const
+	{
+        return (int) m_pD->m_surfaces.size();
+	}
+
+
+	//---------------------------------------------------------------------------------------------
+	Node* NodeComponentNew::GetInputNode( int i ) const
+	{
+        check( i >=0 && i < GetInputCount() );
+
+		NodePtr pResult;
+
+        if ( i<(int)m_pD->m_surfaces.size() )
+		{
+            pResult = m_pD->m_surfaces[i];
+		}
+
+		return pResult.get();
+	}
+
+
+	//---------------------------------------------------------------------------------------------
+	void NodeComponentNew::SetInputNode( int i, NodePtr pNode )
+	{
+        check( i >=0 && i < GetInputCount() );
+
+        if ( i<(int)m_pD->m_surfaces.size() )
+		{
+            m_pD->m_surfaces[ i ] = dynamic_cast<NodeSurface*>(pNode.get());
+		}
+	}
+
+
+	//---------------------------------------------------------------------------------------------
+	// Own Interface
+	//---------------------------------------------------------------------------------------------
+	const char* NodeComponentNew::GetName() const
+	{
+		const char* strResult = m_pD->m_name.c_str();
+
+		return strResult;
+	}
+
+
+	//---------------------------------------------------------------------------------------------
+	void NodeComponentNew::SetName( const char* strName )
+	{
+		if (strName)
+		{
+			m_pD->m_name = strName;
+		}
+		else
+		{
+			m_pD->m_name = "";
+		}
+	}
+
+
+	//---------------------------------------------------------------------------------------------
+	uint16 mu::NodeComponentNew::GetId() const
+    {
+    	return m_pD->m_id;
+    }
+
+
+	//---------------------------------------------------------------------------------------------
+	void mu::NodeComponentNew::SetId(uint16 id)
+    {
+    	m_pD->m_id = id;
+    }
+
+
+	//---------------------------------------------------------------------------------------------
+    int NodeComponentNew::GetSurfaceCount() const
+	{
+        return (int)m_pD->m_surfaces.size();
+	}
+
+
+	//---------------------------------------------------------------------------------------------
+    void NodeComponentNew::SetSurfaceCount( int num )
+	{
+		check( num >=0 );
+        m_pD->m_surfaces.resize( num );
+	}
+
+
+	//---------------------------------------------------------------------------------------------
+    NodeSurface* NodeComponentNew::GetSurface( int index ) const
+	{
+        check( index >=0 && index < (int)m_pD->m_surfaces.size() );
+
+        return m_pD->m_surfaces[ index ].get();
+	}
+
+
+	//---------------------------------------------------------------------------------------------
+    void NodeComponentNew::SetSurface( int index, NodeSurface* pNode )
+	{
+        check( index >=0 && index < (int)m_pD->m_surfaces.size() );
+
+        m_pD->m_surfaces[ index ] = pNode;
+	}
+
+}
+
+
