[mutable] Move Mutable plugin from the restricted folder to the experimental plugins folder.

diff --git a/Engine/Plugins/Experimental/Mutable/Source/MutableTools/Private/NodeComponentNewPrivate.h b/Engine/Plugins/Experimental/Mutable/Source/MutableTools/Private/NodeComponentNewPrivate.h
new file mode 100644
--- /dev/null
+++ b/Engine/Plugins/Experimental/Mutable/Source/MutableTools/Private/NodeComponentNewPrivate.h
@@ -0,0 +1,78 @@
+// Copyright Epic Games, Inc. All Rights Reserved.
+
+
+#pragma once
+
+
+#include "NodeComponentPrivate.h"
+#include "NodeComponentNew.h"
+#include "NodePatchImagePrivate.h"
+
+#include "NodeMesh.h"
+#include "NodePatchMesh.h"
+#include "NodeImage.h"
+
+#include "MemoryPrivate.h"
+
+
+namespace mu
+{
+
+	//---------------------------------------------------------------------------------------------
+	//!
+	//---------------------------------------------------------------------------------------------
+	class NodeComponentNew::Private : public NodeComponent::Private
+	{
+	public:
+
+		MUTABLE_DEFINE_CONST_VISITABLE()
+
+	public:
+
+		static NODE_TYPE s_type;
+
+		string m_name;
+
+		uint16 m_id = 0;
+
+        vector<NodeSurfacePtr> m_surfaces;
+
+		//!
+		void Serialise( OutputArchive& arch ) const
+		{
+            uint32_t ver = 3;
+			
+			arch << ver;
+
+			arch << m_name;
+			arch << m_id;
+            arch << m_surfaces;
+        }
+
+		//!
+		void Unserialise( InputArchive& arch )
+		{
+            uint32_t ver;
+			arch >> ver;
+			check(ver >= 2);
+
+			arch >> m_name;
+
+			if (ver >= 3)
+			{
+				arch >> m_id;
+			}
+
+            arch >> m_surfaces;
+		}
+
+		// NodeComponent::Private interface
+        const NodeComponentNew::Private* GetParentComponentNew() const override
+		{
+			return this;
+		}
+
+	};
+
+}
+
diff --git a/Engine/Plugins/Experimental/Mutable/Source/MutableTools/Private/NodeImage.cpp b/Engine/Plugins/Experimental/Mutable/Source/MutableTools/Private/NodeImage.cpp
new file mode 100644
--- /dev/null
+++ b/Engine/Plugins/Experimental/Mutable/Source/MutableTools/Private/NodeImage.cpp
@@ -0,0 +1,121 @@
+// Copyright Epic Games, Inc. All Rights Reserved.
+
+
+#include "NodeImage.h"
+#include "NodePrivate.h"
+
+#include "NodeImageConstant.h"
+#include "NodeImageDifference.h"
+#include "NodeImageInterpolate.h"
+#include "NodeImageInterpolate3.h"
+#include "NodeImageSaturate.h"
+#include "NodeImageLuminance.h"
+#include "NodeImageTable.h"
+#include "NodeImageSwizzle.h"
+#include "NodeImageFormat.h"
+#include "NodeImageSelectColour.h"
+#include "NodeImageColourMap.h"
+#include "NodeImageGradient.h"
+#include "NodeImageBinarise.h"
+#include "NodeImageLayer.h"
+#include "NodeImageLayerColour.h"
+#include "NodeImageMultiLayer.h"
+#include "NodeImageResize.h"
+#include "NodeImagePlainColour.h"
+#include "NodeImageProject.h"
+#include "NodeImageMipmap.h"
+#include "NodeImageSwitch.h"
+#include "NodeImageConditional.h"
+#include "NodeImageParameter.h"
+#include "NodeImageInvert.h"
+#include "NodeImageVariation.h"
+#include "NodeImageNormalComposite.h"
+#include "NodeImageTransform.h"
+
+
+namespace mu
+{
+
+	//---------------------------------------------------------------------------------------------
+	// Static initialisation
+	//---------------------------------------------------------------------------------------------
+	static NODE_TYPE s_nodeImageType = 	NODE_TYPE( "NodeImage", Node::GetStaticType() );
+
+
+	//---------------------------------------------------------------------------------------------
+	void NodeImage::Serialise( const NodeImage* p, OutputArchive& arch )
+	{
+        uint32_t ver = 0;
+		arch << ver;
+
+		arch << uint32_t(p->Type);
+		p->SerialiseWrapper(arch);
+    }
+
+        
+	//---------------------------------------------------------------------------------------------
+	NodeImagePtr NodeImage::StaticUnserialise( InputArchive& arch )
+	{
+        uint32_t ver;
+		arch >> ver;
+		check( ver == 0 );
+
+        uint32_t id;
+		arch >> id;
+
+		switch (id)
+		{
+		case 1 :  return NodeImageConstant::StaticUnserialise( arch ); break;
+		case 2 :  return NodeImageDifference::StaticUnserialise( arch ); break;
+		//case 3 :  return NodeImageIdentity::StaticUnserialise( arch ); break;
+		case 4 :  return NodeImageInterpolate::StaticUnserialise( arch ); break;
+		case 6 :  return NodeImageSaturate::StaticUnserialise( arch ); break;
+		//case 8 :  return NodeSelectImage::StaticUnserialise( arch ); break;
+		case 9 :  return NodeImageTable::StaticUnserialise( arch ); break;
+		case 10 :  return NodeImageSwizzle::StaticUnserialise( arch ); break;
+		case 11 :  return NodeImageSelectColour::StaticUnserialise( arch ); break;
+		case 12 :  return NodeImageColourMap::StaticUnserialise( arch ); break;
+		case 13 :  return NodeImageGradient::StaticUnserialise( arch ); break;
+		//case 14 :  return NodeImageVolumeLayer::StaticUnserialise( arch ); break;
+		case 17 :  return NodeImageBinarise::StaticUnserialise( arch ); break;
+		case 18 :  return NodeImageLuminance::StaticUnserialise( arch ); break;
+		case 19 :  return NodeImageLayer::StaticUnserialise( arch ); break;
+		case 20 :  return NodeImageLayerColour::StaticUnserialise( arch ); break;
+		case 21 :  return NodeImageResize::StaticUnserialise( arch ); break;
+		case 22 :  return NodeImagePlainColour::StaticUnserialise( arch ); break;
+		case 23 :  return NodeImageInterpolate3::StaticUnserialise( arch ); break;
+        case 24 :  return NodeImageProject::StaticUnserialise( arch ); break;
+        case 25 :  return NodeImageMipmap::StaticUnserialise( arch ); break;
+        case 26 :  return NodeImageSwitch::StaticUnserialise( arch ); break;
+		case 27 :  return NodeImageConditional::StaticUnserialise( arch ); break;
+        case 28 :  return NodeImageFormat::StaticUnserialise( arch ); break;
+        case 29 :  return NodeImageParameter::StaticUnserialise( arch ); break;
+        case 30 :  return NodeImageMultiLayer::StaticUnserialise( arch ); break;
+        case 31 :  return NodeImageInvert::StaticUnserialise( arch ); break;
+        case 32 :  return NodeImageVariation::StaticUnserialise( arch ); break;
+        case 33 :  return NodeImageNormalComposite::StaticUnserialise( arch ); break;
+        case 34 :  return NodeImageTransform::StaticUnserialise( arch ); break;
+        default : check(false);
+		}
+
+		return 0;
+	}
+
+
+	//---------------------------------------------------------------------------------------------
+	const NODE_TYPE* NodeImage::GetType() const
+	{
+		return GetStaticType();
+	}
+
+
+	//---------------------------------------------------------------------------------------------
+	const NODE_TYPE* NodeImage::GetStaticType()
+	{
+		return &s_nodeImageType;
+	}
+
+
+}
+
+
diff --git a/Engine/Plugins/Experimental/Mutable/Source/MutableTools/Private/NodeImageBinarise.cpp b/Engine/Plugins/Experimental/Mutable/Source/MutableTools/Private/NodeImageBinarise.cpp
new file mode 100644
--- /dev/null
+++ b/Engine/Plugins/Experimental/Mutable/Source/MutableTools/Private/NodeImageBinarise.cpp
@@ -0,0 +1,100 @@
+// Copyright Epic Games, Inc. All Rights Reserved.
+
+
+#include "NodeImageBinarise.h"
+#include "NodeImageBinarisePrivate.h"
+
+#include "NodeScalar.h"
+
+#include "ImagePrivate.h"
+
+
+#define NODE_INPUT_COUNT 	2
+
+
+namespace mu
+{
+
+	//---------------------------------------------------------------------------------------------
+	// Static initialisation
+	//---------------------------------------------------------------------------------------------
+	NODE_TYPE NodeImageBinarise::Private::s_type =
+			NODE_TYPE( "ImageMultiply", NodeImage::GetStaticType() );
+
+	//---------------------------------------------------------------------------------------------
+
+	MUTABLE_IMPLEMENT_NODE( NodeImageBinarise, EType::Binarise, Node, Node::EType::Image)
+
+
+	//---------------------------------------------------------------------------------------------
+	// Node Interface
+	//---------------------------------------------------------------------------------------------
+	void NodeImageBinarise::SetInputNode( int i, NodePtr pNode )
+	{
+		check( i>=0 && i<NODE_INPUT_COUNT );
+
+		switch (i)
+		{
+        case 0: m_pD->m_pBase = dynamic_cast<NodeImage*>(pNode.get()); break;
+        case 1: m_pD->m_pThreshold = dynamic_cast<NodeScalar*>(pNode.get()); break;
+		}
+	}
+
+
+	//---------------------------------------------------------------------------------------------
+	int NodeImageBinarise::GetInputCount() const
+	{
+		return NODE_INPUT_COUNT;
+	}
+
+
+	//---------------------------------------------------------------------------------------------
+	Node* NodeImageBinarise::GetInputNode( int i ) const
+	{
+		check( i>=0 && i<NODE_INPUT_COUNT );
+
+		Node* pResult = 0;
+
+		switch (i)
+		{
+        case 0: pResult = m_pD->m_pBase.get(); break;
+        case 1: pResult = m_pD->m_pThreshold.get(); break;
+		}
+
+		return pResult;
+	}
+
+
+	//---------------------------------------------------------------------------------------------
+	// Own Interface
+	//---------------------------------------------------------------------------------------------
+	NodeImagePtr NodeImageBinarise::GetBase() const
+	{
+		return m_pD->m_pBase;
+	}
+
+
+	//---------------------------------------------------------------------------------------------
+	void NodeImageBinarise::SetBase( NodeImagePtr pNode )
+	{
+		m_pD->m_pBase = pNode;
+	}
+
+
+	//---------------------------------------------------------------------------------------------
+	NodeScalarPtr NodeImageBinarise::GetThreshold() const
+	{
+		return m_pD->m_pThreshold;
+	}
+
+
+	//---------------------------------------------------------------------------------------------
+	void NodeImageBinarise::SetThreshold( NodeScalarPtr pNode )
+	{
+		m_pD->m_pThreshold = pNode;
+	}
+
+}
+
+#undef NODE_INPUT_COUNT
+
diff --git a/Engine/Plugins/Experimental/Mutable/Source/MutableTools/Private/NodeImageBinarisePrivate.h b/Engine/Plugins/Experimental/Mutable/Source/MutableTools/Private/NodeImageBinarisePrivate.h
new file mode 100644
--- /dev/null
+++ b/Engine/Plugins/Experimental/Mutable/Source/MutableTools/Private/NodeImageBinarisePrivate.h
@@ -0,0 +1,53 @@
+// Copyright Epic Games, Inc. All Rights Reserved.
+
+
+#pragma once
+
+#include "NodeImagePrivate.h"
+#include "NodeImageBinarise.h"
+#include "AST.h"
+
+#include "MemoryPrivate.h"
+
+
+namespace mu
+{
+
+
+	class NodeImageBinarise::Private : public NodeImage::Private
+	{
+	public:
+
+		MUTABLE_DEFINE_CONST_VISITABLE();
+
+	public:
+
+		static NODE_TYPE s_type;
+
+		NodeImagePtr m_pBase;
+		NodeScalarPtr m_pThreshold;
+
+		//!
+		void Serialise( OutputArchive& arch ) const
+		{
+            uint32_t ver = 0;
+			arch << ver;
+
+			arch << m_pBase;
+			arch << m_pThreshold;
+		}
+
+		//!
+		void Unserialise( InputArchive& arch )
+		{
+            uint32_t ver;
+			arch >> ver;
+			check(ver==0);
+
+			arch >> m_pBase;
+			arch >> m_pThreshold;
+		}
+	};
+
+
+}
diff --git a/Engine/Plugins/Experimental/Mutable/Source/MutableTools/Private/NodeImageColourMap.cpp b/Engine/Plugins/Experimental/Mutable/Source/MutableTools/Private/NodeImageColourMap.cpp
new file mode 100644
--- /dev/null
+++ b/Engine/Plugins/Experimental/Mutable/Source/MutableTools/Private/NodeImageColourMap.cpp
@@ -0,0 +1,112 @@
+// Copyright Epic Games, Inc. All Rights Reserved.
+
+
+#include "NodeImageColourMap.h"
+#include "NodeImageColourMapPrivate.h"
+
+#include "NodeColour.h"
+
+#include "ImagePrivate.h"
+
+
+namespace mu
+{
+
+	//---------------------------------------------------------------------------------------------
+	// Static initialisation
+	//---------------------------------------------------------------------------------------------
+	NODE_TYPE NodeImageColourMap::Private::s_type =
+			NODE_TYPE( "ImageColourMap", NodeImage::GetStaticType() );
+
+
+	//---------------------------------------------------------------------------------------------
+	//!
+	//---------------------------------------------------------------------------------------------
+
+	MUTABLE_IMPLEMENT_NODE( NodeImageColourMap, EType::ColourMap, Node, Node::EType::Image)
+
+
+	//---------------------------------------------------------------------------------------------
+	// Node Interface
+	//---------------------------------------------------------------------------------------------
+	int NodeImageColourMap::GetInputCount() const
+	{
+		return 3;
+	}
+
+
+	//---------------------------------------------------------------------------------------------
+	Node* NodeImageColourMap::GetInputNode( int i ) const
+	{
+		check( i>=0 && i< GetInputCount());
+
+		Node* pResult = 0;
+
+		switch (i)
+		{
+        case 0: pResult = m_pD->m_pBase.get(); break;
+        case 1: pResult = m_pD->m_pMask.get(); break;
+        case 2: pResult = m_pD->m_pMap.get(); break;
+		}
+
+		return pResult;
+	}
+
+
+	//---------------------------------------------------------------------------------------------
+	void NodeImageColourMap::SetInputNode( int i, NodePtr pNode )
+	{
+		check( i>=0 && i< GetInputCount());
+
+		switch (i)
+		{
+        case 0: m_pD->m_pBase = dynamic_cast<NodeImage*>( pNode.get() ); break;
+        case 1: m_pD->m_pMask = dynamic_cast<NodeImage*>( pNode.get() ); break;
+        case 2: m_pD->m_pMap = dynamic_cast<NodeImage*>( pNode.get() ); break;
+		}
+	}
+
+
+	//---------------------------------------------------------------------------------------------
+	// Own Interface
+	//---------------------------------------------------------------------------------------------
+	NodeImagePtr NodeImageColourMap::GetBase() const
+	{
+		return m_pD->m_pBase.get();
+	}
+
+
+	//---------------------------------------------------------------------------------------------
+	void NodeImageColourMap::SetBase( NodeImagePtr pNode )
+	{
+		m_pD->m_pBase = pNode;
+	}
+
+
+	//---------------------------------------------------------------------------------------------
+	NodeImagePtr NodeImageColourMap::GetMap() const
+	{
+		return m_pD->m_pMap.get();
+	}
+
+
+	//---------------------------------------------------------------------------------------------
+	void NodeImageColourMap::SetMap( NodeImagePtr pNode )
+	{
+		m_pD->m_pMap = pNode;
+	}
+
+
+	//---------------------------------------------------------------------------------------------
+	NodeImagePtr NodeImageColourMap::GetMask() const
+	{
+		return m_pD->m_pMask.get();
+	}
+
+	//---------------------------------------------------------------------------------------------
+	void NodeImageColourMap::SetMask( NodeImagePtr pNode )
+	{
+		m_pD->m_pMask = pNode;
+	}
+
+}
diff --git a/Engine/Plugins/Experimental/Mutable/Source/MutableTools/Private/NodeImageColourMapPrivate.h b/Engine/Plugins/Experimental/Mutable/Source/MutableTools/Private/NodeImageColourMapPrivate.h
new file mode 100644
--- /dev/null
+++ b/Engine/Plugins/Experimental/Mutable/Source/MutableTools/Private/NodeImageColourMapPrivate.h
@@ -0,0 +1,57 @@
+// Copyright Epic Games, Inc. All Rights Reserved.
+
+
+#pragma once
+
+#include "NodeImagePrivate.h"
+#include "NodeImageColourMap.h"
+#include "AST.h"
+
+#include "MemoryPrivate.h"
+
+
+namespace mu
+{
+
+
+	class NodeImageColourMap::Private : public NodeImage::Private
+	{
+	public:
+
+		MUTABLE_DEFINE_CONST_VISITABLE()
+
+	public:
+
+		static NODE_TYPE s_type;
+
+		NodeImagePtr m_pBase;
+		NodeImagePtr m_pMask;
+		NodeImagePtr m_pMap;
+
+		//!
+		void Serialise( OutputArchive& arch ) const
+		{
+            uint32_t ver = 0;
+			arch << ver;
+
+			arch << m_pBase;
+			arch << m_pMask;
+			arch << m_pMap;
+		}
+
+		//!
+		void Unserialise( InputArchive& arch )
+		{
+            uint32_t ver;
+			arch >> ver;
+			check(ver==0);
+
+			arch >> m_pBase;
+			arch >> m_pMask;
+			arch >> m_pMap;
+		}
+	};
+
+
+}
+
diff --git a/Engine/Plugins/Experimental/Mutable/Source/MutableTools/Private/NodeImageConditional.cpp b/Engine/Plugins/Experimental/Mutable/Source/MutableTools/Private/NodeImageConditional.cpp
new file mode 100644
--- /dev/null
+++ b/Engine/Plugins/Experimental/Mutable/Source/MutableTools/Private/NodeImageConditional.cpp
@@ -0,0 +1,139 @@
+// Copyright Epic Games, Inc. All Rights Reserved.
+
+
+#include "MutableTools/Public/NodeImageConditional.h"
+#include "MutableTools/Private/NodeImageConditionalPrivate.h"
+
+#include "MutableTools/Public/NodeBool.h"
+
+#include "MutableRuntime/Private/ImagePrivate.h"
+
+
+namespace mu
+{
+
+	//---------------------------------------------------------------------------------------------
+	// Static initialisation
+	//---------------------------------------------------------------------------------------------
+    NODE_TYPE NodeImageConditional::Private::s_type =
+            NODE_TYPE( "ImageConditional", NodeImage::GetStaticType() );
+
+
+	//---------------------------------------------------------------------------------------------
+	//!
+	//---------------------------------------------------------------------------------------------
+
+    MUTABLE_IMPLEMENT_NODE( NodeImageConditional, EType::Conditional, Node, Node::EType::Image)
+
+
+	//---------------------------------------------------------------------------------------------
+	// Node Interface
+	//---------------------------------------------------------------------------------------------
+    int NodeImageConditional::GetInputCount() const
+	{
+        return 3;
+	}
+
+
+	//---------------------------------------------------------------------------------------------
+    Node* NodeImageConditional::GetInputNode( int i ) const
+	{
+		check( i>=0 && i<GetInputCount() );
+
+		Node* pResult = 0;
+
+        switch (i)
+		{
+        case 0:
+            pResult = m_pD->m_parameter.get();
+            break;
+
+        case 1:
+            pResult = m_pD->m_true.get();
+            break;
+
+        case 2:
+            pResult = m_pD->m_false.get();
+            break;
+
+        default:
+			break;
+		}
+
+		return pResult;
+	}
+
+
+	//---------------------------------------------------------------------------------------------
+    void NodeImageConditional::SetInputNode( int i, NodePtr pNode )
+	{
+		check( i>=0 && i<GetInputCount() );
+
+        switch (i)
+		{
+        case 0:
+            m_pD->m_parameter = dynamic_cast<NodeBool*>(pNode.get());
+            break;
+
+        case 1:
+            m_pD->m_true = dynamic_cast<NodeImage*>(pNode.get());
+            break;
+
+        case 2:
+            m_pD->m_false = dynamic_cast<NodeImage*>(pNode.get());
+            break;
+
+        default:
+			break;
+		}
+	}
+
+
+	//---------------------------------------------------------------------------------------------
+	// Own Interface
+	//---------------------------------------------------------------------------------------------
+    NodeBool* NodeImageConditional::GetParameter() const
+	{
+        return m_pD->m_parameter.get();
+	}
+
+
+	//---------------------------------------------------------------------------------------------
+    void NodeImageConditional::SetParameter( NodeBool* pNode )
+	{
+        m_pD->m_parameter = pNode;
+	}
+
+
+    //---------------------------------------------------------------------------------------------
+    NodeImage* NodeImageConditional::GetOptionTrue() const
+    {
+        return m_pD->m_true.get();
+    }
+
+
+    //---------------------------------------------------------------------------------------------
+    void NodeImageConditional::SetOptionTrue( NodeImage* pNode )
+    {
+        m_pD->m_true = pNode;
+    }
+
+
+    //---------------------------------------------------------------------------------------------
+    NodeImage* NodeImageConditional::GetOptionFalse() const
+    {
+        return m_pD->m_false.get();
+    }
+
+
+    //---------------------------------------------------------------------------------------------
+    void NodeImageConditional::SetOptionFalse( NodeImage* pNode )
+    {
+        m_pD->m_false = pNode;
+    }
+
+
+
+}
+
+
diff --git a/Engine/Plugins/Experimental/Mutable/Source/MutableTools/Private/NodeImageConditionalPrivate.h b/Engine/Plugins/Experimental/Mutable/Source/MutableTools/Private/NodeImageConditionalPrivate.h
new file mode 100644
--- /dev/null
+++ b/Engine/Plugins/Experimental/Mutable/Source/MutableTools/Private/NodeImageConditionalPrivate.h
@@ -0,0 +1,56 @@
+// Copyright Epic Games, Inc. All Rights Reserved.
+
+
+#pragma once
+
+#include "NodeImageConditional.h"
+#include "NodeImagePrivate.h"
+#include "NodeBool.h"
+#include "AST.h"
+
+
+namespace mu
+{
+
+
+    class NodeImageConditional::Private : public NodeImage::Private
+	{
+	public:
+
+		MUTABLE_DEFINE_CONST_VISITABLE()
+
+	public:
+
+		static NODE_TYPE s_type;
+
+        NodeBoolPtr m_parameter;
+        NodeImagePtr m_true;
+        NodeImagePtr m_false;
+
+		//!
+		void Serialise( OutputArchive& arch ) const
+		{
+            uint32_t ver = 0;
+			arch << ver;
+
+            arch << m_parameter;
+            arch << m_true;
+            arch << m_false;
+        }
+
+		//!
+		void Unserialise( InputArchive& arch )
+		{
+            uint32_t ver;
+			arch >> ver;
+			check(ver==0);
+
+            arch >> m_parameter;
+            arch >> m_true;
+            arch >> m_false;
+        }
+	};
+
+
+}
+
diff --git a/Engine/Plugins/Experimental/Mutable/Source/MutableTools/Private/NodeImageConstant.cpp b/Engine/Plugins/Experimental/Mutable/Source/MutableTools/Private/NodeImageConstant.cpp
new file mode 100644
--- /dev/null
+++ b/Engine/Plugins/Experimental/Mutable/Source/MutableTools/Private/NodeImageConstant.cpp
@@ -0,0 +1,67 @@
+// Copyright Epic Games, Inc. All Rights Reserved.
+
+
+#include "NodeImageConstant.h"
+#include "NodeImageConstantPrivate.h"
+#include "ImagePrivate.h"
+
+
+namespace mu
+{
+
+	//---------------------------------------------------------------------------------------------
+	// Static initialisation
+	//---------------------------------------------------------------------------------------------
+	NODE_TYPE NodeImageConstant::Private::s_type =
+			NODE_TYPE( "ImageConstant", NodeImage::GetStaticType() );
+
+
+	//---------------------------------------------------------------------------------------------
+	MUTABLE_IMPLEMENT_NODE( NodeImageConstant, EType::Constant, Node, Node::EType::Image);
+
+
+	//---------------------------------------------------------------------------------------------
+	int NodeImageConstant::GetInputCount() const
+	{
+		return 0;
+	}
+
+
+	//---------------------------------------------------------------------------------------------
+    Node* NodeImageConstant::GetInputNode( int ) const
+	{
+		check( false );
+		return 0;
+	}
+
+
+	//---------------------------------------------------------------------------------------------
+    void NodeImageConstant::SetInputNode( int, NodePtr )
+	{
+		check( false );
+	}
+
+
+	//---------------------------------------------------------------------------------------------
+    ImagePtrConst NodeImageConstant::GetValue() const
+	{
+        Ptr<const Image> pImage;
+        if (m_pD->m_pProxy) pImage = m_pD->m_pProxy->Get();
+        return pImage;
+	}
+
+
+	//---------------------------------------------------------------------------------------------
+	void NodeImageConstant::SetValue( Ptr<const Image> pValue )
+	{
+        m_pD->m_pProxy = new ResourceProxyMemory<Image>( pValue );
+	}
+
+
+    //---------------------------------------------------------------------------------------------
+    void NodeImageConstant::SetValue( Ptr<ResourceProxy<Image>> pImage )
+    {
+        m_pD->m_pProxy = pImage;
+    }
+
+}
diff --git a/Engine/Plugins/Experimental/Mutable/Source/MutableTools/Private/NodeImageConstantPrivate.h b/Engine/Plugins/Experimental/Mutable/Source/MutableTools/Private/NodeImageConstantPrivate.h
new file mode 100644
--- /dev/null
+++ b/Engine/Plugins/Experimental/Mutable/Source/MutableTools/Private/NodeImageConstantPrivate.h
@@ -0,0 +1,62 @@
+// Copyright Epic Games, Inc. All Rights Reserved.
+
+#pragma once
+
+#include "NodeImagePrivate.h"
+#include "NodeImageConstant.h"
+#include "AST.h"
+#include "StreamsPrivate.h"
+
+#include "MemoryPrivate.h"
+
+
+namespace mu
+{
+
+	class NodeImageConstant::Private : public NodeImage::Private
+	{
+	public:
+
+		MUTABLE_DEFINE_CONST_VISITABLE();
+
+	public:
+
+		static NODE_TYPE s_type;
+
+        Ptr<ResourceProxy<Image>> m_pProxy;
+
+		//!
+		void Serialise( OutputArchive& arch ) const
+		{
+            uint32_t ver = 0;
+			arch << ver;
+
+            Ptr<const Image> image;
+            if (m_pProxy)
+            {
+                image = m_pProxy->Get();
+            }
+
+            arch << image;
+		}
+
+		//!
+		void Unserialise( InputArchive& arch )
+		{
+            uint32_t ver;
+			arch >> ver;
+			check(ver==0);
+
+            // Are we using proxies?
+            m_pProxy = arch.NewImageProxy();
+            if (!m_pProxy)
+            {
+                // Normal serialisation
+                ImagePtr image;
+                arch >> image;
+                m_pProxy = new ResourceProxyMemory<Image>( image.get() );
+            }
+		}
+	};
+
+}
diff --git a/Engine/Plugins/Experimental/Mutable/Source/MutableTools/Private/NodeImageDifference.cpp b/Engine/Plugins/Experimental/Mutable/Source/MutableTools/Private/NodeImageDifference.cpp
new file mode 100644
--- /dev/null
+++ b/Engine/Plugins/Experimental/Mutable/Source/MutableTools/Private/NodeImageDifference.cpp
@@ -0,0 +1,100 @@
+// Copyright Epic Games, Inc. All Rights Reserved.
+
+
+#include "NodeImageDifference.h"
+#include "NodeImageDifferencePrivate.h"
+
+#include "NodeColour.h"
+
+#include "ImagePrivate.h"
+
+
+#define NODE_INPUT_COUNT 	2
+
+
+namespace mu
+{
+
+	//---------------------------------------------------------------------------------------------
+	// Static initialisation
+	//---------------------------------------------------------------------------------------------
+	NODE_TYPE NodeImageDifference::Private::s_type =
+			NODE_TYPE( "NodeImageDifference", NodeImage::GetStaticType() );
+
+	//---------------------------------------------------------------------------------------------
+
+	MUTABLE_IMPLEMENT_NODE( NodeImageDifference, EType::Difference, Node, Node::EType::Image);
+
+
+	//---------------------------------------------------------------------------------------------
+	// Node Interface
+	//---------------------------------------------------------------------------------------------
+	void NodeImageDifference::SetInputNode( int i, NodePtr pNode )
+	{
+		check( i>=0 && i<NODE_INPUT_COUNT );
+
+		switch (i)
+		{
+        case 0: m_pD->m_pA = dynamic_cast<NodeImage*>(pNode.get()); break;
+        case 1: m_pD->m_pB = dynamic_cast<NodeImage*>(pNode.get()); break;
+		}
+	}
+
+
+	//---------------------------------------------------------------------------------------------
+	int NodeImageDifference::GetInputCount() const
+	{
+		return NODE_INPUT_COUNT;
+	}
+
+
+	//---------------------------------------------------------------------------------------------
+	Node* NodeImageDifference::GetInputNode( int i ) const
+	{
+		check( i>=0 && i<NODE_INPUT_COUNT );
+
+		Node* pResult = 0;
+
+		switch (i)
+		{
+        case 0: pResult = m_pD->m_pA.get(); break;
+        case 1: pResult = m_pD->m_pB.get(); break;
+		}
+
+		return pResult;
+	}
+
+
+	//---------------------------------------------------------------------------------------------
+	// Own Interface
+	//---------------------------------------------------------------------------------------------
+	NodeImagePtr NodeImageDifference::GetA() const
+	{
+		return m_pD->m_pA;
+	}
+
+
+	//---------------------------------------------------------------------------------------------
+	void NodeImageDifference::SetA( NodeImagePtr pNode )
+	{
+		m_pD->m_pA = pNode;
+	}
+
+
+	//---------------------------------------------------------------------------------------------
+	NodeImagePtr NodeImageDifference::GetB() const
+	{
+		return m_pD->m_pB;
+	}
+
+
+	//---------------------------------------------------------------------------------------------
+	void NodeImageDifference::SetB( NodeImagePtr pNode )
+	{
+		m_pD->m_pB = pNode;
+	}
+
+}
+
+#undef NODE_INPUT_COUNT
+
diff --git a/Engine/Plugins/Experimental/Mutable/Source/MutableTools/Private/NodeImageDifferencePrivate.h b/Engine/Plugins/Experimental/Mutable/Source/MutableTools/Private/NodeImageDifferencePrivate.h
new file mode 100644
--- /dev/null
+++ b/Engine/Plugins/Experimental/Mutable/Source/MutableTools/Private/NodeImageDifferencePrivate.h
@@ -0,0 +1,52 @@
+// Copyright Epic Games, Inc. All Rights Reserved.
+
+#pragma once
+
+#include "NodeImagePrivate.h"
+#include "NodeImageDifference.h"
+#include "AST.h"
+
+#include "MemoryPrivate.h"
+
+
+namespace mu
+{
+
+
+	class NodeImageDifference::Private : public NodeImage::Private
+	{
+	public:
+
+		MUTABLE_DEFINE_CONST_VISITABLE();
+
+	public:
+
+		static NODE_TYPE s_type;
+
+		NodeImagePtr m_pA;
+		NodeImagePtr m_pB;
+
+		//!
+		void Serialise( OutputArchive& arch ) const
+		{
+            uint32_t ver = 0;
+			arch << ver;
+
+			arch << m_pA;
+			arch << m_pB;
+		}
+
+		//!
+		void Unserialise( InputArchive& arch )
+		{
+            uint32_t ver;
+			arch >> ver;
+			check(ver==0);
+
+			arch >> m_pA;
+			arch >> m_pB;
+		}
+	};
+
+
+}
diff --git a/Engine/Plugins/Experimental/Mutable/Source/MutableTools/Private/NodeImageFormat.cpp b/Engine/Plugins/Experimental/Mutable/Source/MutableTools/Private/NodeImageFormat.cpp
new file mode 100644
--- /dev/null
+++ b/Engine/Plugins/Experimental/Mutable/Source/MutableTools/Private/NodeImageFormat.cpp
@@ -0,0 +1,89 @@
+// Copyright Epic Games, Inc. All Rights Reserved.
+
+
+#include "NodeImageFormat.h"
+#include "NodeImageFormatPrivate.h"
+
+
+#define NODE_INPUT_COUNT 	4
+
+
+namespace mu
+{
+
+	//---------------------------------------------------------------------------------------------
+	// Static initialisation
+	//---------------------------------------------------------------------------------------------
+	NODE_TYPE NodeImageFormat::Private::s_type =
+			NODE_TYPE( "ImageFormat", NodeImage::GetStaticType() );
+
+
+	//---------------------------------------------------------------------------------------------
+	//!
+	//---------------------------------------------------------------------------------------------
+
+	MUTABLE_IMPLEMENT_NODE( NodeImageFormat, EType::Format, Node, Node::EType::Image)
+
+
+	//---------------------------------------------------------------------------------------------
+	// Node Interface
+	//---------------------------------------------------------------------------------------------
+	int NodeImageFormat::GetInputCount() const
+	{
+		return 1;
+	}
+
+
+	//---------------------------------------------------------------------------------------------
+	Node* NodeImageFormat::GetInputNode( int i ) const
+	{
+		check( i>=0 && i<GetInputCount() );
+        (void)i;
+
+		return m_pD->m_source.get();
+	}
+
+
+	//---------------------------------------------------------------------------------------------
+	void NodeImageFormat::SetInputNode( int i, NodePtr pNode )
+	{
+		check( i>=0 && i<GetInputCount() );
+        (void)i;
+
+		m_pD->m_source = dynamic_cast<NodeImage*>(pNode.get());
+	}
+
+
+	//---------------------------------------------------------------------------------------------
+	// Own Interface
+	//---------------------------------------------------------------------------------------------
+	EImageFormat NodeImageFormat::GetFormat() const
+	{
+		return m_pD->m_format;
+	}
+
+
+	//---------------------------------------------------------------------------------------------
+    void NodeImageFormat::SetFormat(EImageFormat format, EImageFormat formatIfAlpha )
+	{
+		m_pD->m_format = format;
+        m_pD->m_formatIfAlpha = formatIfAlpha;
+	}
+
+
+	//---------------------------------------------------------------------------------------------
+	NodeImagePtr NodeImageFormat::GetSource() const
+	{
+		return m_pD->m_source.get();
+	}
+
+
+	//---------------------------------------------------------------------------------------------
+	void NodeImageFormat::SetSource(NodeImagePtr pNode )
+	{
+		m_pD->m_source = pNode;
+	}
+
+}
+
+#undef NODE_INPUT_COUNT
diff --git a/Engine/Plugins/Experimental/Mutable/Source/MutableTools/Private/NodeImageFormatPrivate.h b/Engine/Plugins/Experimental/Mutable/Source/MutableTools/Private/NodeImageFormatPrivate.h
new file mode 100644
--- /dev/null
+++ b/Engine/Plugins/Experimental/Mutable/Source/MutableTools/Private/NodeImageFormatPrivate.h
@@ -0,0 +1,59 @@
+// Copyright Epic Games, Inc. All Rights Reserved.
+
+#pragma once
+
+#include "NodeImageFormat.h"
+#include "NodeImagePrivate.h"
+#include "AST.h"
+#include "ImagePrivate.h"
+
+
+namespace mu
+{
+
+
+	class NodeImageFormat::Private : public NodeImage::Private
+	{
+	public:
+
+		MUTABLE_DEFINE_CONST_VISITABLE()
+
+	public:
+
+		static NODE_TYPE s_type;
+
+		EImageFormat m_format = EImageFormat::IF_NONE;
+		EImageFormat m_formatIfAlpha = EImageFormat::IF_NONE;
+        NodeImagePtr m_source;
+
+		//!
+		void Serialise( OutputArchive& arch ) const
+		{
+            uint32_t ver = 1;
+			arch << ver;
+
+            arch << (uint32)m_format;
+            arch << (uint32)m_formatIfAlpha;
+            arch << m_source;
+		}
+
+		//!
+		void Unserialise( InputArchive& arch )
+		{
+            uint32_t ver;
+			arch >> ver;
+            check(ver==1);
+
+			uint32 format;
+			arch >> format;
+			m_format = EImageFormat(format);
+
+			arch >> format;
+			m_formatIfAlpha = EImageFormat(format);
+
+			arch >> m_source;
+		}
+	};
+
+
+}
diff --git a/Engine/Plugins/Experimental/Mutable/Source/MutableTools/Private/NodeImageGradient.cpp b/Engine/Plugins/Experimental/Mutable/Source/MutableTools/Private/NodeImageGradient.cpp
new file mode 100644
--- /dev/null
+++ b/Engine/Plugins/Experimental/Mutable/Source/MutableTools/Private/NodeImageGradient.cpp
@@ -0,0 +1,123 @@
+// Copyright Epic Games, Inc. All Rights Reserved.
+
+
+#include "NodeImageGradient.h"
+#include "NodeImageGradientPrivate.h"
+
+#include "NodeColour.h"
+
+#include "ImagePrivate.h"
+
+
+#define NODE_INPUT_COUNT 	2
+
+
+namespace mu
+{
+
+	//---------------------------------------------------------------------------------------------
+	// Static initialisation
+	//---------------------------------------------------------------------------------------------
+	NODE_TYPE NodeImageGradient::Private::s_type =
+			NODE_TYPE( "ImageGradient", NodeImage::GetStaticType() );
+
+
+	//---------------------------------------------------------------------------------------------
+	//!
+	//---------------------------------------------------------------------------------------------
+
+	MUTABLE_IMPLEMENT_NODE( NodeImageGradient, EType::Gradient, Node, Node::EType::Image)
+
+
+	//---------------------------------------------------------------------------------------------
+	// Node Interface
+	//---------------------------------------------------------------------------------------------
+	int NodeImageGradient::GetInputCount() const
+	{
+		return NODE_INPUT_COUNT;
+	}
+
+
+	//---------------------------------------------------------------------------------------------
+	Node* NodeImageGradient::GetInputNode( int i ) const
+	{
+		check( i>=0 && i<NODE_INPUT_COUNT );
+
+		Node* pResult = 0;
+
+		switch (i)
+		{
+        case 0: pResult = m_pD->m_pColour0.get(); break;
+        case 1: pResult = m_pD->m_pColour1.get(); break;
+		}
+
+		return pResult;
+	}
+
+
+	//---------------------------------------------------------------------------------------------
+	void NodeImageGradient::SetInputNode( int i, NodePtr pNode )
+	{
+		check( i>=0 && i<NODE_INPUT_COUNT );
+
+		switch (i)
+		{
+        case 0: m_pD->m_pColour0 = dynamic_cast<NodeColour*>( pNode.get() ); break;
+        case 1: m_pD->m_pColour1 = dynamic_cast<NodeColour*>( pNode.get() ); break;
+		}
+	}
+
+
+	//---------------------------------------------------------------------------------------------
+	// Own Interface
+	//---------------------------------------------------------------------------------------------
+	NodeColourPtr NodeImageGradient::GetColour0() const
+	{
+		return m_pD->m_pColour0.get();
+	}
+
+
+	//---------------------------------------------------------------------------------------------
+	void NodeImageGradient::SetColour0( NodeColourPtr pNode )
+	{
+		m_pD->m_pColour0 = pNode;
+	}
+
+
+	//---------------------------------------------------------------------------------------------
+	NodeColourPtr NodeImageGradient::GetColour1() const
+	{
+		return m_pD->m_pColour1.get();
+	}
+
+
+	//---------------------------------------------------------------------------------------------
+	void NodeImageGradient::SetColour1( NodeColourPtr pNode )
+	{
+		m_pD->m_pColour1 = pNode;
+	}
+
+
+	//---------------------------------------------------------------------------------------------
+	int NodeImageGradient::GetSizeX() const
+	{
+		return m_pD->m_size[0];
+	}
+
+
+	//---------------------------------------------------------------------------------------------
+	int NodeImageGradient::GetSizeY() const
+	{
+		return m_pD->m_size[1];
+	}
+
+	//---------------------------------------------------------------------------------------------
+	void NodeImageGradient::SetSize( int x, int y )
+	{
+		m_pD->m_size[0] = x;
+		m_pD->m_size[1] = y;
+	}
+
+}
+
+#undef NODE_INPUT_COUNT
diff --git a/Engine/Plugins/Experimental/Mutable/Source/MutableTools/Private/NodeImageGradientPrivate.h b/Engine/Plugins/Experimental/Mutable/Source/MutableTools/Private/NodeImageGradientPrivate.h
new file mode 100644
--- /dev/null
+++ b/Engine/Plugins/Experimental/Mutable/Source/MutableTools/Private/NodeImageGradientPrivate.h
@@ -0,0 +1,58 @@
+// Copyright Epic Games, Inc. All Rights Reserved.
+
+
+#pragma once
+
+
+#include "NodeImagePrivate.h"
+#include "NodeImageGradient.h"
+#include "AST.h"
+
+#include "MemoryPrivate.h"
+#include "MutableMath.h"
+
+
+namespace mu
+{
+
+
+	class NodeImageGradient::Private : public NodeImage::Private
+	{
+	public:
+
+		MUTABLE_DEFINE_CONST_VISITABLE()
+
+	public:
+
+		static NODE_TYPE s_type;
+
+		NodeColourPtr m_pColour0;
+		NodeColourPtr m_pColour1;
+		vec2<int> m_size = { 256,1 };		
+
+		//!
+		void Serialise( OutputArchive& arch ) const
+		{
+            uint32_t ver = 0;
+			arch << ver;
+
+			arch << m_pColour0;
+			arch << m_pColour1;
+			arch << m_size;
+		}
+
+		//!
+		void Unserialise( InputArchive& arch )
+		{
+            uint32_t ver;
+			arch >> ver;
+			check(ver==0);
+
+			arch >> m_pColour0;
+			arch >> m_pColour1;
+			arch >> m_size;
+		}
+	};
+
+
+}
diff --git a/Engine/Plugins/Experimental/Mutable/Source/MutableTools/Private/NodeImageInterpolate.cpp b/Engine/Plugins/Experimental/Mutable/Source/MutableTools/Private/NodeImageInterpolate.cpp
new file mode 100644
--- /dev/null
+++ b/Engine/Plugins/Experimental/Mutable/Source/MutableTools/Private/NodeImageInterpolate.cpp
@@ -0,0 +1,127 @@
+// Copyright Epic Games, Inc. All Rights Reserved.
+
+
+#include "NodeImageInterpolate.h"
+#include "NodeImageInterpolatePrivate.h"
+
+#include "NodeScalar.h"
+
+#include "ImagePrivate.h"
+
+
+namespace mu
+{
+
+	//---------------------------------------------------------------------------------------------
+	// Static initialisation
+	//---------------------------------------------------------------------------------------------
+	NODE_TYPE NodeImageInterpolate::Private::s_type =
+			NODE_TYPE( "ImageInterpolate", NodeImage::GetStaticType() );
+
+
+	//---------------------------------------------------------------------------------------------
+	//!
+	//---------------------------------------------------------------------------------------------
+
+	MUTABLE_IMPLEMENT_NODE( NodeImageInterpolate, EType::Interpolate, Node, Node::EType::Image)
+
+
+	//---------------------------------------------------------------------------------------------
+	// Node Interface
+	//---------------------------------------------------------------------------------------------
+	int NodeImageInterpolate::GetInputCount() const
+	{
+		return 1 + (int)m_pD->m_targets.size();
+	}
+
+
+	//---------------------------------------------------------------------------------------------
+	Node* NodeImageInterpolate::GetInputNode( int i ) const
+	{
+		check( i>=0 && i<GetInputCount() );
+
+		Node* pResult = 0;
+
+		switch (i)
+		{
+		case 0:
+			pResult = m_pD->m_pFactor.get();
+			break;
+
+		default:
+			pResult = m_pD->m_targets[i-1].get();
+			break;
+		}
+
+		return pResult;
+	}
+
+
+	//---------------------------------------------------------------------------------------------
+	void NodeImageInterpolate::SetInputNode( int i, NodePtr pNode )
+	{
+		check( i>=0 && i<GetInputCount() );
+
+		switch (i)
+		{
+		case 0:
+			m_pD->m_pFactor = dynamic_cast<NodeScalar*>(pNode.get());
+			break;
+
+		default:
+			m_pD->m_targets[i-1] = dynamic_cast<NodeImage*>(pNode.get());
+			break;
+		}
+	}
+
+
+	//---------------------------------------------------------------------------------------------
+	// Own Interface
+	//---------------------------------------------------------------------------------------------
+	NodeScalarPtr NodeImageInterpolate::GetFactor() const
+	{
+		return m_pD->m_pFactor.get();
+	}
+
+
+	//---------------------------------------------------------------------------------------------
+	void NodeImageInterpolate::SetFactor( NodeScalarPtr pNode )
+	{
+		m_pD->m_pFactor = pNode;
+	}
+
+
+	//---------------------------------------------------------------------------------------------
+	NodeImagePtr NodeImageInterpolate::GetTarget( int t ) const
+	{
+		check( t>=0 && t<(int)m_pD->m_targets.size() );
+		return m_pD->m_targets[t].get();
+	}
+
+
+	//---------------------------------------------------------------------------------------------
+	void NodeImageInterpolate::SetTarget( int t, NodeImagePtr pNode )
+	{
+		check( t>=0 && t<(int)m_pD->m_targets.size() );
+		m_pD->m_targets[t] = pNode;
+	}
+
+
+	//---------------------------------------------------------------------------------------------
+	void NodeImageInterpolate::SetTargetCount( int t )
+	{
+		m_pD->m_targets.resize(t);
+	}
+
+
+	//---------------------------------------------------------------------------------------------
+	int NodeImageInterpolate::GetTargetCount() const
+	{
+		return int(m_pD->m_targets.size());
+	}
+
+
+
+}
+
+
diff --git a/Engine/Plugins/Experimental/Mutable/Source/MutableTools/Private/NodeImageInterpolate3.cpp b/Engine/Plugins/Experimental/Mutable/Source/MutableTools/Private/NodeImageInterpolate3.cpp
new file mode 100644
--- /dev/null
+++ b/Engine/Plugins/Experimental/Mutable/Source/MutableTools/Private/NodeImageInterpolate3.cpp
@@ -0,0 +1,153 @@
+// Copyright Epic Games, Inc. All Rights Reserved.
+
+
+#include "NodeImageInterpolate3.h"
+#include "NodeImageInterpolate3Private.h"
+
+#include "NodeScalar.h"
+
+#include "ImagePrivate.h"
+
+
+namespace mu
+{
+
+	//---------------------------------------------------------------------------------------------
+	// Static initialisation
+	//---------------------------------------------------------------------------------------------
+	NODE_TYPE NodeImageInterpolate3::Private::s_type =
+			NODE_TYPE( "ImageInterpolate3", NodeImage::GetStaticType() );
+
+
+	//---------------------------------------------------------------------------------------------
+	//!
+	//---------------------------------------------------------------------------------------------
+
+	MUTABLE_IMPLEMENT_NODE( NodeImageInterpolate3, EType::Interpolate3, Node, Node::EType::Image)
+
+
+	//---------------------------------------------------------------------------------------------
+	// Node Interface
+	//---------------------------------------------------------------------------------------------
+	int NodeImageInterpolate3::GetInputCount() const
+	{
+		return 5;
+	}
+
+
+	//---------------------------------------------------------------------------------------------
+	Node* NodeImageInterpolate3::GetInputNode( int i ) const
+	{
+		check( i>=0 && i<GetInputCount() );
+
+		Node* pResult = 0;
+
+		switch (i)
+		{
+		case 0: pResult = m_pD->m_pFactor1.get(); break;
+		case 1: pResult = m_pD->m_pFactor2.get(); break;
+		case 2: pResult = m_pD->m_pTarget0.get(); break;
+		case 3: pResult = m_pD->m_pTarget1.get(); break;
+		case 4: pResult = m_pD->m_pTarget2.get(); break;
+		default:
+			break;
+		}
+
+		return pResult;
+	}
+
+
+	//---------------------------------------------------------------------------------------------
+	void NodeImageInterpolate3::SetInputNode( int i, NodePtr pNode )
+	{
+		check( i>=0 && i<GetInputCount() );
+
+		switch (i)
+		{
+		case 0: m_pD->m_pFactor1 = dynamic_cast<NodeScalar*>(pNode.get()); break;
+		case 1: m_pD->m_pFactor2 = dynamic_cast<NodeScalar*>(pNode.get()); break;
+		case 2: m_pD->m_pTarget0 = dynamic_cast<NodeImage*>(pNode.get()); break;
+		case 3: m_pD->m_pTarget1 = dynamic_cast<NodeImage*>(pNode.get()); break;
+		case 4: m_pD->m_pTarget2 = dynamic_cast<NodeImage*>(pNode.get()); break;
+		default:
+			break;
+		}
+	}
+
+
+	//---------------------------------------------------------------------------------------------
+	// Own Interface
+	//---------------------------------------------------------------------------------------------
+	NodeScalarPtr NodeImageInterpolate3::GetFactor1() const
+	{
+		return m_pD->m_pFactor1.get();
+	}
+
+
+	//---------------------------------------------------------------------------------------------
+	void NodeImageInterpolate3::SetFactor1( NodeScalarPtr pNode )
+	{
+		m_pD->m_pFactor1 = pNode;
+	}
+
+
+	//---------------------------------------------------------------------------------------------
+	NodeScalarPtr NodeImageInterpolate3::GetFactor2() const
+	{
+		return m_pD->m_pFactor2.get();
+	}
+
+
+	//---------------------------------------------------------------------------------------------
+	void NodeImageInterpolate3::SetFactor2( NodeScalarPtr pNode )
+	{
+		m_pD->m_pFactor2 = pNode;
+	}
+
+
+	//---------------------------------------------------------------------------------------------
+	NodeImagePtr NodeImageInterpolate3::GetTarget0() const
+	{
+		return m_pD->m_pTarget0.get();
+	}
+
+
+	//---------------------------------------------------------------------------------------------
+	void NodeImageInterpolate3::SetTarget0( NodeImagePtr pNode )
+	{
+		m_pD->m_pTarget0 = pNode;
+	}
+
+
+	//---------------------------------------------------------------------------------------------
+	NodeImagePtr NodeImageInterpolate3::GetTarget1() const
+	{
+		return m_pD->m_pTarget1.get();
+	}
+
+
+	//---------------------------------------------------------------------------------------------
+	void NodeImageInterpolate3::SetTarget1( NodeImagePtr pNode )
+	{
+		m_pD->m_pTarget1 = pNode;
+	}
+
+
+	//---------------------------------------------------------------------------------------------
+	NodeImagePtr NodeImageInterpolate3::GetTarget2() const
+	{
+		return m_pD->m_pTarget2.get();
+	}
+
+
+	//---------------------------------------------------------------------------------------------
+	void NodeImageInterpolate3::SetTarget2( NodeImagePtr pNode )
+	{
+		m_pD->m_pTarget2 = pNode;
+	}
+
+
+
+}
+
+
diff --git a/Engine/Plugins/Experimental/Mutable/Source/MutableTools/Private/NodeImageInterpolate3Private.h b/Engine/Plugins/Experimental/Mutable/Source/MutableTools/Private/NodeImageInterpolate3Private.h
new file mode 100644
--- /dev/null
+++ b/Engine/Plugins/Experimental/Mutable/Source/MutableTools/Private/NodeImageInterpolate3Private.h
@@ -0,0 +1,60 @@
+// Copyright Epic Games, Inc. All Rights Reserved.
+
+#pragma once
+
+#include "NodeImageInterpolate3.h"
+#include "NodeImagePrivate.h"
+#include "NodeScalar.h"
+#include "AST.h"
+
+
+namespace mu
+{
+
+
+	class NodeImageInterpolate3::Private : public NodeImage::Private
+	{
+	public:
+
+		MUTABLE_DEFINE_CONST_VISITABLE()
+
+	public:
+
+		static NODE_TYPE s_type;
+
+		NodeScalarPtr m_pFactor1;
+		NodeScalarPtr m_pFactor2;
+		NodeImagePtr m_pTarget0;
+		NodeImagePtr m_pTarget1;
+		NodeImagePtr m_pTarget2;
+
+		//!
+		void Serialise( OutputArchive& arch ) const
+		{
+            uint32_t ver = 0;
+			arch << ver;
+
+			arch << m_pFactor1;
+			arch << m_pFactor2;
+			arch << m_pTarget0;
+			arch << m_pTarget1;
+			arch << m_pTarget2;
+		}
+
+		//!
+		void Unserialise( InputArchive& arch )
+		{
+            uint32_t ver;
+			arch >> ver;
+			check(ver==0);
+
+			arch >> m_pFactor1;
+			arch >> m_pFactor2;
+			arch >> m_pTarget0;
+			arch >> m_pTarget1;
+			arch >> m_pTarget2;
+		}
+	};
+
+
+}
diff --git a/Engine/Plugins/Experimental/Mutable/Source/MutableTools/Private/NodeImageInterpolatePrivate.h b/Engine/Plugins/Experimental/Mutable/Source/MutableTools/Private/NodeImageInterpolatePrivate.h
new file mode 100644
--- /dev/null
+++ b/Engine/Plugins/Experimental/Mutable/Source/MutableTools/Private/NodeImageInterpolatePrivate.h
@@ -0,0 +1,50 @@
+// Copyright Epic Games, Inc. All Rights Reserved.
+
+#pragma once
+
+#include "NodeImageInterpolate.h"
+#include "NodeImagePrivate.h"
+#include "NodeScalar.h"
+
+
+namespace mu
+{
+
+
+	class NodeImageInterpolate::Private : public NodeImage::Private
+	{
+	public:
+
+		MUTABLE_DEFINE_CONST_VISITABLE()
+
+	public:
+
+		static NODE_TYPE s_type;
+
+		NodeScalarPtr m_pFactor;
+		vector<NodeImagePtr> m_targets;
+
+		//!
+		void Serialise( OutputArchive& arch ) const
+		{
+            uint32_t ver = 0;
+			arch << ver;
+
+			arch << m_pFactor;
+			arch << m_targets;
+		}
+
+		//!
+		void Unserialise( InputArchive& arch )
+		{
+            uint32_t ver;
+			arch >> ver;
+			check(ver==0);
+
+			arch >> m_pFactor;
+			arch >> m_targets;
+		}
+	};
+
+
+}
diff --git a/Engine/Plugins/Experimental/Mutable/Source/MutableTools/Private/NodeImageInvert.cpp b/Engine/Plugins/Experimental/Mutable/Source/MutableTools/Private/NodeImageInvert.cpp
new file mode 100644
--- /dev/null
+++ b/Engine/Plugins/Experimental/Mutable/Source/MutableTools/Private/NodeImageInvert.cpp
@@ -0,0 +1,67 @@
+// Copyright Epic Games, Inc. All Rights Reserved.
+
+
+#include "NodeImageInvert.h"
+#include "NodeImageInvertPrivate.h"
+
+#include "ImagePrivate.h"
+
+
+#define NODE_INPUT_COUNT 1
+
+namespace mu
+{
+
+	//---------------------------------------------------------------------------------------------
+	// Static initialisation
+	//---------------------------------------------------------------------------------------------
+	NODE_TYPE NodeImageInvert::Private::s_type =
+		NODE_TYPE("ImageInvert", NodeImage::GetStaticType());
+
+	//---------------------------------------------------------------------------------------------
+
+	MUTABLE_IMPLEMENT_NODE( NodeImageInvert, EType::Invert, Node, Node::EType::Image)
+
+	//---------------------------------------------------------------------------------------------
+	// Node Interface
+	//---------------------------------------------------------------------------------------------
+	void NodeImageInvert::SetInputNode(int i, NodePtr pNode)
+	{
+		check(i >= 0 && i < NODE_INPUT_COUNT);
+
+		m_pD->m_pBase = dynamic_cast<NodeImage*>(pNode.get());
+	}
+
+	//---------------------------------------------------------------------------------------------
+	int NodeImageInvert::GetInputCount() const
+	{
+		return NODE_INPUT_COUNT;
+	}
+
+	//---------------------------------------------------------------------------------------------
+	Node* NodeImageInvert::GetInputNode(int i)const
+	{
+		check(i >= 0 && i < NODE_INPUT_COUNT);
+
+		Node* pResult = 0;
+		pResult = m_pD->m_pBase.get();
+
+		return pResult;
+	}
+
+	//---------------------------------------------------------------------------------------------
+	// Own Interface
+	//---------------------------------------------------------------------------------------------
+	NodeImagePtr NodeImageInvert::GetBase()const
+	{
+		return m_pD->m_pBase;
+	}
+
+	void NodeImageInvert::SetBase(NodeImagePtr pNode)
+	{
+		m_pD->m_pBase = pNode;
+	}
+	
+}
+
+#undef NODE_INPUT_COUNT
diff --git a/Engine/Plugins/Experimental/Mutable/Source/MutableTools/Private/NodeImageInvertPrivate.h b/Engine/Plugins/Experimental/Mutable/Source/MutableTools/Private/NodeImageInvertPrivate.h
new file mode 100644
--- /dev/null
+++ b/Engine/Plugins/Experimental/Mutable/Source/MutableTools/Private/NodeImageInvertPrivate.h
@@ -0,0 +1,47 @@
+// Copyright Epic Games, Inc. All Rights Reserved.
+
+#pragma once
+
+#include "NodeImagePrivate.h"
+#include "NodeImageInvert.h"
+#include "AST.h"
+
+#include "MemoryPrivate.h"
+
+
+namespace mu
+{
+
+
+	class NodeImageInvert::Private : public NodeImage::Private
+	{
+	public:
+
+		MUTABLE_DEFINE_CONST_VISITABLE();
+
+	public:
+
+		static NODE_TYPE s_type;
+
+		NodeImagePtr m_pBase;
+		
+		//!
+		void Serialise(OutputArchive& arch) const
+		{
+			uint32_t ver = 0;
+			arch << ver;
+
+			arch << m_pBase;
+		}
+
+		//!
+		void Unserialise(InputArchive& arch)
+		{
+			uint32_t ver;
+			arch >> ver;
+			check(ver == 0);
+
+			arch >> m_pBase;
+		}
+	};
+}
diff --git a/Engine/Plugins/Experimental/Mutable/Source/MutableTools/Private/NodeImageLayer.cpp b/Engine/Plugins/Experimental/Mutable/Source/MutableTools/Private/NodeImageLayer.cpp
new file mode 100644
--- /dev/null
+++ b/Engine/Plugins/Experimental/Mutable/Source/MutableTools/Private/NodeImageLayer.cpp
@@ -0,0 +1,148 @@
+// Copyright Epic Games, Inc. All Rights Reserved.
+
+
+#include "NodeImageLayer.h"
+#include "NodeImageLayerPrivate.h"
+
+#include "NodeColour.h"
+
+#include "ImagePrivate.h"
+
+
+#define NODE_INPUT_COUNT 	3
+
+
+namespace mu
+{
+
+	//---------------------------------------------------------------------------------------------
+	// Static initialisation
+	//---------------------------------------------------------------------------------------------
+	NODE_TYPE NodeImageLayer::Private::s_type =
+			NODE_TYPE( "ImageLayer", NodeImage::GetStaticType() );
+
+
+	//---------------------------------------------------------------------------------------------
+	//!
+	//---------------------------------------------------------------------------------------------
+
+	MUTABLE_IMPLEMENT_NODE( NodeImageLayer, EType::Layer, Node, Node::EType::Image)
+
+
+	//---------------------------------------------------------------------------------------------
+	const char* NodeImageLayer::s_blendTypeName[] =
+	{
+		"SoftLight",
+		"HardLight",
+		"Burn",
+		"Dodge",
+		"Screen",
+		"Overlay",
+		"Blend",
+        "Multiply",
+        "AlphaOverlay"
+	};
+
+
+	//---------------------------------------------------------------------------------------------
+	// Node Interface
+	//---------------------------------------------------------------------------------------------
+	int NodeImageLayer::GetInputCount() const
+	{
+		return NODE_INPUT_COUNT;
+	}
+
+
+	//---------------------------------------------------------------------------------------------
+	Node* NodeImageLayer::GetInputNode( int i ) const
+	{
+		check( i>=0 && i<NODE_INPUT_COUNT );
+
+		Node* pResult = 0;
+
+		switch (i)
+		{
+        case 0: pResult = m_pD->m_pBase.get(); break;
+        case 1: pResult = m_pD->m_pMask.get(); break;
+        case 2: pResult = m_pD->m_pBlended.get(); break;
+		}
+
+		return pResult;
+	}
+
+
+	//---------------------------------------------------------------------------------------------
+	void NodeImageLayer::SetInputNode( int i, NodePtr pNode )
+	{
+		check( i>=0 && i<NODE_INPUT_COUNT );
+
+		switch (i)
+		{
+        case 0: m_pD->m_pBase = dynamic_cast<NodeImage*>( pNode.get() ); break;
+        case 1: m_pD->m_pMask = dynamic_cast<NodeImage*>( pNode.get() ); break;
+        case 2: m_pD->m_pBlended = dynamic_cast<NodeImage*>( pNode.get() ); break;
+		}
+	}
+
+
+	//---------------------------------------------------------------------------------------------
+	// Own Interface
+	//---------------------------------------------------------------------------------------------
+	NodeImagePtr NodeImageLayer::GetBase() const
+	{
+		return m_pD->m_pBase.get();
+	}
+
+
+	//---------------------------------------------------------------------------------------------
+	void NodeImageLayer::SetBase( NodeImagePtr pNode )
+	{
+		m_pD->m_pBase = pNode;
+	}
+
+
+	//---------------------------------------------------------------------------------------------
+	NodeImagePtr NodeImageLayer::GetMask() const
+	{
+		return m_pD->m_pMask.get();
+	}
+
+
+	//---------------------------------------------------------------------------------------------
+	void NodeImageLayer::SetMask( NodeImagePtr pNode )
+	{
+		m_pD->m_pMask = pNode;
+	}
+
+
+	//---------------------------------------------------------------------------------------------
+	NodeImagePtr NodeImageLayer::GetBlended() const
+	{
+		return m_pD->m_pBlended.get();
+	}
+
+
+	//---------------------------------------------------------------------------------------------
+	void NodeImageLayer::SetBlended( NodeImagePtr pNode )
+	{
+		m_pD->m_pBlended = pNode;
+	}
+
+
+	//---------------------------------------------------------------------------------------------
+	EBlendType NodeImageLayer::GetBlendType() const
+	{
+		return m_pD->m_type;
+	}
+
+
+	//---------------------------------------------------------------------------------------------
+	void NodeImageLayer::SetBlendType(EBlendType t)
+	{
+		m_pD->m_type = t;
+	}
+
+
+}
+
+#undef NODE_INPUT_COUNT
diff --git a/Engine/Plugins/Experimental/Mutable/Source/MutableTools/Private/NodeImageLayerColour.cpp b/Engine/Plugins/Experimental/Mutable/Source/MutableTools/Private/NodeImageLayerColour.cpp
new file mode 100644
--- /dev/null
+++ b/Engine/Plugins/Experimental/Mutable/Source/MutableTools/Private/NodeImageLayerColour.cpp
@@ -0,0 +1,128 @@
+// Copyright Epic Games, Inc. All Rights Reserved.
+
+
+#include "NodeImageLayerColour.h"
+#include "NodeImageLayerColourPrivate.h"
+
+#include "NodeColour.h"
+
+#include "ImagePrivate.h"
+
+
+#define NODE_INPUT_COUNT 	3
+
+
+namespace mu
+{
+
+	//---------------------------------------------------------------------------------------------
+	// Static initialisation
+	//---------------------------------------------------------------------------------------------
+	NODE_TYPE NodeImageLayerColour::Private::s_type =
+			NODE_TYPE( "ImageLayerColour", NodeImage::GetStaticType() );
+
+
+	//---------------------------------------------------------------------------------------------
+	//!
+	//---------------------------------------------------------------------------------------------
+
+	MUTABLE_IMPLEMENT_NODE( NodeImageLayerColour, EType::LayerColour, Node, Node::EType::Image);
+
+
+	//---------------------------------------------------------------------------------------------
+	// Node Interface
+	//---------------------------------------------------------------------------------------------
+	int NodeImageLayerColour::GetInputCount() const
+	{
+		return NODE_INPUT_COUNT;
+	}
+
+
+	//---------------------------------------------------------------------------------------------
+	Node* NodeImageLayerColour::GetInputNode( int i ) const
+	{
+		check( i>=0 && i<NODE_INPUT_COUNT );
+
+		Node* pResult = 0;
+
+		switch (i)
+		{
+        case 0: pResult = m_pD->m_pBase.get(); break;
+        case 1: pResult = m_pD->m_pMask.get(); break;
+        case 2: pResult = m_pD->m_pColour.get(); break;
+		}
+
+		return pResult;
+	}
+
+
+	//---------------------------------------------------------------------------------------------
+	void NodeImageLayerColour::SetInputNode( int i, NodePtr pNode )
+	{
+		check( i>=0 && i<NODE_INPUT_COUNT );
+
+		switch (i)
+		{
+        case 0: m_pD->m_pBase = dynamic_cast<NodeImage*>( pNode.get() ); break;
+        case 1: m_pD->m_pMask = dynamic_cast<NodeImage*>( pNode.get() ); break;
+        case 2: m_pD->m_pColour = dynamic_cast<NodeColour*>( pNode.get() ); break;
+		}
+	}
+
+
+	//---------------------------------------------------------------------------------------------
+	// Own Interface
+	//---------------------------------------------------------------------------------------------
+	NodeImagePtr NodeImageLayerColour::GetBase() const
+	{
+		return m_pD->m_pBase.get();
+	}
+
+	//---------------------------------------------------------------------------------------------
+	void NodeImageLayerColour::SetBase( NodeImagePtr pNode )
+	{
+		m_pD->m_pBase = pNode;
+	}
+
+	//---------------------------------------------------------------------------------------------
+	NodeImagePtr NodeImageLayerColour::GetMask() const
+	{
+		return m_pD->m_pMask.get();
+	}
+
+	//---------------------------------------------------------------------------------------------
+	void NodeImageLayerColour::SetMask( NodeImagePtr pNode )
+	{
+		m_pD->m_pMask = pNode;
+	}
+
+	//---------------------------------------------------------------------------------------------
+	NodeColourPtr NodeImageLayerColour::GetColour() const
+	{
+		return m_pD->m_pColour.get();
+	}
+
+	//---------------------------------------------------------------------------------------------
+	void NodeImageLayerColour::SetColour( NodeColourPtr pNode )
+	{
+		m_pD->m_pColour = pNode;
+	}
+
+
+	//---------------------------------------------------------------------------------------------
+	EBlendType NodeImageLayerColour::GetBlendType() const
+	{
+		return m_pD->m_type;
+	}
+
+
+	//---------------------------------------------------------------------------------------------
+	void NodeImageLayerColour::SetBlendType(EBlendType t)
+	{
+		m_pD->m_type = t;
+	}
+
+}
+
+#undef NODE_INPUT_COUNT
+
diff --git a/Engine/Plugins/Experimental/Mutable/Source/MutableTools/Private/NodeImageLayerColourPrivate.h b/Engine/Plugins/Experimental/Mutable/Source/MutableTools/Private/NodeImageLayerColourPrivate.h
new file mode 100644
--- /dev/null
+++ b/Engine/Plugins/Experimental/Mutable/Source/MutableTools/Private/NodeImageLayerColourPrivate.h
@@ -0,0 +1,57 @@
+// Copyright Epic Games, Inc. All Rights Reserved.
+
+#pragma once
+
+#include "NodeImagePrivate.h"
+#include "NodeImageLayerPrivate.h"
+#include "NodeImageLayerColour.h"
+#include "NodeColour.h"
+
+#include "MemoryPrivate.h"
+
+
+namespace mu
+{
+
+	class NodeImageLayerColour::Private : public NodeImage::Private
+	{
+	public:
+
+		MUTABLE_DEFINE_CONST_VISITABLE()
+
+	public:
+
+		static NODE_TYPE s_type;
+
+		NodeImagePtr m_pBase;
+		NodeImagePtr m_pMask;
+		NodeColourPtr m_pColour;
+		EBlendType m_type;
+
+		//!
+		void Serialise( OutputArchive& arch ) const
+		{
+            uint32_t ver = 0;
+			arch << ver;
+
+			arch << m_pBase;
+			arch << m_pMask;
+			arch << m_pColour;
+			arch << m_type;
+		}
+
+		//!
+		void Unserialise( InputArchive& arch )
+		{
+            uint32_t ver;
+			arch >> ver;
+			check(ver==0);
+
+			arch >> m_pBase;
+			arch >> m_pMask;
+			arch >> m_pColour;
+			arch >> m_type;
+		}
+	};
+
+}
diff --git a/Engine/Plugins/Experimental/Mutable/Source/MutableTools/Private/NodeImageLayerPrivate.h b/Engine/Plugins/Experimental/Mutable/Source/MutableTools/Private/NodeImageLayerPrivate.h
new file mode 100644
--- /dev/null
+++ b/Engine/Plugins/Experimental/Mutable/Source/MutableTools/Private/NodeImageLayerPrivate.h
@@ -0,0 +1,60 @@
+// Copyright Epic Games, Inc. All Rights Reserved.
+
+#pragma once
+
+#include "NodeImagePrivate.h"
+#include "NodeImageLayer.h"
+#include "NodeColour.h"
+#include "AST.h"
+
+#include "MemoryPrivate.h"
+
+
+namespace mu
+{
+
+	class NodeImageLayer::Private : public NodeImage::Private
+	{
+	public:
+
+		MUTABLE_DEFINE_CONST_VISITABLE()
+
+	public:
+
+		static NODE_TYPE s_type;
+
+		NodeImagePtr m_pBase;
+		NodeImagePtr m_pMask;
+		NodeImagePtr m_pBlended;
+		EBlendType m_type;
+
+		//!
+		void Serialise( OutputArchive& arch ) const
+		{
+            uint32_t ver = 0;
+			arch << ver;
+
+			arch << m_pBase;
+			arch << m_pMask;
+			arch << m_pBlended;
+            arch << (uint32_t)m_type;
+		}
+
+		//!
+		void Unserialise( InputArchive& arch )
+		{
+            uint32_t ver;
+			arch >> ver;
+			check(ver==0);
+
+			arch >> m_pBase;
+			arch >> m_pMask;
+			arch >> m_pBlended;
+
+            uint32_t t;
+            arch >> t;
+            m_type=(EBlendType)t;
+		}
+	};
+
+}
diff --git a/Engine/Plugins/Experimental/Mutable/Source/MutableTools/Private/NodeImageLuminance.cpp b/Engine/Plugins/Experimental/Mutable/Source/MutableTools/Private/NodeImageLuminance.cpp
new file mode 100644
--- /dev/null
+++ b/Engine/Plugins/Experimental/Mutable/Source/MutableTools/Private/NodeImageLuminance.cpp
@@ -0,0 +1,88 @@
+// Copyright Epic Games, Inc. All Rights Reserved.
+
+
+#include "NodeImageLuminance.h"
+#include "NodeImageLuminancePrivate.h"
+
+#include "NodeScalar.h"
+
+#include "ImagePrivate.h"
+
+
+#define NODE_INPUT_COUNT 	1
+
+
+namespace mu
+{
+
+	//---------------------------------------------------------------------------------------------
+	// Static initialisation
+	//---------------------------------------------------------------------------------------------
+	NODE_TYPE NodeImageLuminance::Private::s_type =
+			NODE_TYPE( "ImageLuminance", NodeImage::GetStaticType() );
+
+
+	//---------------------------------------------------------------------------------------------
+	//!
+	//---------------------------------------------------------------------------------------------
+
+	MUTABLE_IMPLEMENT_NODE( NodeImageLuminance, EType::Luminance, Node, Node::EType::Image)
+
+
+	//---------------------------------------------------------------------------------------------
+	// Node Interface
+	//---------------------------------------------------------------------------------------------
+	int NodeImageLuminance::GetInputCount() const
+	{
+		return NODE_INPUT_COUNT;
+	}
+
+
+	//---------------------------------------------------------------------------------------------
+	Node* NodeImageLuminance::GetInputNode( int i ) const
+	{
+		check( i>=0 && i<NODE_INPUT_COUNT );
+
+		Node* pResult = 0;
+
+		switch (i)
+		{
+        case 0: pResult = m_pD->m_pSource.get(); break;
+		}
+
+		return pResult;
+	}
+
+
+	//---------------------------------------------------------------------------------------------
+	void NodeImageLuminance::SetInputNode( int i, NodePtr pNode )
+	{
+		check( i>=0 && i<NODE_INPUT_COUNT );
+
+		switch (i)
+		{
+        case 0: m_pD->m_pSource = dynamic_cast<NodeImage*>(pNode.get()); break;
+		}
+	}
+
+
+	//---------------------------------------------------------------------------------------------
+	// Own Interface
+	//---------------------------------------------------------------------------------------------
+	NodeImagePtr NodeImageLuminance::GetSource() const
+	{
+		return m_pD->m_pSource.get();
+	}
+
+
+	//---------------------------------------------------------------------------------------------
+	void NodeImageLuminance::SetSource( NodeImagePtr pNode )
+	{
+		m_pD->m_pSource = pNode;
+	}
+
+
+
+}
+
+#undef NODE_INPUT_COUNT
diff --git a/Engine/Plugins/Experimental/Mutable/Source/MutableTools/Private/NodeImageLuminancePrivate.h b/Engine/Plugins/Experimental/Mutable/Source/MutableTools/Private/NodeImageLuminancePrivate.h
new file mode 100644
--- /dev/null
+++ b/Engine/Plugins/Experimental/Mutable/Source/MutableTools/Private/NodeImageLuminancePrivate.h
@@ -0,0 +1,48 @@
+// Copyright Epic Games, Inc. All Rights Reserved.
+
+#pragma once
+
+#include "NodeImagePrivate.h"
+#include "NodeImageLuminance.h"
+
+#include "MemoryPrivate.h"
+
+
+namespace mu
+{
+
+
+	class NodeImageLuminance::Private : public NodeImage::Private
+	{
+	public:
+
+		MUTABLE_DEFINE_CONST_VISITABLE()
+
+	public:
+
+		static NODE_TYPE s_type;
+
+		NodeImagePtr m_pSource;
+
+		//!
+		void Serialise( OutputArchive& arch ) const
+		{
+            uint32_t ver = 0;
+			arch << ver;
+
+			arch << m_pSource;
+		}
+
+		//!
+		void Unserialise( InputArchive& arch )
+		{
+            uint32_t ver;
+			arch >> ver;
+			check(ver==0);
+
+			arch >> m_pSource;
+		}
+	};
+
+
+}
diff --git a/Engine/Plugins/Experimental/Mutable/Source/MutableTools/Private/NodeImageMipmap.cpp b/Engine/Plugins/Experimental/Mutable/Source/MutableTools/Private/NodeImageMipmap.cpp
new file mode 100644
--- /dev/null
+++ b/Engine/Plugins/Experimental/Mutable/Source/MutableTools/Private/NodeImageMipmap.cpp
@@ -0,0 +1,95 @@
+// Copyright Epic Games, Inc. All Rights Reserved.
+
+
+#include "NodeImageMipmap.h"
+#include "NodeImageMipmapPrivate.h"
+
+#include "NodeScalar.h"
+
+#include "ImagePrivate.h"
+
+
+#define NODE_INPUT_COUNT 	1
+
+
+namespace mu
+{
+
+	//---------------------------------------------------------------------------------------------
+	// Static initialisation
+	//---------------------------------------------------------------------------------------------
+    NODE_TYPE NodeImageMipmap::Private::s_type =
+            NODE_TYPE( "ImageMipmap", NodeImage::GetStaticType() );
+
+
+	//---------------------------------------------------------------------------------------------
+	//!
+	//---------------------------------------------------------------------------------------------
+
+    MUTABLE_IMPLEMENT_NODE( NodeImageMipmap, EType::Mipmap, Node, Node::EType::Image)
+
+
+	//---------------------------------------------------------------------------------------------
+	// Node Interface
+	//---------------------------------------------------------------------------------------------
+    int NodeImageMipmap::GetInputCount() const
+	{
+		return NODE_INPUT_COUNT;
+	}
+
+
+	//---------------------------------------------------------------------------------------------
+    Node* NodeImageMipmap::GetInputNode( int i ) const
+	{
+		check( i>=0 && i<NODE_INPUT_COUNT );
+
+		Node* pResult = 0;
+
+		switch (i)
+		{
+        case 0: pResult = m_pD->m_pSource.get(); break;
+		}
+
+		return pResult;
+	}
+
+
+	//---------------------------------------------------------------------------------------------
+    void NodeImageMipmap::SetInputNode( int i, NodePtr pNode )
+	{
+		check( i>=0 && i<NODE_INPUT_COUNT );
+
+		switch (i)
+		{
+        case 0: m_pD->m_pSource = dynamic_cast<NodeImage*>(pNode.get()); break;
+		}
+	}
+
+
+	//---------------------------------------------------------------------------------------------
+	// Own Interface
+	//---------------------------------------------------------------------------------------------
+    NodeImagePtr NodeImageMipmap::GetSource() const
+	{
+		return m_pD->m_pSource.get();
+	}
+
+
+	//---------------------------------------------------------------------------------------------
+    void NodeImageMipmap::SetSource( NodeImagePtr pNode )
+	{
+		m_pD->m_pSource = pNode;
+	}
+
+	//---------------------------------------------------------------------------------------------
+	void NodeImageMipmap::SetMipmapGenerationSettings( EMipmapFilterType filterType,
+													   EAddressMode addressMode,
+													   float sharpenFactor,
+													   bool mipDitherAlpha )
+	{
+		m_pD->m_settings = FMipmapGenerationSettings{ sharpenFactor, filterType, addressMode, mipDitherAlpha };
+	}
+
+}
+
+#undef NODE_INPUT_COUNT
diff --git a/Engine/Plugins/Experimental/Mutable/Source/MutableTools/Private/NodeImageMipmapPrivate.h b/Engine/Plugins/Experimental/Mutable/Source/MutableTools/Private/NodeImageMipmapPrivate.h
new file mode 100644
--- /dev/null
+++ b/Engine/Plugins/Experimental/Mutable/Source/MutableTools/Private/NodeImageMipmapPrivate.h
@@ -0,0 +1,58 @@
+// Copyright Epic Games, Inc. All Rights Reserved.
+
+#pragma once
+
+#include "NodeImagePrivate.h"
+#include "NodeImageMipmap.h"
+#include "NodeScalar.h"
+#include "AST.h"
+
+#include "ImagePrivate.h"
+#include "MemoryPrivate.h"
+
+
+namespace mu
+{	
+    class NodeImageMipmap::Private : public NodeImage::Private
+	{
+	public:
+
+		MUTABLE_DEFINE_CONST_VISITABLE()
+
+	public:
+
+		static NODE_TYPE s_type;
+
+		NodeImagePtr m_pSource;
+		NodeScalarPtr m_pFactor;
+
+		FMipmapGenerationSettings m_settings;
+
+		//!
+		void Serialise( OutputArchive& arch ) const
+		{
+            uint32_t ver = 1;
+			arch << ver;
+
+			arch << m_pSource;
+			arch << m_pFactor;
+
+			arch << m_settings;
+		}
+
+		//!
+		void Unserialise( InputArchive& arch )
+		{
+            uint32_t ver;
+			arch >> ver;
+			check(ver==1);
+
+			arch >> m_pSource;
+			arch >> m_pFactor;
+			
+			arch >> m_settings;
+		}	
+	};
+
+
+}
diff --git a/Engine/Plugins/Experimental/Mutable/Source/MutableTools/Private/NodeImageMultiLayer.cpp b/Engine/Plugins/Experimental/Mutable/Source/MutableTools/Private/NodeImageMultiLayer.cpp
new file mode 100644
--- /dev/null
+++ b/Engine/Plugins/Experimental/Mutable/Source/MutableTools/Private/NodeImageMultiLayer.cpp
@@ -0,0 +1,149 @@
+// Copyright Epic Games, Inc. All Rights Reserved.
+
+
+#include "NodeImageMultiLayer.h"
+#include "NodeImageMultiLayerPrivate.h"
+
+#include "NodeColour.h"
+
+#include "ImagePrivate.h"
+
+
+#define NODE_INPUT_COUNT 	4
+
+
+namespace mu
+{
+
+	//---------------------------------------------------------------------------------------------
+	// Static initialisation
+	//---------------------------------------------------------------------------------------------
+    NODE_TYPE NodeImageMultiLayer::Private::s_type =
+            NODE_TYPE( "ImageMultiLayer", NodeImage::GetStaticType() );
+
+
+	//---------------------------------------------------------------------------------------------
+	//!
+	//---------------------------------------------------------------------------------------------
+
+    MUTABLE_IMPLEMENT_NODE( NodeImageMultiLayer, EType::MultiLayer, Node, Node::EType::Image)
+
+
+	//---------------------------------------------------------------------------------------------
+	// Node Interface
+	//---------------------------------------------------------------------------------------------
+    int NodeImageMultiLayer::GetInputCount() const
+	{
+		return NODE_INPUT_COUNT;
+	}
+
+
+	//---------------------------------------------------------------------------------------------
+    Node* NodeImageMultiLayer::GetInputNode( int i ) const
+	{
+		check( i>=0 && i<NODE_INPUT_COUNT );
+
+		Node* pResult = 0;
+
+		switch (i)
+		{
+        case 0: pResult = m_pD->m_pBase.get(); break;
+        case 1: pResult = m_pD->m_pMask.get(); break;
+        case 2: pResult = m_pD->m_pBlended.get(); break;
+        case 3: pResult = m_pD->m_pRange.get(); break;
+        }
+
+		return pResult;
+	}
+
+
+	//---------------------------------------------------------------------------------------------
+    void NodeImageMultiLayer::SetInputNode( int i, NodePtr pNode )
+	{
+		check( i>=0 && i<NODE_INPUT_COUNT );
+
+		switch (i)
+		{
+        case 0: m_pD->m_pBase = dynamic_cast<NodeImage*>( pNode.get() ); break;
+        case 1: m_pD->m_pMask = dynamic_cast<NodeImage*>( pNode.get() ); break;
+        case 2: m_pD->m_pBlended = dynamic_cast<NodeImage*>( pNode.get() ); break;
+        case 3: m_pD->m_pRange = dynamic_cast<NodeRange*>( pNode.get() ); break;
+        }
+	}
+
+
+	//---------------------------------------------------------------------------------------------
+	// Own Interface
+	//---------------------------------------------------------------------------------------------
+    NodeImagePtr NodeImageMultiLayer::GetBase() const
+	{
+		return m_pD->m_pBase.get();
+	}
+
+
+	//---------------------------------------------------------------------------------------------
+    void NodeImageMultiLayer::SetBase( NodeImagePtr pNode )
+	{
+		m_pD->m_pBase = pNode;
+	}
+
+
+	//---------------------------------------------------------------------------------------------
+    NodeImagePtr NodeImageMultiLayer::GetMask() const
+	{
+		return m_pD->m_pMask.get();
+	}
+
+
+	//---------------------------------------------------------------------------------------------
+    void NodeImageMultiLayer::SetMask( NodeImagePtr pNode )
+	{
+		m_pD->m_pMask = pNode;
+	}
+
+
+    //---------------------------------------------------------------------------------------------
+    NodeImagePtr NodeImageMultiLayer::GetBlended() const
+    {
+        return m_pD->m_pBlended.get();
+    }
+
+
+    //---------------------------------------------------------------------------------------------
+    void NodeImageMultiLayer::SetBlended( NodeImagePtr pNode )
+    {
+        m_pD->m_pBlended = pNode;
+    }
+
+
+    //---------------------------------------------------------------------------------------------
+    NodeRangePtr NodeImageMultiLayer::GetRange() const
+    {
+        return m_pD->m_pRange.get();
+    }
+
+
+    //---------------------------------------------------------------------------------------------
+    void NodeImageMultiLayer::SetRange( NodeRangePtr pNode )
+    {
+        m_pD->m_pRange = pNode;
+    }
+
+
+	//---------------------------------------------------------------------------------------------
+	EBlendType NodeImageMultiLayer::GetBlendType() const
+	{
+		return m_pD->m_type;
+	}
+
+
+	//---------------------------------------------------------------------------------------------
+    void NodeImageMultiLayer::SetBlendType(EBlendType t)
+	{
+		m_pD->m_type = t;
+	}
+
+
+}
+
+#undef NODE_INPUT_COUNT
diff --git a/Engine/Plugins/Experimental/Mutable/Source/MutableTools/Private/NodeImageMultiLayerPrivate.h b/Engine/Plugins/Experimental/Mutable/Source/MutableTools/Private/NodeImageMultiLayerPrivate.h
new file mode 100644
--- /dev/null
+++ b/Engine/Plugins/Experimental/Mutable/Source/MutableTools/Private/NodeImageMultiLayerPrivate.h
@@ -0,0 +1,64 @@
+// Copyright Epic Games, Inc. All Rights Reserved.
+
+#pragma once
+
+#include "NodeImagePrivate.h"
+#include "NodeImageMultiLayer.h"
+#include "NodeColour.h"
+#include "AST.h"
+
+#include "MemoryPrivate.h"
+
+
+namespace mu
+{
+
+    class NodeImageMultiLayer::Private : public NodeImage::Private
+	{
+	public:
+
+		MUTABLE_DEFINE_CONST_VISITABLE()
+
+	public:
+
+		static NODE_TYPE s_type;
+
+		NodeImagePtr m_pBase;
+		NodeImagePtr m_pMask;
+        NodeImagePtr m_pBlended;
+        NodeRangePtr m_pRange;
+		EBlendType m_type;
+
+		//!
+		void Serialise( OutputArchive& arch ) const
+		{
+            uint32_t ver = 0;
+			arch << ver;
+            arch << uint32_t(m_type);
+
+			arch << m_pBase;
+			arch << m_pMask;
+            arch << m_pBlended;
+            arch << m_pRange;
+		}
+
+		//!
+		void Unserialise( InputArchive& arch )
+		{
+            uint32_t ver;
+			arch >> ver;
+			check(ver==0);
+
+            uint32_t t;
+            arch >> t;
+            m_type=(EBlendType)t;
+
+            arch >> m_pBase;
+			arch >> m_pMask;
+            arch >> m_pBlended;
+            arch >> m_pRange;
+
+		}
+	};
+
+}
diff --git a/Engine/Plugins/Experimental/Mutable/Source/MutableTools/Private/NodeImageNormalComposite.cpp b/Engine/Plugins/Experimental/Mutable/Source/MutableTools/Private/NodeImageNormalComposite.cpp
new file mode 100644
--- /dev/null
+++ b/Engine/Plugins/Experimental/Mutable/Source/MutableTools/Private/NodeImageNormalComposite.cpp
@@ -0,0 +1,117 @@
+// Copyright Epic Games, Inc. All Rights Reserved.
+
+
+#include "NodeImageNormalComposite.h"
+#include "NodeImageNormalCompositePrivate.h"
+
+#include "ImagePrivate.h"
+
+
+namespace mu
+{
+
+	//---------------------------------------------------------------------------------------------
+	// Static initialisation
+	//---------------------------------------------------------------------------------------------
+	NODE_TYPE NodeImageNormalComposite::Private::s_type =
+			NODE_TYPE( "ImageNormalComposite", NodeImage::GetStaticType() );
+
+
+	//---------------------------------------------------------------------------------------------
+	//!
+	//---------------------------------------------------------------------------------------------
+
+	MUTABLE_IMPLEMENT_NODE( NodeImageNormalComposite, EType::NormalComposite, Node, Node::EType::Image )
+
+
+	//---------------------------------------------------------------------------------------------
+	// Node Interface
+	//---------------------------------------------------------------------------------------------
+	int NodeImageNormalComposite::GetInputCount() const
+	{
+		return 2;
+	}
+
+
+	//---------------------------------------------------------------------------------------------
+	Node* NodeImageNormalComposite::GetInputNode( int i ) const
+	{
+		check( i>=0 && i< GetInputCount());
+
+		Node* pResult = 0;
+
+		switch (i)
+		{
+        case 0: pResult = m_pD->m_pBase.get(); break;
+        case 1: pResult = m_pD->m_pNormal.get(); break;
+		}
+
+		return pResult;
+	}
+
+
+	//---------------------------------------------------------------------------------------------
+	void NodeImageNormalComposite::SetInputNode( int i, NodePtr pNode )
+	{
+		check( i>=0 && i< GetInputCount());
+
+		switch (i)
+		{
+        case 0: m_pD->m_pBase = dynamic_cast<NodeImage*>( pNode.get() ); break;
+        case 1: m_pD->m_pNormal = dynamic_cast<NodeImage*>( pNode.get() ); break;
+		}
+	}
+
+
+	//---------------------------------------------------------------------------------------------
+	// Own Interface
+	//---------------------------------------------------------------------------------------------
+	NodeImagePtr NodeImageNormalComposite::GetBase() const
+	{
+		return m_pD->m_pBase.get();
+	}
+
+	//---------------------------------------------------------------------------------------------
+	void NodeImageNormalComposite::SetBase( NodeImagePtr pNode )
+	{
+		m_pD->m_pBase = pNode;
+	}
+
+	//---------------------------------------------------------------------------------------------
+	NodeImagePtr NodeImageNormalComposite::GetNormal() const
+	{
+		return m_pD->m_pNormal.get();
+	}
+
+	//---------------------------------------------------------------------------------------------
+	void NodeImageNormalComposite::SetNormal( NodeImagePtr pNode )
+	{
+		m_pD->m_pNormal = pNode;
+	}
+
+	//---------------------------------------------------------------------------------------------
+	float NodeImageNormalComposite::GetPower() const
+	{
+		return m_pD->m_power;
+	}
+
+	
+	//---------------------------------------------------------------------------------------------
+	void NodeImageNormalComposite::SetPower( float power )
+	{
+		m_pD->m_power = power;
+	}
+
+
+	//---------------------------------------------------------------------------------------------
+	ECompositeImageMode NodeImageNormalComposite::GetMode() const
+	{
+		return m_pD->m_mode;
+	}
+
+	//---------------------------------------------------------------------------------------------
+	void NodeImageNormalComposite::SetMode( ECompositeImageMode mode )
+	{
+		m_pD->m_mode = mode;
+	}
+}
diff --git a/Engine/Plugins/Experimental/Mutable/Source/MutableTools/Private/NodeImageNormalCompositePrivate.h b/Engine/Plugins/Experimental/Mutable/Source/MutableTools/Private/NodeImageNormalCompositePrivate.h
new file mode 100644
--- /dev/null
+++ b/Engine/Plugins/Experimental/Mutable/Source/MutableTools/Private/NodeImageNormalCompositePrivate.h
@@ -0,0 +1,56 @@
+// Copyright Epic Games, Inc. All Rights Reserved.
+
+#pragma once
+
+#include "NodeImagePrivate.h"
+#include "NodeImageNormalComposite.h"
+#include "AST.h"
+
+#include "MemoryPrivate.h"
+
+
+namespace mu
+{
+
+	class NodeImageNormalComposite::Private : public NodeImage::Private
+	{
+	public:
+
+		MUTABLE_DEFINE_CONST_VISITABLE()
+
+	public:
+
+		static NODE_TYPE s_type;
+
+		NodeImagePtr m_pBase;
+		NodeImagePtr m_pNormal;
+		float m_power;
+		ECompositeImageMode m_mode;
+
+		//!
+		void Serialise( OutputArchive& arch ) const
+		{
+            uint32_t ver = 0;
+			arch << ver;
+
+			arch << m_pBase;
+			arch << m_pNormal;
+			arch << m_power;
+			arch << m_mode;
+		}
+
+		//!
+		void Unserialise( InputArchive& arch )
+		{
+            uint32_t ver;
+			arch >> ver;
+			check(ver==0);
+
+			arch >> m_pBase;
+			arch >> m_pNormal;
+			arch >> m_power;
+			arch >> m_mode;
+		}
+	};
+
+}
diff --git a/Engine/Plugins/Experimental/Mutable/Source/MutableTools/Private/NodeImageParameter.cpp b/Engine/Plugins/Experimental/Mutable/Source/MutableTools/Private/NodeImageParameter.cpp
new file mode 100644
--- /dev/null
+++ b/Engine/Plugins/Experimental/Mutable/Source/MutableTools/Private/NodeImageParameter.cpp
@@ -0,0 +1,96 @@
+// Copyright Epic Games, Inc. All Rights Reserved.
+
+
+#include "NodeImageParameter.h"
+#include "NodeImageParameterPrivate.h"
+
+#include "MutableMath.h"
+#include "MemoryPrivate.h"
+
+
+namespace mu
+{
+
+	//---------------------------------------------------------------------------------------------
+	// Static initialisation
+	//---------------------------------------------------------------------------------------------
+    NODE_TYPE NodeImageParameter::Private::s_type =
+            NODE_TYPE( "ImageParameter", NodeImage::GetStaticType() );
+
+
+	//---------------------------------------------------------------------------------------------
+	//!
+	//---------------------------------------------------------------------------------------------
+
+    MUTABLE_IMPLEMENT_NODE( NodeImageParameter, EType::Parameter, Node, Node::EType::Image)
+
+
+	//---------------------------------------------------------------------------------------------
+	// Node Interface
+	//---------------------------------------------------------------------------------------------
+    int NodeImageParameter::GetInputCount() const
+	{
+		return 0;
+	}
+
+
+	//---------------------------------------------------------------------------------------------
+    Node* NodeImageParameter::GetInputNode( int ) const
+	{
+		check( false );
+		return 0;
+	}
+
+
+	//---------------------------------------------------------------------------------------------
+    void NodeImageParameter::SetInputNode( int, NodePtr )
+	{
+		check( false );
+	}
+
+
+	//---------------------------------------------------------------------------------------------
+	// Own Interface
+	//---------------------------------------------------------------------------------------------
+    const char* NodeImageParameter::GetName() const
+	{
+		return m_pD->m_name.c_str();
+	}
+
+
+	//---------------------------------------------------------------------------------------------
+    void NodeImageParameter::SetName( const char* strName )
+	{
+		if ( strName )
+		{
+			m_pD->m_name = strName;
+		}
+		else
+		{
+			m_pD->m_name = "";
+		}
+	}
+
+
+	const char* NodeImageParameter::GetUid() const
+	{
+		return m_pD->m_uid.c_str();
+	}
+
+
+	//---------------------------------------------------------------------------------------------
+	void NodeImageParameter::SetUid( const char* strUid )
+	{
+		if ( strUid )
+		{
+			m_pD->m_uid = strUid;
+		}
+		else
+		{
+			m_pD->m_uid = "";
+		}
+	}
+
+}
+
+
diff --git a/Engine/Plugins/Experimental/Mutable/Source/MutableTools/Private/NodeImageParameterPrivate.h b/Engine/Plugins/Experimental/Mutable/Source/MutableTools/Private/NodeImageParameterPrivate.h
new file mode 100644
--- /dev/null
+++ b/Engine/Plugins/Experimental/Mutable/Source/MutableTools/Private/NodeImageParameterPrivate.h
@@ -0,0 +1,47 @@
+// Copyright Epic Games, Inc. All Rights Reserved.
+
+#pragma once
+
+#include "NodeImageParameter.h"
+#include "NodeImagePrivate.h"
+
+
+namespace mu
+{
+
+    class NodeImageParameter::Private : public NodeImage::Private
+	{
+	public:
+
+		MUTABLE_DEFINE_CONST_VISITABLE()
+
+	public:
+
+		static NODE_TYPE s_type;
+
+		string m_name;
+		string m_uid;
+
+		//!
+		void Serialise( OutputArchive& arch ) const
+		{
+            uint32_t ver = 1;
+			arch << ver;
+
+			arch << m_name;
+			arch << m_uid;
+		}
+
+		//!
+		void Unserialise( InputArchive& arch )
+		{
+            uint32_t ver;
+			arch >> ver;
+            check(ver==1);
+
+			arch >> m_name;
+            arch >> m_uid;
+		}
+	};
+
+}
diff --git a/Engine/Plugins/Experimental/Mutable/Source/MutableTools/Private/NodeImagePlainColour.cpp b/Engine/Plugins/Experimental/Mutable/Source/MutableTools/Private/NodeImagePlainColour.cpp
new file mode 100644
--- /dev/null
+++ b/Engine/Plugins/Experimental/Mutable/Source/MutableTools/Private/NodeImagePlainColour.cpp
@@ -0,0 +1,93 @@
+// Copyright Epic Games, Inc. All Rights Reserved.
+
+
+#include "NodeImagePlainColour.h"
+#include "NodeImagePlainColourPrivate.h"
+
+#include "NodeColour.h"
+
+#include "ImagePrivate.h"
+
+
+#define NODE_INPUT_COUNT 	1
+
+
+namespace mu
+{
+
+	//---------------------------------------------------------------------------------------------
+	// Static initialisation
+	//---------------------------------------------------------------------------------------------
+	NODE_TYPE NodeImagePlainColour::Private::s_type =
+			NODE_TYPE( "ImagePlainColour", NodeImage::GetStaticType() );
+
+
+	//---------------------------------------------------------------------------------------------
+
+	MUTABLE_IMPLEMENT_NODE( NodeImagePlainColour, EType::PlainColour, Node, Node::EType::Image)
+
+
+	//---------------------------------------------------------------------------------------------
+	// Node Interface
+	//---------------------------------------------------------------------------------------------
+	void NodeImagePlainColour::SetInputNode( int i, NodePtr pNode )
+	{
+		check( i>=0 && i<NODE_INPUT_COUNT );
+
+		switch (i)
+		{
+        case 0: m_pD->m_pColour = dynamic_cast<NodeColour*>(pNode.get()); break;
+		}
+	}
+
+
+	//---------------------------------------------------------------------------------------------
+	int NodeImagePlainColour::GetInputCount() const
+	{
+		return NODE_INPUT_COUNT;
+	}
+
+
+	//---------------------------------------------------------------------------------------------
+	Node* NodeImagePlainColour::GetInputNode( int i ) const
+	{
+		check( i>=0 && i<NODE_INPUT_COUNT );
+
+		Node* pResult = 0;
+
+		switch (i)
+		{
+        case 0: pResult = m_pD->m_pColour.get(); break;
+		}
+
+		return pResult;
+	}
+
+
+	//---------------------------------------------------------------------------------------------
+	// Own Interface
+	//---------------------------------------------------------------------------------------------
+	NodeColourPtr NodeImagePlainColour::GetColour() const
+	{
+		return m_pD->m_pColour;
+	}
+
+
+	//---------------------------------------------------------------------------------------------
+	void NodeImagePlainColour::SetColour( NodeColourPtr pNode )
+	{
+		m_pD->m_pColour = pNode;
+	}
+
+
+	//---------------------------------------------------------------------------------------------
+	void NodeImagePlainColour::SetSize( int x, int y )
+	{
+		m_pD->m_sizeX = x;
+		m_pD->m_sizeY = y;
+	}
+
+}
+
+#undef NODE_INPUT_COUNT
+
diff --git a/Engine/Plugins/Experimental/Mutable/Source/MutableTools/Private/NodeImagePlainColourPrivate.h b/Engine/Plugins/Experimental/Mutable/Source/MutableTools/Private/NodeImagePlainColourPrivate.h
new file mode 100644
--- /dev/null
+++ b/Engine/Plugins/Experimental/Mutable/Source/MutableTools/Private/NodeImagePlainColourPrivate.h
@@ -0,0 +1,60 @@
+// Copyright Epic Games, Inc. All Rights Reserved.
+
+#pragma once
+
+#include "NodeImagePrivate.h"
+#include "NodeImagePlainColour.h"
+#include "AST.h"
+
+#include "MemoryPrivate.h"
+
+
+namespace mu
+{
+
+
+	class NodeImagePlainColour::Private : public NodeImage::Private
+	{
+	public:
+
+		MUTABLE_DEFINE_CONST_VISITABLE()
+
+	public:
+
+		Private()
+		{
+            m_sizeX = 4;
+            m_sizeY = 4;
+		}
+
+		static NODE_TYPE s_type;
+
+		NodeColourPtr m_pColour;
+		int m_sizeX, m_sizeY;
+
+		//!
+		void Serialise( OutputArchive& arch ) const
+		{
+            uint32_t ver = 0;
+			arch << ver;
+
+			arch << m_pColour;
+			arch << m_sizeX;
+			arch << m_sizeY;
+		}
+
+		//!
+		void Unserialise( InputArchive& arch )
+		{
+            uint32_t ver;
+			arch >> ver;
+			check(ver==0);
+
+			arch >> m_pColour;
+			arch >> m_sizeX;
+			arch >> m_sizeY;
+		}
+	};
+
+
+}
diff --git a/Engine/Plugins/Experimental/Mutable/Source/MutableTools/Private/NodeImagePrivate.h b/Engine/Plugins/Experimental/Mutable/Source/MutableTools/Private/NodeImagePrivate.h
new file mode 100644
--- /dev/null
+++ b/Engine/Plugins/Experimental/Mutable/Source/MutableTools/Private/NodeImagePrivate.h
@@ -0,0 +1,19 @@
+// Copyright Epic Games, Inc. All Rights Reserved.
+
+#pragma once
+
+#include "NodePrivate.h"
+#include "NodeImage.h"
+
+#include "MemoryPrivate.h"
+
+
+namespace mu
+{
+
+
+	class NodeImage::Private : public Node::Private
+	{
+	};
+
+}
diff --git a/Engine/Plugins/Experimental/Mutable/Source/MutableTools/Private/NodeImageProject.cpp b/Engine/Plugins/Experimental/Mutable/Source/MutableTools/Private/NodeImageProject.cpp
new file mode 100644
--- /dev/null
+++ b/Engine/Plugins/Experimental/Mutable/Source/MutableTools/Private/NodeImageProject.cpp
@@ -0,0 +1,226 @@
+// Copyright Epic Games, Inc. All Rights Reserved.
+
+
+#include "NodeImageProject.h"
+#include "NodeImageProjectPrivate.h"
+
+#include "NodeScalar.h"
+
+#include "ImagePrivate.h"
+
+
+
+namespace mu
+{
+
+	//---------------------------------------------------------------------------------------------
+	// Static initialisation
+	//---------------------------------------------------------------------------------------------
+	NODE_TYPE NodeImageProject::Private::s_type =
+			NODE_TYPE( "ImageProject", NodeImage::GetStaticType() );
+
+
+	//---------------------------------------------------------------------------------------------
+	//!
+	//---------------------------------------------------------------------------------------------
+
+	MUTABLE_IMPLEMENT_NODE( NodeImageProject, EType::Project, Node, Node::EType::Image)
+
+
+	//---------------------------------------------------------------------------------------------
+	// Node Interface
+	//---------------------------------------------------------------------------------------------
+	int NodeImageProject::GetInputCount() const
+	{
+		return 6;
+	}
+
+
+	//---------------------------------------------------------------------------------------------
+	Node* NodeImageProject::GetInputNode( int i ) const
+	{
+		check( i>=0 && i<GetInputCount() );
+
+		Node* pResult = 0;
+
+		switch (i)
+		{
+		case 0:
+			pResult = m_pD->m_pMesh.get();
+			break;
+
+		case 1:
+			pResult = m_pD->m_pProjector.get();
+			break;
+
+		case 2:
+			pResult = m_pD->m_pAngleFadeStart.get();
+			break;
+
+		case 3:
+			pResult = m_pD->m_pAngleFadeEnd.get();
+			break;
+
+		case 4:
+			pResult = m_pD->m_pMask.get();
+			break;
+
+		default:
+			pResult = m_pD->m_pImage.get();
+			break;
+		}
+
+		return pResult;
+	}
+
+
+	//---------------------------------------------------------------------------------------------
+	void NodeImageProject::SetInputNode( int i, NodePtr pNode )
+	{
+		check( i>=0 && i<GetInputCount() );
+
+		switch (i)
+		{
+		case 0:
+			m_pD->m_pMesh = dynamic_cast<NodeMesh*>(pNode.get());
+			break;
+
+		case 1:
+			m_pD->m_pProjector = dynamic_cast<NodeProjector*>(pNode.get());
+			break;
+
+		case 2:
+			m_pD->m_pAngleFadeStart = dynamic_cast<NodeScalar*>(pNode.get());
+			break;
+
+		case 3:
+			m_pD->m_pAngleFadeEnd = dynamic_cast<NodeScalar*>(pNode.get());
+			break;
+
+		case 4:
+			m_pD->m_pMask = dynamic_cast<NodeImage*>(pNode.get());
+			break;
+
+		default:
+			m_pD->m_pImage = dynamic_cast<NodeImage*>(pNode.get());
+			break;
+		}
+	}
+
+
+	//---------------------------------------------------------------------------------------------
+	// Own Interface
+	//---------------------------------------------------------------------------------------------
+	NodeMeshPtr NodeImageProject::GetMesh() const
+	{
+		return m_pD->m_pMesh.get();
+	}
+
+
+	//---------------------------------------------------------------------------------------------
+	void NodeImageProject::SetMesh( NodeMeshPtr pNode )
+	{
+		m_pD->m_pMesh = pNode;
+	}
+
+
+	//---------------------------------------------------------------------------------------------
+	NodeProjectorPtr NodeImageProject::GetProjector() const
+	{
+		return m_pD->m_pProjector.get();
+	}
+
+
+	//---------------------------------------------------------------------------------------------
+	void NodeImageProject::SetProjector( NodeProjectorPtr pNode )
+	{
+		m_pD->m_pProjector = pNode;
+	}
+
+
+	//---------------------------------------------------------------------------------------------
+	NodeScalarPtr NodeImageProject::GetAngleFadeStart() const
+	{
+		return m_pD->m_pAngleFadeStart.get();
+	}
+
+
+	//---------------------------------------------------------------------------------------------
+	void NodeImageProject::SetAngleFadeStart( NodeScalarPtr pNode )
+	{
+		m_pD->m_pAngleFadeStart = pNode;
+	}
+
+
+	//---------------------------------------------------------------------------------------------
+	NodeScalarPtr NodeImageProject::GetAngleFadeEnd() const
+	{
+		return m_pD->m_pAngleFadeEnd.get();
+	}
+
+
+	//---------------------------------------------------------------------------------------------
+	void NodeImageProject::SetAngleFadeEnd( NodeScalarPtr pNode )
+	{
+		m_pD->m_pAngleFadeEnd = pNode;
+	}
+
+
+	//---------------------------------------------------------------------------------------------
+	NodeImagePtr NodeImageProject::GetTargetMask() const
+	{
+		return m_pD->m_pMask;
+	}
+
+
+	//---------------------------------------------------------------------------------------------
+	void NodeImageProject::SetTargetMask( NodeImagePtr pNode )
+	{
+		m_pD->m_pMask = pNode;
+	}
+
+
+	//---------------------------------------------------------------------------------------------
+	NodeImagePtr NodeImageProject::GetImage() const
+	{
+		return m_pD->m_pImage;
+	}
+
+
+	//---------------------------------------------------------------------------------------------
+	void NodeImageProject::SetImage( NodeImagePtr pNode )
+	{
+		m_pD->m_pImage = pNode;
+	}
+
+
+    //---------------------------------------------------------------------------------------------
+    uint8_t NodeImageProject::GetLayout() const
+    {
+        return m_pD->m_layout;
+    }
+
+
+    //---------------------------------------------------------------------------------------------
+    void NodeImageProject::SetLayout( uint8_t l )
+    {
+        m_pD->m_layout = l;
+    }
+
+
+	//---------------------------------------------------------------------------------------------
+	const FUintVector2& NodeImageProject::GetImageSize() const
+	{
+		return m_pD->m_imageSize;
+	}
+
+
+	//---------------------------------------------------------------------------------------------
+	void NodeImageProject::SetImageSize(const FUintVector2& size)
+	{
+		m_pD->m_imageSize = size;
+	}
+
+}
+
+
