[mutable] Move Mutable plugin from the restricted folder to the experimental plugins folder.

diff --git a/Engine/Plugins/Experimental/Mutable/Source/MutableTools/Private/NodeImageProjectPrivate.h b/Engine/Plugins/Experimental/Mutable/Source/MutableTools/Private/NodeImageProjectPrivate.h
new file mode 100644
--- /dev/null
+++ b/Engine/Plugins/Experimental/Mutable/Source/MutableTools/Private/NodeImageProjectPrivate.h
@@ -0,0 +1,72 @@
+// Copyright Epic Games, Inc. All Rights Reserved.
+
+#pragma once
+
+#include "MutableTools/Public/NodeImageProject.h"
+
+#include "MutableTools/Private/NodeImagePrivate.h"
+#include "MutableTools/Public/NodeScalar.h"
+#include "MutableTools/Public/NodeColour.h"
+#include "MutableTools/Public/NodeMesh.h"
+#include "MutableTools/Public/NodeProjector.h"
+#include "MutableRuntime/Private/MutableMath.h"
+
+namespace mu
+{
+
+
+	class NodeImageProject::Private : public NodeImage::Private
+	{
+	public:
+
+		MUTABLE_DEFINE_CONST_VISITABLE()
+
+	public:
+
+		static NODE_TYPE s_type;
+
+		NodeProjectorPtr m_pProjector;
+		NodeMeshPtr m_pMesh;
+		NodeScalarPtr m_pAngleFadeStart;
+		NodeScalarPtr m_pAngleFadeEnd;
+		NodeImagePtr m_pImage;
+		NodeImagePtr m_pMask;
+        uint8_t m_layout = 0;
+		FUintVector2 m_imageSize;
+
+		//!
+		void Serialise( OutputArchive& arch ) const
+		{
+            uint32_t ver = 2;
+			arch << ver;
+
+			arch << m_pProjector;
+			arch << m_pMesh;
+			arch << m_pAngleFadeStart;
+			arch << m_pAngleFadeEnd;
+			arch << m_pImage;
+            arch << m_pMask;
+            arch << m_layout;
+			arch << m_imageSize;
+        }
+
+		//!
+		void Unserialise( InputArchive& arch )
+		{
+            uint32_t ver;
+			arch >> ver;
+            check(ver==2);
+
+			arch >> m_pProjector;
+			arch >> m_pMesh;
+			arch >> m_pAngleFadeStart;
+			arch >> m_pAngleFadeEnd;
+			arch >> m_pImage;
+			arch >> m_pMask;
+            arch >> m_layout;
+			arch >> m_imageSize;
+		}
+	};
+
+
+}
diff --git a/Engine/Plugins/Experimental/Mutable/Source/MutableTools/Private/NodeImageResize.cpp b/Engine/Plugins/Experimental/Mutable/Source/MutableTools/Private/NodeImageResize.cpp
new file mode 100644
--- /dev/null
+++ b/Engine/Plugins/Experimental/Mutable/Source/MutableTools/Private/NodeImageResize.cpp
@@ -0,0 +1,98 @@
+// Copyright Epic Games, Inc. All Rights Reserved.
+
+
+#include "NodeImageResize.h"
+#include "NodeImageResizePrivate.h"
+
+#include "NodeScalar.h"
+
+#include "ImagePrivate.h"
+
+
+#define NODE_INPUT_COUNT 	1
+
+
+namespace mu
+{
+
+	//---------------------------------------------------------------------------------------------
+	// Static initialisation
+	//---------------------------------------------------------------------------------------------
+	NODE_TYPE NodeImageResize::Private::s_type =
+			NODE_TYPE( "ImageResize", NodeImage::GetStaticType() );
+
+	//---------------------------------------------------------------------------------------------
+
+	MUTABLE_IMPLEMENT_NODE( NodeImageResize, EType::Resize, Node, Node::EType::Image)
+
+
+	//---------------------------------------------------------------------------------------------
+	// Node Interface
+	//---------------------------------------------------------------------------------------------
+	void NodeImageResize::SetInputNode( int i, NodePtr pNode )
+	{
+		check( i>=0 && i<NODE_INPUT_COUNT );
+
+		switch (i)
+		{
+        case 0: m_pD->m_pBase = dynamic_cast<NodeImage*>(pNode.get()); break;
+		}
+	}
+
+
+	//---------------------------------------------------------------------------------------------
+	int NodeImageResize::GetInputCount() const
+	{
+		return NODE_INPUT_COUNT;
+	}
+
+
+	//---------------------------------------------------------------------------------------------
+	Node* NodeImageResize::GetInputNode( int i ) const
+	{
+		check( i>=0 && i<NODE_INPUT_COUNT );
+
+		Node* pResult = 0;
+
+		switch (i)
+		{
+        case 0: pResult = m_pD->m_pBase.get(); break;
+		}
+
+		return pResult;
+	}
+
+
+	//---------------------------------------------------------------------------------------------
+	// Own Interface
+	//---------------------------------------------------------------------------------------------
+	NodeImagePtr NodeImageResize::GetBase() const
+	{
+		return m_pD->m_pBase;
+	}
+
+
+	//---------------------------------------------------------------------------------------------
+	void NodeImageResize::SetBase( NodeImagePtr pNode )
+	{
+		m_pD->m_pBase = pNode;
+	}
+
+
+	//---------------------------------------------------------------------------------------------
+	void NodeImageResize::SetRelative( bool rel )
+	{
+		m_pD->m_relative = rel;
+	}
+
+
+	//---------------------------------------------------------------------------------------------
+	void NodeImageResize::SetSize( float x, float y )
+	{
+		m_pD->m_sizeX = x;
+		m_pD->m_sizeY = y;
+	}
+
+}
+
+#undef NODE_INPUT_COUNT
diff --git a/Engine/Plugins/Experimental/Mutable/Source/MutableTools/Private/NodeImageResizePrivate.h b/Engine/Plugins/Experimental/Mutable/Source/MutableTools/Private/NodeImageResizePrivate.h
new file mode 100644
--- /dev/null
+++ b/Engine/Plugins/Experimental/Mutable/Source/MutableTools/Private/NodeImageResizePrivate.h
@@ -0,0 +1,57 @@
+// Copyright Epic Games, Inc. All Rights Reserved.
+
+#pragma once
+
+#include "NodeImagePrivate.h"
+#include "NodeImageResize.h"
+#include "AST.h"
+
+#include "MemoryPrivate.h"
+
+
+namespace mu
+{
+
+
+	class NodeImageResize::Private : public NodeImage::Private
+	{
+	public:
+
+		MUTABLE_DEFINE_CONST_VISITABLE()
+
+	public:
+
+		static NODE_TYPE s_type;
+
+		NodeImagePtr m_pBase;
+		bool m_relative = true;
+		float m_sizeX = 0.5f, m_sizeY = 0.5f;
+
+		//!
+		void Serialise( OutputArchive& arch ) const
+		{
+            uint32_t ver = 0;
+			arch << ver;
+
+			arch << m_pBase;
+			arch << m_relative;
+			arch << m_sizeX;
+			arch << m_sizeY;
+		}
+
+		//!
+		void Unserialise( InputArchive& arch )
+		{
+            uint32_t ver;
+			arch >> ver;
+			check(ver==0);
+
+			arch >> m_pBase;
+			arch >> m_relative;
+			arch >> m_sizeX;
+			arch >> m_sizeY;
+		}
+	};
+
+
+}
diff --git a/Engine/Plugins/Experimental/Mutable/Source/MutableTools/Private/NodeImageSaturate.cpp b/Engine/Plugins/Experimental/Mutable/Source/MutableTools/Private/NodeImageSaturate.cpp
new file mode 100644
--- /dev/null
+++ b/Engine/Plugins/Experimental/Mutable/Source/MutableTools/Private/NodeImageSaturate.cpp
@@ -0,0 +1,105 @@
+// Copyright Epic Games, Inc. All Rights Reserved.
+
+
+#include "NodeImageSaturate.h"
+#include "NodeImageSaturatePrivate.h"
+
+#include "NodeScalar.h"
+
+#include "ImagePrivate.h"
+
+
+#define NODE_INPUT_COUNT 	2
+
+
+namespace mu
+{
+
+	//---------------------------------------------------------------------------------------------
+	// Static initialisation
+	//---------------------------------------------------------------------------------------------
+	NODE_TYPE NodeImageSaturate::Private::s_type =
+			NODE_TYPE( "ImageSaturate", NodeImage::GetStaticType() );
+
+
+	//---------------------------------------------------------------------------------------------
+	//!
+	//---------------------------------------------------------------------------------------------
+
+	MUTABLE_IMPLEMENT_NODE( NodeImageSaturate, EType::Saturate, Node, Node::EType::Image)
+
+
+	//---------------------------------------------------------------------------------------------
+	// Node Interface
+	//---------------------------------------------------------------------------------------------
+	int NodeImageSaturate::GetInputCount() const
+	{
+		return NODE_INPUT_COUNT;
+	}
+
+
+	//---------------------------------------------------------------------------------------------
+	Node* NodeImageSaturate::GetInputNode( int i ) const
+	{
+		check( i>=0 && i<NODE_INPUT_COUNT );
+
+		Node* pResult = 0;
+
+		switch (i)
+		{
+		case 0: pResult = m_pD->m_pFactor.get();
+		case 1: pResult = m_pD->m_pSource.get();
+		}
+
+		return pResult;
+	}
+
+
+	//---------------------------------------------------------------------------------------------
+	void NodeImageSaturate::SetInputNode( int i, NodePtr pNode )
+	{
+		check( i>=0 && i<NODE_INPUT_COUNT );
+
+		switch (i)
+		{
+		case 0: m_pD->m_pFactor = dynamic_cast<NodeScalar*>(pNode.get());
+		case 1: m_pD->m_pSource = dynamic_cast<NodeImage*>(pNode.get());
+		}
+	}
+
+
+	//---------------------------------------------------------------------------------------------
+	// Own Interface
+	//---------------------------------------------------------------------------------------------
+	NodeScalarPtr NodeImageSaturate::GetFactor() const
+	{
+		return m_pD->m_pFactor.get();
+	}
+
+
+	//---------------------------------------------------------------------------------------------
+	void NodeImageSaturate::SetFactor( NodeScalarPtr pNode )
+	{
+		m_pD->m_pFactor = pNode;
+	}
+
+
+	//---------------------------------------------------------------------------------------------
+	NodeImagePtr NodeImageSaturate::GetSource() const
+	{
+		return m_pD->m_pSource.get();
+	}
+
+
+	//---------------------------------------------------------------------------------------------
+	void NodeImageSaturate::SetSource( NodeImagePtr pNode )
+	{
+		m_pD->m_pSource = pNode;
+	}
+
+
+
+}
+
+#undef NODE_INPUT_COUNT
+
diff --git a/Engine/Plugins/Experimental/Mutable/Source/MutableTools/Private/NodeImageSaturatePrivate.h b/Engine/Plugins/Experimental/Mutable/Source/MutableTools/Private/NodeImageSaturatePrivate.h
new file mode 100644
--- /dev/null
+++ b/Engine/Plugins/Experimental/Mutable/Source/MutableTools/Private/NodeImageSaturatePrivate.h
@@ -0,0 +1,53 @@
+// Copyright Epic Games, Inc. All Rights Reserved.
+
+#pragma once
+
+#include "NodeImagePrivate.h"
+#include "NodeImageSaturate.h"
+#include "NodeScalar.h"
+#include "AST.h"
+
+#include "MemoryPrivate.h"
+
+
+namespace mu
+{
+
+
+	class NodeImageSaturate::Private : public NodeImage::Private
+	{
+	public:
+
+		MUTABLE_DEFINE_CONST_VISITABLE()
+
+	public:
+
+		static NODE_TYPE s_type;
+
+		NodeImagePtr m_pSource;
+		NodeScalarPtr m_pFactor;
+
+		//!
+		void Serialise( OutputArchive& arch ) const
+		{
+            uint32_t ver = 0;
+			arch << ver;
+
+			arch << m_pSource;
+			arch << m_pFactor;
+		}
+
+		//!
+		void Unserialise( InputArchive& arch )
+		{
+            uint32_t ver;
+			arch >> ver;
+			check(ver==0);
+
+			arch >> m_pSource;
+			arch >> m_pFactor;
+		}
+	};
+
+
+}
diff --git a/Engine/Plugins/Experimental/Mutable/Source/MutableTools/Private/NodeImageSelectColour.cpp b/Engine/Plugins/Experimental/Mutable/Source/MutableTools/Private/NodeImageSelectColour.cpp
new file mode 100644
--- /dev/null
+++ b/Engine/Plugins/Experimental/Mutable/Source/MutableTools/Private/NodeImageSelectColour.cpp
@@ -0,0 +1,103 @@
+// Copyright Epic Games, Inc. All Rights Reserved.
+
+
+#include "NodeImageSelectColour.h"
+#include "NodeImageSelectColourPrivate.h"
+
+#include "NodeColour.h"
+
+#include "ImagePrivate.h"
+
+
+#define NODE_INPUT_COUNT 	2
+
+
+namespace mu
+{
+
+	//---------------------------------------------------------------------------------------------
+	// Static initialisation
+	//---------------------------------------------------------------------------------------------
+	NODE_TYPE NodeImageSelectColour::Private::s_type =
+			NODE_TYPE( "ImageSelectColour", NodeImage::GetStaticType() );
+
+	//---------------------------------------------------------------------------------------------
+	//!
+	//---------------------------------------------------------------------------------------------
+
+	MUTABLE_IMPLEMENT_NODE( NodeImageSelectColour, EType::SelectColour, Node, Node::EType::Image)
+
+
+	//---------------------------------------------------------------------------------------------
+	// Node Interface
+	//---------------------------------------------------------------------------------------------
+	int NodeImageSelectColour::GetInputCount() const
+	{
+		return NODE_INPUT_COUNT;
+	}
+
+
+	//---------------------------------------------------------------------------------------------
+	Node* NodeImageSelectColour::GetInputNode( int i ) const
+	{
+		check( i>=0 && i<NODE_INPUT_COUNT );
+
+		Node* pResult = 0;
+
+		switch (i)
+		{
+        case 0: pResult = m_pD->m_pColour.get(); break;
+        case 1: pResult = m_pD->m_pSource.get(); break;
+		}
+
+		return pResult;
+	}
+
+
+	//---------------------------------------------------------------------------------------------
+	void NodeImageSelectColour::SetInputNode( int i, NodePtr pNode )
+	{
+		check( i>=0 && i<NODE_INPUT_COUNT );
+
+		switch (i)
+		{
+        case 0: m_pD->m_pColour = dynamic_cast<NodeColour*>(pNode.get()); break;
+        case 1: m_pD->m_pSource = dynamic_cast<NodeImage*>(pNode.get()); break;
+		}
+	}
+
+
+	//---------------------------------------------------------------------------------------------
+	// Own Interface
+	//---------------------------------------------------------------------------------------------
+	NodeColourPtr NodeImageSelectColour::GetColour() const
+	{
+		return m_pD->m_pColour.get();
+	}
+
+
+	//---------------------------------------------------------------------------------------------
+	void NodeImageSelectColour::SetColour( NodeColourPtr pNode )
+	{
+		m_pD->m_pColour = pNode;
+	}
+
+
+	//---------------------------------------------------------------------------------------------
+	NodeImagePtr NodeImageSelectColour::GetSource() const
+	{
+		return m_pD->m_pSource.get();
+	}
+
+
+	//---------------------------------------------------------------------------------------------
+	void NodeImageSelectColour::SetSource( NodeImagePtr pNode )
+	{
+		m_pD->m_pSource = pNode;
+	}
+
+
+
+}
+
+#undef NODE_INPUT_COUNT
diff --git a/Engine/Plugins/Experimental/Mutable/Source/MutableTools/Private/NodeImageSelectColourPrivate.h b/Engine/Plugins/Experimental/Mutable/Source/MutableTools/Private/NodeImageSelectColourPrivate.h
new file mode 100644
--- /dev/null
+++ b/Engine/Plugins/Experimental/Mutable/Source/MutableTools/Private/NodeImageSelectColourPrivate.h
@@ -0,0 +1,53 @@
+// Copyright Epic Games, Inc. All Rights Reserved.
+
+#pragma once
+
+#include "NodeImagePrivate.h"
+#include "NodeImageSelectColour.h"
+#include "NodeColour.h"
+#include "AST.h"
+
+#include "MemoryPrivate.h"
+
+
+namespace mu
+{
+
+
+	class NodeImageSelectColour::Private : public NodeImage::Private
+	{
+	public:
+
+		MUTABLE_DEFINE_CONST_VISITABLE()
+
+	public:
+
+		static NODE_TYPE s_type;
+
+		NodeImagePtr m_pSource;
+		NodeColourPtr m_pColour;
+
+		//!
+		void Serialise( OutputArchive& arch ) const
+		{
+            uint32_t ver = 0;
+			arch << ver;
+
+			arch << m_pSource;
+			arch << m_pColour;
+		}
+
+		//!
+		void Unserialise( InputArchive& arch )
+		{
+            uint32_t ver;
+			arch >> ver;
+			check(ver==0);
+
+			arch >> m_pSource;
+			arch >> m_pColour;
+		}
+	};
+
+
+}
diff --git a/Engine/Plugins/Experimental/Mutable/Source/MutableTools/Private/NodeImageSwitch.cpp b/Engine/Plugins/Experimental/Mutable/Source/MutableTools/Private/NodeImageSwitch.cpp
new file mode 100644
--- /dev/null
+++ b/Engine/Plugins/Experimental/Mutable/Source/MutableTools/Private/NodeImageSwitch.cpp
@@ -0,0 +1,127 @@
+// Copyright Epic Games, Inc. All Rights Reserved.
+
+
+#include "NodeImageSwitch.h"
+#include "NodeImageSwitchPrivate.h"
+
+#include "NodeScalar.h"
+
+#include "ImagePrivate.h"
+
+
+namespace mu
+{
+
+	//---------------------------------------------------------------------------------------------
+	// Static initialisation
+	//---------------------------------------------------------------------------------------------
+	NODE_TYPE NodeImageSwitch::Private::s_type =
+			NODE_TYPE( "ImageSwitch", NodeImage::GetStaticType() );
+
+
+	//---------------------------------------------------------------------------------------------
+	//!
+	//---------------------------------------------------------------------------------------------
+
+	MUTABLE_IMPLEMENT_NODE( NodeImageSwitch, EType::Switch, Node, Node::EType::Image)
+
+
+	//---------------------------------------------------------------------------------------------
+	// Node Interface
+	//---------------------------------------------------------------------------------------------
+	int NodeImageSwitch::GetInputCount() const
+	{
+		return 1 + (int)m_pD->m_options.size();
+	}
+
+
+	//---------------------------------------------------------------------------------------------
+	Node* NodeImageSwitch::GetInputNode( int i ) const
+	{
+		check( i>=0 && i<GetInputCount() );
+
+		Node* pResult = 0;
+
+		switch (i)
+		{
+		case 0:
+			pResult = m_pD->m_pParameter.get();
+			break;
+
+		default:
+			pResult = m_pD->m_options[i-1].get();
+			break;
+		}
+
+		return pResult;
+	}
+
+
+	//---------------------------------------------------------------------------------------------
+	void NodeImageSwitch::SetInputNode( int i, NodePtr pNode )
+	{
+		check( i>=0 && i<GetInputCount() );
+
+		switch (i)
+		{
+		case 0:
+			m_pD->m_pParameter = dynamic_cast<NodeScalar*>(pNode.get());
+			break;
+
+		default:
+			m_pD->m_options[i-1] = dynamic_cast<NodeImage*>(pNode.get());
+			break;
+		}
+	}
+
+
+	//---------------------------------------------------------------------------------------------
+	// Own Interface
+	//---------------------------------------------------------------------------------------------
+	NodeScalarPtr NodeImageSwitch::GetParameter() const
+	{
+		return m_pD->m_pParameter.get();
+	}
+
+
+	//---------------------------------------------------------------------------------------------
+	void NodeImageSwitch::SetParameter( NodeScalarPtr pNode )
+	{
+		m_pD->m_pParameter = pNode;
+	}
+
+
+	//---------------------------------------------------------------------------------------------
+	void NodeImageSwitch::SetOptionCount( int t )
+	{
+		m_pD->m_options.resize(t);
+	}
+
+
+	//---------------------------------------------------------------------------------------------
+	int NodeImageSwitch::GetOptionCount() const 
+	{
+		return int(m_pD->m_options.size());
+	}
+
+
+	//---------------------------------------------------------------------------------------------
+	NodeImagePtr NodeImageSwitch::GetOption( int t ) const
+	{
+		check( t>=0 && t<(int)m_pD->m_options.size() );
+		return m_pD->m_options[t].get();
+	}
+
+
+	//---------------------------------------------------------------------------------------------
+	void NodeImageSwitch::SetOption( int t, NodeImagePtr pNode )
+	{
+		check( t>=0 && t<(int)m_pD->m_options.size() );
+		m_pD->m_options[t] = pNode;
+	}
+
+
+
+}
+
+
diff --git a/Engine/Plugins/Experimental/Mutable/Source/MutableTools/Private/NodeImageSwitchPrivate.h b/Engine/Plugins/Experimental/Mutable/Source/MutableTools/Private/NodeImageSwitchPrivate.h
new file mode 100644
--- /dev/null
+++ b/Engine/Plugins/Experimental/Mutable/Source/MutableTools/Private/NodeImageSwitchPrivate.h
@@ -0,0 +1,50 @@
+// Copyright Epic Games, Inc. All Rights Reserved.
+
+#pragma once
+
+#include "NodeImageSwitch.h"
+#include "NodeImagePrivate.h"
+#include "NodeScalar.h"
+
+
+namespace mu
+{
+
+
+	class NodeImageSwitch::Private : public NodeImage::Private
+	{
+	public:
+
+		MUTABLE_DEFINE_CONST_VISITABLE()
+
+	public:
+
+		static NODE_TYPE s_type;
+
+		NodeScalarPtr m_pParameter;
+		vector<NodeImagePtr> m_options;
+
+		//!
+		void Serialise( OutputArchive& arch ) const
+		{
+            uint32_t ver = 0;
+			arch << ver;
+
+			arch << m_pParameter;
+			arch << m_options;
+		}
+
+		//!
+		void Unserialise( InputArchive& arch )
+		{
+            uint32_t ver;
+			arch >> ver;
+			check(ver==0);
+
+			arch >> m_pParameter;
+			arch >> m_options;
+		}
+	};
+
+
+}
diff --git a/Engine/Plugins/Experimental/Mutable/Source/MutableTools/Private/NodeImageSwizzle.cpp b/Engine/Plugins/Experimental/Mutable/Source/MutableTools/Private/NodeImageSwizzle.cpp
new file mode 100644
--- /dev/null
+++ b/Engine/Plugins/Experimental/Mutable/Source/MutableTools/Private/NodeImageSwizzle.cpp
@@ -0,0 +1,114 @@
+// Copyright Epic Games, Inc. All Rights Reserved.
+
+
+#include "NodeImageSwizzle.h"
+#include "NodeImageSwizzlePrivate.h"
+
+
+#define NODE_INPUT_COUNT 	4
+
+
+namespace mu
+{
+
+	//---------------------------------------------------------------------------------------------
+	// Static initialisation
+	//---------------------------------------------------------------------------------------------
+	NODE_TYPE NodeImageSwizzle::Private::s_type =
+			NODE_TYPE( "ImageSwizzle", NodeImage::GetStaticType() );
+
+
+	//---------------------------------------------------------------------------------------------
+	//!
+	//---------------------------------------------------------------------------------------------
+
+	MUTABLE_IMPLEMENT_NODE( NodeImageSwizzle, EType::Swizzle, Node, Node::EType::Image)
+
+
+	//---------------------------------------------------------------------------------------------
+	// Node Interface
+	//---------------------------------------------------------------------------------------------
+	int NodeImageSwizzle::GetInputCount() const
+	{
+		return (int)m_pD->m_sources.size();
+	}
+
+
+	//---------------------------------------------------------------------------------------------
+	Node* NodeImageSwizzle::GetInputNode( int i ) const
+	{
+		check( i>=0 && i<GetInputCount() );
+
+		Node* pResult = 0;
+
+		pResult = m_pD->m_sources[i].get();
+
+		return pResult;
+	}
+
+
+	//---------------------------------------------------------------------------------------------
+	void NodeImageSwizzle::SetInputNode( int i, NodePtr pNode )
+	{
+		check( i>=0 && i<GetInputCount() );
+
+		m_pD->m_sources[i] = dynamic_cast<NodeImage*>(pNode.get());
+	}
+
+
+	//---------------------------------------------------------------------------------------------
+	// Own Interface
+	//---------------------------------------------------------------------------------------------
+	EImageFormat NodeImageSwizzle::GetFormat() const
+	{
+		return m_pD->m_format;
+	}
+
+
+	//---------------------------------------------------------------------------------------------
+	void NodeImageSwizzle::SetFormat(EImageFormat format )
+	{
+		m_pD->m_format = format;
+
+		int channelCount = GetImageFormatData( format ).m_channels;
+		m_pD->m_sources.resize( channelCount );
+		m_pD->m_sourceChannels.resize( channelCount );
+	}
+
+
+	//---------------------------------------------------------------------------------------------
+	NodeImagePtr NodeImageSwizzle::GetSource( int t ) const
+	{
+		check( t>=0 && t<(int)m_pD->m_sources.size() );
+		return m_pD->m_sources[t].get();
+	}
+
+
+	//---------------------------------------------------------------------------------------------
+	void NodeImageSwizzle::SetSource( int t, NodeImagePtr pNode )
+	{
+		check( t>=0 && t<(int)m_pD->m_sources.size() );
+		m_pD->m_sources[t] = pNode;
+	}
+
+
+	//---------------------------------------------------------------------------------------------
+	int NodeImageSwizzle::GetSourceChannel( int t ) const
+	{
+		check( t>=0 && t<(int)m_pD->m_sourceChannels.size() );
+		return m_pD->m_sourceChannels[t];
+	}
+
+
+	//---------------------------------------------------------------------------------------------
+	void NodeImageSwizzle::SetSourceChannel( int t, int channel )
+	{
+		check( t>=0 && t<(int)m_pD->m_sourceChannels.size() );
+		m_pD->m_sourceChannels[t] = channel;
+	}
+
+
+
+}
+
+#undef NODE_INPUT_COUNT
diff --git a/Engine/Plugins/Experimental/Mutable/Source/MutableTools/Private/NodeImageSwizzlePrivate.h b/Engine/Plugins/Experimental/Mutable/Source/MutableTools/Private/NodeImageSwizzlePrivate.h
new file mode 100644
--- /dev/null
+++ b/Engine/Plugins/Experimental/Mutable/Source/MutableTools/Private/NodeImageSwizzlePrivate.h
@@ -0,0 +1,56 @@
+// Copyright Epic Games, Inc. All Rights Reserved.
+
+#pragma once
+
+#include "NodeImageSwizzle.h"
+#include "NodeImagePrivate.h"
+#include "ImagePrivate.h"
+
+
+namespace mu
+{
+
+
+	class NodeImageSwizzle::Private : public NodeImage::Private
+	{
+	public:
+
+		MUTABLE_DEFINE_CONST_VISITABLE()
+
+	public:
+
+		static NODE_TYPE s_type;
+
+		EImageFormat m_format;
+		vector<NodeImagePtr> m_sources;
+		vector<int> m_sourceChannels;
+
+		//!
+		void Serialise( OutputArchive& arch ) const
+		{
+            uint32_t ver = 0;
+			arch << ver;
+
+			arch << (uint32)m_format;
+			arch << m_sources;
+			arch << m_sourceChannels;
+		}
+
+		//!
+		void Unserialise( InputArchive& arch )
+		{
+            uint32_t ver;
+			arch >> ver;
+			check(ver==0);
+
+			uint32 format;
+			arch >> format;
+			m_format = (EImageFormat)format;
+
+			arch >> m_sources;
+			arch >> m_sourceChannels;
+		}
+	};
+
+
+}
diff --git a/Engine/Plugins/Experimental/Mutable/Source/MutableTools/Private/NodeImageTable.cpp b/Engine/Plugins/Experimental/Mutable/Source/MutableTools/Private/NodeImageTable.cpp
new file mode 100644
--- /dev/null
+++ b/Engine/Plugins/Experimental/Mutable/Source/MutableTools/Private/NodeImageTable.cpp
@@ -0,0 +1,99 @@
+// Copyright Epic Games, Inc. All Rights Reserved.
+
+
+#include "NodeImageTable.h"
+#include "NodeImageTablePrivate.h"
+
+
+namespace mu
+{
+
+
+	//---------------------------------------------------------------------------------------------
+	// Static initialisation
+	//---------------------------------------------------------------------------------------------
+	NODE_TYPE NodeImageTable::Private::s_type =
+			NODE_TYPE( "TableImage", NodeImage::GetStaticType() );
+
+
+	//---------------------------------------------------------------------------------------------
+	//!
+	//---------------------------------------------------------------------------------------------
+
+	MUTABLE_IMPLEMENT_NODE( NodeImageTable, EType::Table, Node, Node::EType::Image);
+
+
+	//---------------------------------------------------------------------------------------------
+	// Node Interface
+	//---------------------------------------------------------------------------------------------
+	int NodeImageTable::GetInputCount() const
+	{
+		return 0;
+	}
+
+
+	//---------------------------------------------------------------------------------------------
+	Node* NodeImageTable::GetInputNode( int i ) const
+	{
+		check( i >=0 && i < 0 );
+        (void)i;
+        return 0;
+	}
+
+
+	//---------------------------------------------------------------------------------------------
+    void NodeImageTable::SetInputNode( int i, NodePtr )
+	{
+		check( i >=0 && i < 0 );
+        (void)i;
+        //m_pD->m_pObject = dynamic_cast<NodeObject*>( pNode.get() );
+	}
+
+
+	//---------------------------------------------------------------------------------------------
+	// Own Interface
+	//---------------------------------------------------------------------------------------------
+	const char* NodeImageTable::GetColumn() const
+	{
+		return m_pD->m_columnName.c_str();
+	}
+
+
+	//---------------------------------------------------------------------------------------------
+	void NodeImageTable::SetColumn( const char* strName )
+	{
+		if (strName)
+		{
+			m_pD->m_columnName = strName;
+		}
+		else
+		{
+			m_pD->m_columnName = "";
+		}
+	}
+
+
+	//---------------------------------------------------------------------------------------------
+	void NodeImageTable::SetParameterName( const char* strName )
+	{
+		m_pD->m_parameterName = strName;
+	}
+
+
+	//---------------------------------------------------------------------------------------------
+	void NodeImageTable::SetTable( TablePtr pTable )
+	{
+		m_pD->m_pTable = pTable;
+	}
+
+
+	//---------------------------------------------------------------------------------------------
+	TablePtr NodeImageTable::GetTable() const
+	{
+		return m_pD->m_pTable;
+	}
+
+
+}
+
+
diff --git a/Engine/Plugins/Experimental/Mutable/Source/MutableTools/Private/NodeImageTablePrivate.h b/Engine/Plugins/Experimental/Mutable/Source/MutableTools/Private/NodeImageTablePrivate.h
new file mode 100644
--- /dev/null
+++ b/Engine/Plugins/Experimental/Mutable/Source/MutableTools/Private/NodeImageTablePrivate.h
@@ -0,0 +1,56 @@
+// Copyright Epic Games, Inc. All Rights Reserved.
+
+#pragma once
+
+#include "NodePrivate.h"
+#include "NodeImageTable.h"
+#include "Table.h"
+#include "AST.h"
+
+#include "MemoryPrivate.h"
+
+
+namespace mu
+{
+
+
+	class NodeImageTable::Private : public Node::Private
+	{
+	public:
+
+		MUTABLE_DEFINE_CONST_VISITABLE()
+
+	public:
+
+		static NODE_TYPE s_type;
+
+		string m_parameterName;
+		TablePtr m_pTable;
+		string m_columnName;
+
+		//!
+		void Serialise( OutputArchive& arch ) const
+		{
+            uint32_t ver = 1;
+			arch << ver;
+
+			arch << m_parameterName;
+			arch << m_pTable;
+			arch << m_columnName;
+		}
+
+		//!
+		void Unserialise( InputArchive& arch )
+		{
+            uint32_t ver;
+			arch >> ver;
+            check(ver==1);
+
+			arch >> m_parameterName;
+			arch >> m_pTable;
+			arch >> m_columnName;
+		}
+
+	};
+
+}
diff --git a/Engine/Plugins/Experimental/Mutable/Source/MutableTools/Private/NodeImageTransform.cpp b/Engine/Plugins/Experimental/Mutable/Source/MutableTools/Private/NodeImageTransform.cpp
new file mode 100644
--- /dev/null
+++ b/Engine/Plugins/Experimental/Mutable/Source/MutableTools/Private/NodeImageTransform.cpp
@@ -0,0 +1,153 @@
+// Copyright Epic Games, Inc. All Rights Reserved.
+
+
+#include "NodeImageTransform.h"
+#include "NodeImageTransformPrivate.h"
+
+#include "NodeScalar.h"
+
+#include "ImagePrivate.h"
+
+
+#define NODE_INPUT_COUNT 	6
+
+
+namespace mu
+{
+
+	//---------------------------------------------------------------------------------------------
+	// Static initialisation
+	//---------------------------------------------------------------------------------------------
+	NODE_TYPE NodeImageTransform::Private::s_type =
+			NODE_TYPE( "ImageTransform", NodeImage::GetStaticType() );
+
+	//---------------------------------------------------------------------------------------------
+
+	MUTABLE_IMPLEMENT_NODE( NodeImageTransform, EType::Transform, Node, Node::EType::Image)
+
+
+	//---------------------------------------------------------------------------------------------
+	// Node Interface
+	//---------------------------------------------------------------------------------------------
+	void NodeImageTransform::SetInputNode( int i, NodePtr pNode )
+	{
+		check( i>=0 && i<NODE_INPUT_COUNT );
+
+		switch (i)
+		{
+        case 0: m_pD->m_pBase = dynamic_cast<NodeImage*>(pNode.get()); break;
+        case 1: m_pD->m_pOffsetX = dynamic_cast<NodeScalar*>(pNode.get()); break;
+        case 2: m_pD->m_pOffsetY = dynamic_cast<NodeScalar*>(pNode.get()); break;
+        case 3: m_pD->m_pScaleX = dynamic_cast<NodeScalar*>(pNode.get()); break;
+        case 4: m_pD->m_pScaleY = dynamic_cast<NodeScalar*>(pNode.get()); break;
+        case 5: m_pD->m_pRotation = dynamic_cast<NodeScalar*>(pNode.get()); break;
+		}
+	}
+
+
+	//---------------------------------------------------------------------------------------------
+	int NodeImageTransform::GetInputCount() const
+	{
+		return NODE_INPUT_COUNT;
+	}
+
+
+	//---------------------------------------------------------------------------------------------
+	Node* NodeImageTransform::GetInputNode( int i ) const
+	{
+		check( i>=0 && i<NODE_INPUT_COUNT );
+
+		Node* pResult = 0;
+
+		switch (i)
+		{
+        case 0: pResult = m_pD->m_pBase.get(); break;
+		case 1: pResult = m_pD->m_pOffsetX.get(); break;
+		case 2: pResult = m_pD->m_pOffsetY.get(); break;
+		case 3: pResult = m_pD->m_pScaleX.get(); break;
+		case 4: pResult = m_pD->m_pScaleY.get(); break;
+		case 5: pResult = m_pD->m_pRotation.get(); break;
+		}
+
+		return pResult;
+	}
+
+
+	//---------------------------------------------------------------------------------------------
+	// Own Interface
+	//---------------------------------------------------------------------------------------------
+	NodeImagePtr NodeImageTransform::GetBase() const
+	{
+		return m_pD->m_pBase;
+	}
+
+	//---------------------------------------------------------------------------------------------
+	NodeScalarPtr NodeImageTransform::GetOffsetX() const
+	{
+		return m_pD->m_pOffsetX;
+	}
+
+	//---------------------------------------------------------------------------------------------
+	NodeScalarPtr NodeImageTransform::GetOffsetY() const
+	{
+		return m_pD->m_pOffsetY;
+	}
+
+	//---------------------------------------------------------------------------------------------
+	NodeScalarPtr NodeImageTransform::GetScaleX() const
+	{
+		return m_pD->m_pScaleX;
+	}
+
+	//---------------------------------------------------------------------------------------------
+	NodeScalarPtr NodeImageTransform::GetScaleY() const
+	{
+		return m_pD->m_pScaleY;
+	}
+
+	//---------------------------------------------------------------------------------------------
+	NodeScalarPtr NodeImageTransform::GetRotation() const
+	{
+		return m_pD->m_pRotation;
+	}
+
+	//---------------------------------------------------------------------------------------------
+	void NodeImageTransform::SetBase( NodeImagePtr pNode )
+	{
+		m_pD->m_pBase = pNode;
+	}
+
+	//---------------------------------------------------------------------------------------------
+	void NodeImageTransform::SetOffsetX( NodeScalarPtr pNode )
+	{
+		m_pD->m_pOffsetX = pNode;
+	}
+
+	//---------------------------------------------------------------------------------------------
+	void NodeImageTransform::SetOffsetY( NodeScalarPtr pNode )
+	{
+		m_pD->m_pOffsetY = pNode;
+	}
+
+	//---------------------------------------------------------------------------------------------
+	void NodeImageTransform::SetScaleX( NodeScalarPtr pNode )
+	{
+		m_pD->m_pScaleX = pNode;
+	}
+
+	//---------------------------------------------------------------------------------------------
+	void NodeImageTransform::SetScaleY( NodeScalarPtr pNode )
+	{
+		m_pD->m_pScaleY = pNode;
+	}
+
+	//---------------------------------------------------------------------------------------------
+	void NodeImageTransform::SetRotation( NodeScalarPtr pNode )
+	{
+		m_pD->m_pRotation = pNode;
+	}
+
+}
+
+#undef NODE_INPUT_COUNT
+
diff --git a/Engine/Plugins/Experimental/Mutable/Source/MutableTools/Private/NodeImageTransformPrivate.h b/Engine/Plugins/Experimental/Mutable/Source/MutableTools/Private/NodeImageTransformPrivate.h
new file mode 100644
--- /dev/null
+++ b/Engine/Plugins/Experimental/Mutable/Source/MutableTools/Private/NodeImageTransformPrivate.h
@@ -0,0 +1,65 @@
+// Copyright Epic Games, Inc. All Rights Reserved.
+
+
+#pragma once
+
+#include "NodeImagePrivate.h"
+#include "NodeImageTransform.h"
+#include "AST.h"
+
+#include "MemoryPrivate.h"
+
+
+namespace mu
+{
+
+
+	class NodeImageTransform::Private : public NodeImage::Private
+	{
+	public:
+
+		MUTABLE_DEFINE_CONST_VISITABLE();
+
+	public:
+
+		static NODE_TYPE s_type;
+
+		NodeImagePtr m_pBase;
+		NodeScalarPtr m_pOffsetX;
+		NodeScalarPtr m_pOffsetY;
+		NodeScalarPtr m_pScaleX;
+		NodeScalarPtr m_pScaleY;
+		NodeScalarPtr m_pRotation;
+
+		//!
+		void Serialise( OutputArchive& arch ) const
+		{
+            uint32_t ver = 0;
+			arch << ver;
+
+			arch << m_pBase;
+			arch << m_pOffsetX;
+			arch << m_pOffsetY;
+			arch << m_pScaleX;
+			arch << m_pScaleY;
+			arch << m_pRotation;
+		}
+
+		//!
+		void Unserialise( InputArchive& arch )
+		{
+            uint32_t ver;
+			arch >> ver;
+			check(ver==0);
+
+			arch >> m_pBase;
+			arch >> m_pOffsetX;
+			arch >> m_pOffsetY;
+			arch >> m_pScaleX;
+			arch >> m_pScaleY;
+			arch >> m_pRotation;
+		}
+	};
+
+
+}
diff --git a/Engine/Plugins/Experimental/Mutable/Source/MutableTools/Private/NodeImageVariation.cpp b/Engine/Plugins/Experimental/Mutable/Source/MutableTools/Private/NodeImageVariation.cpp
new file mode 100644
--- /dev/null
+++ b/Engine/Plugins/Experimental/Mutable/Source/MutableTools/Private/NodeImageVariation.cpp
@@ -0,0 +1,115 @@
+// Copyright Epic Games, Inc. All Rights Reserved.
+
+
+#include "NodeImageVariation.h"
+#include "NodeImageVariationPrivate.h"
+
+
+namespace mu
+{
+
+    //---------------------------------------------------------------------------------------------
+    // Static initialisation
+    //---------------------------------------------------------------------------------------------
+    NODE_TYPE NodeImageVariation::Private::s_type =
+        NODE_TYPE( "ImageVariation", NodeImage::GetStaticType() );
+
+
+    //---------------------------------------------------------------------------------------------
+    //---------------------------------------------------------------------------------------------
+    //---------------------------------------------------------------------------------------------
+
+    MUTABLE_IMPLEMENT_NODE( NodeImageVariation, EType::Variation, Node, Node::EType::Image)
+
+
+    //---------------------------------------------------------------------------------------------
+    int NodeImageVariation::GetInputCount() const { return 1 + int( m_pD->m_variations.size() ); }
+
+
+    //---------------------------------------------------------------------------------------------
+    Node* NodeImageVariation::GetInputNode( int i ) const
+    {
+        check( i >= 0 && i < GetInputCount() );
+
+        if ( i == 0 )
+        {
+            return m_pD->m_defaultImage.get();
+        }
+        i -= 1;
+
+        if ( i < int( m_pD->m_variations.size() ) )
+        {
+            return m_pD->m_variations[i].m_image.get();
+        }
+        i -= int( m_pD->m_variations.size() );
+
+        return nullptr;
+    }
+
+
+    //---------------------------------------------------------------------------------------------
+    void NodeImageVariation::SetInputNode( int i, NodePtr pNode )
+    {
+        check( i >= 0 && i < GetInputCount() );
+
+        if ( i == 0 )
+        {
+            m_pD->m_defaultImage = dynamic_cast<NodeImage*>( pNode.get() );
+            return;
+        }
+
+        i -= 1;
+        if ( i < int( m_pD->m_variations.size() ) )
+        {
+
+            m_pD->m_variations[i].m_image = dynamic_cast<NodeImage*>( pNode.get() );
+            return;
+        }
+        i -= (int)m_pD->m_variations.size();
+    }
+
+
+    //---------------------------------------------------------------------------------------------
+    // Own Interface
+    //---------------------------------------------------------------------------------------------
+    void NodeImageVariation::SetDefaultImage( NodeImage* p ) { m_pD->m_defaultImage = p; }
+
+
+    //---------------------------------------------------------------------------------------------
+    int NodeImageVariation::GetVariationCount() const { return int( m_pD->m_variations.size() ); }
+
+
+    //---------------------------------------------------------------------------------------------
+    void NodeImageVariation::SetVariationCount( int num )
+    {
+        check( num >= 0 );
+        m_pD->m_variations.resize( num );
+    }
+
+    //---------------------------------------------------------------------------------------------
+    void NodeImageVariation::SetVariationTag( int index, const char* strTag )
+    {
+        check( index >= 0 && index < (int)m_pD->m_variations.size() );
+        check( strTag );
+
+        if ( strTag )
+        {
+            m_pD->m_variations[index].m_tag = strTag;
+        }
+        else
+        {
+            m_pD->m_variations[index].m_tag = "";
+        }
+    }
+
+
+    //---------------------------------------------------------------------------------------------
+    void NodeImageVariation::SetVariationImage( int index, NodeImage* pNode )
+    {
+        check( index >= 0 && index < (int)m_pD->m_variations.size() );
+
+        m_pD->m_variations[index].m_image = pNode;
+    }
+
+
+} // namespace mu
diff --git a/Engine/Plugins/Experimental/Mutable/Source/MutableTools/Private/NodeImageVariationPrivate.h b/Engine/Plugins/Experimental/Mutable/Source/MutableTools/Private/NodeImageVariationPrivate.h
new file mode 100644
--- /dev/null
+++ b/Engine/Plugins/Experimental/Mutable/Source/MutableTools/Private/NodeImageVariationPrivate.h
@@ -0,0 +1,80 @@
+// Copyright Epic Games, Inc. All Rights Reserved.
+
+#pragma once
+
+#include "NodePrivate.h"
+#include "NodeImageVariation.h"
+
+#include "MemoryPrivate.h"
+
+
+namespace mu
+{
+
+    //---------------------------------------------------------------------------------------------
+    //!
+    //---------------------------------------------------------------------------------------------
+    class NodeImageVariation::Private : public Node::Private
+    {
+    public:
+        MUTABLE_DEFINE_CONST_VISITABLE()
+
+    public:
+        Private() {}
+
+        static NODE_TYPE s_type;
+
+        NodeImagePtr m_defaultImage;
+
+        struct VARIATION
+        {
+            NodeImagePtr m_image;
+            string m_tag;
+
+            //!
+            void Serialise( OutputArchive& arch ) const
+            {
+                uint32_t ver = 1;
+                arch << ver;
+
+                arch << m_tag;
+                arch << m_image;
+            }
+
+            void Unserialise( InputArchive& arch )
+            {
+                uint32_t ver = 0;
+                arch >> ver;
+                check( ver == 1 );
+
+                arch >> m_tag;
+                arch >> m_image;
+            }
+        };
+
+        vector<VARIATION> m_variations;
+
+        //!
+        void Serialise( OutputArchive& arch ) const
+        {
+            uint32_t ver = 1;
+            arch << ver;
+
+            arch << m_defaultImage;
+            arch << m_variations;
+        }
+
+        //!
+        void Unserialise( InputArchive& arch )
+        {
+            uint32_t ver;
+            arch >> ver;
+			check(ver == 1);
+
+            arch >> m_defaultImage;
+            arch >> m_variations;
+        }
+    };
+
+} // namespace mu
+
diff --git a/Engine/Plugins/Experimental/Mutable/Source/MutableTools/Private/NodeLOD.cpp b/Engine/Plugins/Experimental/Mutable/Source/MutableTools/Private/NodeLOD.cpp
new file mode 100644
--- /dev/null
+++ b/Engine/Plugins/Experimental/Mutable/Source/MutableTools/Private/NodeLOD.cpp
@@ -0,0 +1,136 @@
+// Copyright Epic Games, Inc. All Rights Reserved.
+
+
+#include "MutableTools/Public/NodeLOD.h"
+#include "MutableTools/Private/NodeLODPrivate.h"
+
+
+namespace mu
+{
+
+
+	//---------------------------------------------------------------------------------------------
+	// Static initialisation
+	//---------------------------------------------------------------------------------------------
+	NODE_TYPE NodeLOD::Private::s_type =
+			NODE_TYPE( "LOD", NodeLOD::GetStaticType() );
+
+
+	//---------------------------------------------------------------------------------------------
+	//!
+	//---------------------------------------------------------------------------------------------
+
+	MUTABLE_IMPLEMENT_NODE( NodeLOD, EType::LOD, Node, Node::EType::LOD);
+
+
+	//---------------------------------------------------------------------------------------------
+	// Node Interface
+	//---------------------------------------------------------------------------------------------
+	int NodeLOD::GetInputCount() const
+	{
+		return int(m_pD->m_components.size()+m_pD->m_modifiers.size());
+	}
+
+
+	//---------------------------------------------------------------------------------------------
+	Node* NodeLOD::GetInputNode( int i ) const
+	{
+        check(i >= 0 && i < GetInputCount());
+
+        if (i < int(m_pD->m_components.size()))
+        {
+            return m_pD->m_components[i].get();
+		}
+
+        i -= int(m_pD->m_components.size());
+
+        if (i < int(m_pD->m_modifiers.size()))
+        {
+            return m_pD->m_modifiers[i].get();
+        }
+
+        return nullptr;
+    }
+
+
+    //---------------------------------------------------------------------------------------------
+	void NodeLOD::SetInputNode( int i, NodePtr pNode )
+	{
+        check(i >= 0 && i < GetInputCount());
+
+        if (i<int(m_pD->m_components.size()))
+		{
+            m_pD->m_components[i] = dynamic_cast<NodeComponent*>(pNode.get());
+            return;
+        }
+
+        i -= int(m_pD->m_components.size());
+        if (i < int(m_pD->m_modifiers.size()))
+        {
+            m_pD->m_modifiers[i] = dynamic_cast<NodeModifier*>(pNode.get());
+        }
+    }
+
+
+	//---------------------------------------------------------------------------------------------
+	// Own Interface
+	//---------------------------------------------------------------------------------------------
+	int NodeLOD::GetComponentCount() const
+	{
+		return (int)m_pD->m_components.size();
+	}
+
+
+	//---------------------------------------------------------------------------------------------
+	void NodeLOD::SetComponentCount(int num)
+	{
+		check(num >= 0);
+		m_pD->m_components.resize(num);
+	}
+
+
+	//---------------------------------------------------------------------------------------------
+	NodeComponentPtr NodeLOD::GetComponent(int index) const
+	{
+		check(index >= 0 && index < (int)m_pD->m_components.size());
+
+		return m_pD->m_components[index].get();
+	}
+
+
+	//---------------------------------------------------------------------------------------------
+	void NodeLOD::SetComponent(int index, NodeComponentPtr pComponent)
+	{
+		check(index >= 0 && index < (int)m_pD->m_components.size());
+
+		m_pD->m_components[index] = pComponent;
+	}
+
+
+	//---------------------------------------------------------------------------------------------
+	int NodeLOD::GetModifierCount() const
+	{
+		return (int)m_pD->m_modifiers.size();
+	}
+
+
+	//---------------------------------------------------------------------------------------------
+	void NodeLOD::SetModifierCount(int num)
+	{
+		check(num >= 0);
+		m_pD->m_modifiers.resize(num);
+	}
+
+
+	//---------------------------------------------------------------------------------------------
+	void NodeLOD::SetModifier(int index, NodeModifierPtr pModifier)
+	{
+		check(index >= 0 && index < (int)m_pD->m_modifiers.size());
+
+		m_pD->m_modifiers[index] = pModifier;
+	}
+
+
+}
+
+
diff --git a/Engine/Plugins/Experimental/Mutable/Source/MutableTools/Private/NodeLODPrivate.h b/Engine/Plugins/Experimental/Mutable/Source/MutableTools/Private/NodeLODPrivate.h
new file mode 100644
--- /dev/null
+++ b/Engine/Plugins/Experimental/Mutable/Source/MutableTools/Private/NodeLODPrivate.h
@@ -0,0 +1,53 @@
+// Copyright Epic Games, Inc. All Rights Reserved.
+
+#pragma once
+
+#include "NodePrivate.h"
+#include "NodeLOD.h"
+#include "NodeComponentPrivate.h"
+#include "NodeModifierPrivate.h"
+#include "AST.h"
+
+#include "MemoryPrivate.h"
+
+
+namespace mu
+{
+
+
+	class NodeLOD::Private : public Node::Private
+	{
+	public:
+
+		MUTABLE_DEFINE_CONST_VISITABLE();
+
+	public:
+
+		static NODE_TYPE s_type;
+
+		vector<NodeComponentPtr> m_components;
+		vector<NodeModifierPtr> m_modifiers;
+
+		//!
+		void Serialise( OutputArchive& arch ) const
+		{
+            uint32_t ver = 1;
+			arch << ver;
+
+			arch << m_components;
+			arch << m_modifiers;
+		}
+
+		//!
+		void Unserialise( InputArchive& arch )
+		{
+            uint32_t ver;
+			arch >> ver;
+            check(ver == 1);
+
+			arch >> m_components;
+			arch >> m_modifiers;
+		}
+	};
+
+}
diff --git a/Engine/Plugins/Experimental/Mutable/Source/MutableTools/Private/NodeLayout.cpp b/Engine/Plugins/Experimental/Mutable/Source/MutableTools/Private/NodeLayout.cpp
new file mode 100644
--- /dev/null
+++ b/Engine/Plugins/Experimental/Mutable/Source/MutableTools/Private/NodeLayout.cpp
@@ -0,0 +1,398 @@
+// Copyright Epic Games, Inc. All Rights Reserved.
+
+
+#include "NodeLayout.h"
+#include "NodeLayoutPrivate.h"
+
+#include "NodeImage.h"
+
+
+namespace mu
+{
+
+	//---------------------------------------------------------------------------------------------
+	//---------------------------------------------------------------------------------------------
+	//---------------------------------------------------------------------------------------------
+	// Static initialisation
+	//---------------------------------------------------------------------------------------------
+	static NODE_TYPE s_nodeLayoutType = NODE_TYPE( "NodeLayout", Node::GetStaticType() );
+
+
+	//---------------------------------------------------------------------------------------------
+	void NodeLayout::Serialise( const NodeLayout* p, OutputArchive& arch )
+	{
+        uint32_t ver = 0;
+		arch << ver;
+
+	#define SERIALISE_CHILDREN( C, ID ) \
+		( const C* pTyped = dynamic_cast<const C*>(p) )			\
+		{ 														\
+            arch << (uint32_t)ID;								\
+			C::Serialise( pTyped, arch );						\
+		}														\
+
+		if SERIALISE_CHILDREN( NodeLayoutBlocks				, 0 )
+		//else if SERIALISE_CHILDREN( NodeMeshIdentity		, 1 )
+
+#undef SERIALISE_CHILDREN
+	}
+
+
+	//---------------------------------------------------------------------------------------------
+	NodeLayoutPtr NodeLayout::StaticUnserialise( InputArchive& arch )
+	{
+        uint32_t ver;
+		arch >> ver;
+		check( ver == 0 );
+
+        uint32_t id;
+		arch >> id;
+
+		switch (id)
+		{
+		case 0 :  return NodeLayoutBlocks::StaticUnserialise( arch ); break;
+		default : check(false);
+		}
+
+		return 0;
+	}
+
+
+	//---------------------------------------------------------------------------------------------
+	const NODE_TYPE* NodeLayout::GetType() const
+	{
+		return GetStaticType();
+	}
+
+
+	//---------------------------------------------------------------------------------------------
+	const NODE_TYPE* NodeLayout::GetStaticType()
+	{
+		return &s_nodeLayoutType;
+	}
+
+
+
+	//---------------------------------------------------------------------------------------------
+	//---------------------------------------------------------------------------------------------
+	//---------------------------------------------------------------------------------------------
+	// Static initialisation
+	//---------------------------------------------------------------------------------------------
+	NODE_TYPE NodeLayoutBlocks::Private::s_type =
+			NODE_TYPE( "LayoutBlocks", NodeLayout::GetStaticType() );
+
+
+	//---------------------------------------------------------------------------------------------
+	//!
+	//---------------------------------------------------------------------------------------------
+
+	MUTABLE_IMPLEMENT_NODE( NodeLayoutBlocks, EType::Blocks, Node, Node::EType::Layout)
+
+
+	//---------------------------------------------------------------------------------------------
+	// Node Interface
+	//---------------------------------------------------------------------------------------------
+	int NodeLayoutBlocks::GetInputCount() const
+	{
+		return 0;
+	}
+
+
+	//---------------------------------------------------------------------------------------------
+    Node* NodeLayoutBlocks::GetInputNode( int ) const
+	{
+		check(false);
+		return 0;
+	}
+
+
+	//---------------------------------------------------------------------------------------------
+    void NodeLayoutBlocks::SetInputNode( int, NodePtr )
+	{
+		check(false);
+	}
+
+
+	//---------------------------------------------------------------------------------------------
+	// Own Interface
+	//---------------------------------------------------------------------------------------------
+	void NodeLayoutBlocks::SetGridSize( int x, int y )
+	{
+		m_pD->m_pLayout->SetGridSize( x, y );
+	}
+
+
+	//---------------------------------------------------------------------------------------------
+	void NodeLayoutBlocks::SetMaxGridSize(int x, int y)
+	{
+		m_pD->m_pLayout->SetMaxGridSize(x, y);
+	}
+
+
+	//---------------------------------------------------------------------------------------------
+	void NodeLayoutBlocks::GetGridSize( int* pX, int* pY ) const
+	{
+		FIntPoint grid = m_pD->m_pLayout->GetGridSize();
+		if (pX) *pX = grid[0];
+		if (pY) *pY = grid[1];
+	}
+
+
+	//---------------------------------------------------------------------------------------------
+	void NodeLayoutBlocks::GetMaxGridSize(int* pX, int* pY) const
+	{
+		m_pD->m_pLayout->GetMaxGridSize(pX, pY);
+	}
+
+
+	//---------------------------------------------------------------------------------------------
+	void NodeLayoutBlocks::SetBlockCount( int n )
+	{
+		m_pD->m_pLayout->SetBlockCount( n );
+	}
+
+
+	//---------------------------------------------------------------------------------------------
+	int NodeLayoutBlocks::GetBlockCount()
+	{
+		return m_pD->m_pLayout->GetBlockCount();
+	}
+
+
+	//---------------------------------------------------------------------------------------------
+    void NodeLayoutBlocks::SetBlock( int index, int minx, int miny, int sizex, int sizey )
+	{
+        m_pD->m_pLayout->SetBlock( index, minx, miny, sizex, sizey );
+	}
+
+
+	//---------------------------------------------------------------------------------------------
+	void NodeLayoutBlocks::GetBlock(int index, int* minx, int* miny, int* sizex, int* sizey)
+	{
+		if (minx && miny && sizex && sizey)
+		{
+			m_pD->m_pLayout->GetBlock(index, minx, miny, sizex, sizey);
+		}
+	}
+
+
+	//---------------------------------------------------------------------------------------------
+	void NodeLayoutBlocks::SetBlockPriority(int index, int priority)
+	{
+		m_pD->m_pLayout->SetBlockPriority(index, priority);
+	}
+
+
+	//---------------------------------------------------------------------------------------------
+	void NodeLayoutBlocks::SetLayoutPackingStrategy(mu::EPackStrategy strategy)
+	{
+		m_pD->m_pLayout->SetLayoutPackingStrategy(strategy);
+	}
+
+	//---------------------------------------------------------------------------------------------
+	NodeLayoutBlocksPtr NodeLayoutBlocks::GenerateLayoutBlocks(const MeshPtr pMesh, int layoutIndex, int gridSizeX, int gridSizeY)
+	{
+		NodeLayoutBlocksPtr layout = nullptr;
+
+		if (pMesh && layoutIndex >=0 && gridSizeX+gridSizeY>0)
+		{
+			int indexCount = pMesh->GetIndexCount();
+			vector< vec2f > UVs(indexCount*2);
+
+			UntypedMeshBufferIteratorConst indexIt(pMesh->GetIndexBuffers(), MBS_VERTEXINDEX, 0);
+			UntypedMeshBufferIteratorConst texIt(pMesh->GetVertexBuffers(), MBS_TEXCOORDS, layoutIndex);
+			
+			//Getting UVs face by face
+			for (int v = 0; v < indexCount/3; ++v)
+			{
+				uint32_t i_1 = indexIt.GetAsUINT32(); 
+				indexIt++;
+				uint32_t i_2 = indexIt.GetAsUINT32();
+				indexIt++;
+				uint32_t i_3 = indexIt.GetAsUINT32();
+				indexIt++;
+
+				float uv_1[2] = { 0.0f,0.0f };
+				ConvertData(0, uv_1, MBF_FLOAT32, (texIt + i_1).ptr(), texIt.GetFormat());
+				ConvertData(1, uv_1, MBF_FLOAT32, (texIt + i_1).ptr(), texIt.GetFormat());
+				
+				float uv_2[2] = { 0.0f,0.0f };
+				ConvertData(0, uv_2, MBF_FLOAT32, (texIt + i_2).ptr(), texIt.GetFormat());
+				ConvertData(1, uv_2, MBF_FLOAT32, (texIt + i_2).ptr(), texIt.GetFormat());
+
+				float uv_3[2] = { 0.0f,0.0f };
+				ConvertData(0, uv_3, MBF_FLOAT32, (texIt + i_3).ptr(), texIt.GetFormat());
+				ConvertData(1, uv_3, MBF_FLOAT32, (texIt + i_3).ptr(), texIt.GetFormat());
+
+				
+				UVs[v * 6 + 0][0] = uv_1[0];
+				UVs[v * 6 + 0][1] = uv_1[1];
+				UVs[v * 6 + 1][0] = uv_2[0];
+				UVs[v * 6 + 1][1] = uv_2[1];
+						
+				UVs[v * 6 + 2][0] = uv_2[0];
+				UVs[v * 6 + 2][1] = uv_2[1];
+				UVs[v * 6 + 3][0] = uv_3[0];
+				UVs[v * 6 + 3][1] = uv_3[1];
+						
+				UVs[v * 6 + 4][0] = uv_3[0]; 
+				UVs[v * 6 + 4][1] = uv_3[1]; 
+				UVs[v * 6 + 5][0] = uv_1[0];
+				UVs[v * 6 + 5][1] = uv_1[1];
+			}
+
+			layout = new NodeLayoutBlocks;
+			layout->SetGridSize(gridSizeX, gridSizeY);
+			layout->SetMaxGridSize(gridSizeX, gridSizeY);
+			layout->SetLayoutPackingStrategy(EPackStrategy::RESIZABLE_LAYOUT);
+			
+			vector<box<vec2<int>>> blocks;
+			
+			//Generating blocks
+			for (int i = 0; i < indexCount; ++i)
+			{
+				vec2<int> a, b;
+			
+				a[0] = (int)floor(UVs[i * 2][0] * gridSizeX);
+				a[1] = (int)floor(UVs[i * 2][1] * gridSizeY);
+									  
+				b[0] = (int)floor(UVs[i * 2 +1][0] * gridSizeX);
+				b[1] = (int)floor(UVs[i * 2 +1][1] * gridSizeY);
+
+				//floor of UV = 1*gridSize is gridSize which is not a valid range
+				if (a[0] == gridSizeX){ a[0] = gridSizeX-1; }
+				if (a[1] == gridSizeY){	a[1] = gridSizeY-1;	}
+				if (b[0] == gridSizeX){	b[0] = gridSizeX-1;	}
+				if (b[1] == gridSizeY){ b[1] = gridSizeY-1;	}
+			
+				//a and b are in the same block
+				if (a == b)
+				{
+					bool contains = false;
+			
+					for (size_t it = 0; it < blocks.size(); ++it)
+					{
+						if (blocks[it].Contains(a) || blocks[it].Contains(b))
+						{
+							contains = true;
+						}
+					}
+					
+					//There is no block that contains them 
+					if (!contains)
+					{
+						box<vec2<int>> currBlock;
+						currBlock.min = a;
+						currBlock.size = vec2<int>(1, 1);
+			
+						blocks.push_back(currBlock);
+					}
+				}
+				else //they are in different blocks
+				{
+					int idxA = -1;
+					int idxB = -1;
+					
+					//Getting the blocks that contain them
+					for (size_t it = 0; it < blocks.size(); ++it)
+					{
+						if (blocks[it].Contains(a))
+						{
+							idxA = (int)it;
+						}
+						if (blocks[it].Contains(b))
+						{
+							idxB = (int)it;
+						}
+					}
+					
+					//The blocks are not the same
+					if (idxA != idxB)
+					{
+						box<vec2<int>> currBlock;
+						
+						//One of the blocks doesn't exist
+						if (idxA != -1 && idxB == -1)
+						{
+							currBlock.min = b;
+							currBlock.size = vec2<int>(1, 1);
+							blocks[idxA].Bound(currBlock);
+						}
+						else if (idxB != -1 && idxA == -1)
+						{
+							currBlock.min = a;
+							currBlock.size = vec2<int>(1, 1);
+							blocks[idxB].Bound(currBlock);
+						}
+						else //Both exist
+						{
+							blocks[idxA].Bound(blocks[idxB]);
+							blocks.erase(blocks.begin() + idxB);
+						}
+					}
+					else //the blocks doesn't exist
+					{
+						if (idxA == -1)
+						{
+							box<vec2<int>> currBlockA;
+							box<vec2<int>> currBlockB;
+			
+							currBlockA.min = a;
+							currBlockB.min = b;
+							currBlockA.size = vec2<int>(1, 1);
+							currBlockB.size = vec2<int>(1, 1);
+			
+							currBlockA.Bound(currBlockB);
+							blocks.push_back(currBlockA);
+						}
+					}
+				}
+			}
+			
+			bool intersections = true;
+			
+			//Cheking if the blocks intersect with each other
+			while (intersections)
+			{
+				intersections = false;
+			
+				for (size_t i = 0; !intersections && i < blocks.size(); ++i)
+				{
+					for (size_t j = 0; j < blocks.size(); ++j)
+					{
+						if (i != j && blocks[i].Contains(blocks[j].min))
+						{
+							blocks[i].Bound(blocks[j]);
+							blocks.erase(blocks.begin() + j);
+							intersections = true;
+							break;
+						}
+					}
+				}
+			}
+			
+			int numBlocks = (int)blocks.size();
+			
+			//Generating layout blocks
+			if (numBlocks > 0)
+			{
+				layout->SetBlockCount(numBlocks);
+			
+				for (int i = 0; i < numBlocks; ++i)
+				{
+					int blockIndex = i;
+					int minX = blocks[i].min[0];
+					int minY = blocks[i].min[1];
+					int sizeX = blocks[i].size[0];
+					int sizeY = blocks[i].size[1];
+			
+					layout->SetBlock(blockIndex, minX, minY, sizeX, sizeY);
+				}
+			}
+		}
+
+		return layout;
+	}
+}
+
+
diff --git a/Engine/Plugins/Experimental/Mutable/Source/MutableTools/Private/NodeLayoutPrivate.h b/Engine/Plugins/Experimental/Mutable/Source/MutableTools/Private/NodeLayoutPrivate.h
new file mode 100644
--- /dev/null
+++ b/Engine/Plugins/Experimental/Mutable/Source/MutableTools/Private/NodeLayoutPrivate.h
@@ -0,0 +1,68 @@
+// Copyright Epic Games, Inc. All Rights Reserved.
+
+#pragma once
+
+#include "MutableTools/Private/NodePrivate.h"
+#include "NodeLayout.h"
+#include "NodeImage.h"
+#include "MutableTools/Private/AST.h"
+#include "MutableRuntime/Private/MemoryPrivate.h"
+
+
+namespace mu
+{
+
+
+	class NodeLayout::Private : public Node::Private
+	{
+	public:
+
+		virtual Layout* GetLayout() = 0;
+
+	};
+
+
+	class NodeLayoutBlocks::Private : public NodeLayout::Private
+	{
+	public:
+
+		MUTABLE_DEFINE_CONST_VISITABLE()
+
+	public:
+
+		Private()
+		{
+			m_pLayout = new Layout();
+		}
+
+		static NODE_TYPE s_type;
+
+		LayoutPtr m_pLayout;
+
+		//!
+		void Serialise( OutputArchive& arch ) const
+		{
+            uint32_t ver = 0;
+			arch << ver;
+
+			arch << m_pLayout;
+		}
+
+		//!
+		void Unserialise( InputArchive& arch )
+		{
+            uint32_t ver;
+			arch >> ver;
+			check(ver==0);
+
+			arch >> m_pLayout;
+		}
+
+        // NodeLayout::Private interface
+        Layout* GetLayout() override
+        {
+            return m_pLayout.get();
+        }
+    };
+
+}
diff --git a/Engine/Plugins/Experimental/Mutable/Source/MutableTools/Private/NodeMesh.cpp b/Engine/Plugins/Experimental/Mutable/Source/MutableTools/Private/NodeMesh.cpp
new file mode 100644
--- /dev/null
+++ b/Engine/Plugins/Experimental/Mutable/Source/MutableTools/Private/NodeMesh.cpp
@@ -0,0 +1,99 @@
+// Copyright Epic Games, Inc. All Rights Reserved.
+
+#include "NodeMesh.h"
+#include "NodeMeshPrivate.h"
+
+#include "NodeMesh.h"
+#include "NodeMeshConstant.h"
+#include "NodeMeshFormat.h"
+#include "NodeMeshInterpolate.h"
+#include "NodeMeshTable.h"
+#include "NodeMeshSubtract.h"
+#include "NodeMeshTangents.h"
+#include "NodeMeshMorph.h"
+#include "NodeMeshMakeMorph.h"
+#include "NodeMeshSwitch.h"
+#include "NodeMeshFragment.h"
+#include "NodeMeshTransform.h"
+#include "NodeMeshClipMorphPlane.h"
+#include "NodeMeshClipWithMesh.h"
+#include "NodeMeshApplyPose.h"
+#include "NodeMeshVariation.h"
+#include "NodeMeshGeometryOperation.h"
+#include "NodeMeshReshape.h"
+#include "NodeMeshClipDeform.h"
+
+namespace mu
+{
+
+	//---------------------------------------------------------------------------------------------
+	// Static initialisation
+	//---------------------------------------------------------------------------------------------
+	static NODE_TYPE s_nodeMeshType = NODE_TYPE( "NodeMesh", Node::GetStaticType() );
+
+
+	//---------------------------------------------------------------------------------------------
+	void NodeMesh::Serialise( const NodeMesh* p, OutputArchive& arch )
+	{
+        uint32_t ver = 0;
+		arch << ver;
+
+		arch << uint32_t(p->Type);
+		p->SerialiseWrapper( arch );
+    }
+
+
+    //---------------------------------------------------------------------------------------------
+	NodeMeshPtr NodeMesh::StaticUnserialise( InputArchive& arch )
+	{
+        uint32_t ver;
+		arch >> ver;
+		check( ver == 0 );
+
+        uint32_t id;
+		arch >> id;
+
+		switch (EType(id))
+		{
+		case EType::Constant:		return NodeMeshConstant::StaticUnserialise( arch ); break;
+		case EType::Interpolate:	return NodeMeshInterpolate::StaticUnserialise( arch ); break;
+		case EType::Table:			return NodeMeshTable::StaticUnserialise( arch ); break;
+		case EType::Subtract:		return NodeMeshSubtract::StaticUnserialise( arch ); break;
+		case EType::Format:			return NodeMeshFormat::StaticUnserialise( arch ); break;
+		case EType::Tangents:		return NodeMeshTangents::StaticUnserialise( arch ); break;
+		case EType::Morph:			return NodeMeshMorph::StaticUnserialise( arch ); break;
+		case EType::MakeMorph:		return NodeMeshMakeMorph::StaticUnserialise( arch ); break;
+		case EType::Switch:			return NodeMeshSwitch::StaticUnserialise( arch ); break;
+        case EType::Fragment:		return NodeMeshFragment::StaticUnserialise( arch ); break;
+		case EType::Transform:		return NodeMeshTransform::StaticUnserialise( arch ); break;
+		case EType::ClipMorphPlane: return NodeMeshClipMorphPlane::StaticUnserialise( arch ); break;
+        case EType::ClipWithMesh:	return NodeMeshClipWithMesh::StaticUnserialise( arch ); break;
+        case EType::ApplyPose:		return NodeMeshApplyPose::StaticUnserialise( arch ); break;
+		case EType::Variation:		return NodeMeshVariation::StaticUnserialise( arch ); break;
+		case EType::GeometryOperation:	return NodeMeshGeometryOperation::StaticUnserialise( arch ); break;
+		case EType::Reshape:			return NodeMeshReshape::StaticUnserialise( arch ); break;
+		case EType::ClipDeform:			return NodeMeshClipDeform::StaticUnserialise( arch ); break;
+		default : check(false);
+		}
+
+		return 0;
+	}
+
+
+	//---------------------------------------------------------------------------------------------
+	const NODE_TYPE* NodeMesh::GetType() const
+	{
+		return GetStaticType();
+	}
+
+
+	//---------------------------------------------------------------------------------------------
+	const NODE_TYPE* NodeMesh::GetStaticType()
+	{
+		return &s_nodeMeshType;
+	}
+
+
+}
+
+
diff --git a/Engine/Plugins/Experimental/Mutable/Source/MutableTools/Private/NodeMeshApplyPose.cpp b/Engine/Plugins/Experimental/Mutable/Source/MutableTools/Private/NodeMeshApplyPose.cpp
new file mode 100644
--- /dev/null
+++ b/Engine/Plugins/Experimental/Mutable/Source/MutableTools/Private/NodeMeshApplyPose.cpp
@@ -0,0 +1,134 @@
+// Copyright Epic Games, Inc. All Rights Reserved.
+
+
+#include "NodeMeshApplyPose.h"
+#include "NodeMeshApplyPosePrivate.h"
+
+#include "NodeScalar.h"
+
+#include "MeshPrivate.h"
+
+
+
+namespace mu
+{
+
+	//---------------------------------------------------------------------------------------------
+	// Static initialisation
+	//---------------------------------------------------------------------------------------------
+    NODE_TYPE NodeMeshApplyPose::Private::s_type =
+            NODE_TYPE( "MeshApplyPose", NodeMesh::GetStaticType() );
+
+
+	//---------------------------------------------------------------------------------------------
+	//!
+	//---------------------------------------------------------------------------------------------
+
+    MUTABLE_IMPLEMENT_NODE( NodeMeshApplyPose, EType::ApplyPose, Node, Node::EType::Mesh)
+
+
+	//---------------------------------------------------------------------------------------------
+	// Node Interface
+	//---------------------------------------------------------------------------------------------
+    int NodeMeshApplyPose::GetInputCount() const
+	{
+        return 2;
+	}
+
+
+	//---------------------------------------------------------------------------------------------
+    Node* NodeMeshApplyPose::GetInputNode( int i ) const
+	{
+		check( i>=0 && i<GetInputCount() );
+
+		Node* pResult = 0;
+
+		switch (i)
+		{
+		case 0:
+            pResult = m_pD->m_pBase.get();
+            break;
+        case 1:
+            pResult = m_pD->m_pPose.get();
+			break;
+        default:
+            check(false);
+            break;
+		}
+
+		return pResult;
+	}
+
+
+	//---------------------------------------------------------------------------------------------
+    void NodeMeshApplyPose::SetInputNode( int i, NodePtr pNode )
+	{
+		check( i>=0 && i<GetInputCount() );
+
+		switch (i)
+		{
+		case 0:
+            m_pD->m_pBase = dynamic_cast<NodeMesh*>(pNode.get());
+            break;
+
+		case 1:
+            m_pD->m_pPose = dynamic_cast<NodeMesh*>(pNode.get());
+			break;
+
+        default:
+            check(false);
+            break;
+        }
+	}
+
+
+	//---------------------------------------------------------------------------------------------
+	// Own Interface
+	//---------------------------------------------------------------------------------------------
+    NodeMeshPtr NodeMeshApplyPose::GetBase() const
+    {
+        return m_pD->m_pBase.get();
+    }
+
+
+    //---------------------------------------------------------------------------------------------
+    void NodeMeshApplyPose::SetBase( NodeMeshPtr pNode )
+    {
+        m_pD->m_pBase = pNode;
+    }
+
+
+    //---------------------------------------------------------------------------------------------
+    NodeMeshPtr NodeMeshApplyPose::GetPose() const
+    {
+        return m_pD->m_pPose.get();
+    }
+
+
+    //---------------------------------------------------------------------------------------------
+    void NodeMeshApplyPose::SetPose( NodeMeshPtr pNode )
+    {
+        m_pD->m_pPose = pNode;
+    }
+
+
+	//---------------------------------------------------------------------------------------------
+    NodeLayoutPtr NodeMeshApplyPose::Private::GetLayout( int index ) const
+	{
+		NodeLayoutPtr pResult;
+
+		if ( m_pBase )
+		{
+			NodeMesh::Private* pPrivate =
+					dynamic_cast<NodeMesh::Private*>( m_pBase->GetBasePrivate() );
+
+			pResult = pPrivate->GetLayout( index );
+		}
+
+		return pResult;
+	}
+
+
+}
+
+
diff --git a/Engine/Plugins/Experimental/Mutable/Source/MutableTools/Private/NodeMeshApplyPosePrivate.h b/Engine/Plugins/Experimental/Mutable/Source/MutableTools/Private/NodeMeshApplyPosePrivate.h
new file mode 100644
--- /dev/null
+++ b/Engine/Plugins/Experimental/Mutable/Source/MutableTools/Private/NodeMeshApplyPosePrivate.h
@@ -0,0 +1,56 @@
+// Copyright Epic Games, Inc. All Rights Reserved.
+
+#pragma once
+
+#include "NodeMeshApplyPose.h"
+#include "NodeScalar.h"
+#include "NodeMeshPrivate.h"
+#include "AST.h"
+
+namespace mu
+{
+
+
+    class NodeMeshApplyPose::Private : public NodeMesh::Private
+	{
+	public:
+
+		MUTABLE_DEFINE_CONST_VISITABLE()
+
+	public:
+
+		static NODE_TYPE s_type;
+
+        NodeMeshPtr m_pBase;
+        NodeMeshPtr m_pPose;
+
+		//!
+		void Serialise( OutputArchive& arch ) const
+		{
+            uint32_t ver = 0;
+			arch << ver;
+
+            arch << m_pBase;
+            arch << m_pPose;
+		}
+
+		//!
+		void Unserialise( InputArchive& arch )
+		{
+            uint32_t ver;
+			arch >> ver;
+			check(ver==0);
+
+            arch >> m_pBase;
+            arch >> m_pPose;
+        }
+
+
+        // NodeMesh::Private interface
+        NodeLayoutPtr GetLayout( int index ) const override;
+
+	};
+
+
+}
+
diff --git a/Engine/Plugins/Experimental/Mutable/Source/MutableTools/Private/NodeMeshClipDeform.cpp b/Engine/Plugins/Experimental/Mutable/Source/MutableTools/Private/NodeMeshClipDeform.cpp
new file mode 100644
--- /dev/null
+++ b/Engine/Plugins/Experimental/Mutable/Source/MutableTools/Private/NodeMeshClipDeform.cpp
@@ -0,0 +1,121 @@
+// Copyright Epic Games, Inc. All Rights Reserved.
+
+#include "NodeMeshClipDeform.h"
+#include "NodeMeshClipDeformPrivate.h"
+
+#include "NodeScalar.h"
+
+#include "MeshPrivate.h"
+#include "ImagePrivate.h"
+
+
+namespace mu
+{
+
+	//---------------------------------------------------------------------------------------------
+	// Static initialisation
+	//---------------------------------------------------------------------------------------------
+	NODE_TYPE NodeMeshClipDeform::Private::s_type =
+			NODE_TYPE( "MeshClipDeform", NodeMesh::GetStaticType() );
+
+
+	//---------------------------------------------------------------------------------------------
+	//!
+	//---------------------------------------------------------------------------------------------
+
+	MUTABLE_IMPLEMENT_NODE( NodeMeshClipDeform, EType::ClipDeform, Node, Node::EType::Mesh)
+
+
+	//---------------------------------------------------------------------------------------------
+	// Node Interface
+	//---------------------------------------------------------------------------------------------
+	int NodeMeshClipDeform::GetInputCount() const
+	{
+		return 2;
+	}
+
+
+	//---------------------------------------------------------------------------------------------
+	Node* NodeMeshClipDeform::GetInputNode( int i ) const
+	{
+		check( i>=0 && i<GetInputCount() );
+
+		Node* pResult = nullptr;
+
+		switch (i)
+		{
+		case 0: pResult = m_pD->m_pBaseMesh.get(); break;
+		case 1: pResult = m_pD->m_pClipShape.get(); break;
+		default:
+			break;
+		}
+
+		return pResult;
+	}
+
+
+	//---------------------------------------------------------------------------------------------
+	void NodeMeshClipDeform::SetInputNode( int i, NodePtr pNode )
+	{
+		check( i>=0 && i<GetInputCount() );
+
+		switch (i)
+		{
+		case 0: m_pD->m_pBaseMesh = dynamic_cast<NodeMesh*>(pNode.get()); break;
+		case 1: m_pD->m_pClipShape = dynamic_cast<NodeMesh*>(pNode.get()); break;
+		default:
+			break;
+		}
+	}
+
+
+	//---------------------------------------------------------------------------------------------
+	// Own Interface
+	//---------------------------------------------------------------------------------------------
+	const Ptr<NodeMesh>& NodeMeshClipDeform::GetBaseMesh() const
+	{
+		return m_pD->m_pBaseMesh;
+	}
+
+
+	//---------------------------------------------------------------------------------------------
+	void NodeMeshClipDeform::SetBaseMesh(const Ptr<NodeMesh>& pNode)
+	{
+		m_pD->m_pBaseMesh = pNode;
+	}
+
+
+	//---------------------------------------------------------------------------------------------
+	const Ptr<NodeMesh>& NodeMeshClipDeform::GetClipShape() const
+	{
+		return m_pD->m_pClipShape;
+	}
+
+
+	//---------------------------------------------------------------------------------------------
+	void NodeMeshClipDeform::SetClipShape(const Ptr<NodeMesh>& pNode)
+	{
+		m_pD->m_pClipShape = pNode;
+	}
+
+
+	//---------------------------------------------------------------------------------------------
+	Ptr<NodeLayout> NodeMeshClipDeform::Private::GetLayout( int index ) const
+	{
+		NodeLayoutPtr pResult;
+
+		if (m_pBaseMesh)
+		{
+			NodeMesh::Private* pPrivate =
+					dynamic_cast<NodeMesh::Private*>(m_pBaseMesh->GetBasePrivate() );
+
+			pResult = pPrivate->GetLayout( index );
+		}
+
+		return pResult;
+	}
+
+
+}
+
+
diff --git a/Engine/Plugins/Experimental/Mutable/Source/MutableTools/Private/NodeMeshClipDeformPrivate.h b/Engine/Plugins/Experimental/Mutable/Source/MutableTools/Private/NodeMeshClipDeformPrivate.h
new file mode 100644
--- /dev/null
+++ b/Engine/Plugins/Experimental/Mutable/Source/MutableTools/Private/NodeMeshClipDeformPrivate.h
@@ -0,0 +1,56 @@
+// Copyright Epic Games, Inc. All Rights Reserved.
+
+#pragma once
+
+#include "NodeMeshClipDeform.h"
+#include "NodeScalar.h"
+#include "NodeMeshPrivate.h"
+#include "NodeImagePrivate.h"
+#include "AST.h"
+
+namespace mu
+{
+
+	class NodeMeshClipDeform::Private : public NodeMesh::Private
+	{
+	public:
+
+		MUTABLE_DEFINE_CONST_VISITABLE()
+
+	public:
+
+		static NODE_TYPE s_type;
+
+		Ptr<NodeMesh> m_pBaseMesh;
+		Ptr<NodeMesh> m_pClipShape;
+		Ptr<NodeImage> m_pShapeWeights;
+
+        //!
+		void Serialise( OutputArchive& arch ) const
+		{
+            uint32_t ver = 0;
+			arch << ver;
+
+			arch << m_pBaseMesh;
+			arch << m_pClipShape;
+			arch << m_pShapeWeights;
+        }
+
+		//!
+		void Unserialise( InputArchive& arch )
+		{
+            uint32_t ver;
+			arch >> ver;
+			check(ver == 0);
+
+			arch >> m_pBaseMesh;
+			arch >> m_pClipShape;
+			arch >> m_pShapeWeights;	
+		}
+
+		// NodeMesh::Private interface
+        Ptr<NodeLayout> GetLayout( int index ) const override;
+
+	};
+
+}
diff --git a/Engine/Plugins/Experimental/Mutable/Source/MutableTools/Private/NodeMeshClipMorphPlane.cpp b/Engine/Plugins/Experimental/Mutable/Source/MutableTools/Private/NodeMeshClipMorphPlane.cpp
new file mode 100644
--- /dev/null
+++ b/Engine/Plugins/Experimental/Mutable/Source/MutableTools/Private/NodeMeshClipMorphPlane.cpp
@@ -0,0 +1,154 @@
+// Copyright Epic Games, Inc. All Rights Reserved.
+
+
+#include "NodeMeshClipMorphPlane.h"
+#include "NodeMeshClipMorphPlanePrivate.h"
+#include "MeshPrivate.h"
+
+
+#define NODE_INPUT_COUNT 	1
+
+
+namespace mu
+{
+
+	//---------------------------------------------------------------------------------------------
+	// Static initialisation
+	//---------------------------------------------------------------------------------------------
+    NODE_TYPE NodeMeshClipMorphPlane::Private::s_type =
+            NODE_TYPE( "MeshClipMorphPlane", NodeMesh::GetStaticType() );
+
+
+	//---------------------------------------------------------------------------------------------
+	//!
+	//---------------------------------------------------------------------------------------------
+
+    MUTABLE_IMPLEMENT_NODE( NodeMeshClipMorphPlane, EType::ClipMorphPlane, Node, Node::EType::Mesh)
+
+
+	//---------------------------------------------------------------------------------------------
+	// Node Interface
+	//---------------------------------------------------------------------------------------------
+    int NodeMeshClipMorphPlane::GetInputCount() const
+	{
+		return NODE_INPUT_COUNT;
+	}
+
+
+	//---------------------------------------------------------------------------------------------
+    Node* NodeMeshClipMorphPlane::GetInputNode( int i ) const
+	{
+		check( i>=0 && i<NODE_INPUT_COUNT );
+        (void)i;
+        return m_pD->m_pSource.get();
+	}
+
+
+	//---------------------------------------------------------------------------------------------
+    void NodeMeshClipMorphPlane::SetInputNode( int i, NodePtr pNode )
+	{
+		check( i>=0 && i<NODE_INPUT_COUNT );
+		if (i==0)
+		{
+			m_pD->m_pSource = dynamic_cast<NodeMesh*>( pNode.get() );
+		}
+	}
+
+
+	//---------------------------------------------------------------------------------------------
+    const char* NodeMeshClipMorphPlane::GetInputName( int i ) const
+	{
+		check( i>=0 && i<NODE_INPUT_COUNT );
+        (void)i;
+        return "Source";
+	}
+
+
+	//---------------------------------------------------------------------------------------------
+    const NODE_TYPE* NodeMeshClipMorphPlane::GetInputType( int i ) const
+	{
+		check( i>=0 && i<NODE_INPUT_COUNT );
+        (void)i;
+        return NodeMesh::GetStaticType();
+	}
+
+
+	//---------------------------------------------------------------------------------------------
+	// Own Interface
+	//---------------------------------------------------------------------------------------------
+    NodeMeshPtr NodeMeshClipMorphPlane::GetSource() const
+	{
+		return m_pD->m_pSource;
+	}
+
+
+    //---------------------------------------------------------------------------------------------
+    void NodeMeshClipMorphPlane::SetSource( NodeMesh* p )
+    {
+        m_pD->m_pSource = p;
+    }
+
+
+	//---------------------------------------------------------------------------------------------
+	void NodeMeshClipMorphPlane::SetPlane(float centerX, float centerY, float centerZ, float normalX, float normalY, float normalZ)
+	{
+		m_pD->m_origin = vec3f(centerX, centerY, centerZ);
+		m_pD->m_normal = vec3f(normalX, normalY, normalZ);
+	}
+
+
+	//---------------------------------------------------------------------------------------------
+	void NodeMeshClipMorphPlane::SetParams(float dist, float factor)
+	{
+		m_pD->m_dist = dist;
+		m_pD->m_factor = factor;
+	}
+
+	//---------------------------------------------------------------------------------------------
+	void NodeMeshClipMorphPlane::SetMorphEllipse(float radius1, float radius2, float rotation)
+	{
+		m_pD->m_radius1 = radius1;
+		m_pD->m_radius2 = radius2;
+		m_pD->m_rotation = rotation;
+	}
+
+	//---------------------------------------------------------------------------------------------
+	void NodeMeshClipMorphPlane::SetVertexSelectionBox(float centerX, float centerY, float centerZ, float radiusX, float radiusY, float radiusZ)
+	{
+		m_pD->m_vertexSelectionType = Private::VS_SHAPE;
+		m_pD->m_selectionBoxOrigin = vec3f(centerX, centerY, centerZ);
+		m_pD->m_selectionBoxRadius = vec3f(radiusX, radiusY, radiusZ);
+	}
+
+	//---------------------------------------------------------------------------------------------
+	void NodeMeshClipMorphPlane::SetVertexSelectionBone(const char* strBoneName, float maxEffectRadius)
+	{
+		m_pD->m_vertexSelectionType = Private::VS_BONE_HIERARCHY;
+		m_pD->m_vertexSelectionBone = strBoneName;
+		m_pD->m_maxEffectRadius = maxEffectRadius;
+	}
+
+	//---------------------------------------------------------------------------------------------
+    NodeLayoutPtr NodeMeshClipMorphPlane::Private::GetLayout( int index ) const
+	{
+		NodeLayoutPtr pResult;
+
+		if ( m_pSource )
+		{
+			NodeMesh::Private* pPrivate =
+					dynamic_cast<NodeMesh::Private*>( m_pSource->GetBasePrivate() );
+
+			pResult = pPrivate->GetLayout( index );
+		}
+
+		return pResult;
+	}
+
+	//---------------------------------------------------------------------------------------------
+	void NodeMeshClipMorphPlane::AddTag(const char* tagName)
+	{
+		m_pD->m_tags.push_back(tagName);
+	}
+}
+
+#undef NODE_INPUT_COUNT
diff --git a/Engine/Plugins/Experimental/Mutable/Source/MutableTools/Private/NodeMeshClipMorphPlanePrivate.h b/Engine/Plugins/Experimental/Mutable/Source/MutableTools/Private/NodeMeshClipMorphPlanePrivate.h
new file mode 100644
--- /dev/null
+++ b/Engine/Plugins/Experimental/Mutable/Source/MutableTools/Private/NodeMeshClipMorphPlanePrivate.h
@@ -0,0 +1,127 @@
+// Copyright Epic Games, Inc. All Rights Reserved.
+
+#pragma once
+
+#include "NodeMeshPrivate.h"
+#include "NodeMeshClipMorphPlane.h"
+#include "AST.h"
+
+#include "MemoryPrivate.h"
+#include "MutableMath.h"
+
+
+namespace mu
+{
+
+
+    class NodeMeshClipMorphPlane::Private : public NodeMesh::Private
+	{
+	public:
+
+		MUTABLE_DEFINE_CONST_VISITABLE()
+
+	public:
+
+		Private()
+			: m_dist(0.0f), m_factor(0.0f)
+			, m_radius1(0.0f), m_radius2(0.0f), m_rotation(0.0f)
+			, m_vertexSelectionType(VS_ALL)
+			, m_maxEffectRadius(-1.f)
+		{
+		}
+
+		static NODE_TYPE s_type;
+
+		NodeMeshPtr m_pSource;
+
+		// Morph field parameters
+
+		//! Distance to the plane of last affected vertex
+		float m_dist;
+
+		//! "Linearity" factor of the influence.
+		float m_factor;
+
+		// Ellipse location
+		vec3f m_origin;
+		vec3f m_normal;
+		float m_radius1, m_radius2, m_rotation;
+
+		//! Typed of vertex selection
+		typedef enum
+		{
+			//! All vertices, so no extra info is needed
+			VS_ALL,
+
+			//! Select vertices inside a shape
+			VS_SHAPE,
+
+			//! Select all vertices affected by any bone in a sub hierarchy
+			VS_BONE_HIERARCHY,
+
+		} VERTEX_SELECTION;
+
+		// Vertex selection box 
+        uint8_t m_vertexSelectionType;
+		vec3f m_selectionBoxOrigin;
+		vec3f m_selectionBoxRadius;
+		string m_vertexSelectionBone;
+
+		vector<mu::string> m_tags;
+
+		// Max distance a vertex can have to the bone in order to be affected. A negative value
+		// means no limit.
+		float m_maxEffectRadius;
+
+		//!
+		void Serialise( OutputArchive& arch ) const
+		{
+            uint32_t ver = 3;
+			arch << ver;
+
+			arch << m_pSource;
+			arch << m_origin;
+			arch << m_normal;
+			arch << m_dist;
+			arch << m_factor;
+			arch << m_radius1;
+			arch << m_radius2;
+			arch << m_rotation;
+			arch << m_vertexSelectionType;
+			arch << m_selectionBoxOrigin;
+			arch << m_selectionBoxRadius;
+			arch << m_tags;
+			arch << m_vertexSelectionBone;
+			arch << m_maxEffectRadius;
+		}
+
+		//!
+		void Unserialise( InputArchive& arch )
+		{
+            uint32_t ver;
+			arch >> ver;
+            check(ver==3);
+
+			arch >> m_pSource;
+			arch >> m_origin;
+			arch >> m_normal;
+			arch >> m_dist;
+			arch >> m_factor;
+			arch >> m_radius1;
+			arch >> m_radius2;
+			arch >> m_rotation;
+			arch >> m_vertexSelectionType;
+			arch >> m_selectionBoxOrigin;
+			arch >> m_selectionBoxRadius;
+			arch >> m_tags;
+			arch >> m_vertexSelectionBone;
+			arch >> m_maxEffectRadius;
+		}
+
+        // NodeMesh::Private interface
+        NodeLayoutPtr GetLayout( int index ) const override;
+
+	};
+
+
+}
diff --git a/Engine/Plugins/Experimental/Mutable/Source/MutableTools/Private/NodeMeshClipWithMesh.cpp b/Engine/Plugins/Experimental/Mutable/Source/MutableTools/Private/NodeMeshClipWithMesh.cpp
new file mode 100644
--- /dev/null
+++ b/Engine/Plugins/Experimental/Mutable/Source/MutableTools/Private/NodeMeshClipWithMesh.cpp
@@ -0,0 +1,122 @@
+// Copyright Epic Games, Inc. All Rights Reserved.
+
+
+#include "NodeMeshClipWithMesh.h"
+#include "NodeMeshClipWithMeshPrivate.h"
+#include "MeshPrivate.h"
+
+
+#define NODE_INPUT_COUNT 	1
+
+
+namespace mu
+{
+
+	//---------------------------------------------------------------------------------------------
+	// Static initialisation
+	//---------------------------------------------------------------------------------------------
+    NODE_TYPE NodeMeshClipWithMesh::Private::s_type =
+            NODE_TYPE( "MeshClipMorphPlane", NodeMesh::GetStaticType() );
+
+
+	//---------------------------------------------------------------------------------------------
+	//!
+	//---------------------------------------------------------------------------------------------
+
+    MUTABLE_IMPLEMENT_NODE( NodeMeshClipWithMesh, EType::ClipWithMesh, Node, Node::EType::Mesh)
+
+
+	//---------------------------------------------------------------------------------------------
+	// Node Interface
+	//---------------------------------------------------------------------------------------------
+    int NodeMeshClipWithMesh::GetInputCount() const
+	{
+		return NODE_INPUT_COUNT;
+	}
+
+
+	//---------------------------------------------------------------------------------------------
+    Node* NodeMeshClipWithMesh::GetInputNode( int i ) const
+	{
+		check( i>=0 && i<NODE_INPUT_COUNT );
+        (void)i;
+        return m_pD->m_pSource.get();
+	}
+
+
+	//---------------------------------------------------------------------------------------------
+    void NodeMeshClipWithMesh::SetInputNode( int i, NodePtr pNode )
+	{
+		check( i>=0 && i<NODE_INPUT_COUNT );
+		if (i==0)
+		{
+			m_pD->m_pSource = dynamic_cast<NodeMesh*>( pNode.get() );
+		}
+	}
+
+
+	//---------------------------------------------------------------------------------------------
+    const char* NodeMeshClipWithMesh::GetInputName( int i ) const
+	{
+		check( i>=0 && i<NODE_INPUT_COUNT );
+        (void)i;
+        return "Source";
+	}
+
+
+	//---------------------------------------------------------------------------------------------
+    const NODE_TYPE* NodeMeshClipWithMesh::GetInputType( int i ) const
+	{
+		check( i>=0 && i<NODE_INPUT_COUNT );
+        (void)i;
+        return NodeMesh::GetStaticType();
+	}
+
+
+	//---------------------------------------------------------------------------------------------
+	// Own Interface
+	//---------------------------------------------------------------------------------------------
+    NodeMeshPtr NodeMeshClipWithMesh::GetSource() const
+	{
+		return m_pD->m_pSource;
+	}
+
+
+    //---------------------------------------------------------------------------------------------
+    void NodeMeshClipWithMesh::SetSource(NodeMesh* p)
+    {
+        m_pD->m_pSource = p;
+    }
+
+
+	//---------------------------------------------------------------------------------------------
+	void NodeMeshClipWithMesh::SetClipMesh(NodeMesh* p)
+	{
+		m_pD->m_pClipMesh = p;
+	}
+
+
+	//---------------------------------------------------------------------------------------------
+    NodeLayoutPtr NodeMeshClipWithMesh::Private::GetLayout( int index ) const
+	{
+		NodeLayoutPtr pResult;
+
+		if ( m_pSource )
+		{
+			NodeMesh::Private* pPrivate =
+					dynamic_cast<NodeMesh::Private*>( m_pSource->GetBasePrivate() );
+
+			pResult = pPrivate->GetLayout( index );
+		}
+
+		return pResult;
+	}
+
+	//---------------------------------------------------------------------------------------------
+	void NodeMeshClipWithMesh::AddTag(const char* tagName)
+	{
+		m_pD->m_tags.push_back(tagName);
+	}
+}
+
+#undef NODE_INPUT_COUNT
diff --git a/Engine/Plugins/Experimental/Mutable/Source/MutableTools/Private/NodeMeshClipWithMeshPrivate.h b/Engine/Plugins/Experimental/Mutable/Source/MutableTools/Private/NodeMeshClipWithMeshPrivate.h
new file mode 100644
--- /dev/null
+++ b/Engine/Plugins/Experimental/Mutable/Source/MutableTools/Private/NodeMeshClipWithMeshPrivate.h
@@ -0,0 +1,64 @@
+// Copyright Epic Games, Inc. All Rights Reserved.
+
+#pragma once
+
+#include "NodeMeshPrivate.h"
+#include "NodeMeshClipWithMesh.h"
+#include "AST.h"
+
+#include "MemoryPrivate.h"
+#include "MutableMath.h"
+
+
+namespace mu
+{
+
+
+    class NodeMeshClipWithMesh::Private : public NodeMesh::Private
+	{
+	public:
+
+		MUTABLE_DEFINE_CONST_VISITABLE()
+
+	public:
+
+		Private()
+		{
+		}
+
+		static NODE_TYPE s_type;
+
+		NodeMeshPtr m_pSource;
+		NodeMeshPtr m_pClipMesh;
+
+		vector<mu::string> m_tags;
+
+		//!
+		void Serialise( OutputArchive& arch ) const
+		{
+            uint32_t ver = 1;
+			arch << ver;
+
+			arch << m_pSource;
+			arch << m_pClipMesh;
+			arch << m_tags;
+		}
+
+		//!
+		void Unserialise( InputArchive& arch )
+		{
+            uint32_t ver;
+			arch >> ver;
+
+			arch >> m_pSource;
+			arch >> m_pClipMesh;
+			arch >> m_tags;
+		}
+
+		// NodeMesh::Private interface
+        NodeLayoutPtr GetLayout( int index ) const override;
+
+	};
+
+
+}
diff --git a/Engine/Plugins/Experimental/Mutable/Source/MutableTools/Private/NodeMeshConstant.cpp b/Engine/Plugins/Experimental/Mutable/Source/MutableTools/Private/NodeMeshConstant.cpp
new file mode 100644
--- /dev/null
+++ b/Engine/Plugins/Experimental/Mutable/Source/MutableTools/Private/NodeMeshConstant.cpp
@@ -0,0 +1,128 @@
+// Copyright Epic Games, Inc. All Rights Reserved.
+
+
+#include "NodeMeshConstant.h"
+#include "NodeMeshConstantPrivate.h"
+
+#include "NodeLayout.h"
+#include "NodeLayoutPrivate.h"
+
+#include "MeshPrivate.h"
+#include "MemoryPrivate.h"
+
+
+namespace mu
+{
+
+	//---------------------------------------------------------------------------------------------
+	// Static initialisation
+	//---------------------------------------------------------------------------------------------
+	NODE_TYPE NodeMeshConstant::Private::s_type =
+			NODE_TYPE( "MeshConstant", NodeMesh::GetStaticType() );
+
+
+	//---------------------------------------------------------------------------------------------
+	//!
+	//---------------------------------------------------------------------------------------------
+
+	MUTABLE_IMPLEMENT_NODE( NodeMeshConstant, EType::Constant, Node, Node::EType::Mesh);
+
+
+	//---------------------------------------------------------------------------------------------
+	// Node Interface
+	//---------------------------------------------------------------------------------------------
+	int NodeMeshConstant::GetInputCount() const
+	{
+		return (int)m_pD->m_layouts.size();
+	}
+
+
+	//---------------------------------------------------------------------------------------------
+	Node* NodeMeshConstant::GetInputNode( int i ) const
+	{
+		check( i>=0 && i<(int)m_pD->m_layouts.size() );
+		return m_pD->m_layouts[i].get();
+	}
+
+
+	//---------------------------------------------------------------------------------------------
+	void NodeMeshConstant::SetInputNode( int i, NodePtr pNode )
+	{
+		check( i>=0 && i<(int)m_pD->m_layouts.size() );
+
+		m_pD->m_layouts[i] = dynamic_cast<NodeLayout*>( pNode.get() );
+	}
+
+
+	//---------------------------------------------------------------------------------------------
+	// Own Interface
+	//---------------------------------------------------------------------------------------------
+	MeshPtr NodeMeshConstant::GetValue() const
+	{
+		return m_pD->m_pValue;
+	}
+
+
+	//---------------------------------------------------------------------------------------------
+	void NodeMeshConstant::SetValue( MeshPtr pValue )
+	{
+		m_pD->m_pValue = pValue;
+
+        if (m_pD->m_pValue)
+        {
+            // Ensure mesh is well formed
+            m_pD->m_pValue->EnsureSurfaceData();
+        }
+    }
+
+
+	//---------------------------------------------------------------------------------------------
+	int NodeMeshConstant::GetLayoutCount() const
+	{
+		return (int)m_pD->m_layouts.size();
+	}
+
+
+	//---------------------------------------------------------------------------------------------
+	void NodeMeshConstant::SetLayoutCount( int num )
+	{
+		check( num >=0 );
+		m_pD->m_layouts.resize( num );
+	}
+
+
+	//---------------------------------------------------------------------------------------------
+	NodeLayoutPtr NodeMeshConstant::GetLayout( int index ) const
+	{
+		check( index >=0 && index < (int)m_pD->m_layouts.size() );
+
+		return m_pD->GetLayout( index );
+	}
+
+
+	//---------------------------------------------------------------------------------------------
+	NodeLayoutPtr NodeMeshConstant::Private::GetLayout( int index ) const
+	{
+		NodeLayoutPtr pResult;
+
+		if ( index >=0 && index < (int)m_layouts.size() )
+		{
+			pResult = m_layouts[ index ].get();
+		}
+
+		return pResult;
+	}
+
+
+	//---------------------------------------------------------------------------------------------
+	void NodeMeshConstant::SetLayout( int index, NodeLayoutPtr pLayout )
+	{
+		check( index >=0 && index < (int)m_pD->m_layouts.size() );
+
+		m_pD->m_layouts[ index ] = pLayout;
+	}
+
+
+}
+
+
diff --git a/Engine/Plugins/Experimental/Mutable/Source/MutableTools/Private/NodeMeshConstantPrivate.h b/Engine/Plugins/Experimental/Mutable/Source/MutableTools/Private/NodeMeshConstantPrivate.h
new file mode 100644
--- /dev/null
+++ b/Engine/Plugins/Experimental/Mutable/Source/MutableTools/Private/NodeMeshConstantPrivate.h
@@ -0,0 +1,60 @@
+// Copyright Epic Games, Inc. All Rights Reserved.
+
+#pragma once
+
+#include "NodeMeshPrivate.h"
+
+#include "NodeMeshConstant.h"
+#include "NodeLayout.h"
+#include "AST.h"
+
+
+namespace mu
+{
+
+	class NodeMeshConstant::Private : public NodeMesh::Private
+	{
+	public:
+
+		MUTABLE_DEFINE_CONST_VISITABLE()
+
+	public:
+
+		static NODE_TYPE s_type;
+
+		MeshPtr m_pValue;
+
+		vector<NodeLayoutPtr> m_layouts;
+
+		//!
+		void Serialise( OutputArchive& arch ) const
+		{
+            uint32_t ver = 0;
+			arch << ver;
+
+			arch << m_pValue;
+			arch << m_layouts;
+		}
+
+		//!
+		void Unserialise( InputArchive& arch )
+		{
+            uint32_t ver;
+			arch >> ver;
+			check(ver==0);
+
+			arch >> m_pValue;
+			arch >> m_layouts;
+
+			//if (m_pValue)
+			//{
+			//	m_pValue->CheckIntegrity();
+			//}
+        }
+
+		// NodeMesh::Private interface
+        NodeLayoutPtr GetLayout( int index ) const override;
+
+	};
+
+}
diff --git a/Engine/Plugins/Experimental/Mutable/Source/MutableTools/Private/NodeMeshFormat.cpp b/Engine/Plugins/Experimental/Mutable/Source/MutableTools/Private/NodeMeshFormat.cpp
new file mode 100644
--- /dev/null
+++ b/Engine/Plugins/Experimental/Mutable/Source/MutableTools/Private/NodeMeshFormat.cpp
@@ -0,0 +1,115 @@
+// Copyright Epic Games, Inc. All Rights Reserved.
+
+
+#include "NodeMeshFormat.h"
+#include "NodeMeshFormatPrivate.h"
+#include "MeshPrivate.h"
+
+
+namespace mu
+{
+
+	//---------------------------------------------------------------------------------------------
+	// Static initialisation
+	//---------------------------------------------------------------------------------------------
+	NODE_TYPE NodeMeshFormat::Private::s_type =
+			NODE_TYPE( "MeshFormat", NodeMesh::GetStaticType() );
+
+
+	//---------------------------------------------------------------------------------------------
+	//!
+	//---------------------------------------------------------------------------------------------
+
+	MUTABLE_IMPLEMENT_NODE( NodeMeshFormat, EType::Format, Node, Node::EType::Mesh)
+
+
+	//---------------------------------------------------------------------------------------------
+	// Node Interface
+	//---------------------------------------------------------------------------------------------
+	int NodeMeshFormat::GetInputCount() const
+	{
+		return 1;
+	}
+
+
+	//---------------------------------------------------------------------------------------------
+	Node* NodeMeshFormat::GetInputNode( int i ) const
+	{
+		check( i==0 );
+        (void)i;
+        return m_pD->m_pSource.get();
+	}
+
+
+	//---------------------------------------------------------------------------------------------
+	void NodeMeshFormat::SetInputNode( int i, NodePtr pNode )
+	{
+		check( i==0 );
+        (void)i;
+        m_pD->m_pSource = dynamic_cast<NodeMesh*>( pNode.get() );
+	}
+
+
+	//---------------------------------------------------------------------------------------------
+	// Own Interface
+	//---------------------------------------------------------------------------------------------
+	NodeMeshPtr NodeMeshFormat::GetSource() const
+	{
+		return m_pD->m_pSource;
+	}
+
+
+	//---------------------------------------------------------------------------------------------
+	void NodeMeshFormat::SetSource( NodeMesh* pValue )
+	{
+		m_pD->m_pSource = pValue;
+	}
+
+
+	//---------------------------------------------------------------------------------------------
+	FMeshBufferSet& NodeMeshFormat::GetVertexBuffers()
+	{
+		return m_pD->m_VertexBuffers;
+	}
+
+
+	//---------------------------------------------------------------------------------------------
+	FMeshBufferSet& NodeMeshFormat::GetIndexBuffers()
+	{
+		return m_pD->m_IndexBuffers;
+	}
+
+
+	//---------------------------------------------------------------------------------------------
+	FMeshBufferSet& NodeMeshFormat::GetFaceBuffers()
+	{
+		return m_pD->m_FaceBuffers;
+	}
+
+
+	//---------------------------------------------------------------------------------------------
+	void NodeMeshFormat::SetRebuildTangents( bool enabled )
+	{
+		m_pD->m_rebuildTangents = enabled;
+	}
+
+
+	//---------------------------------------------------------------------------------------------
+	NodeLayoutPtr NodeMeshFormat::Private::GetLayout( int index ) const
+	{
+		NodeLayoutPtr pResult;
+
+		if ( m_pSource )
+		{
+			NodeMesh::Private* pPrivate =
+					dynamic_cast<NodeMesh::Private*>( m_pSource->GetBasePrivate() );
+
+			pResult = pPrivate->GetLayout( index );
+		}
+
+		return pResult;
+	}
+
+}
+
+
diff --git a/Engine/Plugins/Experimental/Mutable/Source/MutableTools/Private/NodeMeshFormatPrivate.h b/Engine/Plugins/Experimental/Mutable/Source/MutableTools/Private/NodeMeshFormatPrivate.h
new file mode 100644
--- /dev/null
+++ b/Engine/Plugins/Experimental/Mutable/Source/MutableTools/Private/NodeMeshFormatPrivate.h
@@ -0,0 +1,75 @@
+// Copyright Epic Games, Inc. All Rights Reserved.
+
+#pragma once
+
+#include "NodeMeshPrivate.h"
+#include "NodeMeshFormat.h"
+#include "AST.h"
+
+#include "MemoryPrivate.h"
+
+
+namespace mu
+{
+
+
+	class NodeMeshFormat::Private : public NodeMesh::Private
+	{
+	public:
+
+		MUTABLE_DEFINE_CONST_VISITABLE()
+
+	public:
+
+		Private()
+		{
+			m_rebuildTangents = false;
+		}
+
+
+		static NODE_TYPE s_type;
+
+		//! Source mesh to transform
+		NodeMeshPtr m_pSource;
+
+		//! New mesh format. The buffers in the sets have no elements, but they define the formats.
+		//! If they are null it means that they are left with the original format.
+		FMeshBufferSet m_VertexBuffers;
+		FMeshBufferSet m_IndexBuffers;
+		FMeshBufferSet m_FaceBuffers;
+
+		bool m_rebuildTangents;
+
+		//!
+		void Serialise( OutputArchive& arch ) const
+		{
+            uint32_t ver = 3;
+			arch << ver;
+
+			arch << m_pSource;
+			arch << m_VertexBuffers;
+			arch << m_IndexBuffers;
+			arch << m_FaceBuffers;
+			arch << m_rebuildTangents;
+		}
+
+		//!
+		void Unserialise( InputArchive& arch )
+		{
+            uint32_t ver;
+			arch >> ver;
+            check(ver==3);
+
+			arch >> m_pSource;
+			arch >> m_VertexBuffers;
+			arch >> m_IndexBuffers;
+			arch >> m_FaceBuffers;
+			arch >> m_rebuildTangents;
+		}
+
+		// NodeMesh::Private interface
+        NodeLayoutPtr GetLayout( int index ) const override;
+	};
+
+
+}
