[mutable] Move Mutable plugin from the restricted folder to the experimental plugins folder.

diff --git a/Engine/Plugins/Experimental/Mutable/Source/MutableTools/Private/NodeMeshFragment.cpp b/Engine/Plugins/Experimental/Mutable/Source/MutableTools/Private/NodeMeshFragment.cpp
new file mode 100644
--- /dev/null
+++ b/Engine/Plugins/Experimental/Mutable/Source/MutableTools/Private/NodeMeshFragment.cpp
@@ -0,0 +1,158 @@
+// Copyright Epic Games, Inc. All Rights Reserved.
+
+
+#include "NodeMeshFragment.h"
+#include "NodeMeshFragmentPrivate.h"
+
+#include "NodeScalar.h"
+
+#include "MeshPrivate.h"
+
+
+namespace mu
+{
+
+	//---------------------------------------------------------------------------------------------
+	// Static initialisation
+	//---------------------------------------------------------------------------------------------
+	NODE_TYPE NodeMeshFragment::Private::s_type =
+			NODE_TYPE( "MeshFragment", NodeMesh::GetStaticType() );
+
+
+	//---------------------------------------------------------------------------------------------
+	//!
+	//---------------------------------------------------------------------------------------------
+
+	MUTABLE_IMPLEMENT_NODE( NodeMeshFragment, EType::Fragment, Node, Node::EType::Mesh)
+
+
+	//---------------------------------------------------------------------------------------------
+	// Node Interface
+	//---------------------------------------------------------------------------------------------
+	int NodeMeshFragment::GetInputCount() const
+	{
+		return 1;
+	}
+
+
+	//---------------------------------------------------------------------------------------------
+	Node* NodeMeshFragment::GetInputNode( int i ) const
+	{
+		check( i>=0 && i<GetInputCount() );
+        (void)i;
+
+		Node* pResult = 0;
+
+		pResult = m_pD->m_pMesh.get();
+
+		return pResult;
+	}
+
+
+	//---------------------------------------------------------------------------------------------
+	void NodeMeshFragment::SetInputNode( int i, NodePtr pNode )
+	{
+		check( i>=0 && i<GetInputCount() );
+
+		switch (i)
+		{
+		case 0:
+			m_pD->m_pMesh = dynamic_cast<NodeMesh*>(pNode.get());
+			break;
+
+		default:
+			break;
+		}
+	}
+
+
+	//---------------------------------------------------------------------------------------------
+	// Own Interface
+	//---------------------------------------------------------------------------------------------
+    NodeMeshFragment::FRAGMENT_TYPE NodeMeshFragment::GetFragmentType() const
+    {
+        return m_pD->m_fragmentType;
+    }
+
+
+    //---------------------------------------------------------------------------------------------
+    void NodeMeshFragment::SetFragmentType( FRAGMENT_TYPE type )
+    {
+        m_pD->m_fragmentType = type;
+    }
+
+
+    //---------------------------------------------------------------------------------------------
+    NodeMeshPtr NodeMeshFragment::GetMesh() const
+	{
+		return m_pD->m_pMesh.get();
+	}
+
+
+	//---------------------------------------------------------------------------------------------
+	void NodeMeshFragment::SetMesh( NodeMeshPtr pNode )
+	{
+		m_pD->m_pMesh = pNode;
+	}
+
+
+	//---------------------------------------------------------------------------------------------
+    int NodeMeshFragment::GetLayoutOrGroup() const
+	{
+        return m_pD->m_layoutOrGroup;
+	}
+
+
+	//---------------------------------------------------------------------------------------------
+    void NodeMeshFragment::SetLayoutOrGroup( int l )
+	{
+        m_pD->m_layoutOrGroup = l;
+	}
+
+
+	//---------------------------------------------------------------------------------------------
+	void NodeMeshFragment::SetBlockCount( int t )
+	{
+		m_pD->m_blocks.resize(t);
+	}
+
+
+	//---------------------------------------------------------------------------------------------
+	int NodeMeshFragment::GetBlock( int t ) const
+	{
+		check( t>=0 && t<(int)m_pD->m_blocks.size() );
+		return m_pD->m_blocks[t];
+	}
+
+
+	//---------------------------------------------------------------------------------------------
+	void NodeMeshFragment::SetBlock( int t, int b )
+	{
+		check( t>=0 && t<(int)m_pD->m_blocks.size() );
+		m_pD->m_blocks[t] = b;
+	}
+
+
+	//---------------------------------------------------------------------------------------------
+	NodeLayoutPtr NodeMeshFragment::Private::GetLayout( int index ) const
+	{
+		NodeLayoutPtr pResult;
+
+		if ( m_pMesh )
+		{
+			// TODO: Cut a fragment out of the layout.
+
+			NodeMesh::Private* pPrivate =
+					dynamic_cast<NodeMesh::Private*>( m_pMesh->GetBasePrivate() );
+
+			pResult = pPrivate->GetLayout( index );
+		}
+
+		return pResult;
+	}
+
+
+
+}
+
+
diff --git a/Engine/Plugins/Experimental/Mutable/Source/MutableTools/Private/NodeMeshFragmentPrivate.h b/Engine/Plugins/Experimental/Mutable/Source/MutableTools/Private/NodeMeshFragmentPrivate.h
new file mode 100644
--- /dev/null
+++ b/Engine/Plugins/Experimental/Mutable/Source/MutableTools/Private/NodeMeshFragmentPrivate.h
@@ -0,0 +1,72 @@
+// Copyright Epic Games, Inc. All Rights Reserved.
+
+#pragma once
+
+#include "NodeMeshFragment.h"
+#include "NodeScalar.h"
+#include "NodeMeshPrivate.h"
+#include "AST.h"
+
+
+namespace mu
+{
+
+
+	class NodeMeshFragment::Private : public NodeMesh::Private
+	{
+	public:
+
+		MUTABLE_DEFINE_CONST_VISITABLE()
+
+	public:
+
+		Private()
+		{
+            m_layoutOrGroup = -1;
+		}
+
+		static NODE_TYPE s_type;
+
+		NodeMeshPtr m_pMesh;
+        int m_layoutOrGroup;
+		vector<int> m_blocks;
+
+        FRAGMENT_TYPE m_fragmentType;
+
+		//!
+		void Serialise( OutputArchive& arch ) const
+		{
+            uint32_t ver = 1;
+			arch << ver;
+
+			arch << m_pMesh;
+            arch << (int32_t)m_fragmentType;
+            arch << m_layoutOrGroup;
+			arch << m_blocks;
+
+		}
+
+		//!
+		void Unserialise( InputArchive& arch )
+		{
+            uint32_t ver;
+			arch >> ver;
+            check(ver==1);
+
+			arch >> m_pMesh;
+
+            int32_t temp;
+            arch >> temp;
+            m_fragmentType = (FRAGMENT_TYPE)temp;
+
+            arch >> m_layoutOrGroup;
+			arch >> m_blocks;
+		}
+
+
+		// NodeMesh::Private interface
+        NodeLayoutPtr GetLayout( int index ) const override;
+	};
+
+
+}
diff --git a/Engine/Plugins/Experimental/Mutable/Source/MutableTools/Private/NodeMeshGeometryOperation.cpp b/Engine/Plugins/Experimental/Mutable/Source/MutableTools/Private/NodeMeshGeometryOperation.cpp
new file mode 100644
--- /dev/null
+++ b/Engine/Plugins/Experimental/Mutable/Source/MutableTools/Private/NodeMeshGeometryOperation.cpp
@@ -0,0 +1,154 @@
+// Copyright Epic Games, Inc. All Rights Reserved.
+
+
+#include "NodeMeshGeometryOperation.h"
+#include "NodeMeshGeometryOperationPrivate.h"
+
+#include "NodeScalar.h"
+
+#include "MeshPrivate.h"
+
+
+
+namespace mu
+{
+
+	//---------------------------------------------------------------------------------------------
+	// Static initialisation
+	//---------------------------------------------------------------------------------------------
+	NODE_TYPE NodeMeshGeometryOperation::Private::s_type =
+			NODE_TYPE( "MeshGeometryOperation", NodeMesh::GetStaticType() );
+
+
+	//---------------------------------------------------------------------------------------------
+	//!
+	//---------------------------------------------------------------------------------------------
+
+	MUTABLE_IMPLEMENT_NODE( NodeMeshGeometryOperation, EType::GeometryOperation, Node, Node::EType::Mesh)
+
+
+	//---------------------------------------------------------------------------------------------
+	// Node Interface
+	//---------------------------------------------------------------------------------------------
+	int NodeMeshGeometryOperation::GetInputCount() const
+	{
+		return 4;
+	}
+
+
+	//---------------------------------------------------------------------------------------------
+	Node* NodeMeshGeometryOperation::GetInputNode( int i ) const
+	{
+		check( i>=0 && i<GetInputCount() );
+
+		Node* pResult = nullptr;
+
+		switch (i)
+		{
+		case 0: pResult = m_pD->m_pMeshA.get(); break;
+		case 1: pResult = m_pD->m_pMeshB.get(); break;
+		case 2: pResult = m_pD->m_pScalarA.get(); break;
+		case 3: pResult = m_pD->m_pScalarB.get(); break;
+		default:
+			break;
+		}
+
+		return pResult;
+	}
+
+
+	//---------------------------------------------------------------------------------------------
+	void NodeMeshGeometryOperation::SetInputNode( int i, NodePtr pNode )
+	{
+		check( i>=0 && i<GetInputCount() );
+
+		switch (i)
+		{
+		case 0: m_pD->m_pMeshA = dynamic_cast<NodeMesh*>(pNode.get()); break;
+		case 1: m_pD->m_pMeshB = dynamic_cast<NodeMesh*>(pNode.get()); break;
+		case 2: m_pD->m_pScalarA = dynamic_cast<NodeScalar*>(pNode.get()); break;
+		case 3: m_pD->m_pScalarB = dynamic_cast<NodeScalar*>(pNode.get()); break;
+		default:
+			break;
+		}
+	}
+
+
+	//---------------------------------------------------------------------------------------------
+	// Own Interface
+	//---------------------------------------------------------------------------------------------
+	const NodeScalarPtr& NodeMeshGeometryOperation::GetScalarA() const
+	{
+		return m_pD->m_pScalarA;
+	}
+
+
+	//---------------------------------------------------------------------------------------------
+	void NodeMeshGeometryOperation::SetScalarA(const NodeScalarPtr& pNode)
+	{
+		m_pD->m_pScalarA = pNode;
+	}
+
+
+	//---------------------------------------------------------------------------------------------
+	const NodeScalarPtr& NodeMeshGeometryOperation::GetScalarB() const
+	{
+		return m_pD->m_pScalarB;
+	}
+
+
+	//---------------------------------------------------------------------------------------------
+	void NodeMeshGeometryOperation::SetScalarB(const NodeScalarPtr& pNode)
+	{
+		m_pD->m_pScalarB = pNode;
+	}
+
+
+	//---------------------------------------------------------------------------------------------
+	const NodeMeshPtr& NodeMeshGeometryOperation::GetMeshA() const
+	{
+		return m_pD->m_pMeshA;
+	}
+
+
+	//---------------------------------------------------------------------------------------------
+	void NodeMeshGeometryOperation::SetMeshA(const NodeMeshPtr& pNode)
+	{
+		m_pD->m_pMeshA = pNode;
+	}
+
+
+	//---------------------------------------------------------------------------------------------
+	const NodeMeshPtr& NodeMeshGeometryOperation::GetMeshB() const
+	{
+		return m_pD->m_pMeshB;
+	}
+
+
+	//---------------------------------------------------------------------------------------------
+	void NodeMeshGeometryOperation::SetMeshB(const NodeMeshPtr& pNode)
+	{
+		m_pD->m_pMeshB = pNode;
+	}
+
+
+	//---------------------------------------------------------------------------------------------
+	NodeLayoutPtr NodeMeshGeometryOperation::Private::GetLayout( int index ) const
+	{
+		NodeLayoutPtr pResult;
+
+		if (m_pMeshA)
+		{
+			NodeMesh::Private* pPrivate =
+					dynamic_cast<NodeMesh::Private*>(m_pMeshA->GetBasePrivate() );
+
+			pResult = pPrivate->GetLayout( index );
+		}
+
+		return pResult;
+	}
+
+
+}
+
+
diff --git a/Engine/Plugins/Experimental/Mutable/Source/MutableTools/Private/NodeMeshGeometryOperationPrivate.h b/Engine/Plugins/Experimental/Mutable/Source/MutableTools/Private/NodeMeshGeometryOperationPrivate.h
new file mode 100644
--- /dev/null
+++ b/Engine/Plugins/Experimental/Mutable/Source/MutableTools/Private/NodeMeshGeometryOperationPrivate.h
@@ -0,0 +1,59 @@
+// Copyright Epic Games, Inc. All Rights Reserved.
+
+#pragma once
+
+#include "NodeMeshGeometryOperation.h"
+#include "NodeScalar.h"
+#include "NodeMeshPrivate.h"
+#include "AST.h"
+
+namespace mu
+{
+
+
+	class NodeMeshGeometryOperation::Private : public NodeMesh::Private
+	{
+	public:
+
+		MUTABLE_DEFINE_CONST_VISITABLE()
+
+	public:
+
+		static NODE_TYPE s_type;
+
+		NodeMeshPtr m_pMeshA;
+		NodeMeshPtr m_pMeshB;
+		NodeScalarPtr m_pScalarA;
+		NodeScalarPtr m_pScalarB;
+
+        //!
+		void Serialise( OutputArchive& arch ) const
+		{
+            uint32_t ver = 0;
+			arch << ver;
+
+			arch << m_pMeshA;
+			arch << m_pMeshB;
+			arch << m_pScalarA;
+            arch << m_pScalarB;
+        }
+
+		//!
+		void Unserialise( InputArchive& arch )
+		{
+            uint32_t ver;
+			arch >> ver;
+			check(ver==0);
+
+			arch >> m_pMeshA;
+			arch >> m_pMeshB;
+			arch >> m_pScalarA;
+			arch >> m_pScalarB;
+		}
+
+		// NodeMesh::Private interface
+        NodeLayoutPtr GetLayout( int index ) const override;
+
+	};
+
+}
diff --git a/Engine/Plugins/Experimental/Mutable/Source/MutableTools/Private/NodeMeshInterpolate.cpp b/Engine/Plugins/Experimental/Mutable/Source/MutableTools/Private/NodeMeshInterpolate.cpp
new file mode 100644
--- /dev/null
+++ b/Engine/Plugins/Experimental/Mutable/Source/MutableTools/Private/NodeMeshInterpolate.cpp
@@ -0,0 +1,170 @@
+// Copyright Epic Games, Inc. All Rights Reserved.
+
+
+#include "NodeMeshInterpolate.h"
+#include "NodeMeshInterpolatePrivate.h"
+
+#include "NodeScalar.h"
+
+#include "MeshPrivate.h"
+
+
+#define NODE_INPUT_COUNT 	4
+
+
+namespace mu
+{
+
+	//---------------------------------------------------------------------------------------------
+	// Static initialisation
+	//---------------------------------------------------------------------------------------------
+	NODE_TYPE NodeMeshInterpolate::Private::s_type =
+			NODE_TYPE( "MeshInterpolate", NodeMesh::GetStaticType() );
+
+
+	//---------------------------------------------------------------------------------------------
+	//!
+	//---------------------------------------------------------------------------------------------
+
+	MUTABLE_IMPLEMENT_NODE( NodeMeshInterpolate, EType::Interpolate, Node, Node::EType::Mesh)
+
+
+	//---------------------------------------------------------------------------------------------
+	// Node Interface
+	//---------------------------------------------------------------------------------------------
+	int NodeMeshInterpolate::GetInputCount() const
+	{
+		return 1 + (int)m_pD->m_targets.size();
+	}
+
+
+	//---------------------------------------------------------------------------------------------
+	Node* NodeMeshInterpolate::GetInputNode( int i ) const
+	{
+		check( i>=0 && i<GetInputCount() );
+
+		Node* pResult = 0;
+
+		switch (i)
+		{
+		case 0:
+			pResult = m_pD->m_pFactor.get();
+			break;
+
+		default:
+			pResult = m_pD->m_targets[i-1].get();
+			break;
+		}
+
+		return pResult;
+	}
+
+
+	//---------------------------------------------------------------------------------------------
+	void NodeMeshInterpolate::SetInputNode( int i, NodePtr pNode )
+	{
+		check( i>=0 && i<GetInputCount() );
+
+		switch (i)
+		{
+		case 0:
+			m_pD->m_pFactor = dynamic_cast<NodeScalar*>(pNode.get());
+			break;
+
+		default:
+			m_pD->m_targets[i-1] = dynamic_cast<NodeMesh*>(pNode.get());
+			break;
+		}
+	}
+
+
+	//---------------------------------------------------------------------------------------------
+	// Own Interface
+	//---------------------------------------------------------------------------------------------
+	NodeScalarPtr NodeMeshInterpolate::GetFactor() const
+	{
+		return m_pD->m_pFactor.get();
+	}
+
+
+	//---------------------------------------------------------------------------------------------
+	void NodeMeshInterpolate::SetFactor( NodeScalarPtr pNode )
+	{
+		m_pD->m_pFactor = pNode;
+	}
+
+
+	//---------------------------------------------------------------------------------------------
+	NodeMeshPtr NodeMeshInterpolate::GetTarget( int t ) const
+	{
+		check( t>=0 && t<(int)m_pD->m_targets.size() );
+
+		return m_pD->m_targets[t].get();
+	}
+
+
+	//---------------------------------------------------------------------------------------------
+	void NodeMeshInterpolate::SetTarget( int t, NodeMeshPtr pNode )
+	{
+		check( t>=0 && t<(int)m_pD->m_targets.size() );
+		m_pD->m_targets[t] = pNode;
+	}
+
+
+	//---------------------------------------------------------------------------------------------
+	void NodeMeshInterpolate::SetTargetCount( int t )
+	{
+		m_pD->m_targets.resize(t);
+	}
+
+
+	//---------------------------------------------------------------------------------------------
+	int NodeMeshInterpolate::GetTargetCount() const
+	{
+		return (int)m_pD->m_targets.size();
+	}
+
+
+	//---------------------------------------------------------------------------------------------
+	void NodeMeshInterpolate::SetChannelCount( int t )
+	{
+		m_pD->m_channels.resize(t);
+	}
+
+
+	//---------------------------------------------------------------------------------------------
+	void NodeMeshInterpolate::SetChannel( int i, MESH_BUFFER_SEMANTIC semantic, int semanticIndex )
+	{
+		check( i>=0 && i<(int)m_pD->m_channels.size() );
+
+		if ( i>=0 && i<(int)m_pD->m_channels.size() )
+		{
+			m_pD->m_channels[i].semantic = semantic;
+			m_pD->m_channels[i].semanticIndex = semanticIndex;
+		}
+	}
+
+
+	//---------------------------------------------------------------------------------------------
+	//---------------------------------------------------------------------------------------------
+	//---------------------------------------------------------------------------------------------
+	NodeLayoutPtr NodeMeshInterpolate::Private::GetLayout( int index ) const
+	{
+		NodeLayoutPtr pResult;
+
+		// TODO: Substract layouts too? Usually they are ignored.
+		if ( m_targets.size()>0 && m_targets[0] )
+		{
+			NodeMesh::Private* pPrivate =
+					dynamic_cast<NodeMesh::Private*>( m_targets[0]->GetBasePrivate() );
+
+			pResult = pPrivate->GetLayout( index );
+		}
+
+		return pResult;
+	}
+
+
+}
+
+#undef NODE_INPUT_COUNT
diff --git a/Engine/Plugins/Experimental/Mutable/Source/MutableTools/Private/NodeMeshInterpolatePrivate.h b/Engine/Plugins/Experimental/Mutable/Source/MutableTools/Private/NodeMeshInterpolatePrivate.h
new file mode 100644
--- /dev/null
+++ b/Engine/Plugins/Experimental/Mutable/Source/MutableTools/Private/NodeMeshInterpolatePrivate.h
@@ -0,0 +1,106 @@
+// Copyright Epic Games, Inc. All Rights Reserved.
+
+#pragma once
+
+#include "NodeMeshInterpolate.h"
+#include "NodeScalar.h"
+#include "NodeMeshPrivate.h"
+#include "AST.h"
+
+#include "MeshPrivate.h"
+
+
+namespace mu
+{
+
+
+	class NodeMeshInterpolate::Private : public NodeMesh::Private
+	{
+	public:
+
+		MUTABLE_DEFINE_CONST_VISITABLE()
+
+	public:
+
+		static NODE_TYPE s_type;
+
+		NodeScalarPtr m_pFactor;
+		vector<NodeMeshPtr> m_targets;
+
+		//!
+		struct CHANNEL
+		{
+			CHANNEL()
+			{
+				semantic = MBS_NONE;
+				semanticIndex = 0;
+			}
+
+			CHANNEL( MESH_BUFFER_SEMANTIC asemantic,
+                     int32_t asemanticIndex )
+			{
+				semantic = asemantic;
+				semanticIndex = asemanticIndex;
+			}
+
+			MESH_BUFFER_SEMANTIC semantic;
+            int32_t semanticIndex;
+
+			//-------------------------------------------------------------------------------------
+			void Serialise( OutputArchive& arch ) const
+			{
+                const int32_t ver = 0;
+				arch << ver;
+
+				arch << semantic;
+				arch << semanticIndex;
+			}
+
+			//-------------------------------------------------------------------------------------
+			void Unserialise( InputArchive& arch )
+			{
+                int32_t ver = 0;
+				arch >> ver;
+				check(ver == 0);
+
+				arch >> semantic;
+				arch >> semanticIndex;
+			}
+		};
+
+		vector<CHANNEL> m_channels;
+
+		//!
+		void Serialise( OutputArchive& arch ) const
+		{
+            uint32_t ver = 1;
+			arch << ver;
+
+			arch << m_pFactor;
+			arch << m_targets;
+			arch << m_channels;
+		}
+
+		//!
+		void Unserialise( InputArchive& arch )
+		{
+            uint32_t ver;
+			arch >> ver;
+            check(ver==1);
+
+			arch >> m_pFactor;
+			arch >> m_targets;
+
+			if (ver>=1)
+			{
+				arch >> m_channels;
+			}
+		}
+
+		// NodeMesh::Private interface
+        NodeLayoutPtr GetLayout( int index ) const override;
+
+	};
+
+
+}
diff --git a/Engine/Plugins/Experimental/Mutable/Source/MutableTools/Private/NodeMeshMakeMorph.cpp b/Engine/Plugins/Experimental/Mutable/Source/MutableTools/Private/NodeMeshMakeMorph.cpp
new file mode 100644
--- /dev/null
+++ b/Engine/Plugins/Experimental/Mutable/Source/MutableTools/Private/NodeMeshMakeMorph.cpp
@@ -0,0 +1,111 @@
+// Copyright Epic Games, Inc. All Rights Reserved.
+
+
+#include "NodeMeshMakeMorph.h"
+#include "NodeMeshMakeMorphPrivate.h"
+#include "MeshPrivate.h"
+
+
+#define NODE_INPUT_COUNT 	2
+
+
+namespace mu
+{
+
+	//---------------------------------------------------------------------------------------------
+	// Static initialisation
+	//---------------------------------------------------------------------------------------------
+    NODE_TYPE NodeMeshMakeMorph::Private::s_type =
+			NODE_TYPE( "MeshMakeMorph", NodeMesh::GetStaticType() );
+
+
+	//---------------------------------------------------------------------------------------------
+	//!
+	//---------------------------------------------------------------------------------------------
+
+    MUTABLE_IMPLEMENT_NODE( NodeMeshMakeMorph, EType::MakeMorph, Node, Node::EType::Mesh)
+
+
+	//---------------------------------------------------------------------------------------------
+	// Node Interface
+	//---------------------------------------------------------------------------------------------
+    int NodeMeshMakeMorph::GetInputCount() const
+	{
+		return NODE_INPUT_COUNT;
+	}
+
+
+	//---------------------------------------------------------------------------------------------
+    Node* NodeMeshMakeMorph::GetInputNode( int i ) const
+	{
+		check( i>=0 && i<NODE_INPUT_COUNT );
+		return i == 0 ? m_pD->m_pBase.get() : m_pD->m_pTarget.get();
+	}
+
+
+	//---------------------------------------------------------------------------------------------
+    void NodeMeshMakeMorph::SetInputNode( int i, NodePtr pNode )
+	{
+		check( i>=0 && i<NODE_INPUT_COUNT );
+		if (i==0)
+		{
+			m_pD->m_pBase = dynamic_cast<NodeMesh*>( pNode.get() );
+		}
+		else
+		{
+			m_pD->m_pTarget = dynamic_cast<NodeMesh*>( pNode.get() );
+		}
+	}
+
+
+	//---------------------------------------------------------------------------------------------
+	// Own Interface
+	//---------------------------------------------------------------------------------------------
+    NodeMeshPtr NodeMeshMakeMorph::GetBase() const
+	{
+		return m_pD->m_pBase;
+	}
+
+
+	//---------------------------------------------------------------------------------------------
+    void NodeMeshMakeMorph::SetBase( NodeMesh* p )
+	{
+		m_pD->m_pBase = p;
+	}
+
+
+	//---------------------------------------------------------------------------------------------
+    NodeMeshPtr NodeMeshMakeMorph::GetTarget() const
+	{
+		return m_pD->m_pTarget;
+	}
+
+
+	//---------------------------------------------------------------------------------------------
+    void NodeMeshMakeMorph::SetTarget( NodeMesh* p )
+	{
+		m_pD->m_pTarget = p;
+	}
+
+
+	//---------------------------------------------------------------------------------------------
+    NodeLayoutPtr NodeMeshMakeMorph::Private::GetLayout( int index ) const
+	{
+		NodeLayoutPtr pResult;
+
+		// TODO: Substract layouts too? Usually they are ignored.
+		if ( m_pBase )
+		{
+			NodeMesh::Private* pPrivate =
+					dynamic_cast<NodeMesh::Private*>( m_pBase->GetBasePrivate() );
+
+			pResult = pPrivate->GetLayout( index );
+		}
+
+		return pResult;
+	}
+
+
+}
+
+#undef NODE_INPUT_COUNT
diff --git a/Engine/Plugins/Experimental/Mutable/Source/MutableTools/Private/NodeMeshMakeMorphPrivate.h b/Engine/Plugins/Experimental/Mutable/Source/MutableTools/Private/NodeMeshMakeMorphPrivate.h
new file mode 100644
--- /dev/null
+++ b/Engine/Plugins/Experimental/Mutable/Source/MutableTools/Private/NodeMeshMakeMorphPrivate.h
@@ -0,0 +1,56 @@
+// Copyright Epic Games, Inc. All Rights Reserved.
+
+#pragma once
+
+#include "NodeMeshPrivate.h"
+#include "NodeMeshMakeMorph.h"
+#include "AST.h"
+
+#include "MemoryPrivate.h"
+
+
+namespace mu
+{
+
+
+    class NodeMeshMakeMorph::Private : public NodeMesh::Private
+	{
+	public:
+
+		MUTABLE_DEFINE_CONST_VISITABLE()
+
+	public:
+
+		static NODE_TYPE s_type;
+
+		NodeMeshPtr m_pBase;
+		NodeMeshPtr m_pTarget;
+
+		//!
+		void Serialise( OutputArchive& arch ) const
+		{
+            uint32_t ver = 0;
+			arch << ver;
+
+			arch << m_pBase;
+			arch << m_pTarget;
+		}
+
+		//!
+		void Unserialise( InputArchive& arch )
+		{
+            uint32_t ver;
+			arch >> ver;
+			check(ver==0);
+
+			arch >> m_pBase;
+			arch >> m_pTarget;
+		}
+
+		// NodeMesh::Private interface
+        NodeLayoutPtr GetLayout( int index ) const override;
+
+	};
+
+
+}
diff --git a/Engine/Plugins/Experimental/Mutable/Source/MutableTools/Private/NodeMeshMorph.cpp b/Engine/Plugins/Experimental/Mutable/Source/MutableTools/Private/NodeMeshMorph.cpp
new file mode 100644
--- /dev/null
+++ b/Engine/Plugins/Experimental/Mutable/Source/MutableTools/Private/NodeMeshMorph.cpp
@@ -0,0 +1,205 @@
+// Copyright Epic Games, Inc. All Rights Reserved.
+
+
+#include "NodeMeshMorph.h"
+#include "NodeMeshMorphPrivate.h"
+
+#include "NodeScalar.h"
+
+#include "MeshPrivate.h"
+
+
+
+namespace mu
+{
+
+	//---------------------------------------------------------------------------------------------
+	// Static initialisation
+	//---------------------------------------------------------------------------------------------
+	NODE_TYPE NodeMeshMorph::Private::s_type =
+			NODE_TYPE( "MeshMorph", NodeMesh::GetStaticType() );
+
+
+	//---------------------------------------------------------------------------------------------
+	//!
+	//---------------------------------------------------------------------------------------------
+
+	MUTABLE_IMPLEMENT_NODE( NodeMeshMorph, EType::Morph, Node, Node::EType::Mesh)
+
+
+	//---------------------------------------------------------------------------------------------
+	// Node Interface
+	//---------------------------------------------------------------------------------------------
+	int NodeMeshMorph::GetInputCount() const
+	{
+		return 2 + (int)m_pD->m_morphs.size();
+	}
+
+
+	//---------------------------------------------------------------------------------------------
+	Node* NodeMeshMorph::GetInputNode( int i ) const
+	{
+		check( i>=0 && i<GetInputCount() );
+
+		Node* pResult = 0;
+
+		switch (i)
+		{
+		case 0:
+			pResult = m_pD->m_pFactor.get();
+			break;
+
+		case 1:
+			pResult = m_pD->m_pBase.get();
+			break;
+
+		default:
+			pResult = m_pD->m_morphs[i-2].get();
+			break;
+		}
+
+		return pResult;
+	}
+
+
+	//---------------------------------------------------------------------------------------------
+	void NodeMeshMorph::SetInputNode( int i, NodePtr pNode )
+	{
+		check( i>=0 && i<GetInputCount() );
+
+		switch (i)
+		{
+		case 0:
+			m_pD->m_pFactor = dynamic_cast<NodeScalar*>(pNode.get());
+			break;
+
+		case 1:
+			m_pD->m_pBase = dynamic_cast<NodeMesh*>(pNode.get());
+			break;
+
+		default:
+			m_pD->m_morphs[i-2] = dynamic_cast<NodeMesh*>(pNode.get());
+			break;
+		}
+	}
+
+
+	//---------------------------------------------------------------------------------------------
+	// Own Interface
+	//---------------------------------------------------------------------------------------------
+	NodeScalarPtr NodeMeshMorph::GetFactor() const
+	{
+		return m_pD->m_pFactor.get();
+	}
+
+
+	//---------------------------------------------------------------------------------------------
+	void NodeMeshMorph::SetFactor( NodeScalarPtr pNode )
+	{
+		m_pD->m_pFactor = pNode;
+	}
+
+
+	//---------------------------------------------------------------------------------------------
+	NodeMeshPtr NodeMeshMorph::GetBase() const
+	{
+		return m_pD->m_pBase.get();
+	}
+
+
+	//---------------------------------------------------------------------------------------------
+	void NodeMeshMorph::SetBase( NodeMeshPtr pNode )
+	{
+		m_pD->m_pBase = pNode;
+	}
+
+
+	//---------------------------------------------------------------------------------------------
+	NodeMeshPtr NodeMeshMorph::GetMorph( int t ) const
+	{
+		check( t>=0 && t<(int)m_pD->m_morphs.size() );
+
+		return m_pD->m_morphs[t].get();
+	}
+
+
+	//---------------------------------------------------------------------------------------------
+	void NodeMeshMorph::SetMorph( int t, NodeMeshPtr pNode )
+	{
+		check( t>=0 && t<(int)m_pD->m_morphs.size() );
+		m_pD->m_morphs[t] = pNode;
+	}
+
+	//---------------------------------------------------------------------------------------------
+	void NodeMeshMorph::SetMorphCount( int t )
+	{
+		m_pD->m_morphs.resize(t);
+	}
+
+	//---------------------------------------------------------------------------------------------
+	int NodeMeshMorph::GetMorphCount() const
+	{
+		return (int)m_pD->m_morphs.size();
+	}
+
+	//---------------------------------------------------------------------------------------------
+	void NodeMeshMorph::SetMorphIndicesAreRelative( bool relative )
+	{
+		m_pD->m_vertexIndicesAreRelative = relative;
+	}
+
+	//---------------------------------------------------------------------------------------------
+	void NodeMeshMorph::SetReshapeSkeleton(bool bEnable)
+	{
+		m_pD->m_reshapeSkeleton = bEnable;
+	}	
+
+	//---------------------------------------------------------------------------------------------
+	void NodeMeshMorph::SetDeformAllBones(bool bEnable)
+	{
+		m_pD->m_deformAllBones = bEnable;
+	}
+
+	//---------------------------------------------------------------------------------------------
+	void NodeMeshMorph::SetDeformAllPhysics(bool bEnable)
+	{
+		m_pD->m_deformAllPhysics = bEnable;
+	}
+
+	//---------------------------------------------------------------------------------------------
+	void NodeMeshMorph::AddBoneToDeform(const char* BoneName)
+	{
+		m_pD->m_bonesToDeform.Emplace(BoneName);
+	}
+
+	void NodeMeshMorph::AddPhysicsBodyToDeform(const char* BoneName)
+	{
+		m_pD->m_physicsToDeform.Emplace(BoneName);
+	}
+
+	//---------------------------------------------------------------------------------------------
+	void NodeMeshMorph::SetReshapePhysicsVolumes(bool bEnable)
+	{
+		m_pD->m_reshapePhysicsVolumes = bEnable;
+	}
+
+	//---------------------------------------------------------------------------------------------
+	NodeLayoutPtr NodeMeshMorph::Private::GetLayout( int index ) const
+	{
+		NodeLayoutPtr pResult;
+
+		if ( m_pBase )
+		{
+			NodeMesh::Private* pPrivate =
+					dynamic_cast<NodeMesh::Private*>( m_pBase->GetBasePrivate() );
+
+			pResult = pPrivate->GetLayout( index );
+		}
+
+		return pResult;
+	}
+
+
+}
+
+
diff --git a/Engine/Plugins/Experimental/Mutable/Source/MutableTools/Private/NodeMeshMorphPrivate.h b/Engine/Plugins/Experimental/Mutable/Source/MutableTools/Private/NodeMeshMorphPrivate.h
new file mode 100644
--- /dev/null
+++ b/Engine/Plugins/Experimental/Mutable/Source/MutableTools/Private/NodeMeshMorphPrivate.h
@@ -0,0 +1,112 @@
+// Copyright Epic Games, Inc. All Rights Reserved.
+
+#pragma once
+
+#include "NodeMeshMorph.h"
+#include "NodeScalar.h"
+#include "NodeMeshPrivate.h"
+#include "AST.h"
+
+namespace mu
+{
+
+
+	class NodeMeshMorph::Private : public NodeMesh::Private
+	{
+	public:
+
+		MUTABLE_DEFINE_CONST_VISITABLE()
+
+	public:
+
+		static NODE_TYPE s_type;
+
+		NodeScalarPtr m_pFactor;
+		NodeMeshPtr m_pBase;
+		vector<NodeMeshPtr> m_morphs;
+        bool m_vertexIndicesAreRelative = false;
+		
+		bool m_reshapeSkeleton = false;
+		bool m_reshapePhysicsVolumes = false;
+		
+		bool m_deformAllBones = false;
+		bool m_deformAllPhysics = false;
+		
+		TArray<string> m_bonesToDeform;
+		TArray<string> m_physicsToDeform;
+
+        //!
+		void Serialise( OutputArchive& arch ) const
+		{
+            uint32_t ver = 3;
+			arch << ver;
+
+			arch << m_pFactor;
+			arch << m_pBase;
+			arch << m_morphs;
+            arch << m_vertexIndicesAreRelative;
+
+			arch << m_reshapeSkeleton;
+			arch << m_reshapePhysicsVolumes;
+			arch << m_bonesToDeform;
+
+			arch << m_deformAllBones;
+			arch << m_deformAllPhysics;
+			
+			arch << m_physicsToDeform;
+
+        }
+
+		//!
+		void Unserialise( InputArchive& arch )
+		{
+            uint32_t ver;
+			arch >> ver;
+			check(ver <= 3);
+
+			arch >> m_pFactor;
+			arch >> m_pBase;
+			arch >> m_morphs;
+            arch >> m_vertexIndicesAreRelative;
+	
+			if (ver >= 2)
+			{
+				arch >> m_reshapeSkeleton;
+				arch >> m_reshapePhysicsVolumes;
+	
+				// This repetition is needed, there was a bug where m_reshapePhysicsVolumes was serialized twice in ver 2.
+				if (ver == 2) 
+				{
+					arch >> m_reshapePhysicsVolumes;
+				}
+				arch >> m_bonesToDeform;			
+			}
+			else
+			{
+				m_reshapeSkeleton = false;
+				m_reshapePhysicsVolumes = false;
+				m_bonesToDeform.Empty();
+			}
+
+			if (ver >= 3)
+			{
+				arch >> m_deformAllBones;
+				arch >> m_deformAllPhysics;
+				arch >> m_physicsToDeform;
+			}
+			else
+			{
+				m_deformAllBones = false;
+				m_deformAllPhysics = false;
+				m_physicsToDeform.Empty();
+			}
+
+		}
+
+		// NodeMesh::Private interface
+        NodeLayoutPtr GetLayout( int index ) const override;
+
+	};
+
+
+}
diff --git a/Engine/Plugins/Experimental/Mutable/Source/MutableTools/Private/NodeMeshPrivate.h b/Engine/Plugins/Experimental/Mutable/Source/MutableTools/Private/NodeMeshPrivate.h
new file mode 100644
--- /dev/null
+++ b/Engine/Plugins/Experimental/Mutable/Source/MutableTools/Private/NodeMeshPrivate.h
@@ -0,0 +1,23 @@
+// Copyright Epic Games, Inc. All Rights Reserved.
+
+#pragma once
+
+#include "NodePrivate.h"
+
+#include "NodeMesh.h"
+#include "NodeLayout.h"
+
+
+namespace mu
+{
+
+
+	class NodeMesh::Private : public Node::Private
+	{
+	public:
+
+		virtual NodeLayoutPtr GetLayout( int index ) const = 0;
+
+	};
+
+}
diff --git a/Engine/Plugins/Experimental/Mutable/Source/MutableTools/Private/NodeMeshReshape.cpp b/Engine/Plugins/Experimental/Mutable/Source/MutableTools/Private/NodeMeshReshape.cpp
new file mode 100644
--- /dev/null
+++ b/Engine/Plugins/Experimental/Mutable/Source/MutableTools/Private/NodeMeshReshape.cpp
@@ -0,0 +1,179 @@
+// Copyright Epic Games, Inc. All Rights Reserved.
+
+#include "NodeMeshReshape.h"
+#include "NodeMeshReshapePrivate.h"
+
+#include "NodeScalar.h"
+
+#include "MeshPrivate.h"
+
+
+
+namespace mu
+{
+
+	//---------------------------------------------------------------------------------------------
+	// Static initialisation
+	//---------------------------------------------------------------------------------------------
+	NODE_TYPE NodeMeshReshape::Private::s_type =
+			NODE_TYPE( "MeshReshape", NodeMesh::GetStaticType() );
+
+
+	//---------------------------------------------------------------------------------------------
+	//!
+	//---------------------------------------------------------------------------------------------
+
+	MUTABLE_IMPLEMENT_NODE( NodeMeshReshape, EType::Reshape, Node, Node::EType::Mesh)
+
+
+	//---------------------------------------------------------------------------------------------
+	// Node Interface
+	//---------------------------------------------------------------------------------------------
+	int NodeMeshReshape::GetInputCount() const
+	{
+		return 3;
+	}
+
+
+	//---------------------------------------------------------------------------------------------
+	Node* NodeMeshReshape::GetInputNode( int i ) const
+	{
+		check( i>=0 && i<GetInputCount() );
+
+		Node* pResult = nullptr;
+
+		switch (i)
+		{
+		case 0: pResult = m_pD->m_pBaseMesh.get(); break;
+		case 1: pResult = m_pD->m_pBaseShape.get(); break;
+		case 2: pResult = m_pD->m_pTargetShape.get(); break;
+		default:
+			break;
+		}
+
+		return pResult;
+	}
+
+
+	//---------------------------------------------------------------------------------------------
+	void NodeMeshReshape::SetInputNode( int i, NodePtr pNode )
+	{
+		check( i>=0 && i<GetInputCount() );
+
+		switch (i)
+		{
+		case 0: m_pD->m_pBaseMesh = dynamic_cast<NodeMesh*>(pNode.get()); break;
+		case 1: m_pD->m_pBaseShape = dynamic_cast<NodeMesh*>(pNode.get()); break;
+		case 2: m_pD->m_pTargetShape = dynamic_cast<NodeMesh*>(pNode.get()); break;
+		default:
+			break;
+		}
+	}
+
+
+	//---------------------------------------------------------------------------------------------
+	// Own Interface
+	//---------------------------------------------------------------------------------------------
+	const Ptr<NodeMesh>& NodeMeshReshape::GetBaseMesh() const
+	{
+		return m_pD->m_pBaseMesh;
+	}
+
+
+	//---------------------------------------------------------------------------------------------
+	void NodeMeshReshape::SetBaseMesh(const Ptr<NodeMesh>& pNode)
+	{
+		m_pD->m_pBaseMesh = pNode;
+	}
+
+
+	//---------------------------------------------------------------------------------------------
+	const Ptr<NodeMesh>& NodeMeshReshape::GetBaseShape() const
+	{
+		return m_pD->m_pBaseShape;
+	}
+
+
+	//---------------------------------------------------------------------------------------------
+	void NodeMeshReshape::SetBaseShape(const Ptr<NodeMesh>& pNode)
+	{
+		m_pD->m_pBaseShape = pNode;
+	}
+
+
+	//---------------------------------------------------------------------------------------------
+	const Ptr<NodeMesh>& NodeMeshReshape::GetTargetShape() const
+	{
+		return m_pD->m_pTargetShape;
+	}
+
+
+	//---------------------------------------------------------------------------------------------
+	void NodeMeshReshape::SetTargetShape(const Ptr<NodeMesh>& pNode)
+	{
+		m_pD->m_pTargetShape = pNode;
+	}
+
+
+	//---------------------------------------------------------------------------------------------
+	void NodeMeshReshape::SetReshapeSkeleton(bool bEnable)
+	{
+		m_pD->m_reshapeSkeleton = bEnable;
+	}
+	
+
+	//---------------------------------------------------------------------------------------------
+	void NodeMeshReshape::SetEnableRigidParts(bool bEnable)
+	{
+		m_pD->m_enableRigidParts = bEnable;
+	}
+
+	//---------------------------------------------------------------------------------------------
+	void NodeMeshReshape::SetDeformAllBones(bool bEnable)
+	{
+		m_pD->m_deformAllBones = bEnable;
+	}
+
+	//---------------------------------------------------------------------------------------------
+	void NodeMeshReshape::SetDeformAllPhysics(bool bEnable)
+	{
+		m_pD->m_deformAllPhysics = bEnable;
+	}
+
+	//---------------------------------------------------------------------------------------------
+	void NodeMeshReshape::AddBoneToDeform(const char* BoneName)
+	{
+		m_pD->m_bonesToDeform.Emplace(BoneName);
+	}
+
+	//---------------------------------------------------------------------------------------------
+	void NodeMeshReshape::AddPhysicsBodyToDeform(const char* BoneName)
+	{
+		m_pD->m_physicsToDeform.Emplace(BoneName);
+	}
+
+	void NodeMeshReshape::SetReshapePhysicsVolumes(bool bEnable)
+	{
+		m_pD->m_reshapePhysicsVolumes = bEnable;
+	}
+	
+	//---------------------------------------------------------------------------------------------
+	Ptr<NodeLayout> NodeMeshReshape::Private::GetLayout( int index ) const
+	{
+		NodeLayoutPtr pResult;
+
+		if (m_pBaseMesh)
+		{
+			NodeMesh::Private* pPrivate =
+					dynamic_cast<NodeMesh::Private*>(m_pBaseMesh->GetBasePrivate() );
+
+			pResult = pPrivate->GetLayout( index );
+		}
+
+		return pResult;
+	}
+
+
+}
+
+
diff --git a/Engine/Plugins/Experimental/Mutable/Source/MutableTools/Private/NodeMeshReshapePrivate.h b/Engine/Plugins/Experimental/Mutable/Source/MutableTools/Private/NodeMeshReshapePrivate.h
new file mode 100644
--- /dev/null
+++ b/Engine/Plugins/Experimental/Mutable/Source/MutableTools/Private/NodeMeshReshapePrivate.h
@@ -0,0 +1,85 @@
+// Copyright Epic Games, Inc. All Rights Reserved.
+
+#pragma once
+
+#include "NodeMeshReshape.h"
+#include "NodeScalar.h"
+#include "NodeMeshPrivate.h"
+#include "AST.h"
+
+namespace mu
+{
+
+
+	class NodeMeshReshape::Private : public NodeMesh::Private
+	{
+	public:
+
+		MUTABLE_DEFINE_CONST_VISITABLE()
+
+	public:
+
+		static NODE_TYPE s_type;
+
+		Ptr<NodeMesh> m_pBaseMesh;
+		Ptr<NodeMesh> m_pBaseShape;
+		Ptr<NodeMesh> m_pTargetShape;
+		bool m_reshapeSkeleton = false;
+		bool m_enableRigidParts = false;
+		bool m_deformAllBones = false;
+		bool m_deformAllPhysics = false;
+		bool m_reshapePhysicsVolumes = false;
+		
+		TArray<string> m_bonesToDeform;
+		TArray<string> m_physicsToDeform;
+        //!
+		void Serialise( OutputArchive& arch ) const
+		{
+            uint32_t ver = 5;
+			arch << ver;
+
+			arch << m_pBaseMesh;
+			arch << m_pBaseShape;
+			arch << m_pTargetShape;
+			arch << m_reshapeSkeleton;
+			arch << m_enableRigidParts;
+			arch << m_deformAllBones;
+			arch << m_bonesToDeform;
+			arch << m_reshapePhysicsVolumes;
+			arch << m_deformAllPhysics;
+			arch << m_physicsToDeform;
+        }
+
+		//!
+		void Unserialise( InputArchive& arch )
+		{
+            uint32_t ver;
+			arch >> ver;
+			check(ver<=5);
+
+			arch >> m_pBaseMesh;
+			arch >> m_pBaseShape;
+			arch >> m_pTargetShape;
+			arch >> m_reshapeSkeleton;
+			arch >> m_enableRigidParts;
+			arch >> m_deformAllBones;
+			arch >> m_bonesToDeform;
+			
+		  	if (ver >= 4)
+		 	{
+				arch >> m_reshapePhysicsVolumes;
+		 	}
+
+			if (ver >= 5)
+			{
+				arch >> m_deformAllPhysics;
+				arch >> m_physicsToDeform;
+			}
+		}
+
+		// NodeMesh::Private interface
+        Ptr<NodeLayout> GetLayout( int index ) const override;
+
+	};
+
+}
diff --git a/Engine/Plugins/Experimental/Mutable/Source/MutableTools/Private/NodeMeshSubtract.cpp b/Engine/Plugins/Experimental/Mutable/Source/MutableTools/Private/NodeMeshSubtract.cpp
new file mode 100644
--- /dev/null
+++ b/Engine/Plugins/Experimental/Mutable/Source/MutableTools/Private/NodeMeshSubtract.cpp
@@ -0,0 +1,111 @@
+// Copyright Epic Games, Inc. All Rights Reserved.
+
+
+#include "NodeMeshSubtract.h"
+#include "NodeMeshSubtractPrivate.h"
+#include "MeshPrivate.h"
+
+
+#define NODE_INPUT_COUNT 	2
+
+
+namespace mu
+{
+
+	//---------------------------------------------------------------------------------------------
+	// Static initialisation
+	//---------------------------------------------------------------------------------------------
+	NODE_TYPE NodeMeshSubtract::Private::s_type =
+			NODE_TYPE( "MeshSubtract", NodeMesh::GetStaticType() );
+
+
+	//---------------------------------------------------------------------------------------------
+	//!
+	//---------------------------------------------------------------------------------------------
+
+	MUTABLE_IMPLEMENT_NODE( NodeMeshSubtract, EType::Subtract, Node, Node::EType::Mesh)
+
+
+	//---------------------------------------------------------------------------------------------
+	// Node Interface
+	//---------------------------------------------------------------------------------------------
+	int NodeMeshSubtract::GetInputCount() const
+	{
+		return NODE_INPUT_COUNT;
+	}
+
+
+	//---------------------------------------------------------------------------------------------
+	Node* NodeMeshSubtract::GetInputNode( int i ) const
+	{
+		check( i>=0 && i<NODE_INPUT_COUNT );
+		return i == 0 ? m_pD->m_pA.get() : m_pD->m_pB.get();
+	}
+
+
+	//---------------------------------------------------------------------------------------------
+	void NodeMeshSubtract::SetInputNode( int i, NodePtr pNode )
+	{
+		check( i>=0 && i<NODE_INPUT_COUNT );
+		if (i==0)
+		{
+			m_pD->m_pA = dynamic_cast<NodeMesh*>( pNode.get() );
+		}
+		else
+		{
+			m_pD->m_pB = dynamic_cast<NodeMesh*>( pNode.get() );
+		}
+	}
+
+
+	//---------------------------------------------------------------------------------------------
+	// Own Interface
+	//---------------------------------------------------------------------------------------------
+	NodeMeshPtr NodeMeshSubtract::GetA() const
+	{
+		return m_pD->m_pA;
+	}
+
+
+	//---------------------------------------------------------------------------------------------
+	void NodeMeshSubtract::SetA( NodeMesh* p )
+	{
+		m_pD->m_pA = p;
+	}
+
+
+	//---------------------------------------------------------------------------------------------
+	NodeMeshPtr NodeMeshSubtract::GetB() const
+	{
+		return m_pD->m_pB;
+	}
+
+
+	//---------------------------------------------------------------------------------------------
+	void NodeMeshSubtract::SetB( NodeMesh* p )
+	{
+		m_pD->m_pB = p;
+	}
+
+
+	//---------------------------------------------------------------------------------------------
+	NodeLayoutPtr NodeMeshSubtract::Private::GetLayout( int index ) const
+	{
+		NodeLayoutPtr pResult;
+
+		// TODO: Substract layouts too? Usually they are ignored.
+		if ( m_pA )
+		{
+			NodeMesh::Private* pPrivate =
+					dynamic_cast<NodeMesh::Private*>( m_pA->GetBasePrivate() );
+
+			pResult = pPrivate->GetLayout( index );
+		}
+
+		return pResult;
+	}
+
+
+}
+
+#undef NODE_INPUT_COUNT
diff --git a/Engine/Plugins/Experimental/Mutable/Source/MutableTools/Private/NodeMeshSubtractPrivate.h b/Engine/Plugins/Experimental/Mutable/Source/MutableTools/Private/NodeMeshSubtractPrivate.h
new file mode 100644
--- /dev/null
+++ b/Engine/Plugins/Experimental/Mutable/Source/MutableTools/Private/NodeMeshSubtractPrivate.h
@@ -0,0 +1,56 @@
+// Copyright Epic Games, Inc. All Rights Reserved.
+
+#pragma once
+
+#include "NodeMeshPrivate.h"
+#include "NodeMeshSubtract.h"
+#include "AST.h"
+
+#include "MemoryPrivate.h"
+
+
+namespace mu
+{
+
+
+	class NodeMeshSubtract::Private : public NodeMesh::Private
+	{
+	public:
+
+		MUTABLE_DEFINE_CONST_VISITABLE()
+
+	public:
+
+		static NODE_TYPE s_type;
+
+		NodeMeshPtr m_pA;
+		NodeMeshPtr m_pB;
+
+		//!
+		void Serialise( OutputArchive& arch ) const
+		{
+            uint32_t ver = 0;
+			arch << ver;
+
+			arch << m_pA;
+			arch << m_pB;
+		}
+
+		//!
+		void Unserialise( InputArchive& arch )
+		{
+            uint32_t ver;
+			arch >> ver;
+			check(ver==0);
+
+			arch >> m_pA;
+			arch >> m_pB;
+		}
+
+		// NodeMesh::Private interface
+        NodeLayoutPtr GetLayout( int index ) const override;
+
+	};
+
+
+}
diff --git a/Engine/Plugins/Experimental/Mutable/Source/MutableTools/Private/NodeMeshSwitch.cpp b/Engine/Plugins/Experimental/Mutable/Source/MutableTools/Private/NodeMeshSwitch.cpp
new file mode 100644
--- /dev/null
+++ b/Engine/Plugins/Experimental/Mutable/Source/MutableTools/Private/NodeMeshSwitch.cpp
@@ -0,0 +1,137 @@
+// Copyright Epic Games, Inc. All Rights Reserved.
+
+
+#include "NodeMeshSwitch.h"
+#include "NodeMeshSwitchPrivate.h"
+
+#include "NodeScalar.h"
+
+#include "MeshPrivate.h"
+
+
+namespace mu
+{
+
+	//---------------------------------------------------------------------------------------------
+	// Static initialisation
+	//---------------------------------------------------------------------------------------------
+	NODE_TYPE NodeMeshSwitch::Private::s_type =
+			NODE_TYPE( "MeshSwitch", NodeMesh::GetStaticType() );
+
+
+	//---------------------------------------------------------------------------------------------
+	//!
+	//---------------------------------------------------------------------------------------------
+
+	MUTABLE_IMPLEMENT_NODE( NodeMeshSwitch, EType::Switch, Node, Node::EType::Mesh)
+
+
+	//---------------------------------------------------------------------------------------------
+	// Node Interface
+	//---------------------------------------------------------------------------------------------
+	int NodeMeshSwitch::GetInputCount() const
+	{
+		return 1 + (int)m_pD->m_options.size();
+	}
+
+
+	//---------------------------------------------------------------------------------------------
+	Node* NodeMeshSwitch::GetInputNode( int i ) const
+	{
+		check( i>=0 && i<GetInputCount() );
+
+		Node* pResult = 0;
+
+		switch (i)
+		{
+		case 0:
+			pResult = m_pD->m_pParameter.get();
+			break;
+
+		default:
+			pResult = m_pD->m_options[i-1].get();
+			break;
+		}
+
+		return pResult;
+	}
+
+
+	//---------------------------------------------------------------------------------------------
+	void NodeMeshSwitch::SetInputNode( int i, NodePtr pNode )
+	{
+		check( i>=0 && i<GetInputCount() );
+
+		switch (i)
+		{
+		case 0:
+			m_pD->m_pParameter = dynamic_cast<NodeScalar*>(pNode.get());
+			break;
+
+		default:
+			m_pD->m_options[i-1] = dynamic_cast<NodeMesh*>(pNode.get());
+			break;
+		}
+	}
+
+
+	//---------------------------------------------------------------------------------------------
+	// Own Interface
+	//---------------------------------------------------------------------------------------------
+	NodeScalarPtr NodeMeshSwitch::GetParameter() const
+	{
+		return m_pD->m_pParameter.get();
+	}
+
+
+	//---------------------------------------------------------------------------------------------
+	void NodeMeshSwitch::SetParameter( NodeScalarPtr pNode )
+	{
+		m_pD->m_pParameter = pNode;
+	}
+
+
+	//---------------------------------------------------------------------------------------------
+	void NodeMeshSwitch::SetOptionCount( int t )
+	{
+		m_pD->m_options.resize(t);
+	}
+
+
+	//---------------------------------------------------------------------------------------------
+	NodeMeshPtr NodeMeshSwitch::GetOption( int t ) const
+	{
+		check( t>=0 && t<(int)m_pD->m_options.size() );
+		return m_pD->m_options[t].get();
+	}
+
+
+	//---------------------------------------------------------------------------------------------
+	void NodeMeshSwitch::SetOption( int t, NodeMeshPtr pNode )
+	{
+		check( t>=0 && t<(int)m_pD->m_options.size() );
+		m_pD->m_options[t] = pNode;
+	}
+
+
+	//---------------------------------------------------------------------------------------------
+	NodeLayoutPtr NodeMeshSwitch::Private::GetLayout( int index ) const
+	{
+		NodeLayoutPtr pResult;
+
+		if (m_options.size()>0 && m_options[0] )
+		{
+			NodeMesh::Private* pPrivate =
+					dynamic_cast<NodeMesh::Private*>( m_options[0]->GetBasePrivate() );
+
+			pResult = pPrivate->GetLayout( index );
+		}
+
+		return pResult;
+	}
+
+
+
+}
+
+
diff --git a/Engine/Plugins/Experimental/Mutable/Source/MutableTools/Private/NodeMeshSwitchPrivate.h b/Engine/Plugins/Experimental/Mutable/Source/MutableTools/Private/NodeMeshSwitchPrivate.h
new file mode 100644
--- /dev/null
+++ b/Engine/Plugins/Experimental/Mutable/Source/MutableTools/Private/NodeMeshSwitchPrivate.h
@@ -0,0 +1,55 @@
+// Copyright Epic Games, Inc. All Rights Reserved.
+
+#pragma once
+
+#include "NodeMeshSwitch.h"
+#include "NodeScalar.h"
+#include "NodeMeshPrivate.h"
+#include "AST.h"
+
+
+namespace mu
+{
+
+
+	class NodeMeshSwitch::Private : public NodeMesh::Private
+	{
+	public:
+
+		MUTABLE_DEFINE_CONST_VISITABLE()
+
+	public:
+
+		static NODE_TYPE s_type;
+
+		NodeScalarPtr m_pParameter;
+		vector<NodeMeshPtr> m_options;
+
+		//!
+		void Serialise( OutputArchive& arch ) const
+		{
+            uint32_t ver = 0;
+			arch << ver;
+
+			arch << m_pParameter;
+			arch << m_options;
+		}
+
+		//!
+		void Unserialise( InputArchive& arch )
+		{
+            uint32_t ver;
+			arch >> ver;
+			check(ver==0);
+
+			arch >> m_pParameter;
+			arch >> m_options;
+		}
+
+
+		// NodeMesh::Private interface
+        NodeLayoutPtr GetLayout( int index ) const override;
+	};
+
+
+}
diff --git a/Engine/Plugins/Experimental/Mutable/Source/MutableTools/Private/NodeMeshTable.cpp b/Engine/Plugins/Experimental/Mutable/Source/MutableTools/Private/NodeMeshTable.cpp
new file mode 100644
--- /dev/null
+++ b/Engine/Plugins/Experimental/Mutable/Source/MutableTools/Private/NodeMeshTable.cpp
@@ -0,0 +1,142 @@
+// Copyright Epic Games, Inc. All Rights Reserved.
+
+
+#include "NodeMeshTable.h"
+#include "NodeMeshTablePrivate.h"
+
+
+namespace mu
+{
+
+
+	//---------------------------------------------------------------------------------------------
+	// Static initialisation
+	//---------------------------------------------------------------------------------------------
+	NODE_TYPE NodeMeshTable::Private::s_type =
+			NODE_TYPE( "TableMesh", NodeMesh::GetStaticType() );
+
+
+	//---------------------------------------------------------------------------------------------
+	//!
+	//---------------------------------------------------------------------------------------------
+
+	MUTABLE_IMPLEMENT_NODE( NodeMeshTable, EType::Table, Node, Node::EType::Mesh);
+
+
+	//---------------------------------------------------------------------------------------------
+	// Node Interface
+	//---------------------------------------------------------------------------------------------
+	int NodeMeshTable::GetInputCount() const
+	{
+		return 0;
+	}
+
+
+	//---------------------------------------------------------------------------------------------
+	Node* NodeMeshTable::GetInputNode( int i ) const
+	{
+		check( i >=0 && i < 0 );
+        (void)i;
+        return 0;
+	}
+
+
+	//---------------------------------------------------------------------------------------------
+    void NodeMeshTable::SetInputNode( int i, NodePtr )
+	{
+		check( i >=0 && i < 0 );
+        (void)i;
+        //m_pD->m_pObject = dynamic_cast<NodeObject*>( pNode.get() );
+	}
+
+
+	//---------------------------------------------------------------------------------------------
+	// Own Interface
+	//---------------------------------------------------------------------------------------------
+	const char* NodeMeshTable::GetColumn() const
+	{
+		return m_pD->m_columnName.c_str();
+	}
+
+
+	//---------------------------------------------------------------------------------------------
+	void NodeMeshTable::SetColumn( const char* strName )
+	{
+		if (strName)
+		{
+			m_pD->m_columnName = strName;
+		}
+		else
+		{
+			m_pD->m_columnName = "";
+		}
+	}
+
+
+	//---------------------------------------------------------------------------------------------
+	void NodeMeshTable::SetParameterName( const char* strName )
+	{
+		m_pD->m_parameterName = strName;
+	}
+
+
+	//---------------------------------------------------------------------------------------------
+	void NodeMeshTable::SetTable( TablePtr pTable )
+	{
+		m_pD->m_pTable = pTable;
+	}
+
+
+	//---------------------------------------------------------------------------------------------
+	TablePtr NodeMeshTable::GetTable() const
+	{
+		return m_pD->m_pTable;
+	}
+
+
+	//---------------------------------------------------------------------------------------------
+	int NodeMeshTable::GetLayoutCount() const
+	{
+		return (int)m_pD->m_layouts.size();
+	}
+
+
+	//---------------------------------------------------------------------------------------------
+	void NodeMeshTable::SetLayoutCount( int i )
+	{
+		m_pD->m_layouts.resize( i );
+	}
+
+
+	//---------------------------------------------------------------------------------------------
+	NodeLayoutPtr NodeMeshTable::GetLayout( int i ) const
+	{
+		check( i>=0 && i<GetLayoutCount() );
+		return m_pD->GetLayout( i );
+	}
+
+
+	//---------------------------------------------------------------------------------------------
+	void NodeMeshTable::SetLayout( int i, NodeLayoutPtr pLayout )
+	{
+		check( i>=0 && i<GetLayoutCount() );
+		m_pD->m_layouts[i] = pLayout;
+	}
+
+
+	//---------------------------------------------------------------------------------------------
+	NodeLayoutPtr NodeMeshTable::Private::GetLayout( int i ) const
+	{
+		NodeLayoutPtr pResult;
+
+		if ( i>=0 && i<(int)m_layouts.size() )
+		{
+			pResult = m_layouts[i];
+		}
+
+		return pResult;
+	}
+
+}
+
+
diff --git a/Engine/Plugins/Experimental/Mutable/Source/MutableTools/Private/NodeMeshTablePrivate.h b/Engine/Plugins/Experimental/Mutable/Source/MutableTools/Private/NodeMeshTablePrivate.h
new file mode 100644
--- /dev/null
+++ b/Engine/Plugins/Experimental/Mutable/Source/MutableTools/Private/NodeMeshTablePrivate.h
@@ -0,0 +1,64 @@
+// Copyright Epic Games, Inc. All Rights Reserved.
+
+#pragma once
+
+#include "NodeMeshPrivate.h"
+#include "NodeMeshTable.h"
+#include "NodeLayout.h"
+#include "Table.h"
+#include "AST.h"
+
+#include "MemoryPrivate.h"
+
+
+namespace mu
+{
+
+
+	class NodeMeshTable::Private : public NodeMesh::Private
+	{
+	public:
+
+		MUTABLE_DEFINE_CONST_VISITABLE()
+
+	public:
+
+		static NODE_TYPE s_type;
+
+		string m_parameterName;
+		TablePtr m_pTable;
+		string m_columnName;
+
+		vector<NodeLayoutPtr> m_layouts;
+
+		//!
+		void Serialise( OutputArchive& arch ) const
+		{
+            uint32_t ver = 1;
+			arch << ver;
+
+			arch << m_parameterName;
+			arch << m_pTable;
+			arch << m_columnName;
+			arch << m_layouts;
+		}
+
+		//!
+		void Unserialise( InputArchive& arch )
+		{
+            uint32_t ver;
+			arch >> ver;
+            check(ver==1);
+
+			arch >> m_parameterName;
+			arch >> m_pTable;
+			arch >> m_columnName;
+			arch >> m_layouts;
+		}
+
+		// NodeMesh::Private interface
+        NodeLayoutPtr GetLayout( int index ) const override;
+
+	};
+
+}
diff --git a/Engine/Plugins/Experimental/Mutable/Source/MutableTools/Private/NodeMeshTangents.cpp b/Engine/Plugins/Experimental/Mutable/Source/MutableTools/Private/NodeMeshTangents.cpp
new file mode 100644
--- /dev/null
+++ b/Engine/Plugins/Experimental/Mutable/Source/MutableTools/Private/NodeMeshTangents.cpp
@@ -0,0 +1,94 @@
+// Copyright Epic Games, Inc. All Rights Reserved.
+
+
+#include "NodeMeshTangents.h"
+#include "NodeMeshTangentsPrivate.h"
+#include "MeshPrivate.h"
+
+
+#define NODE_INPUT_COUNT 	1
+
+
+namespace mu
+{
+
+	//---------------------------------------------------------------------------------------------
+	// Static initialisation
+	//---------------------------------------------------------------------------------------------
+	NODE_TYPE NodeMeshTangents::Private::s_type =
+			NODE_TYPE( "MeshTangents", NodeMesh::GetStaticType() );
+
+
+	//---------------------------------------------------------------------------------------------
+	//!
+	//---------------------------------------------------------------------------------------------
+
+	MUTABLE_IMPLEMENT_NODE( NodeMeshTangents, EType::Tangents, Node, Node::EType::Mesh)
+
+
+	//---------------------------------------------------------------------------------------------
+	// Node Interface
+	//---------------------------------------------------------------------------------------------
+	int NodeMeshTangents::GetInputCount() const
+	{
+		return NODE_INPUT_COUNT;
+	}
+
+
+	//---------------------------------------------------------------------------------------------
+	Node* NodeMeshTangents::GetInputNode( int i ) const
+	{
+		check( i>=0 && i<NODE_INPUT_COUNT );
+        (void)i;
+		return m_pD->m_pSource.get();
+	}
+
+
+	//---------------------------------------------------------------------------------------------
+	void NodeMeshTangents::SetInputNode( int i, NodePtr pNode )
+	{
+		check( i>=0 && i<NODE_INPUT_COUNT );
+        (void)i;
+        if (i==0)
+		{
+			m_pD->m_pSource = dynamic_cast<NodeMesh*>( pNode.get() );
+		}
+	}
+
+
+	//---------------------------------------------------------------------------------------------
+	// Own Interface
+	//---------------------------------------------------------------------------------------------
+	NodeMeshPtr NodeMeshTangents::GetSource() const
+	{
+		return m_pD->m_pSource;
+	}
+
+
+	//---------------------------------------------------------------------------------------------
+	void NodeMeshTangents::SetSource( NodeMesh* p )
+	{
+		m_pD->m_pSource = p;
+	}
+
+
+	//---------------------------------------------------------------------------------------------
+	NodeLayoutPtr NodeMeshTangents::Private::GetLayout( int index ) const
+	{
+		NodeLayoutPtr pResult;
+
+		if ( m_pSource )
+		{
+			NodeMesh::Private* pPrivate =
+					dynamic_cast<NodeMesh::Private*>( m_pSource->GetBasePrivate() );
+
+			pResult = pPrivate->GetLayout( index );
+		}
+
+		return pResult;
+	}
+
+
+}
+
+#undef NODE_INPUT_COUNT
diff --git a/Engine/Plugins/Experimental/Mutable/Source/MutableTools/Private/NodeMeshTangentsPrivate.h b/Engine/Plugins/Experimental/Mutable/Source/MutableTools/Private/NodeMeshTangentsPrivate.h
new file mode 100644
--- /dev/null
+++ b/Engine/Plugins/Experimental/Mutable/Source/MutableTools/Private/NodeMeshTangentsPrivate.h
@@ -0,0 +1,53 @@
+// Copyright Epic Games, Inc. All Rights Reserved.
+
+#pragma once
+
+#include "NodeMeshPrivate.h"
+#include "NodeMeshTangents.h"
+#include "AST.h"
+
+#include "MemoryPrivate.h"
+
+
+namespace mu
+{
+
+
+	class NodeMeshTangents::Private : public NodeMesh::Private
+	{
+	public:
+
+		MUTABLE_DEFINE_CONST_VISITABLE()
+
+	public:
+
+		static NODE_TYPE s_type;
+
+		NodeMeshPtr m_pSource;
+
+		//!
+		void Serialise( OutputArchive& arch ) const
+		{
+            uint32_t ver = 0;
+			arch << ver;
+
+			arch << m_pSource;
+		}
+
+		//!
+		void Unserialise( InputArchive& arch )
+		{
+            uint32_t ver;
+			arch >> ver;
+			check(ver==0);
+
+			arch >> m_pSource;
+		}
+
+		// NodeMesh::Private interface
+        NodeLayoutPtr GetLayout( int index ) const override;
+
+	};
+
+
+}
diff --git a/Engine/Plugins/Experimental/Mutable/Source/MutableTools/Private/NodeMeshTransform.cpp b/Engine/Plugins/Experimental/Mutable/Source/MutableTools/Private/NodeMeshTransform.cpp
new file mode 100644
--- /dev/null
+++ b/Engine/Plugins/Experimental/Mutable/Source/MutableTools/Private/NodeMeshTransform.cpp
@@ -0,0 +1,107 @@
+// Copyright Epic Games, Inc. All Rights Reserved.
+
+
+#include "NodeMeshTransform.h"
+#include "NodeMeshTransformPrivate.h"
+#include "MeshPrivate.h"
+
+
+#define NODE_INPUT_COUNT 	1
+
+
+namespace mu
+{
+
+	//---------------------------------------------------------------------------------------------
+	// Static initialisation
+	//---------------------------------------------------------------------------------------------
+    NODE_TYPE NodeMeshTransform::Private::s_type =
+            NODE_TYPE( "MeshTransform", NodeMesh::GetStaticType() );
+
+
+	//---------------------------------------------------------------------------------------------
+	//!
+	//---------------------------------------------------------------------------------------------
+
+    MUTABLE_IMPLEMENT_NODE( NodeMeshTransform, EType::Transform, Node, Node::EType::Mesh)
+
+
+	//---------------------------------------------------------------------------------------------
+	// Node Interface
+	//---------------------------------------------------------------------------------------------
+    int NodeMeshTransform::GetInputCount() const
+	{
+		return NODE_INPUT_COUNT;
+	}
+
+
+	//---------------------------------------------------------------------------------------------
+    Node* NodeMeshTransform::GetInputNode( int i ) const
+	{
+		check( i>=0 && i<NODE_INPUT_COUNT );
+        (void)i;
+        return m_pD->m_pSource.get();
+	}
+
+
+	//---------------------------------------------------------------------------------------------
+    void NodeMeshTransform::SetInputNode( int i, NodePtr pNode )
+	{
+		check( i>=0 && i<NODE_INPUT_COUNT );
+		if (i==0)
+		{
+			m_pD->m_pSource = dynamic_cast<NodeMesh*>( pNode.get() );
+		}
+	}
+
+
+	//---------------------------------------------------------------------------------------------
+	// Own Interface
+	//---------------------------------------------------------------------------------------------
+    NodeMeshPtr NodeMeshTransform::GetSource() const
+	{
+		return m_pD->m_pSource;
+	}
+
+
+    //---------------------------------------------------------------------------------------------
+    void NodeMeshTransform::SetSource( NodeMesh* p )
+    {
+        m_pD->m_pSource = p;
+    }
+
+
+    //---------------------------------------------------------------------------------------------
+    void NodeMeshTransform::SetTransform( const float* mat )
+    {
+        memcpy( &m_pD->m_transform[0][0], mat, 16*sizeof(float) );
+    }
+
+
+    //---------------------------------------------------------------------------------------------
+    void NodeMeshTransform::GetTransform( float* mat ) const
+    {
+        memcpy( mat, &m_pD->m_transform[0][0], 16*sizeof(float) );
+    }
+
+
+	//---------------------------------------------------------------------------------------------
+    NodeLayoutPtr NodeMeshTransform::Private::GetLayout( int index ) const
+	{
+		NodeLayoutPtr pResult;
+
+		if ( m_pSource )
+		{
+			NodeMesh::Private* pPrivate =
+					dynamic_cast<NodeMesh::Private*>( m_pSource->GetBasePrivate() );
+
+			pResult = pPrivate->GetLayout( index );
+		}
+
+		return pResult;
+	}
+
+
+}
+
+#undef NODE_INPUT_COUNT
diff --git a/Engine/Plugins/Experimental/Mutable/Source/MutableTools/Private/NodeMeshTransformPrivate.h b/Engine/Plugins/Experimental/Mutable/Source/MutableTools/Private/NodeMeshTransformPrivate.h
new file mode 100644
--- /dev/null
+++ b/Engine/Plugins/Experimental/Mutable/Source/MutableTools/Private/NodeMeshTransformPrivate.h
@@ -0,0 +1,57 @@
+// Copyright Epic Games, Inc. All Rights Reserved.
+
+#pragma once
+
+#include "NodeMeshPrivate.h"
+#include "NodeMeshTransform.h"
+#include "AST.h"
+
+#include "MemoryPrivate.h"
+#include "MutableMath.h"
+
+
+namespace mu
+{
+
+
+    class NodeMeshTransform::Private : public NodeMesh::Private
+	{
+	public:
+
+		MUTABLE_DEFINE_CONST_VISITABLE()
+
+	public:
+
+		static NODE_TYPE s_type;
+
+		NodeMeshPtr m_pSource;
+        mat4f m_transform;
+
+		//!
+		void Serialise( OutputArchive& arch ) const
+		{
+            uint32_t ver = 0;
+			arch << ver;
+
+			arch << m_pSource;
+            arch << m_transform;
+		}
+
+		//!
+		void Unserialise( InputArchive& arch )
+		{
+            uint32_t ver;
+			arch >> ver;
+			check(ver==0);
+
+			arch >> m_pSource;
+            arch >> m_transform;
+		}
+
+		// NodeMesh::Private interface
+        NodeLayoutPtr GetLayout( int index ) const override;
+
+	};
+
+
+}
diff --git a/Engine/Plugins/Experimental/Mutable/Source/MutableTools/Private/NodeMeshVariation.cpp b/Engine/Plugins/Experimental/Mutable/Source/MutableTools/Private/NodeMeshVariation.cpp
new file mode 100644
--- /dev/null
+++ b/Engine/Plugins/Experimental/Mutable/Source/MutableTools/Private/NodeMeshVariation.cpp
@@ -0,0 +1,115 @@
+// Copyright Epic Games, Inc. All Rights Reserved.
+
+
+#include "NodeMeshVariation.h"
+#include "NodeMeshVariationPrivate.h"
+
+
+namespace mu
+{
+
+    //---------------------------------------------------------------------------------------------
+    // Static initialisation
+    //---------------------------------------------------------------------------------------------
+    NODE_TYPE NodeMeshVariation::Private::s_type =
+        NODE_TYPE( "MeshVariation", NodeMesh::GetStaticType() );
+
+
+    //---------------------------------------------------------------------------------------------
+    //---------------------------------------------------------------------------------------------
+    //---------------------------------------------------------------------------------------------
+
+    MUTABLE_IMPLEMENT_NODE( NodeMeshVariation, EType::Variation, Node, Node::EType::Mesh)
+
+
+    //---------------------------------------------------------------------------------------------
+    int NodeMeshVariation::GetInputCount() const { return 1 + int( m_pD->m_variations.size() ); }
+
+
+    //---------------------------------------------------------------------------------------------
+    Node* NodeMeshVariation::GetInputNode( int i ) const
+    {
+        check( i >= 0 && i < GetInputCount() );
+
+        if ( i == 0 )
+        {
+            return m_pD->m_defaultMesh.get();
+        }
+        i -= 1;
+
+        if ( i < int( m_pD->m_variations.size() ) )
+        {
+            return m_pD->m_variations[i].m_mesh.get();
+        }
+        i -= int( m_pD->m_variations.size() );
+
+        return nullptr;
+    }
+
+
+    //---------------------------------------------------------------------------------------------
+    void NodeMeshVariation::SetInputNode( int i, NodePtr pNode )
+    {
+        check( i >= 0 && i < GetInputCount() );
+
+        if ( i == 0 )
+        {
+            m_pD->m_defaultMesh = dynamic_cast<NodeMesh*>( pNode.get() );
+            return;
+        }
+
+        i -= 1;
+        if ( i < int( m_pD->m_variations.size() ) )
+        {
+
+            m_pD->m_variations[i].m_mesh = dynamic_cast<NodeMesh*>( pNode.get() );
+            return;
+        }
+        i -= (int)m_pD->m_variations.size();
+    }
+
+
+    //---------------------------------------------------------------------------------------------
+    // Own Interface
+    //---------------------------------------------------------------------------------------------
+    void NodeMeshVariation::SetDefaultMesh( NodeMesh* p ) { m_pD->m_defaultMesh = p; }
+
+
+    //---------------------------------------------------------------------------------------------
+    int NodeMeshVariation::GetVariationCount() const { return int( m_pD->m_variations.size() ); }
+
+
+    //---------------------------------------------------------------------------------------------
+    void NodeMeshVariation::SetVariationCount( int num )
+    {
+        check( num >= 0 );
+        m_pD->m_variations.resize( num );
+    }
+
+    //---------------------------------------------------------------------------------------------
+    void NodeMeshVariation::SetVariationTag( int index, const char* strTag )
+    {
+        check( index >= 0 && index < (int)m_pD->m_variations.size() );
+        check( strTag );
+
+        if ( strTag )
+        {
+            m_pD->m_variations[index].m_tag = strTag;
+        }
+        else
+        {
+            m_pD->m_variations[index].m_tag = "";
+        }
+    }
+
+
+    //---------------------------------------------------------------------------------------------
+    void NodeMeshVariation::SetVariationMesh( int index, NodeMesh* pNode )
+    {
+        check( index >= 0 && index < (int)m_pD->m_variations.size() );
+
+        m_pD->m_variations[index].m_mesh = pNode;
+    }
+
+
+} // namespace mu
diff --git a/Engine/Plugins/Experimental/Mutable/Source/MutableTools/Private/NodeMeshVariationPrivate.h b/Engine/Plugins/Experimental/Mutable/Source/MutableTools/Private/NodeMeshVariationPrivate.h
new file mode 100644
--- /dev/null
+++ b/Engine/Plugins/Experimental/Mutable/Source/MutableTools/Private/NodeMeshVariationPrivate.h
@@ -0,0 +1,79 @@
+// Copyright Epic Games, Inc. All Rights Reserved.
+
+#pragma once
+
+#include "MutableTools/Private/NodePrivate.h"
+#include "MutableTools/Public/NodeMeshVariation.h"
+
+#include "MutableRuntime/Private/MemoryPrivate.h"
+
+
+namespace mu
+{
+
+    //---------------------------------------------------------------------------------------------
+    //!
+    //---------------------------------------------------------------------------------------------
+    class NodeMeshVariation::Private : public Node::Private
+    {
+    public:
+        MUTABLE_DEFINE_CONST_VISITABLE()
+
+    public:
+        Private() {}
+
+        static NODE_TYPE s_type;
+
+        NodeMeshPtr m_defaultMesh;
+
+        struct VARIATION
+        {
+            NodeMeshPtr m_mesh;
+            string m_tag;
+
+            //!
+            void Serialise( OutputArchive& arch ) const
+            {
+                uint32_t ver = 1;
+                arch << ver;
+
+                arch << m_tag;
+                arch << m_mesh;
+            }
+
+            void Unserialise( InputArchive& arch )
+            {
+                uint32_t ver = 0;
+                arch >> ver;
+                check( ver == 1 );
+
+                arch >> m_tag;
+                arch >> m_mesh;
+            }
+        };
+
+        vector<VARIATION> m_variations;
+
+        //!
+        void Serialise( OutputArchive& arch ) const
+        {
+            uint32_t ver = 1;
+            arch << ver;
+
+            arch << m_defaultMesh;
+            arch << m_variations;
+        }
+
+        //!
+        void Unserialise( InputArchive& arch )
+        {
+            uint32_t ver;
+            arch >> ver;
+			check(ver == 1 );
+
+            arch >> m_defaultMesh;
+            arch >> m_variations;
+        }
+    };
+
+} // namespace mu
diff --git a/Engine/Plugins/Experimental/Mutable/Source/MutableTools/Private/NodeModifier.cpp b/Engine/Plugins/Experimental/Mutable/Source/MutableTools/Private/NodeModifier.cpp
new file mode 100644
--- /dev/null
+++ b/Engine/Plugins/Experimental/Mutable/Source/MutableTools/Private/NodeModifier.cpp
@@ -0,0 +1,127 @@
+// Copyright Epic Games, Inc. All Rights Reserved.
+
+
+#include "NodeModifier.h"
+#include "NodeModifierPrivate.h"
+
+#include "NodeModifierMeshClipMorphPlane.h"
+#include "NodeModifierMeshClipWithMesh.h"
+#include "NodeModifierMeshClipDeform.h"
+
+
+namespace mu
+{
+
+
+	//---------------------------------------------------------------------------------------------
+	// Static initialisation
+	//---------------------------------------------------------------------------------------------
+	static NODE_TYPE s_nodeModifierType =
+			NODE_TYPE( "NodeModifier", Node::GetStaticType() );
+
+
+	//---------------------------------------------------------------------------------------------
+	//---------------------------------------------------------------------------------------------
+	//---------------------------------------------------------------------------------------------
+	const NODE_TYPE* NodeModifier::GetType() const
+	{
+		return GetStaticType();
+	}
+
+
+	//---------------------------------------------------------------------------------------------
+	const NODE_TYPE* NodeModifier::GetStaticType()
+	{
+          return &s_nodeModifierType;
+        }
+
+
+	//---------------------------------------------------------------------------------------------
+	void NodeModifier::Serialise( const NodeModifier* p, OutputArchive& arch )
+	{
+        uint32_t ver = 0;
+		arch << ver;
+
+#define SERIALISE_CHILDREN( C, ID ) \
+        ( const C* pTyped##ID = dynamic_cast<const C*>(p) )			\
+        {                                                           \
+            arch << (uint32_t)ID;									\
+            C::Serialise( pTyped##ID, arch );						\
+		}
+
+		if SERIALISE_CHILDREN(NodeModifierMeshClipMorphPlane, 			0 )
+		else if SERIALISE_CHILDREN(NodeModifierMeshClipWithMesh, 		1 )
+		else if SERIALISE_CHILDREN(NodeModifierMeshClipDeform, 		    2 )
+		else check(false);
+
+#undef SERIALISE_CHILDREN
+    }
+
+        
+	//---------------------------------------------------------------------------------------------
+	NodeModifierPtr NodeModifier::StaticUnserialise( InputArchive& arch )
+	{
+        uint32_t ver;
+		arch >> ver;
+		check( ver == 0 );
+
+        uint32_t id;
+		arch >> id;
+
+		switch (id)
+		{
+		case 0 :  return NodeModifierMeshClipMorphPlane::StaticUnserialise( arch ); break;
+		case 1 :  return NodeModifierMeshClipWithMesh::StaticUnserialise( arch ); break;
+		case 2 :  return NodeModifierMeshClipDeform::StaticUnserialise( arch ); break;
+		default : check(false);
+		}
+
+		return 0;
+	}
+
+	//---------------------------------------------------------------------------------------------
+	void NodeModifier::AddTag(const char* tagName)
+	{
+		NodeModifier::Private* pD = dynamic_cast<NodeModifier::Private*>(GetBasePrivate());
+		check(pD);
+
+		pD->m_tags.push_back(tagName);
+	}
+
+
+    //---------------------------------------------------------------------------------------------
+    int NodeModifier::GetTagCount() const
+    {
+        NodeModifier::Private* pD = dynamic_cast<NodeModifier::Private*>(GetBasePrivate());
+        check(pD);
+
+        return int( pD->m_tags.size());
+    }
+
+
+    //---------------------------------------------------------------------------------------------
+    const char* NodeModifier::GetTag( int i ) const
+    {
+        NodeModifier::Private* pD = dynamic_cast<NodeModifier::Private*>(GetBasePrivate());
+        check(pD);
+
+        if (i>=0 && i<GetTagCount())
+        {
+            return pD->m_tags[i].c_str();
+        }
+        return nullptr;
+    }
+
+
+	//---------------------------------------------------------------------------------------------
+	void NodeModifier::SetStage(bool bBeforeNormalOperation)
+	{
+		NodeModifier::Private* pD = dynamic_cast<NodeModifier::Private*>(GetBasePrivate());
+		check(pD);
+
+		pD->m_applyBeforeNormalOperations = bBeforeNormalOperation; 
+	}
+
+}
+
+
diff --git a/Engine/Plugins/Experimental/Mutable/Source/MutableTools/Private/NodeModifierMeshClipDeform.cpp b/Engine/Plugins/Experimental/Mutable/Source/MutableTools/Private/NodeModifierMeshClipDeform.cpp
new file mode 100644
--- /dev/null
+++ b/Engine/Plugins/Experimental/Mutable/Source/MutableTools/Private/NodeModifierMeshClipDeform.cpp
@@ -0,0 +1,66 @@
+// Copyright Epic Games, Inc. All Rights Reserved.
+
+
+#include "NodeModifierMeshClipDeform.h"
+#include "NodeModifierMeshClipDeformPrivate.h"
+#include "MeshPrivate.h"
+
+
+namespace mu
+{
+
+	//---------------------------------------------------------------------------------------------
+	// Static initialisation
+	//---------------------------------------------------------------------------------------------
+	NODE_TYPE NodeModifierMeshClipDeform::Private::s_type =
+			NODE_TYPE( "NodeModifierMeshClipDeform", NodeModifier::GetStaticType() );
+
+
+	//---------------------------------------------------------------------------------------------
+	//!
+	//---------------------------------------------------------------------------------------------
+
+	MUTABLE_IMPLEMENT_NODE( NodeModifierMeshClipDeform, EType::MeshClipDeform, Node, Node::EType::Modifier)
+
+
+	//---------------------------------------------------------------------------------------------
+	// Node Interface
+	//---------------------------------------------------------------------------------------------
+	int NodeModifierMeshClipDeform::GetInputCount() const
+	{
+		return 0;
+	}
+
+
+	//---------------------------------------------------------------------------------------------
+	Node* NodeModifierMeshClipDeform::GetInputNode( int i ) const
+	{
+		check( i>=0 && i< GetInputCount());
+		(void)i;
+		return nullptr;
+	}
+
+
+	//---------------------------------------------------------------------------------------------
+	void NodeModifierMeshClipDeform::SetInputNode( int32 i, NodePtr )
+	{
+		check( i>=0 && i< GetInputCount());
+		(void)i;
+	}
+
+
+	//---------------------------------------------------------------------------------------------
+	// Own Interface
+	//---------------------------------------------------------------------------------------------
+	void NodeModifierMeshClipDeform::SetClipMesh(NodeMesh* ClipMesh)
+	{
+		m_pD->ClipMesh = ClipMesh;
+	}
+
+	void NodeModifierMeshClipDeform::SetBindingMethod(EShapeBindingMethod BindingMethod)
+	{
+		check(BindingMethod != EShapeBindingMethod::ReshapeClosestProject);
+		m_pD->BindingMethod = BindingMethod;
+	}
+
+}
diff --git a/Engine/Plugins/Experimental/Mutable/Source/MutableTools/Private/NodeModifierMeshClipDeformPrivate.h b/Engine/Plugins/Experimental/Mutable/Source/MutableTools/Private/NodeModifierMeshClipDeformPrivate.h
new file mode 100644
--- /dev/null
+++ b/Engine/Plugins/Experimental/Mutable/Source/MutableTools/Private/NodeModifierMeshClipDeformPrivate.h
@@ -0,0 +1,68 @@
+// Copyright Epic Games, Inc. All Rights Reserved.
+
+#pragma once
+
+#include "NodeModifierPrivate.h"
+#include "NodeModifierMeshClipDeform.h"
+#include "NodeMesh.h"
+#include "AST.h"
+
+#include "MemoryPrivate.h"
+
+namespace mu
+{
+
+
+    class NodeModifierMeshClipDeform::Private : public NodeModifier::Private
+    {
+    public:
+
+        MUTABLE_DEFINE_CONST_VISITABLE()
+
+    public:
+
+		Private()
+		{
+		}
+
+		static NODE_TYPE s_type;
+
+		//! 
+		NodeMeshPtr ClipMesh;
+		EShapeBindingMethod BindingMethod = EShapeBindingMethod::ClipDeformClosestProject;
+
+		//!
+		void Serialise( OutputArchive& arch ) const
+		{
+			NodeModifier::Private::Serialise(arch);
+
+            uint32_t ver = 1;
+			arch << ver;
+
+			arch << ClipMesh;
+			arch << BindingMethod;
+		}
+
+		//!
+		void Unserialise( InputArchive& arch )
+		{
+			NodeModifier::Private::Unserialise( arch );
+			
+            uint32_t ver;
+			arch >> ver;
+            check(ver<=1);
+
+			arch >> ClipMesh;
+			if (ver >= 1)
+			{
+				arch >> BindingMethod;
+			}
+			else
+			{
+				BindingMethod = EShapeBindingMethod::ClipDeformClosestProject;
+			}
+	
+		}
+    };
+
+}
diff --git a/Engine/Plugins/Experimental/Mutable/Source/MutableTools/Private/NodeModifierMeshClipMorphPlane.cpp b/Engine/Plugins/Experimental/Mutable/Source/MutableTools/Private/NodeModifierMeshClipMorphPlane.cpp
new file mode 100644
--- /dev/null
+++ b/Engine/Plugins/Experimental/Mutable/Source/MutableTools/Private/NodeModifierMeshClipMorphPlane.cpp
@@ -0,0 +1,100 @@
+// Copyright Epic Games, Inc. All Rights Reserved.
+
+
+#include "NodeModifierMeshClipMorphPlane.h"
+#include "NodeModifierMeshClipMorphPlanePrivate.h"
+#include "MeshPrivate.h"
+
+
+#define NODE_INPUT_COUNT 	0
+
+
+namespace mu
+{
+
+	//---------------------------------------------------------------------------------------------
+	// Static initialisation
+	//---------------------------------------------------------------------------------------------
+    NODE_TYPE NodeModifierMeshClipMorphPlane::Private::s_type =
+            NODE_TYPE( "ModifierMeshClipMorphPlane", NodeModifier::GetStaticType() );
+
+
+	//---------------------------------------------------------------------------------------------
+	//!
+	//---------------------------------------------------------------------------------------------
+
+    MUTABLE_IMPLEMENT_NODE( NodeModifierMeshClipMorphPlane, EType::MeshClipMorphPlane, Node, Node::EType::Modifier)
+
+
+	//---------------------------------------------------------------------------------------------
+	// Node Interface
+	//---------------------------------------------------------------------------------------------
+    int NodeModifierMeshClipMorphPlane::GetInputCount() const
+	{
+		return NODE_INPUT_COUNT;
+	}
+
+
+	//---------------------------------------------------------------------------------------------
+    Node* NodeModifierMeshClipMorphPlane::GetInputNode( int i ) const
+	{
+		check( i>=0 && i<NODE_INPUT_COUNT );
+        (void)i;
+		return nullptr;
+	}
+
+
+	//---------------------------------------------------------------------------------------------
+    void NodeModifierMeshClipMorphPlane::SetInputNode( int i, NodePtr )
+	{
+		check( i>=0 && i<NODE_INPUT_COUNT );
+        (void)i;
+    }
+
+
+	//---------------------------------------------------------------------------------------------
+	// Own Interface
+	//---------------------------------------------------------------------------------------------
+    void NodeModifierMeshClipMorphPlane::SetPlane( float centerX, float centerY, float centerZ,
+                                                   float normalX, float normalY, float normalZ)
+	{
+		m_pD->m_origin = vec3f(centerX, centerY, centerZ);
+		m_pD->m_normal = vec3f(normalX, normalY, normalZ);
+	}
+
+
+	//---------------------------------------------------------------------------------------------
+	void NodeModifierMeshClipMorphPlane::SetParams(float dist, float factor)
+	{
+		m_pD->m_dist = dist;
+		m_pD->m_factor = factor;
+	}
+
+	//---------------------------------------------------------------------------------------------
+	void NodeModifierMeshClipMorphPlane::SetMorphEllipse(float radius1, float radius2, float rotation)
+	{
+		m_pD->m_radius1 = radius1;
+		m_pD->m_radius2 = radius2;
+		m_pD->m_rotation = rotation;
+	}
+
+	//---------------------------------------------------------------------------------------------
+	void NodeModifierMeshClipMorphPlane::SetVertexSelectionBox(float centerX, float centerY, float centerZ, float radiusX, float radiusY, float radiusZ)
+	{
+		m_pD->m_vertexSelectionType = Private::VS_SHAPE;
+		m_pD->m_selectionBoxOrigin = vec3f(centerX, centerY, centerZ);
+		m_pD->m_selectionBoxRadius = vec3f(radiusX, radiusY, radiusZ);
+	}
+
+	//---------------------------------------------------------------------------------------------
+	void NodeModifierMeshClipMorphPlane::SetVertexSelectionBone(const char* strBoneName, float maxEffectRadius)
+	{
+		m_pD->m_vertexSelectionType = Private::VS_BONE_HIERARCHY;
+		m_pD->m_vertexSelectionBone = strBoneName;
+		m_pD->m_maxEffectRadius = maxEffectRadius;
+	}
+
+}
+
+#undef NODE_INPUT_COUNT
+
diff --git a/Engine/Plugins/Experimental/Mutable/Source/MutableTools/Private/NodeModifierMeshClipWithMesh.cpp b/Engine/Plugins/Experimental/Mutable/Source/MutableTools/Private/NodeModifierMeshClipWithMesh.cpp
new file mode 100644
--- /dev/null
+++ b/Engine/Plugins/Experimental/Mutable/Source/MutableTools/Private/NodeModifierMeshClipWithMesh.cpp
@@ -0,0 +1,60 @@
+// Copyright Epic Games, Inc. All Rights Reserved.
+
+
+#include "NodeModifierMeshClipWithMesh.h"
+#include "NodeModifierMeshClipWithMeshPrivate.h"
+#include "MeshPrivate.h"
+
+
+namespace mu
+{
+
+	//---------------------------------------------------------------------------------------------
+	// Static initialisation
+	//---------------------------------------------------------------------------------------------
+    NODE_TYPE NodeModifierMeshClipWithMesh::Private::s_type =
+            NODE_TYPE( "NodeModifierMeshClipWithMesh", NodeModifier::GetStaticType() );
+
+
+	//---------------------------------------------------------------------------------------------
+	//!
+	//---------------------------------------------------------------------------------------------
+
+    MUTABLE_IMPLEMENT_NODE( NodeModifierMeshClipWithMesh, EType::MeshClipWithMesh, Node, Node::EType::Modifier)
+
+
+	//---------------------------------------------------------------------------------------------
+	// Node Interface
+	//---------------------------------------------------------------------------------------------
+    int NodeModifierMeshClipWithMesh::GetInputCount() const
+	{
+		return 0;
+	}
+
+
+	//---------------------------------------------------------------------------------------------
+    Node* NodeModifierMeshClipWithMesh::GetInputNode( int i ) const
+	{
+		check( i>=0 && i< GetInputCount());
+        (void)i;
+        return nullptr;
+	}
+
+
+	//---------------------------------------------------------------------------------------------
+    void NodeModifierMeshClipWithMesh::SetInputNode( int i, NodePtr )
+	{
+		check( i>=0 && i< GetInputCount());
+        (void)i;
+    }
+
+
+	//---------------------------------------------------------------------------------------------
+	// Own Interface
+	//---------------------------------------------------------------------------------------------
+	void NodeModifierMeshClipWithMesh::SetClipMesh(NodeMesh* clipMesh)
+	{
+		m_pD->m_clipMesh = clipMesh;
+	}
+
+}
diff --git a/Engine/Plugins/Experimental/Mutable/Source/MutableTools/Private/NodeModifierMeshClipWithMeshPrivate.h b/Engine/Plugins/Experimental/Mutable/Source/MutableTools/Private/NodeModifierMeshClipWithMeshPrivate.h
new file mode 100644
--- /dev/null
+++ b/Engine/Plugins/Experimental/Mutable/Source/MutableTools/Private/NodeModifierMeshClipWithMeshPrivate.h
@@ -0,0 +1,60 @@
+// Copyright Epic Games, Inc. All Rights Reserved.
+
+#pragma once
+
+#include "NodeModifierPrivate.h"
+#include "NodeModifierMeshClipWithMesh.h"
+#include "NodeMesh.h"
+#include "AST.h"
+
+#include "MemoryPrivate.h"
+#include "MutableMath.h"
+
+
+namespace mu
+{
+
+    class NodeModifierMeshClipWithMesh::Private : public NodeModifier::Private
+	{
+	public:
+
+		MUTABLE_DEFINE_CONST_VISITABLE()
+
+	public:
+
+		Private()
+		{
+		}
+
+		static NODE_TYPE s_type;
+
+		//! 
+		NodeMeshPtr m_clipMesh;
+
+		//!
+		void Serialise( OutputArchive& arch ) const
+		{
+			NodeModifier::Private::Serialise(arch);
+
+            uint32_t ver = 0;
+			arch << ver;
+
+			arch << m_clipMesh;
+		}
+
+		//!
+		void Unserialise( InputArchive& arch )
+		{
+			NodeModifier::Private::Unserialise( arch );
+			
+            uint32_t ver;
+			arch >> ver;
+            check(ver<=0);
+
+			arch >> m_clipMesh;
+		}
+
+	};
+
+
+}
diff --git a/Engine/Plugins/Experimental/Mutable/Source/MutableTools/Private/NodeModifierPrivate.h b/Engine/Plugins/Experimental/Mutable/Source/MutableTools/Private/NodeModifierPrivate.h
new file mode 100644
--- /dev/null
+++ b/Engine/Plugins/Experimental/Mutable/Source/MutableTools/Private/NodeModifierPrivate.h
@@ -0,0 +1,47 @@
+// Copyright Epic Games, Inc. All Rights Reserved.
+
+#pragma once
+
+#include "MutableTools/Private/NodePrivate.h"
+#include "MutableTools/Public/NodeModifier.h"
+
+#include "MutableRuntime/Private/MemoryPrivate.h"
+
+
+namespace mu
+{
+
+	class NodeModifier::Private : public Node::Private
+	{
+	public:
+
+		// Tags that target surface need to have enabled to receive this modifier.
+		vector<mu::string> m_tags;
+
+		// Wether the modifier has to be applied after the normal node operations or before
+		bool m_applyBeforeNormalOperations = true;
+
+		//!
+		void Serialise(OutputArchive& arch) const
+		{
+            uint32_t ver = 2;
+			arch << ver;
+
+			arch << m_tags;
+			arch << m_applyBeforeNormalOperations;
+		}
+
+		//!
+		void Unserialise(InputArchive& arch)
+		{
+            uint32_t ver;
+			arch >> ver;
+			check(ver == 2);
+
+			arch >> m_tags;
+			arch >> m_applyBeforeNormalOperations;
+		}
+
+	};
+
+}
diff --git a/Engine/Plugins/Experimental/Mutable/Source/MutableTools/Private/NodeObject.cpp b/Engine/Plugins/Experimental/Mutable/Source/MutableTools/Private/NodeObject.cpp
new file mode 100644
--- /dev/null
+++ b/Engine/Plugins/Experimental/Mutable/Source/MutableTools/Private/NodeObject.cpp
@@ -0,0 +1,71 @@
+// Copyright Epic Games, Inc. All Rights Reserved.
+
+
+#include "NodeObject.h"
+#include "NodeObjectPrivate.h"
+
+#include "NodeObjectNew.h"
+#include "NodeObjectGroup.h"
+#include "NodeObjectState.h"
+
+
+namespace mu
+{
+
+
+	//---------------------------------------------------------------------------------------------
+	// Static initialisation
+	//---------------------------------------------------------------------------------------------
+	NODE_TYPE s_nodeObjectType = NODE_TYPE( "NodeObject", NodeObject::GetStaticType() );
+
+
+	//---------------------------------------------------------------------------------------------
+	void NodeObject::Serialise( const NodeObject* p, OutputArchive& arch )
+	{
+        uint32_t ver = 0;
+		arch << ver;
+
+		arch << uint32_t(p->Type);
+		p->SerialiseWrapper(arch);
+    }
+
+    //---------------------------------------------------------------------------------------------
+	NodeObjectPtr NodeObject::StaticUnserialise( InputArchive& arch )
+	{
+        uint32_t ver;
+		arch >> ver;
+		check( ver == 0 );
+
+        uint32_t id;
+		arch >> id;
+
+		switch (id)
+		{
+		case 0 :  return NodeObjectNew::StaticUnserialise( arch ); break;
+		case 1 :  return NodeObjectGroup::StaticUnserialise( arch ); break;
+        case 2 :  return NodeObjectState::StaticUnserialise( arch ); break;
+        //case 3 :  return NodeObjectTransform::StaticUnserialise( arch ); break;
+        default : check(false);
+		}
+
+		return 0;
+	}
+
+
+	//---------------------------------------------------------------------------------------------
+	const NODE_TYPE* NodeObject::GetType() const
+	{
+		return GetStaticType();
+	}
+
+
+	//---------------------------------------------------------------------------------------------
+	const NODE_TYPE* NodeObject::GetStaticType()
+	{
+		return &s_nodeObjectType;
+	}
+
+
+}
+
+
diff --git a/Engine/Plugins/Experimental/Mutable/Source/MutableTools/Private/NodeObjectGroup.cpp b/Engine/Plugins/Experimental/Mutable/Source/MutableTools/Private/NodeObjectGroup.cpp
new file mode 100644
--- /dev/null
+++ b/Engine/Plugins/Experimental/Mutable/Source/MutableTools/Private/NodeObjectGroup.cpp
@@ -0,0 +1,178 @@
+// Copyright Epic Games, Inc. All Rights Reserved.
+
+
+#include "NodeObjectGroup.h"
+#include "NodeObjectGroupPrivate.h"
+
+#include "NodeLayout.h"
+
+
+namespace mu
+{
+
+
+	//---------------------------------------------------------------------------------------------
+	// Static initialisation
+	//---------------------------------------------------------------------------------------------
+	NODE_TYPE NodeObjectGroup::Private::s_type =
+            NODE_TYPE( "ObjectGroup", NodeObjectGroup::GetStaticType() );
+
+
+	//---------------------------------------------------------------------------------------------
+	//!
+	//---------------------------------------------------------------------------------------------
+
+	MUTABLE_IMPLEMENT_NODE( NodeObjectGroup, EType::Group, Node, Node::EType::Object)
+
+
+	//---------------------------------------------------------------------------------------------
+	// Node Interface
+	//---------------------------------------------------------------------------------------------
+	int NodeObjectGroup::GetInputCount() const
+	{
+		return (int)m_pD->m_children.size();
+	}
+
+
+	//---------------------------------------------------------------------------------------------
+	Node* NodeObjectGroup::GetInputNode( int i ) const
+	{
+		check( i>=0 && i<GetInputCount() );
+
+		Node* pResult = 0;
+
+		if ( i<(int)m_pD->m_children.size() )
+		{
+			pResult = m_pD->m_children[i].get();
+		}
+
+		return pResult;
+	}
+
+
+	//---------------------------------------------------------------------------------------------
+	void NodeObjectGroup::SetInputNode( int i, NodePtr pNode )
+	{
+		check( i>=0 && i<GetInputCount() );
+
+		if ( i<(int)m_pD->m_children.size() )
+		{
+			m_pD->m_children[i] = dynamic_cast<NodeObject*>( pNode.get() );
+		}
+	}
+
+
+	//---------------------------------------------------------------------------------------------
+	// Own Interface
+	//---------------------------------------------------------------------------------------------
+	const char* NodeObjectGroup::GetName() const
+	{
+		return m_pD->m_name.c_str();
+	}
+
+
+	//---------------------------------------------------------------------------------------------
+	void NodeObjectGroup::SetName( const char* strName )
+	{
+		if( strName )
+		{
+			m_pD->m_name = strName;
+		}
+		else
+		{
+			m_pD->m_name = "";
+		}
+	}
+
+
+	const char* NodeObjectGroup::GetUid() const
+	{
+		return m_pD->m_uid.c_str();
+	}
+
+
+	//---------------------------------------------------------------------------------------------
+	void NodeObjectGroup::SetUid( const char* strUid )
+	{
+		if( strUid )
+		{
+			m_pD->m_uid = strUid;
+		}
+		else
+		{
+			m_pD->m_uid = "";
+		}
+	}
+
+
+	//---------------------------------------------------------------------------------------------
+	NodeObjectGroup::CHILD_SELECTION NodeObjectGroup::GetSelectionType() const
+	{
+		return m_pD->m_type;
+	}
+
+
+	//---------------------------------------------------------------------------------------------
+	void NodeObjectGroup::SetSelectionType( CHILD_SELECTION t )
+	{
+		m_pD->m_type = t;
+	}
+
+
+	//---------------------------------------------------------------------------------------------
+	int NodeObjectGroup::GetChildCount() const
+	{
+		return (int)m_pD->m_children.size();
+	}
+
+
+	//---------------------------------------------------------------------------------------------
+	void NodeObjectGroup::SetChildCount( int num )
+	{
+		check( num >=0 );
+		m_pD->m_children.resize( num );
+	}
+
+
+	//---------------------------------------------------------------------------------------------
+	NodeObjectPtr NodeObjectGroup::GetChild( int index ) const
+	{
+		check( index >=0 && index < (int)m_pD->m_children.size() );
+
+		return m_pD->m_children[ index ].get();
+	}
+
+
+	//---------------------------------------------------------------------------------------------
+	void NodeObjectGroup::SetChild( int index, NodeObjectPtr pObject )
+	{
+		check( index >=0 && index < (int)m_pD->m_children.size() );
+
+		m_pD->m_children[ index ] = pObject;
+	}
+
+
+	//---------------------------------------------------------------------------------------------
+    NodeLayoutPtr NodeObjectGroup::Private::GetLayout(int lod, int component, int surface, int texture ) const
+	{
+		NodeLayoutPtr pLayout;
+
+		for ( size_t i=0; !pLayout && i<m_children.size(); ++i )
+		{
+			if (m_children[i])
+			{
+				NodeObject::Private* pPrivate =
+						dynamic_cast<NodeObject::Private*>( m_children[i]->GetBasePrivate() );
+
+                pLayout = pPrivate->GetLayout( lod, component, surface, texture );
+			}
+		}
+
+		// TODO: layout index
+		return pLayout;
+	}
+
+
+}
+
+
diff --git a/Engine/Plugins/Experimental/Mutable/Source/MutableTools/Private/NodeObjectGroupPrivate.h b/Engine/Plugins/Experimental/Mutable/Source/MutableTools/Private/NodeObjectGroupPrivate.h
new file mode 100644
--- /dev/null
+++ b/Engine/Plugins/Experimental/Mutable/Source/MutableTools/Private/NodeObjectGroupPrivate.h
@@ -0,0 +1,70 @@
+// Copyright Epic Games, Inc. All Rights Reserved.
+
+#pragma once
+
+#include "NodeObjectPrivate.h"
+
+#include "NodeObjectGroup.h"
+#include "NodeLayout.h"
+
+#include "MemoryPrivate.h"
+
+
+namespace mu
+{
+
+	MUTABLE_DEFINE_ENUM_SERIALISABLE( NodeObjectGroup::CHILD_SELECTION )
+
+
+	class NodeObjectGroup::Private : public NodeObject::Private
+	{
+	public:
+
+		MUTABLE_DEFINE_CONST_VISITABLE()
+
+	public:
+
+		static NODE_TYPE s_type;
+
+		string m_name;
+		string m_uid;
+
+		CHILD_SELECTION m_type;
+
+		//! Set the child selection type
+		void SetSelectionType( CHILD_SELECTION );
+
+		vector<NodeObjectPtr> m_children;
+
+		//!
+		void Serialise( OutputArchive& arch ) const
+		{
+            uint32_t ver = 1;
+			arch << ver;
+
+			arch << m_type;
+			arch << m_name;
+			arch << m_uid;
+			arch << m_children;
+		}
+
+		//!
+		void Unserialise( InputArchive& arch )
+		{
+            uint32_t ver;
+			arch >> ver;
+            check(ver==1);
+
+			arch >> m_type;
+			arch >> m_name;
+            arch >> m_uid;
+			arch >> m_children;
+		}
+
+
+        // NodeObject::Private interface
+        NodeLayoutPtr GetLayout( int lod, int component, int surface, int texture ) const override;
+
+	};
+
+}
