[mutable] Move Mutable plugin from the restricted folder to the experimental plugins folder.

diff --git a/Engine/Plugins/Experimental/Mutable/Source/MutableTools/Private/NodeObjectNew.cpp b/Engine/Plugins/Experimental/Mutable/Source/MutableTools/Private/NodeObjectNew.cpp
new file mode 100644
--- /dev/null
+++ b/Engine/Plugins/Experimental/Mutable/Source/MutableTools/Private/NodeObjectNew.cpp
@@ -0,0 +1,329 @@
+// Copyright Epic Games, Inc. All Rights Reserved.
+
+
+#include "NodeObjectNew.h"
+#include "NodeObjectNewPrivate.h"
+
+#include "NodeMeshPrivate.h"
+#include "NodeLayout.h"
+#include "NodeMeshConstant.h"
+#include "NodeComponentNew.h"
+#include "NodeSurfaceNew.h"
+
+
+namespace mu
+{
+
+
+	//---------------------------------------------------------------------------------------------
+	// Static initialisation
+	//---------------------------------------------------------------------------------------------
+	NODE_TYPE NodeObjectNew::Private::s_type =
+			NODE_TYPE( "Object", NodeObjectNew::GetStaticType() );
+
+
+	//---------------------------------------------------------------------------------------------
+	//!
+	//---------------------------------------------------------------------------------------------
+
+	MUTABLE_IMPLEMENT_NODE( NodeObjectNew, EType::New, Node, Node::EType::Object )
+
+
+	//---------------------------------------------------------------------------------------------
+	// Node Interface
+	//---------------------------------------------------------------------------------------------
+	int NodeObjectNew::GetInputCount() const
+	{
+		return (int)(m_pD->m_lods.size() + m_pD->m_children.size());
+	}
+
+
+	//---------------------------------------------------------------------------------------------
+	Node* NodeObjectNew::GetInputNode( int i ) const
+	{
+		check( i>=0 && i<GetInputCount() );
+
+		Node* pResult = 0;
+
+		if ( i<(int)m_pD->m_lods.size() )
+		{
+			pResult = m_pD->m_lods[i].get();
+		}
+		else
+		{
+			i -= (int)m_pD->m_lods.size();
+			pResult = m_pD->m_children[i].get();
+		}
+
+		return pResult;
+	}
+
+
+	//---------------------------------------------------------------------------------------------
+	void NodeObjectNew::SetInputNode( int i, NodePtr pNode )
+	{
+		check( i>=0 && i<GetInputCount() );
+
+		if ( i<(int)m_pD->m_lods.size() )
+		{
+			m_pD->m_lods[i] = dynamic_cast<NodeLOD*>( pNode.get() );
+		}
+		else
+		{
+			i -= (int)m_pD->m_lods.size();
+			m_pD->m_children[i] = dynamic_cast<NodeObject*>( pNode.get() );
+		}
+	}
+
+
+	//---------------------------------------------------------------------------------------------
+	// Own Interface
+	//---------------------------------------------------------------------------------------------
+	const char* NodeObjectNew::GetName() const
+	{
+		return m_pD->m_name.c_str();
+	}
+
+
+	//---------------------------------------------------------------------------------------------
+	void NodeObjectNew::SetName( const char* strName )
+	{
+		if( strName )
+		{
+			m_pD->m_name = strName;
+		}
+		else
+		{
+			m_pD->m_name = "";
+		}
+	}
+
+
+	const char* NodeObjectNew::GetUid() const
+	{
+		return m_pD->m_uid.c_str();
+	}
+
+
+	//---------------------------------------------------------------------------------------------
+	void NodeObjectNew::SetUid( const char* strUid )
+	{
+		if( strUid )
+		{
+			m_pD->m_uid = strUid;
+		}
+		else
+		{
+			m_pD->m_uid = "";
+		}
+	}
+
+
+	//---------------------------------------------------------------------------------------------
+	int NodeObjectNew::GetLODCount() const
+	{
+		return (int)m_pD->m_lods.size();
+	}
+
+
+	//---------------------------------------------------------------------------------------------
+	void NodeObjectNew::SetLODCount( int num )
+	{
+		check( num >=0 );
+		m_pD->m_lods.resize( num );
+	}
+
+
+	//---------------------------------------------------------------------------------------------
+	NodeLODPtr NodeObjectNew::GetLOD( int index ) const
+	{
+		check( index >=0 && index < (int)m_pD->m_lods.size() );
+
+		return m_pD->m_lods[ index ].get();
+	}
+
+
+	//---------------------------------------------------------------------------------------------
+	void NodeObjectNew::SetLOD( int index, NodeLODPtr pLOD )
+	{
+		check( index >=0 && index < (int)m_pD->m_lods.size() );
+
+		m_pD->m_lods[ index ] = pLOD;
+	}
+
+
+	//---------------------------------------------------------------------------------------------
+	int NodeObjectNew::GetChildCount() const
+	{
+		return (int)m_pD->m_children.size();
+	}
+
+
+	//---------------------------------------------------------------------------------------------
+	void NodeObjectNew::SetChildCount( int num )
+	{
+		check( num >=0 );
+		m_pD->m_children.resize( num );
+	}
+
+
+	//---------------------------------------------------------------------------------------------
+	NodeObjectPtr NodeObjectNew::GetChild( int index ) const
+	{
+		check( index >=0 && index < (int)m_pD->m_children.size() );
+
+		return m_pD->m_children[ index ].get();
+	}
+
+
+	//---------------------------------------------------------------------------------------------
+	void NodeObjectNew::SetChild( int index, NodeObjectPtr pObject )
+	{
+		check( index >=0 && index < (int)m_pD->m_children.size() );
+
+		m_pD->m_children[ index ] = pObject;
+	}
+
+
+	//---------------------------------------------------------------------------------------------
+    NodeLayoutPtr NodeObjectNew::Private::GetLayout( int lod,
+                                                     int component,
+                                                     int surface,
+                                                     int //texture
+                                                     ) const
+	{
+		check( lod>=0 && lod<(int)m_lods.size() );
+
+		NodeLayoutPtr pLayout;
+
+		const NodeComponentNew* pComp =
+			dynamic_cast<const NodeComponentNew*> ( m_lods[lod]->GetComponent( component ).get() );
+
+        if (pComp)
+        {
+            const NodeSurfaceNew* pSurface =
+                    dynamic_cast<const NodeSurfaceNew*> ( pComp->GetSurface( surface ) );
+
+            if ( pSurface )
+            {
+                // TODO: Look for the layout index of the given texture
+                // TODO: Multiple meshes
+                if ( pSurface->GetMeshCount()>0 )
+                {
+                     const NodeMesh* pMesh =
+                             dynamic_cast<const NodeMesh*> ( pSurface->GetMesh( 0 ).get() );
+
+                    if ( pMesh )
+                    {
+                        NodeMesh::Private* pPrivate =
+                                dynamic_cast<NodeMesh::Private*>( pMesh->GetBasePrivate() );
+
+                        pLayout = pPrivate->GetLayout( 0 );
+                    }
+                }
+            }
+        }
+
+		// TODO: layout index
+		return pLayout;
+	}
+
+
+    //---------------------------------------------------------------------------------------------
+    bool NodeObjectNew::Private::HasComponent( const NodeComponent* pComponent ) const
+    {
+        bool found = false;
+        for ( size_t l=0; !found && l<m_lods.size(); ++l )
+        {
+            found = std::find( m_lods[l]->GetPrivate()->m_components.begin(),
+                               m_lods[l]->GetPrivate()->m_components.end(),
+                               pComponent )
+                    != m_lods[l]->GetPrivate()->m_components.end();
+        }
+
+        return found;
+    }
+
+
+	//---------------------------------------------------------------------------------------------
+	int NodeObjectNew::GetStateCount() const
+	{
+		return (int)m_pD->m_states.size();
+	}
+
+
+	//---------------------------------------------------------------------------------------------
+	void NodeObjectNew::SetStateCount( int c )
+	{
+		m_pD->m_states.resize( c );
+	}
+
+
+	//---------------------------------------------------------------------------------------------
+	const char* NodeObjectNew::GetStateName( int s ) const
+	{
+		check( s>=0 && s<GetStateCount() );
+		return m_pD->m_states[s].m_name.c_str();
+	}
+
+
+	//---------------------------------------------------------------------------------------------
+	void NodeObjectNew::SetStateName( int s, const char* n )
+	{
+		check( s>=0 && s<GetStateCount() );
+		m_pD->m_states[s].m_name = n;
+	}
+
+
+	//---------------------------------------------------------------------------------------------
+	bool NodeObjectNew::HasStateParam( int s, const char* param ) const
+	{
+		check( s>=0 && s<GetStateCount() );
+		return std::find( m_pD->m_states[s].m_runtimeParams.begin(),
+						  m_pD->m_states[s].m_runtimeParams.end(),
+						  param )
+				!=
+				m_pD->m_states[s].m_runtimeParams.end();
+	}
+
+
+	//---------------------------------------------------------------------------------------------
+	void NodeObjectNew::AddStateParam( int s, const char* param )
+	{
+		check( s>=0 && s<GetStateCount() );
+
+		if (!HasStateParam(s,param))
+		{
+			m_pD->m_states[s].m_runtimeParams.push_back( param );
+		}
+	}
+
+
+	//---------------------------------------------------------------------------------------------
+	void NodeObjectNew::RemoveStateParam( int s, const char* param )
+	{
+		check( s>=0 && s<GetStateCount() );
+
+		vector<string>::iterator it = std::find( m_pD->m_states[s].m_runtimeParams.begin(),
+														   m_pD->m_states[s].m_runtimeParams.end(),
+														   param );
+		if ( it != m_pD->m_states[s].m_runtimeParams.end() )
+		{
+			m_pD->m_states[s].m_runtimeParams.erase( it );
+		}
+	}
+
+
+    //---------------------------------------------------------------------------------------------
+    void NodeObjectNew::SetStateProperties( int s, bool avoidRuntimeCompression, bool onlyFirstLOD, int firstLOD )
+    {
+        check( s>=0 && s<GetStateCount() );
+
+        m_pD->m_states[s].m_optimisation.m_avoidRuntimeCompression = avoidRuntimeCompression;
+        m_pD->m_states[s].m_optimisation.m_onlyFirstLOD = onlyFirstLOD;
+        m_pD->m_states[s].m_optimisation.m_firstLOD = firstLOD;
+    }
+
+}
+
+
diff --git a/Engine/Plugins/Experimental/Mutable/Source/MutableTools/Private/NodeObjectNewPrivate.h b/Engine/Plugins/Experimental/Mutable/Source/MutableTools/Private/NodeObjectNewPrivate.h
new file mode 100644
--- /dev/null
+++ b/Engine/Plugins/Experimental/Mutable/Source/MutableTools/Private/NodeObjectNewPrivate.h
@@ -0,0 +1,76 @@
+// Copyright Epic Games, Inc. All Rights Reserved.
+
+#pragma once
+
+#include "NodeObjectPrivate.h"
+#include "NodeLODPrivate.h"
+#include "CompilerPrivate.h"
+
+#include "NodeObjectNew.h"
+#include "NodeLayout.h"
+
+#include "MemoryPrivate.h"
+
+
+namespace mu
+{
+
+
+	class NodeObjectNew::Private : public NodeObject::Private
+	{
+	public:
+
+		MUTABLE_DEFINE_CONST_VISITABLE()
+
+	public:
+
+		static NODE_TYPE s_type;
+
+		string m_name;
+		string m_uid;
+
+		vector<NodeLODPtr> m_lods;
+
+		vector<NodeObjectPtr> m_children;
+
+
+		//! List of states
+        vector<OBJECT_STATE> m_states;
+
+
+		//!
+		void Serialise( OutputArchive& arch ) const
+		{
+            uint32_t ver = 2;
+			arch << ver;
+
+			arch << m_name;
+			arch << m_uid;
+			arch << m_lods;
+			arch << m_children;
+			arch << m_states;
+		}
+
+		//!
+		void Unserialise( InputArchive& arch )
+		{
+            uint32_t ver;
+			arch >> ver;
+            check( ver==2 );
+
+			arch >> m_name;
+            arch >> m_uid;
+			arch >> m_lods;
+			arch >> m_children;
+			arch >> m_states;
+		}
+
+        //! Return true if the given component is set in any lod of this object.
+        bool HasComponent( const NodeComponent* pComponent ) const;
+
+        // NodeObject::Private interface
+        NodeLayoutPtr GetLayout( int lod, int component, int surface, int texture ) const override;
+
+	};
+
+}
diff --git a/Engine/Plugins/Experimental/Mutable/Source/MutableTools/Private/NodeObjectPrivate.h b/Engine/Plugins/Experimental/Mutable/Source/MutableTools/Private/NodeObjectPrivate.h
new file mode 100644
--- /dev/null
+++ b/Engine/Plugins/Experimental/Mutable/Source/MutableTools/Private/NodeObjectPrivate.h
@@ -0,0 +1,22 @@
+// Copyright Epic Games, Inc. All Rights Reserved.
+
+#pragma once
+
+#include "NodePrivate.h"
+
+#include "NodeObject.h"
+#include "NodeLayout.h"
+
+namespace mu
+{
+
+	class NodeObject::Private : public Node::Private
+	{
+	public:
+
+        virtual NodeLayoutPtr GetLayout( int lod, int component, int surface, int texture ) const = 0;
+
+	};
+
+}
+
diff --git a/Engine/Plugins/Experimental/Mutable/Source/MutableTools/Private/NodeObjectState.cpp b/Engine/Plugins/Experimental/Mutable/Source/MutableTools/Private/NodeObjectState.cpp
new file mode 100644
--- /dev/null
+++ b/Engine/Plugins/Experimental/Mutable/Source/MutableTools/Private/NodeObjectState.cpp
@@ -0,0 +1,195 @@
+// Copyright Epic Games, Inc. All Rights Reserved.
+
+
+#include "NodeObjectState.h"
+#include "NodeObjectStatePrivate.h"
+
+
+namespace mu
+{
+
+
+	//---------------------------------------------------------------------------------------------
+	// Static initialisation
+	//---------------------------------------------------------------------------------------------
+	NODE_TYPE NodeObjectState::Private::s_type =
+			NODE_TYPE( "ObjectState", NodeObjectState::GetStaticType() );
+
+
+	//---------------------------------------------------------------------------------------------
+    NodeLayoutPtr NodeObjectState::Private::GetLayout(int lod, int component, int surface, int texture ) const
+	{
+		NodeLayoutPtr pResult;
+
+		if ( m_pSource )
+		{
+			const NodeObject::Private* pPrivate = dynamic_cast<const NodeObject::Private*>
+					( m_pSource->GetBasePrivate() );
+            pResult = pPrivate->GetLayout( lod, component, surface, texture );
+		}
+
+		return pResult;
+	}
+
+
+	//---------------------------------------------------------------------------------------------
+	//!
+	//---------------------------------------------------------------------------------------------
+
+	MUTABLE_IMPLEMENT_NODE( NodeObjectState, NodeObject::EType::State, Node, Node::EType::Object)
+
+
+	//---------------------------------------------------------------------------------------------
+	// Node Interface
+	//---------------------------------------------------------------------------------------------
+	int NodeObjectState::GetInputCount() const
+	{
+		return 2;
+	}
+
+
+	//---------------------------------------------------------------------------------------------
+	Node* NodeObjectState::GetInputNode( int i ) const
+	{
+		check( i>=0 && i<GetInputCount() );
+
+		Node* pResult = 0;
+
+		if (i==0)
+		{
+			pResult = m_pD->m_pSource.get();
+		}
+		else
+		{
+			pResult = m_pD->m_pRoot.get();
+		}
+
+		return pResult;
+	}
+
+
+	//---------------------------------------------------------------------------------------------
+	void NodeObjectState::SetInputNode( int i, NodePtr pNode )
+	{
+		check( i>=0 && i<GetInputCount() );
+
+		if (i==0)
+		{
+			m_pD->m_pSource = dynamic_cast<NodeObject*>( pNode.get() );
+		}
+		else
+		{
+			m_pD->m_pRoot = dynamic_cast<NodeObject*>( pNode.get() );
+		}
+	}
+
+
+	//---------------------------------------------------------------------------------------------
+	// NodeObject Interface
+	//---------------------------------------------------------------------------------------------
+	const char* NodeObjectState::GetName() const
+	{
+		const char* strRes = "";
+		return strRes;
+	}
+
+
+	//---------------------------------------------------------------------------------------------
+    void NodeObjectState::SetName( const char* )
+	{
+	}
+
+
+	const char* NodeObjectState::GetUid() const
+	{
+		const char* strRes = "";
+		return strRes;
+	}
+
+
+	//---------------------------------------------------------------------------------------------
+	void NodeObjectState::SetUid( const char* )
+	{
+	}
+
+
+	//---------------------------------------------------------------------------------------------
+	// Own Interface
+	//---------------------------------------------------------------------------------------------
+	NodeObjectPtr NodeObjectState::GetSource() const
+	{
+		return m_pD->m_pSource.get();
+	}
+
+
+	//---------------------------------------------------------------------------------------------
+	void NodeObjectState::SetSource( NodeObjectPtr pSource )
+	{
+		m_pD->m_pSource = pSource;
+	}
+
+
+	//---------------------------------------------------------------------------------------------
+	NodeObjectPtr NodeObjectState::GetStateRoot() const
+	{
+		return m_pD->m_pRoot.get();
+	}
+
+
+	//---------------------------------------------------------------------------------------------
+	void NodeObjectState::SetStateRoot( NodeObjectPtr pState )
+	{
+		m_pD->m_pRoot = pState;
+	}
+
+
+	//---------------------------------------------------------------------------------------------
+	const char* NodeObjectState::GetStateName() const
+	{
+		return m_pD->m_state.m_name.c_str();
+	}
+
+
+	//---------------------------------------------------------------------------------------------
+	void NodeObjectState::SetStateName( const char* n )
+	{
+		m_pD->m_state.m_name = n;
+	}
+
+
+	//---------------------------------------------------------------------------------------------
+	bool NodeObjectState::HasStateParam( const char* param ) const
+	{
+		return std::find( m_pD->m_state.m_runtimeParams.begin(),
+						  m_pD->m_state.m_runtimeParams.end(),
+						  param )
+				!=
+				m_pD->m_state.m_runtimeParams.end();
+	}
+
+
+	//---------------------------------------------------------------------------------------------
+	void NodeObjectState::AddStateParam( const char* param )
+	{
+		if (!HasStateParam(param))
+		{
+			m_pD->m_state.m_runtimeParams.push_back( param );
+		}
+	}
+
+
+	//---------------------------------------------------------------------------------------------
+	void NodeObjectState::RemoveStateParam( const char* param )
+	{
+		vector<string>::iterator it = std::find( m_pD->m_state.m_runtimeParams.begin(),
+														   m_pD->m_state.m_runtimeParams.end(),
+														   param );
+		if ( it != m_pD->m_state.m_runtimeParams.end() )
+		{
+			m_pD->m_state.m_runtimeParams.erase( it );
+		}
+	}
+
+}
+
+
diff --git a/Engine/Plugins/Experimental/Mutable/Source/MutableTools/Private/NodeObjectStatePrivate.h b/Engine/Plugins/Experimental/Mutable/Source/MutableTools/Private/NodeObjectStatePrivate.h
new file mode 100644
--- /dev/null
+++ b/Engine/Plugins/Experimental/Mutable/Source/MutableTools/Private/NodeObjectStatePrivate.h
@@ -0,0 +1,62 @@
+// Copyright Epic Games, Inc. All Rights Reserved.
+
+
+#pragma once
+
+#include "NodeObjectPrivate.h"
+#include "NodeObjectState.h"
+#include "CompilerPrivate.h"
+
+
+namespace mu
+{
+
+	class NodeObjectState::Private : public NodeObject::Private
+	{
+	public:
+
+		MUTABLE_DEFINE_CONST_VISITABLE()
+
+	public:
+
+
+        // NodeObject::Private interface
+        NodeLayoutPtr GetLayout( int lod, int component, int surface, int texture ) const override;
+
+
+		//-----------------------------------------------------------------------------------------
+		// Own interface
+		//-----------------------------------------------------------------------------------------
+		static NODE_TYPE s_type;
+
+        OBJECT_STATE m_state;
+
+		NodeObjectPtr m_pSource;
+		NodeObjectPtr m_pRoot;
+
+		//!
+		void Serialise( OutputArchive& arch ) const
+		{
+            uint32_t ver = 0;
+			arch << ver;
+
+			arch << m_state;
+			arch << m_pSource;
+			arch << m_pRoot;
+		}
+
+		//!
+		void Unserialise( InputArchive& arch )
+		{
+            uint32_t ver;
+			arch >> ver;
+			check(ver==0);
+
+			arch >> m_state;
+			arch >> m_pSource;
+			arch >> m_pRoot;
+		}
+	};
+
+}
+
diff --git a/Engine/Plugins/Experimental/Mutable/Source/MutableTools/Private/NodePatchImage.cpp b/Engine/Plugins/Experimental/Mutable/Source/MutableTools/Private/NodePatchImage.cpp
new file mode 100644
--- /dev/null
+++ b/Engine/Plugins/Experimental/Mutable/Source/MutableTools/Private/NodePatchImage.cpp
@@ -0,0 +1,160 @@
+// Copyright Epic Games, Inc. All Rights Reserved.
+
+
+#include "NodePatchImage.h"
+#include "NodePatchImagePrivate.h"
+
+#include "NodeImage.h"
+
+
+namespace mu
+{
+
+
+    //---------------------------------------------------------------------------------------------
+    // Static initialisation
+    //---------------------------------------------------------------------------------------------
+    NODE_TYPE NodePatchImage::Private::s_type =
+            NODE_TYPE( "PatchTexture", NodePatchImage::GetStaticType() );
+
+
+    //---------------------------------------------------------------------------------------------
+    //!
+    //---------------------------------------------------------------------------------------------
+
+    MUTABLE_IMPLEMENT_NODE( NodePatchImage, EType::PatchImage, Node, Node::EType::PatchImage);
+
+
+    //---------------------------------------------------------------------------------------------
+    // Node Interface
+    //---------------------------------------------------------------------------------------------
+    int NodePatchImage::GetInputCount() const
+    {
+        return 2;
+    }
+
+
+    //---------------------------------------------------------------------------------------------
+    Node* NodePatchImage::GetInputNode( int i ) const
+    {
+        check( i >=0 && i < GetInputCount());
+
+        Node* pResult = 0;
+
+        if ( i==0 )
+        {
+            pResult = m_pD->m_pImage.get();
+        }
+        else
+        {
+            pResult = m_pD->m_pMask.get();
+        }
+
+        return pResult;
+    }
+
+
+    //---------------------------------------------------------------------------------------------
+    void NodePatchImage::SetInputNode( int i, NodePtr pNode )
+    {
+        check( i >=0 && i < GetInputCount());
+
+        if ( i==0 )
+        {
+            m_pD->m_pImage = dynamic_cast<NodeImage*>( pNode.get() );
+        }
+        else
+        {
+            m_pD->m_pMask = dynamic_cast<NodeImage*>( pNode.get() );
+        }
+    }
+
+
+    //---------------------------------------------------------------------------------------------
+    // Own Interface
+    //---------------------------------------------------------------------------------------------
+    void NodePatchImage::SetImage( NodeImagePtr pImage )
+    {
+        m_pD->m_pImage = pImage;
+    }
+
+
+    //---------------------------------------------------------------------------------------------
+    NodeImagePtr NodePatchImage::GetImage() const
+    {
+        return m_pD->m_pImage.get();
+    }
+
+
+    //---------------------------------------------------------------------------------------------
+    void NodePatchImage::SetMask( NodeImagePtr pMask )
+    {
+        m_pD->m_pMask = pMask;
+    }
+
+
+    //---------------------------------------------------------------------------------------------
+    NodeImagePtr NodePatchImage::GetMask() const
+    {
+        return m_pD->m_pMask.get();
+    }
+
+
+    //---------------------------------------------------------------------------------------------
+    void NodePatchImage::SetBlockCount( int c )
+    {
+        m_pD->m_blocks.resize( c );
+    }
+
+
+    //---------------------------------------------------------------------------------------------
+    int NodePatchImage::GetBlockCount() const
+    {
+        return (int)m_pD->m_blocks.size();
+    }
+
+
+    //---------------------------------------------------------------------------------------------
+    void NodePatchImage::SetBlock( int index, int block )
+    {
+        check( index>=0 && index<(int)m_pD->m_blocks.size() );
+        m_pD->m_blocks[ index ] = block;
+    }
+
+
+    //---------------------------------------------------------------------------------------------
+    int NodePatchImage::GetBlock( int index ) const
+    {
+        check( index>=0 && index<(int)m_pD->m_blocks.size() );
+        return m_pD->m_blocks[ index ];
+    }
+
+
+    //---------------------------------------------------------------------------------------------
+	EBlendType NodePatchImage::GetBlendType() const
+    {
+        return m_pD->m_blendType;
+    }
+
+
+    //---------------------------------------------------------------------------------------------
+    void NodePatchImage::SetBlendType(EBlendType t)
+    {
+        m_pD->m_blendType = t;
+    }
+
+
+    //---------------------------------------------------------------------------------------------
+    bool NodePatchImage::GetApplyToAlphaChannel() const
+    {
+        return m_pD->m_applyToAlpha;
+    }
+
+
+    //---------------------------------------------------------------------------------------------
+    void NodePatchImage::SetApplyToAlphaChannel(bool a )
+    {
+        m_pD->m_applyToAlpha = a;
+    }
+
+}
diff --git a/Engine/Plugins/Experimental/Mutable/Source/MutableTools/Private/NodePatchImagePrivate.h b/Engine/Plugins/Experimental/Mutable/Source/MutableTools/Private/NodePatchImagePrivate.h
new file mode 100644
--- /dev/null
+++ b/Engine/Plugins/Experimental/Mutable/Source/MutableTools/Private/NodePatchImagePrivate.h
@@ -0,0 +1,69 @@
+// Copyright Epic Games, Inc. All Rights Reserved.
+
+
+#pragma once
+
+#include "NodePrivate.h"
+#include "NodePatchImage.h"
+#include "NodeImage.h"
+
+#include "MemoryPrivate.h"
+
+
+namespace mu
+{
+
+    class NodePatchImage::Private : public Node::Private
+    {
+    public:
+
+        MUTABLE_DEFINE_CONST_VISITABLE();
+
+    public:
+
+        static NODE_TYPE s_type;
+
+        NodeImagePtr m_pImage;
+        NodeImagePtr m_pMask;
+
+        vector<int> m_blocks;
+
+        EBlendType m_blendType = EBlendType::BT_BLEND;
+
+        // Patch alpha channel as well?
+        bool m_applyToAlpha = false;
+
+        //!
+        void Serialise( OutputArchive& arch ) const
+        {
+            uint32_t ver = 2;
+            arch << ver;
+
+            arch << m_pImage;
+            arch << m_pMask;
+            arch << m_blocks;
+            arch << (uint32_t)m_blendType;
+            arch << m_applyToAlpha;
+        }
+
+        //!
+        void Unserialise( InputArchive& arch )
+        {
+            uint32_t ver;
+            arch >> ver;
+            check(ver==2);
+
+            arch >> m_pImage;
+            arch >> m_pMask;
+            arch >> m_blocks;
+
+			uint32_t t;
+            arch >> t;
+            m_blendType=(EBlendType)t;
+
+            arch >> m_applyToAlpha;
+        }
+    };
+
+}
+
diff --git a/Engine/Plugins/Experimental/Mutable/Source/MutableTools/Private/NodePatchMesh.cpp b/Engine/Plugins/Experimental/Mutable/Source/MutableTools/Private/NodePatchMesh.cpp
new file mode 100644
--- /dev/null
+++ b/Engine/Plugins/Experimental/Mutable/Source/MutableTools/Private/NodePatchMesh.cpp
@@ -0,0 +1,108 @@
+// Copyright Epic Games, Inc. All Rights Reserved.
+
+
+#include "NodePatchMesh.h"
+#include "NodePatchMeshPrivate.h"
+#include "MeshPrivate.h"
+
+
+namespace mu
+{
+
+	//---------------------------------------------------------------------------------------------
+	// Static initialisation
+	//---------------------------------------------------------------------------------------------
+	NODE_TYPE NodePatchMesh::Private::s_type =
+			NODE_TYPE( "MeshPatch", Node::GetStaticType() );
+
+
+	//---------------------------------------------------------------------------------------------
+	//!
+	//---------------------------------------------------------------------------------------------
+
+	MUTABLE_IMPLEMENT_NODE( NodePatchMesh, EType::PatchMesh, Node, Node::EType::PatchMesh);
+
+
+	//---------------------------------------------------------------------------------------------
+	// Node Interface
+	//---------------------------------------------------------------------------------------------
+	int NodePatchMesh::GetInputCount() const
+	{
+		return 2;
+	}
+
+
+	//---------------------------------------------------------------------------------------------
+	Node* NodePatchMesh::GetInputNode( int i ) const
+	{
+		Node* pResult = 0;
+
+		if( i==0 )
+		{
+			pResult = m_pD->m_pRemove.get();
+		}
+		else if (i==1)
+		{
+			pResult = m_pD->m_pAdd.get();
+		}
+		else
+		{
+			check( false );
+		}
+
+		return pResult;
+	}
+
+
+	//---------------------------------------------------------------------------------------------
+	void NodePatchMesh::SetInputNode( int i, NodePtr pNode )
+	{
+		if( i==0 )
+		{
+			m_pD->m_pRemove = dynamic_cast<NodeMesh*>( pNode.get() );
+		}
+		else if (i==1)
+		{
+			m_pD->m_pAdd = dynamic_cast<NodeMesh*>( pNode.get() );
+		}
+		else
+		{
+			check( false );
+		}
+	}
+
+
+	//---------------------------------------------------------------------------------------------
+	// Own Interface
+	//---------------------------------------------------------------------------------------------
+    NodeMesh* NodePatchMesh::GetRemove() const
+	{
+        return m_pD->m_pRemove.get();
+	}
+
+
+	//---------------------------------------------------------------------------------------------
+    void NodePatchMesh::SetRemove( NodeMesh* pRemove )
+	{
+		m_pD->m_pRemove = pRemove;
+	}
+
+
+	//---------------------------------------------------------------------------------------------
+    NodeMesh* NodePatchMesh::GetAdd() const
+	{
+        return m_pD->m_pAdd.get();
+	}
+
+
+	//---------------------------------------------------------------------------------------------
+    void NodePatchMesh::SetAdd( NodeMesh* pAdd )
+	{
+		m_pD->m_pAdd = pAdd;
+	}
+
+
+
+}
+
+
diff --git a/Engine/Plugins/Experimental/Mutable/Source/MutableTools/Private/NodePatchMeshPrivate.h b/Engine/Plugins/Experimental/Mutable/Source/MutableTools/Private/NodePatchMeshPrivate.h
new file mode 100644
--- /dev/null
+++ b/Engine/Plugins/Experimental/Mutable/Source/MutableTools/Private/NodePatchMeshPrivate.h
@@ -0,0 +1,53 @@
+// Copyright Epic Games, Inc. All Rights Reserved.
+
+#pragma once
+
+#include "NodeMeshPrivate.h"
+#include "NodePatchMesh.h"
+#include "AST.h"
+
+#include "MemoryPrivate.h"
+
+
+namespace mu
+{
+
+
+	class NodePatchMesh::Private : public Node::Private
+	{
+	public:
+
+		MUTABLE_DEFINE_CONST_VISITABLE();
+
+	public:
+
+		static NODE_TYPE s_type;
+
+		NodeMeshPtr m_pRemove;
+		NodeMeshPtr m_pAdd;
+
+		//!
+		void Serialise( OutputArchive& arch ) const
+		{
+            uint32_t ver = 1;
+			arch << ver;
+
+			arch << m_pRemove;
+			arch << m_pAdd;
+		}
+
+		//!
+		void Unserialise( InputArchive& arch )
+		{
+            uint32_t ver;
+			arch >> ver;
+            check( ver==1 );
+
+			arch >> m_pRemove;
+			arch >> m_pAdd;
+		}
+
+	};
+
+
+}
diff --git a/Engine/Plugins/Experimental/Mutable/Source/MutableTools/Private/NodePrivate.h b/Engine/Plugins/Experimental/Mutable/Source/MutableTools/Private/NodePrivate.h
new file mode 100644
--- /dev/null
+++ b/Engine/Plugins/Experimental/Mutable/Source/MutableTools/Private/NodePrivate.h
@@ -0,0 +1,115 @@
+// Copyright Epic Games, Inc. All Rights Reserved.
+
+
+#pragma once
+
+#include "Node.h"
+#include "Visitor.h"
+#include "AST.h"
+#include "Operations.h"
+#include "MemoryPrivate.h"
+
+#include <memory>
+
+namespace mu
+{
+
+
+    //!
+    class Node::Private :
+        public Base,
+        public BaseVisitable<Ptr<class ASTOp>,DefaultCatchAll,true>
+    {
+    public:
+
+        //!
+        static void CloneChildren( const Node* pSource, NodePtr pTarget, NodeMapPtr pMap );
+
+        //! This is an opaque context used to attach to reported error messages.
+		const void* m_errorContext = nullptr;
+
+		//! Generic pointer to the node owning this private.
+		const Node* m_pNode = nullptr;
+    };
+
+
+    class NodeMap::Private : public Base
+    {
+    public:
+
+        map< const void*, NodePtr > m_map;
+
+    };
+
+
+#define MUTABLE_IMPLEMENT_NODE( N, T, PN, PT )				\
+    N::N()													\
+    {														\
+		Type = T;											\
+        m_pD = new Private();								\
+		m_pD->m_pNode = this;								\
+		PN::Type = PT;										\
+	}														\
+                                                            \
+    N::~N()													\
+    {														\
+        check( m_pD );										\
+        delete m_pD;										\
+        m_pD = nullptr;										\
+    }														\
+                                                            \
+    void N::Serialise( const N* p, OutputArchive& arch )	\
+    {														\
+        arch << *p->m_pD;									\
+    }														\
+                                                            \
+	void N::SerialiseWrapper( OutputArchive& arch ) const	\
+	{														\
+		N::Serialise(this, arch);							\
+	}														\
+                                                            \
+    Ptr<N> N::StaticUnserialise( InputArchive& arch )				\
+    {														\
+        Ptr<N> pResult = new N();							\
+        arch >> *pResult->m_pD;								\
+        return pResult;										\
+    }														\
+                                                            \
+    NodePtr N::Clone( NodeMapPtr pMap ) const				\
+    {														\
+        Ptr<N> pCloned = new N();							\
+                                                            \
+        *pCloned->m_pD = *m_pD;								\
+                                                            \
+        if (!pMap)											\
+        {													\
+            pMap = new NodeMap;								\
+        }													\
+        pMap->Add( this, pCloned );							\
+                                                            \
+        Private::CloneChildren( this, pCloned, pMap );		\
+                                                            \
+        return pCloned;										\
+    }														\
+                                                            \
+    N::Private* N::GetPrivate() const						\
+    {														\
+        return m_pD;										\
+    }														\
+                                                            \
+    Node::Private* N::GetBasePrivate() const				\
+    {														\
+        return m_pD;										\
+    }														\
+                                                            \
+    const NODE_TYPE* N::GetType() const						\
+    {														\
+        return GetStaticType();								\
+    }														\
+                                                            \
+    const NODE_TYPE* N::GetStaticType()						\
+    {														\
+        return &Private::s_type;							\
+    }
+
+}
diff --git a/Engine/Plugins/Experimental/Mutable/Source/MutableTools/Private/NodeProjector.cpp b/Engine/Plugins/Experimental/Mutable/Source/MutableTools/Private/NodeProjector.cpp
new file mode 100644
--- /dev/null
+++ b/Engine/Plugins/Experimental/Mutable/Source/MutableTools/Private/NodeProjector.cpp
@@ -0,0 +1,67 @@
+// Copyright Epic Games, Inc. All Rights Reserved.
+
+
+#include "NodeProjector.h"
+
+#include "Node.h"
+#include "NodePrivate.h"
+
+
+namespace mu
+{
+
+	//---------------------------------------------------------------------------------------------
+	// Static initialisation
+	//---------------------------------------------------------------------------------------------
+	static NODE_TYPE s_nodeProjectorType = NODE_TYPE( "NodeProjector", Node::GetStaticType() );
+
+
+	//---------------------------------------------------------------------------------------------
+	void NodeProjector::Serialise( const NodeProjector* p, OutputArchive& arch )
+	{
+        uint32_t ver = 0;
+		arch << ver;
+
+		arch << uint32_t(p->Type);
+		p->SerialiseWrapper(arch);
+	}
+
+        
+	//---------------------------------------------------------------------------------------------
+	NodeProjectorPtr NodeProjector::StaticUnserialise( InputArchive& arch )
+	{
+        uint32_t ver;
+		arch >> ver;
+		check( ver == 0 );
+
+        uint32_t id;
+		arch >> id;
+
+		switch (id)
+		{
+		case 0 :  return NodeProjectorConstant::StaticUnserialise( arch ); break;
+		case 1 :  return NodeProjectorParameter::StaticUnserialise( arch ); break;
+		default : check(false);
+		}
+
+		return 0;
+	}
+
+
+	//---------------------------------------------------------------------------------------------
+	const NODE_TYPE* NodeProjector::GetType() const
+	{
+		return GetStaticType();
+	}
+
+
+	//---------------------------------------------------------------------------------------------
+	const NODE_TYPE* NodeProjector::GetStaticType()
+	{
+		return &s_nodeProjectorType;
+	}
+
+
+}
+
+
diff --git a/Engine/Plugins/Experimental/Mutable/Source/MutableTools/Private/NodeProjectorConstant.cpp b/Engine/Plugins/Experimental/Mutable/Source/MutableTools/Private/NodeProjectorConstant.cpp
new file mode 100644
--- /dev/null
+++ b/Engine/Plugins/Experimental/Mutable/Source/MutableTools/Private/NodeProjectorConstant.cpp
@@ -0,0 +1,101 @@
+// Copyright Epic Games, Inc. All Rights Reserved.
+
+
+#include "NodeProjectorPrivate.h"
+
+#include "MutableMath.h"
+
+
+namespace mu
+{
+
+	//---------------------------------------------------------------------------------------------
+	// Static initialisation
+	//---------------------------------------------------------------------------------------------
+	NODE_TYPE NodeProjectorConstant::Private::s_type =
+			NODE_TYPE( "ProjectorConstant", NodeProjector::GetStaticType() );
+
+
+	//---------------------------------------------------------------------------------------------
+	//!
+	//---------------------------------------------------------------------------------------------
+
+	MUTABLE_IMPLEMENT_NODE( NodeProjectorConstant, EType::Constant, Node, Node::EType::Projector)
+
+
+	//---------------------------------------------------------------------------------------------
+	// Node Interface
+	//---------------------------------------------------------------------------------------------
+	int NodeProjectorConstant::GetInputCount() const
+	{
+		return 0;
+	}
+
+
+	//---------------------------------------------------------------------------------------------
+    Node* NodeProjectorConstant::GetInputNode( int ) const
+	{
+		check( false );
+		return 0;
+	}
+
+
+	//---------------------------------------------------------------------------------------------
+    void NodeProjectorConstant::SetInputNode( int, NodePtr )
+	{
+		check( false );
+	}
+
+
+	//---------------------------------------------------------------------------------------------
+	// Own Interface
+	//---------------------------------------------------------------------------------------------
+    void NodeProjectorConstant::GetValue( PROJECTOR_TYPE* pType,
+                                          float* pPosX, float* pPosY, float* pPosZ,
+										  float* pDirX, float* pDirY, float* pDirZ,
+										  float* pUpX, float* pUpY, float* pUpZ,
+                                          float* pScaleU, float* pScaleV, float* pScaleW,
+                                          float* pProjectionAngle ) const
+	{
+        if (pType) *pType = m_pD->m_type;
+
+        if (pPosX) *pPosX = m_pD->m_position[0];
+        if (pPosY) *pPosY = m_pD->m_position[1];
+        if (pPosZ) *pPosZ = m_pD->m_position[2];
+
+        if (pDirX) *pDirX = m_pD->m_direction[0];
+        if (pDirY) *pDirY = m_pD->m_direction[1];
+        if (pDirZ) *pDirZ = m_pD->m_direction[2];
+
+        if (pUpX) *pUpX = m_pD->m_up[0];
+        if (pUpY) *pUpY = m_pD->m_up[1];
+        if (pUpZ) *pUpZ = m_pD->m_up[2];
+
+        if (pScaleU) *pScaleU = m_pD->m_scale[0];
+        if (pScaleV) *pScaleV = m_pD->m_scale[1];
+        if (pScaleW) *pScaleW = m_pD->m_scale[2];
+
+        if (pProjectionAngle) *pProjectionAngle = m_pD->m_projectionAngle;
+    }
+
+
+	//---------------------------------------------------------------------------------------------
+    void NodeProjectorConstant::SetValue( PROJECTOR_TYPE type,
+                                          float posX, float posY, float posZ,
+										  float dirX, float dirY, float dirZ,
+										  float upX, float upY, float upZ,
+                                          float scaleU, float scaleV, float scaleW,
+                                          float projectionAngle )
+	{
+        m_pD->m_type = type;
+		m_pD->m_position = vec3<float>( posX, posY, posZ );
+		m_pD->m_direction = vec3<float>( dirX, dirY, dirZ );
+		m_pD->m_up = vec3<float>( upX, upY, upZ );
+        m_pD->m_scale = vec3<float>( scaleU, scaleV, scaleW );
+        m_pD->m_projectionAngle = projectionAngle;
+	}
+
+
+}
+
+
diff --git a/Engine/Plugins/Experimental/Mutable/Source/MutableTools/Private/NodeProjectorParameter.cpp b/Engine/Plugins/Experimental/Mutable/Source/MutableTools/Private/NodeProjectorParameter.cpp
new file mode 100644
--- /dev/null
+++ b/Engine/Plugins/Experimental/Mutable/Source/MutableTools/Private/NodeProjectorParameter.cpp
@@ -0,0 +1,169 @@
+// Copyright Epic Games, Inc. All Rights Reserved.
+
+
+#include "NodeProjectorPrivate.h"
+
+#include "MutableMath.h"
+#include "MemoryPrivate.h"
+
+
+namespace mu
+{
+
+	//---------------------------------------------------------------------------------------------
+	// Static initialisation
+	//---------------------------------------------------------------------------------------------
+	NODE_TYPE NodeProjectorParameter::Private::s_type =
+			NODE_TYPE( "ProjectorParameter", NodeProjector::GetStaticType() );
+
+
+	//---------------------------------------------------------------------------------------------
+	//!
+	//---------------------------------------------------------------------------------------------
+
+	MUTABLE_IMPLEMENT_NODE( NodeProjectorParameter, EType::Parameter, Node, Node::EType::Projector)
+
+
+	//---------------------------------------------------------------------------------------------
+	// Node Interface
+	//---------------------------------------------------------------------------------------------
+	int NodeProjectorParameter::GetInputCount() const
+	{
+        return int( m_pD->m_ranges.size() );
+    }
+
+
+	//---------------------------------------------------------------------------------------------
+    Node* NodeProjectorParameter::GetInputNode( int i ) const
+	{
+        check( i<GetInputCount() );
+        if (i<GetInputCount())
+        {
+            return m_pD->m_ranges[i].get();
+        }
+        return nullptr;
+    }
+
+
+	//---------------------------------------------------------------------------------------------
+    void NodeProjectorParameter::SetInputNode( int i, NodePtr n )
+	{
+        check( i<GetInputCount() );
+        if (i<GetInputCount())
+        {
+            m_pD->m_ranges[i] = dynamic_cast<NodeRange*>(n.get());
+        }
+    }
+
+
+	//---------------------------------------------------------------------------------------------
+	// Own Interface
+	//---------------------------------------------------------------------------------------------
+	const char* NodeProjectorParameter::GetName() const
+	{
+		return m_pD->m_name.c_str();
+	}
+
+
+	//---------------------------------------------------------------------------------------------
+	void NodeProjectorParameter::SetName( const char* strName )
+	{
+		if ( strName )
+		{
+			m_pD->m_name = strName;
+		}
+		else
+		{
+			m_pD->m_name = "";
+		}
+	}
+
+
+	const char* NodeProjectorParameter::GetUid() const
+	{
+		return m_pD->m_uid.c_str();
+	}
+
+
+	//---------------------------------------------------------------------------------------------
+	void NodeProjectorParameter::SetUid( const char* strUid )
+	{
+		if ( strUid )
+		{
+			m_pD->m_uid = strUid;
+		}
+		else
+		{
+			m_pD->m_uid = "";
+		}
+	}
+
+
+	//---------------------------------------------------------------------------------------------
+    void NodeProjectorParameter::GetDefaultValue( PROJECTOR_TYPE* pType,
+                                                  float* pPosX, float* pPosY, float* pPosZ,
+												  float* pDirX, float* pDirY, float* pDirZ,
+												  float* pUpX, float* pUpY, float* pUpZ,
+                                                  float* pScaleU, float* pScaleV, float* pScaleW,
+                                                  float* pProjectionAngle ) const
+	{
+        if (pType) *pType = m_pD->m_type;
+
+        if (pPosX) *pPosX = m_pD->m_position[0];
+        if (pPosY) *pPosY = m_pD->m_position[1];
+        if (pPosZ) *pPosZ = m_pD->m_position[2];
+
+        if (pDirX) *pDirX = m_pD->m_direction[0];
+        if (pDirY) *pDirY = m_pD->m_direction[1];
+        if (pDirZ) *pDirZ = m_pD->m_direction[2];
+
+        if (pUpX) *pUpX = m_pD->m_up[0];
+        if (pUpY) *pUpY = m_pD->m_up[1];
+        if (pUpZ) *pUpZ = m_pD->m_up[2];
+
+        if (pScaleU) *pScaleU = m_pD->m_scale[0];
+        if (pScaleV) *pScaleV = m_pD->m_scale[1];
+        if (pScaleW) *pScaleW = m_pD->m_scale[2];
+
+        if (pProjectionAngle) *pProjectionAngle = m_pD->m_projectionAngle;
+    }
+
+
+	//---------------------------------------------------------------------------------------------
+    void NodeProjectorParameter::SetDefaultValue( PROJECTOR_TYPE type,
+                                                  float posX, float posY, float posZ,
+												  float dirX, float dirY, float dirZ,
+												  float upX, float upY, float upZ,
+                                                  float scaleU, float scaleV, float scaleW,
+                                                  float projectionAngle )
+	{
+        m_pD->m_type = type;
+        m_pD->m_position = vec3<float>( posX, posY, posZ );
+		m_pD->m_direction = vec3<float>( dirX, dirY, dirZ );
+		m_pD->m_up = vec3<float>( upX, upY, upZ );
+        m_pD->m_scale = vec3<float>( scaleU, scaleV, scaleW );
+        m_pD->m_projectionAngle = projectionAngle;
+    }
+
+
+    //---------------------------------------------------------------------------------------------
+    void NodeProjectorParameter::SetRangeCount( int i )
+    {
+        check(i>=0);
+        m_pD->m_ranges.resize(i);
+    }
+
+
+    //---------------------------------------------------------------------------------------------
+    void NodeProjectorParameter::SetRange( int i, NodeRangePtr pRange )
+    {
+        check( i>=0 && i<int(m_pD->m_ranges.size()) );
+        if ( i>=0 && i<int(m_pD->m_ranges.size()) )
+        {
+            m_pD->m_ranges[i] = pRange;
+        }
+    }
+
+}
+
+
diff --git a/Engine/Plugins/Experimental/Mutable/Source/MutableTools/Private/NodeProjectorPrivate.h b/Engine/Plugins/Experimental/Mutable/Source/MutableTools/Private/NodeProjectorPrivate.h
new file mode 100644
--- /dev/null
+++ b/Engine/Plugins/Experimental/Mutable/Source/MutableTools/Private/NodeProjectorPrivate.h
@@ -0,0 +1,121 @@
+// Copyright Epic Games, Inc. All Rights Reserved.
+
+#pragma once
+
+#include "NodePrivate.h"
+#include "NodeProjector.h"
+
+#include "NodeRange.h"
+#include "MutableMath.h"
+
+
+namespace mu
+{
+
+
+	class NodeProjector::Private : public Node::Private
+	{
+	};
+
+
+	class NodeProjectorConstant::Private : public NodeProjector::Private
+	{
+	public:
+
+		MUTABLE_DEFINE_CONST_VISITABLE()
+
+	public:
+
+		static NODE_TYPE s_type;
+
+        PROJECTOR_TYPE m_type = PROJECTOR_TYPE::PLANAR;
+		vec3<float> m_position;
+		vec3<float> m_direction;
+		vec3<float> m_up;
+        vec3<float> m_scale;
+        float m_projectionAngle = 0.0f;
+
+		//!
+		void Serialise( OutputArchive& arch ) const
+		{
+            uint32_t ver = 2;
+			arch << ver;
+
+            arch << m_type;
+            arch << m_position;
+            arch << m_direction;
+			arch << m_up;
+			arch << m_scale;
+            arch << m_projectionAngle;
+		}
+
+		//!
+		void Unserialise( InputArchive& arch )
+		{
+            uint32_t ver;
+			arch >> ver;
+            check(ver==2);
+
+            arch >> m_type;
+			arch >> m_position;
+			arch >> m_direction;
+			arch >> m_up;
+            arch >> m_scale;
+            arch >> m_projectionAngle;
+        }
+
+	};
+
+
+	class NodeProjectorParameter::Private : public NodeProjectorConstant::Private
+	{
+	public:
+
+		MUTABLE_DEFINE_CONST_VISITABLE()
+
+	public:
+
+		static NODE_TYPE s_type;
+
+		string m_name;
+		string m_uid;
+
+        vector<Ptr<NodeRange>> m_ranges;
+
+		//!
+		void Serialise( OutputArchive& arch ) const
+		{
+			NodeProjectorConstant::Private::Serialise( arch );
+
+            uint32_t ver = 2;
+			arch << ver;
+
+			arch << m_name;
+            arch << m_uid;
+            arch << m_ranges;
+        }
+
+		//!
+		void Unserialise( InputArchive& arch )
+		{
+			NodeProjectorConstant::Private::Unserialise( arch );
+
+            uint32_t ver;
+			arch >> ver;
+            check(ver<=2);
+
+			arch >> m_name;
+            if (ver>=1)
+            {
+                arch >> m_uid;
+            }
+
+            if (ver>=2)
+            {
+                arch >> m_ranges;
+            }
+        }
+	};
+
+
+}
diff --git a/Engine/Plugins/Experimental/Mutable/Source/MutableTools/Private/NodeRange.cpp b/Engine/Plugins/Experimental/Mutable/Source/MutableTools/Private/NodeRange.cpp
new file mode 100644
--- /dev/null
+++ b/Engine/Plugins/Experimental/Mutable/Source/MutableTools/Private/NodeRange.cpp
@@ -0,0 +1,74 @@
+// Copyright Epic Games, Inc. All Rights Reserved.
+
+
+#include "NodeRange.h"
+#include "Node.h"
+#include "NodePrivate.h"
+#include "NodeRangeFromScalar.h"
+
+
+namespace mu
+{
+
+	//---------------------------------------------------------------------------------------------
+	// Static initialisation
+	//---------------------------------------------------------------------------------------------
+    static NODE_TYPE s_nodeRangeType = 	NODE_TYPE( "NodeRange", Node::GetStaticType() );
+
+
+	//---------------------------------------------------------------------------------------------
+    void NodeRange::Serialise( const NodeRange* p, OutputArchive& arch )
+	{
+        uint32_t ver = 0;
+		arch << ver;
+
+#define SERIALISE_CHILDREN( C, ID ) \
+        ( const C* pTyped##ID = dynamic_cast<const C*>(p) )			\
+        {                                                           \
+            arch << (uint32_t)ID;                                   \
+            C::Serialise( pTyped##ID, arch );						\
+        }                                                           \
+
+        if SERIALISE_CHILDREN( NodeRangeFromScalar          , 0  )
+        else check(false);
+
+#undef SERIALISE_CHILDREN
+    }
+
+    //---------------------------------------------------------------------------------------------
+    NodeRangePtr NodeRange::StaticUnserialise( InputArchive& arch )
+	{
+        uint32_t ver;
+		arch >> ver;
+		check( ver == 0 );
+
+        uint32_t id;
+		arch >> id;
+
+		switch (id)
+		{
+        case 0 :  return NodeRangeFromScalar::StaticUnserialise( arch ); break;
+        default : check(false);
+		}
+
+        return nullptr;
+	}
+
+
+	//---------------------------------------------------------------------------------------------
+    const NODE_TYPE* NodeRange::GetType() const
+	{
+		return GetStaticType();
+	}
+
+
+	//---------------------------------------------------------------------------------------------
+    const NODE_TYPE* NodeRange::GetStaticType()
+	{
+        return &s_nodeRangeType;
+	}
+
+
+}
+
+
diff --git a/Engine/Plugins/Experimental/Mutable/Source/MutableTools/Private/NodeRangeFromScalar.cpp b/Engine/Plugins/Experimental/Mutable/Source/MutableTools/Private/NodeRangeFromScalar.cpp
new file mode 100644
--- /dev/null
+++ b/Engine/Plugins/Experimental/Mutable/Source/MutableTools/Private/NodeRangeFromScalar.cpp
@@ -0,0 +1,108 @@
+// Copyright Epic Games, Inc. All Rights Reserved.
+
+
+#include "NodeRangeFromScalar.h"
+#include "NodeRangeFromScalarPrivate.h"
+
+#include "NodeScalar.h"
+
+#include "ImagePrivate.h"
+
+
+#define NODE_INPUT_COUNT 	1
+
+
+namespace mu
+{
+
+	//---------------------------------------------------------------------------------------------
+	// Static initialisation
+	//---------------------------------------------------------------------------------------------
+    NODE_TYPE NodeRangeFromScalar::Private::s_type =
+            NODE_TYPE( "RangeFromScalar", NodeRange::GetStaticType() );
+
+
+	//---------------------------------------------------------------------------------------------
+	//!
+	//---------------------------------------------------------------------------------------------
+
+    MUTABLE_IMPLEMENT_NODE( NodeRangeFromScalar, EType::FromScalar, Node, Node::EType::Range)
+
+
+	//---------------------------------------------------------------------------------------------
+	// Node Interface
+	//---------------------------------------------------------------------------------------------
+    int NodeRangeFromScalar::GetInputCount() const
+	{
+		return NODE_INPUT_COUNT;
+	}
+
+
+	//---------------------------------------------------------------------------------------------
+    Node* NodeRangeFromScalar::GetInputNode( int i ) const
+	{
+		check( i>=0 && i<NODE_INPUT_COUNT );
+
+        Node* pResult = nullptr;
+
+		switch (i)
+		{
+        case 0: pResult = m_pD->m_pSize.get(); break;
+		}
+
+		return pResult;
+	}
+
+
+	//---------------------------------------------------------------------------------------------
+    void NodeRangeFromScalar::SetInputNode( int i, NodePtr pNode )
+	{
+		check( i>=0 && i<NODE_INPUT_COUNT );
+
+		switch (i)
+		{
+        case 0: m_pD->m_pSize = dynamic_cast<NodeScalar*>(pNode.get()); break;
+		}
+	}
+
+
+	//---------------------------------------------------------------------------------------------
+	// Own Interface
+	//---------------------------------------------------------------------------------------------
+	Ptr<NodeScalar> NodeRangeFromScalar::GetSize() const
+	{
+        return m_pD->m_pSize.get();
+	}
+
+
+	//---------------------------------------------------------------------------------------------
+    void NodeRangeFromScalar::SetSize( const Ptr<NodeScalar>& pNode )
+	{
+        m_pD->m_pSize = pNode;
+	}
+
+
+	//---------------------------------------------------------------------------------------------
+    const char* NodeRangeFromScalar::GetName() const
+	{
+        return m_pD->m_name.c_str();
+	}
+
+
+	//---------------------------------------------------------------------------------------------
+    void NodeRangeFromScalar::SetName( const char* strName )
+	{
+        if (!strName)
+        {
+            m_pD->m_name.clear();
+        }
+        else
+        {
+            m_pD->m_name = strName;
+        }
+	}
+
+}
+
+#undef NODE_INPUT_COUNT
+
diff --git a/Engine/Plugins/Experimental/Mutable/Source/MutableTools/Private/NodeRangeFromScalarPrivate.h b/Engine/Plugins/Experimental/Mutable/Source/MutableTools/Private/NodeRangeFromScalarPrivate.h
new file mode 100644
--- /dev/null
+++ b/Engine/Plugins/Experimental/Mutable/Source/MutableTools/Private/NodeRangeFromScalarPrivate.h
@@ -0,0 +1,51 @@
+// Copyright Epic Games, Inc. All Rights Reserved.
+
+#pragma once
+
+#include "NodeRangePrivate.h"
+#include "NodeRangeFromScalar.h"
+#include "NodeRange.h"
+#include "NodeScalar.h"
+
+
+namespace mu
+{
+
+
+    class NodeRangeFromScalar::Private : public NodeRange::Private
+	{
+	public:
+
+		MUTABLE_DEFINE_CONST_VISITABLE()
+
+	public:
+
+		static NODE_TYPE s_type;
+
+        NodeScalarPtr m_pSize;
+        string m_name;
+
+		//!
+		void Serialise( OutputArchive& arch ) const
+		{
+            uint32_t ver = 0;
+			arch << ver;
+
+            arch << m_name;
+            arch << m_pSize;
+		}
+
+		//!
+		void Unserialise( InputArchive& arch )
+		{
+            uint32_t ver;
+			arch >> ver;
+			check(ver==0);
+
+            arch >> m_name;
+            arch >> m_pSize;
+		}
+	};
+
+
+}
diff --git a/Engine/Plugins/Experimental/Mutable/Source/MutableTools/Private/NodeRangePrivate.h b/Engine/Plugins/Experimental/Mutable/Source/MutableTools/Private/NodeRangePrivate.h
new file mode 100644
--- /dev/null
+++ b/Engine/Plugins/Experimental/Mutable/Source/MutableTools/Private/NodeRangePrivate.h
@@ -0,0 +1,19 @@
+// Copyright Epic Games, Inc. All Rights Reserved.
+
+#pragma once
+
+#include "NodePrivate.h"
+#include "NodeRange.h"
+
+#include "MemoryPrivate.h"
+
+
+namespace mu
+{
+
+
+    class NodeRange::Private : public Node::Private
+	{
+	};
+
+}
diff --git a/Engine/Plugins/Experimental/Mutable/Source/MutableTools/Private/NodeScalar.cpp b/Engine/Plugins/Experimental/Mutable/Source/MutableTools/Private/NodeScalar.cpp
new file mode 100644
--- /dev/null
+++ b/Engine/Plugins/Experimental/Mutable/Source/MutableTools/Private/NodeScalar.cpp
@@ -0,0 +1,82 @@
+// Copyright Epic Games, Inc. All Rights Reserved.
+
+
+#include "NodeScalar.h"
+
+#include "Node.h"
+#include "NodePrivate.h"
+
+#include "NodeScalarConstant.h"
+#include "NodeScalarParameter.h"
+#include "NodeScalarEnumParameter.h"
+#include "NodeScalarCurve.h"
+#include "NodeScalarSwitch.h"
+#include "NodeScalarArithmeticOperation.h"
+#include "NodeScalarVariation.h"
+#include "NodeScalarTable.h"
+
+
+namespace mu
+{
+
+	//---------------------------------------------------------------------------------------------
+	// Static initialisation
+	//---------------------------------------------------------------------------------------------
+	static NODE_TYPE s_nodeScalarType = NODE_TYPE( "NodeScalar", Node::GetStaticType() );
+
+
+	//---------------------------------------------------------------------------------------------
+	void NodeScalar::Serialise( const NodeScalar* p, OutputArchive& arch )
+	{
+        uint32_t ver = 0;
+		arch << ver;
+
+		arch << uint32_t(p->Type);
+		p->SerialiseWrapper(arch);
+    }
+
+
+    //---------------------------------------------------------------------------------------------
+	NodeScalarPtr NodeScalar::StaticUnserialise( InputArchive& arch )
+	{
+        uint32_t ver;
+		arch >> ver;
+		check( ver == 0 );
+
+        uint32_t id;
+		arch >> id;
+
+		switch (id)
+		{
+		case 0 :  return NodeScalarConstant::StaticUnserialise( arch ); break;
+		case 1 :  return NodeScalarParameter::StaticUnserialise( arch ); break;
+		case 2 :  return NodeScalarEnumParameter::StaticUnserialise( arch ); break;
+        case 3 :  return NodeScalarCurve::StaticUnserialise( arch ); break;
+        case 4 :  return NodeScalarSwitch::StaticUnserialise( arch ); break;
+        case 5 :  return NodeScalarArithmeticOperation::StaticUnserialise( arch ); break;
+        case 6 :  return NodeScalarVariation::StaticUnserialise( arch ); break;
+		case 7 :  return NodeScalarTable::StaticUnserialise( arch ); break;
+        default : check(false);
+		}
+
+		return 0;
+	}
+
+
+	//---------------------------------------------------------------------------------------------
+	const NODE_TYPE* NodeScalar::GetType() const
+	{
+		return GetStaticType();
+	}
+
+
+	//---------------------------------------------------------------------------------------------
+	const NODE_TYPE* NodeScalar::GetStaticType()
+	{
+		return &s_nodeScalarType;
+	}
+
+
+}
+
+
diff --git a/Engine/Plugins/Experimental/Mutable/Source/MutableTools/Private/NodeScalarArithmeticOperation.cpp b/Engine/Plugins/Experimental/Mutable/Source/MutableTools/Private/NodeScalarArithmeticOperation.cpp
new file mode 100644
--- /dev/null
+++ b/Engine/Plugins/Experimental/Mutable/Source/MutableTools/Private/NodeScalarArithmeticOperation.cpp
@@ -0,0 +1,126 @@
+// Copyright Epic Games, Inc. All Rights Reserved.
+
+
+#include "NodeScalarArithmeticOperation.h"
+#include "NodeScalarArithmeticOperationPrivate.h"
+
+#include "NodeScalar.h"
+
+#include "ImagePrivate.h"
+
+
+#define NODE_INPUT_COUNT 	2
+
+
+namespace mu
+{
+
+	//---------------------------------------------------------------------------------------------
+	// Static initialisation
+	//---------------------------------------------------------------------------------------------
+    NODE_TYPE NodeScalarArithmeticOperation::Private::s_type =
+            NODE_TYPE( "ScalarArithmenticOperation", NodeScalar::GetStaticType() );
+
+
+	//---------------------------------------------------------------------------------------------
+	//!
+	//---------------------------------------------------------------------------------------------
+
+    MUTABLE_IMPLEMENT_NODE( NodeScalarArithmeticOperation, EType::ArithmeticOperation, Node, Node::EType::Scalar)
+
+
+    //---------------------------------------------------------------------------------------------
+    const char* NodeScalarArithmeticOperation::s_opTypeName[] =
+    {
+        "Add",
+        "Subtract",
+        "Multiply",
+        "Divide"
+    };
+
+
+	//---------------------------------------------------------------------------------------------
+	// Node Interface
+	//---------------------------------------------------------------------------------------------
+    int NodeScalarArithmeticOperation::GetInputCount() const
+	{
+		return NODE_INPUT_COUNT;
+	}
+
+
+	//---------------------------------------------------------------------------------------------
+    Node* NodeScalarArithmeticOperation::GetInputNode( int i ) const
+	{
+		check( i>=0 && i<NODE_INPUT_COUNT );
+
+		Node* pResult = 0;
+
+		switch (i)
+		{
+		case 0: pResult = m_pD->m_pA.get(); break;
+		case 1: pResult = m_pD->m_pB.get(); break;
+		}
+
+		return pResult;
+	}
+
+
+	//---------------------------------------------------------------------------------------------
+    void NodeScalarArithmeticOperation::SetInputNode( int i, NodePtr pNode )
+	{
+		check( i>=0 && i<NODE_INPUT_COUNT );
+
+		switch (i)
+		{
+        case 0: m_pD->m_pA = dynamic_cast<NodeScalar*>(pNode.get()); break;
+        case 1: m_pD->m_pB = dynamic_cast<NodeScalar*>(pNode.get()); break;
+		}
+	}
+
+
+	//---------------------------------------------------------------------------------------------
+	// Own Interface
+	//---------------------------------------------------------------------------------------------
+    NodeScalarArithmeticOperation::OPERATION NodeScalarArithmeticOperation::GetOperation() const
+	{
+		return m_pD->m_operation;
+	}
+
+
+	//---------------------------------------------------------------------------------------------
+    void NodeScalarArithmeticOperation::SetOperation(NodeScalarArithmeticOperation::OPERATION o)
+	{
+		m_pD->m_operation = o;
+	}
+
+
+    NodeScalarPtr NodeScalarArithmeticOperation::GetA() const
+	{
+		return m_pD->m_pA.get();
+	}
+
+
+	//---------------------------------------------------------------------------------------------
+    void NodeScalarArithmeticOperation::SetA(NodeScalarPtr pNode)
+	{
+		m_pD->m_pA = pNode;
+	}
+
+
+	//---------------------------------------------------------------------------------------------
+    NodeScalarPtr NodeScalarArithmeticOperation::GetB() const
+	{
+		return m_pD->m_pB.get();
+	}
+
+
+	//---------------------------------------------------------------------------------------------
+    void NodeScalarArithmeticOperation::SetB( NodeScalarPtr pNode )
+	{
+		m_pD->m_pB = pNode;
+	}
+
+}
+
+#undef NODE_INPUT_COUNT
+
diff --git a/Engine/Plugins/Experimental/Mutable/Source/MutableTools/Private/NodeScalarArithmeticOperationPrivate.h b/Engine/Plugins/Experimental/Mutable/Source/MutableTools/Private/NodeScalarArithmeticOperationPrivate.h
new file mode 100644
--- /dev/null
+++ b/Engine/Plugins/Experimental/Mutable/Source/MutableTools/Private/NodeScalarArithmeticOperationPrivate.h
@@ -0,0 +1,55 @@
+// Copyright Epic Games, Inc. All Rights Reserved.
+
+#pragma once
+
+#include "NodeScalarPrivate.h"
+#include "NodeScalarArithmeticOperation.h"
+
+#include "MemoryPrivate.h"
+
+
+namespace mu
+{
+
+    MUTABLE_DEFINE_ENUM_SERIALISABLE(NodeScalarArithmeticOperation::OPERATION)
+
+    class NodeScalarArithmeticOperation::Private : public NodeScalar::Private
+	{
+	public:
+
+		MUTABLE_DEFINE_CONST_VISITABLE()
+
+	public:
+
+		static NODE_TYPE s_type;
+
+        NodeScalarArithmeticOperation::OPERATION m_operation;
+        NodeScalarPtr m_pA;
+        NodeScalarPtr m_pB;
+
+		//!
+		void Serialise( OutputArchive& arch ) const
+		{
+            uint32_t ver = 0;
+			arch << ver;
+
+			arch << m_operation;
+			arch << m_pA;
+			arch << m_pB;
+		}
+
+		//!
+		void Unserialise( InputArchive& arch )
+		{
+            uint32_t ver;
+			arch >> ver;
+			check(ver==0);
+
+			arch >> m_operation;
+			arch >> m_pA;
+			arch >> m_pB;
+		}
+	};
+
+
+}
diff --git a/Engine/Plugins/Experimental/Mutable/Source/MutableTools/Private/NodeScalarConstant.cpp b/Engine/Plugins/Experimental/Mutable/Source/MutableTools/Private/NodeScalarConstant.cpp
new file mode 100644
--- /dev/null
+++ b/Engine/Plugins/Experimental/Mutable/Source/MutableTools/Private/NodeScalarConstant.cpp
@@ -0,0 +1,72 @@
+// Copyright Epic Games, Inc. All Rights Reserved.
+
+
+#include "NodeScalarConstant.h"
+#include "NodeScalarConstantPrivate.h"
+
+#include "MutableMath.h"
+
+
+#define NODE_INPUT_COUNT 	0
+
+
+namespace mu
+{
+
+	//---------------------------------------------------------------------------------------------
+	// Static initialisation
+	//---------------------------------------------------------------------------------------------
+	NODE_TYPE NodeScalarConstant::Private::s_type =
+			NODE_TYPE( "ScalarConstant", NodeScalar::GetStaticType() );
+
+
+	//---------------------------------------------------------------------------------------------
+	//!
+	//---------------------------------------------------------------------------------------------
+
+	MUTABLE_IMPLEMENT_NODE( NodeScalarConstant, EType::Constant, Node, Node::EType::Scalar)
+
+
+	//---------------------------------------------------------------------------------------------
+	// Node Interface
+	//---------------------------------------------------------------------------------------------
+	int NodeScalarConstant::GetInputCount() const
+	{
+		return NODE_INPUT_COUNT;
+	}
+
+
+	//---------------------------------------------------------------------------------------------
+    Node* NodeScalarConstant::GetInputNode( int ) const
+	{
+		check( false );
+		return 0;
+	}
+
+
+	//---------------------------------------------------------------------------------------------
+    void NodeScalarConstant::SetInputNode( int, NodePtr )
+	{
+		check( false );
+	}
+
+
+	//---------------------------------------------------------------------------------------------
+	// Own Interface
+	//---------------------------------------------------------------------------------------------
+	float NodeScalarConstant::GetValue() const
+	{
+		return m_pD->m_value;
+	}
+
+	//---------------------------------------------------------------------------------------------
+	void NodeScalarConstant::SetValue( float v )
+	{
+		m_pD->m_value = v;
+	}
+
+
+}
+
+#undef NODE_INPUT_COUNT
+
diff --git a/Engine/Plugins/Experimental/Mutable/Source/MutableTools/Private/NodeScalarConstantPrivate.h b/Engine/Plugins/Experimental/Mutable/Source/MutableTools/Private/NodeScalarConstantPrivate.h
new file mode 100644
--- /dev/null
+++ b/Engine/Plugins/Experimental/Mutable/Source/MutableTools/Private/NodeScalarConstantPrivate.h
@@ -0,0 +1,51 @@
+// Copyright Epic Games, Inc. All Rights Reserved.
+
+#pragma once
+
+#include "NodeScalarPrivate.h"
+
+#include "NodeScalarConstant.h"
+#include "NodeImage.h"
+#include "AST.h"
+#include "MemoryPrivate.h"
+#include "ParametersPrivate.h"
+
+
+namespace mu
+{
+
+
+	class NodeScalarConstant::Private : public NodeScalar::Private
+	{
+	public:
+
+		MUTABLE_DEFINE_CONST_VISITABLE()
+
+	public:
+
+		static NODE_TYPE s_type;
+
+		float m_value;
+
+		//!
+		void Serialise( OutputArchive& arch ) const
+		{
+            uint32_t ver = 0;
+			arch << ver;
+
+			arch << m_value;
+		}
+
+		//!
+		void Unserialise( InputArchive& arch )
+		{
+            uint32_t ver;
+			arch >> ver;
+			check(ver==0);
+
+			arch >> m_value;
+		}
+	};
+
+
+}
diff --git a/Engine/Plugins/Experimental/Mutable/Source/MutableTools/Private/NodeScalarCurve.cpp b/Engine/Plugins/Experimental/Mutable/Source/MutableTools/Private/NodeScalarCurve.cpp
new file mode 100644
--- /dev/null
+++ b/Engine/Plugins/Experimental/Mutable/Source/MutableTools/Private/NodeScalarCurve.cpp
@@ -0,0 +1,123 @@
+// Copyright Epic Games, Inc. All Rights Reserved.
+
+
+#include "NodeScalarCurve.h"
+#include "NodeScalarCurvePrivate.h"
+
+#include "NodeScalar.h"
+
+
+#define NODE_INPUT_COUNT 	1
+
+
+namespace mu
+{
+
+
+	//---------------------------------------------------------------------------------------------
+	// Static initialisation
+	//---------------------------------------------------------------------------------------------
+    NODE_TYPE NodeScalarCurve::Private::s_type =
+            NODE_TYPE( "Curve", NodeScalarCurve::GetStaticType() );
+
+
+	//---------------------------------------------------------------------------------------------
+	//---------------------------------------------------------------------------------------------
+	//---------------------------------------------------------------------------------------------
+
+    MUTABLE_IMPLEMENT_NODE( NodeScalarCurve, EType::Curve, Node, Node::EType::Scalar)
+
+
+	//---------------------------------------------------------------------------------------------
+	// Node Interface
+	//---------------------------------------------------------------------------------------------
+    int NodeScalarCurve::GetInputCount() const
+	{
+		return NODE_INPUT_COUNT;
+	}
+
+
+	//---------------------------------------------------------------------------------------------
+    Node* NodeScalarCurve::GetInputNode( int i ) const
+	{
+		check( i >=0 && i < GetInputCount() );
+        (void)i;
+
+		return m_pD->m_input_scalar.get();
+	}
+
+
+	//---------------------------------------------------------------------------------------------
+    void NodeScalarCurve::SetInputNode( int i, NodePtr pNode )
+	{
+		check( i >=0 && i < GetInputCount());
+        (void)i;
+
+		m_pD->m_input_scalar = dynamic_cast<NodeScalar*>(pNode.get());
+	}
+
+
+	//---------------------------------------------------------------------------------------------
+	// Own Interface
+	//---------------------------------------------------------------------------------------------
+    int NodeScalarCurve::GetKeyFrameCount() const
+	{
+		return (int)m_pD->m_curve.keyFrames.size();
+	}
+
+
+	//---------------------------------------------------------------------------------------------
+    void NodeScalarCurve::SetKeyFrameCount( int num )
+	{
+		check( num >=0 );
+		m_pD->m_curve.keyFrames.resize( num );
+	}
+
+	//---------------------------------------------------------------------------------------------
+    void NodeScalarCurve::SetDefaultValue(float defaultValue)
+	{
+		m_pD->m_curve.defaultValue = defaultValue;
+	}
+
+
+	//---------------------------------------------------------------------------------------------
+    void NodeScalarCurve::SetKeyFrame(
+            int index, float time, float value,
+            float in_tangent, float in_tangent_weight,
+            float out_tangent, float out_tangent_weight,
+            uint8_t interp_mode, uint8_t tangent_mode, uint8_t tangent_weight_mode )
+	{
+		check( index >=0 && index < (int)m_pD->m_curve.keyFrames.size() );
+
+		CurveKeyFrame& key_frame = m_pD->m_curve.keyFrames[index];
+
+		key_frame.time = time;
+		key_frame.value = value;
+		key_frame.in_tangent = in_tangent;
+		key_frame.in_tangent_weight = in_tangent_weight;
+		key_frame.out_tangent = out_tangent;
+		key_frame.out_tangent_weight = out_tangent_weight;
+		
+		key_frame.interp_mode = interp_mode;
+		key_frame.tangent_mode = tangent_mode;
+		key_frame.tangent_weight_mode = tangent_weight_mode;
+	}
+
+
+	//---------------------------------------------------------------------------------------------
+    NodeScalarPtr NodeScalarCurve::GetT() const
+	{
+		return m_pD->m_input_scalar.get();
+	}
+
+
+	//---------------------------------------------------------------------------------------------
+    void NodeScalarCurve::SetT(NodeScalarPtr pNode)
+	{
+		m_pD->m_input_scalar = pNode;
+	}
+
+}
+
+#undef NODE_INPUT_COUNT
+
diff --git a/Engine/Plugins/Experimental/Mutable/Source/MutableTools/Private/NodeScalarCurvePrivate.h b/Engine/Plugins/Experimental/Mutable/Source/MutableTools/Private/NodeScalarCurvePrivate.h
new file mode 100644
--- /dev/null
+++ b/Engine/Plugins/Experimental/Mutable/Source/MutableTools/Private/NodeScalarCurvePrivate.h
@@ -0,0 +1,56 @@
+// Copyright Epic Games, Inc. All Rights Reserved.
+
+#pragma once
+
+#include "NodePrivate.h"
+#include "NodeScalarCurve.h"
+#include "NodeScalar.h"
+#include "AST.h"
+
+#include "MemoryPrivate.h"
+#include "SerialisationPrivate.h"
+#include "ParametersPrivate.h"
+
+
+namespace mu
+{
+	//---------------------------------------------------------------------------------------------
+	//!
+	//---------------------------------------------------------------------------------------------
+    class NodeScalarCurve::Private : public Node::Private
+	{
+	public:
+
+		MUTABLE_DEFINE_CONST_VISITABLE()
+
+	public:
+
+		static NODE_TYPE s_type;
+
+		NodeScalarPtr m_input_scalar;
+		Curve m_curve;
+
+		//!
+		void Serialise(OutputArchive& arch) const
+		{
+            uint32_t ver = 1;
+			arch << ver;
+
+			arch << m_input_scalar;
+			arch << m_curve;
+		}
+
+		//!
+		void Unserialise(InputArchive& arch)
+		{
+            uint32_t ver;
+			arch >> ver;
+            check(ver == 1);
+
+			arch >> m_input_scalar;
+			arch >> m_curve;
+		}
+	};
+
+
+}
diff --git a/Engine/Plugins/Experimental/Mutable/Source/MutableTools/Private/NodeScalarEnumParameter.cpp b/Engine/Plugins/Experimental/Mutable/Source/MutableTools/Private/NodeScalarEnumParameter.cpp
new file mode 100644
--- /dev/null
+++ b/Engine/Plugins/Experimental/Mutable/Source/MutableTools/Private/NodeScalarEnumParameter.cpp
@@ -0,0 +1,156 @@
+// Copyright Epic Games, Inc. All Rights Reserved.
+
+
+#include "NodeScalarEnumParameter.h"
+#include "NodeScalarEnumParameterPrivate.h"
+
+#include "MutableMath.h"
+#include "MemoryPrivate.h"
+
+
+namespace mu
+{
+
+	//---------------------------------------------------------------------------------------------
+	// Static initialisation
+	//---------------------------------------------------------------------------------------------
+	NODE_TYPE NodeScalarEnumParameter::Private::s_type =
+			NODE_TYPE( "ScalarEnumParameter", NodeScalar::GetStaticType() );
+
+
+	//---------------------------------------------------------------------------------------------
+	//!
+	//---------------------------------------------------------------------------------------------
+
+	MUTABLE_IMPLEMENT_NODE( NodeScalarEnumParameter, EType::EnumParameter, Node, Node::EType::Scalar)
+
+
+	//---------------------------------------------------------------------------------------------
+	int NodeScalarEnumParameter::GetInputCount() const
+	{
+        return int( m_pD->m_ranges.size() );
+	}
+
+
+	//---------------------------------------------------------------------------------------------
+    Node* NodeScalarEnumParameter::GetInputNode( int i ) const
+	{
+        check( i<GetInputCount() );
+        if (i<GetInputCount())
+        {
+            return m_pD->m_ranges[i].get();
+        }
+        return nullptr;
+	}
+
+
+	//---------------------------------------------------------------------------------------------
+    void NodeScalarEnumParameter::SetInputNode( int i, NodePtr n )
+	{
+        check( i<GetInputCount() );
+        if (i<GetInputCount())
+        {
+            m_pD->m_ranges[i] = dynamic_cast<NodeRange*>(n.get());
+        }
+	}
+
+
+	//---------------------------------------------------------------------------------------------
+	const char* NodeScalarEnumParameter::GetName() const
+	{
+		return m_pD->m_name.c_str();
+	}
+
+
+	//---------------------------------------------------------------------------------------------
+	void NodeScalarEnumParameter::SetName( const char* strName )
+	{
+		if ( strName )
+		{
+			m_pD->m_name = strName;
+		}
+		else
+		{
+			m_pD->m_name = "";
+		}
+	}
+
+
+	const char* NodeScalarEnumParameter::GetUid() const
+	{
+		return m_pD->m_uid.c_str();
+	}
+
+
+	//---------------------------------------------------------------------------------------------
+	void NodeScalarEnumParameter::SetUid( const char* strUid )
+	{
+		if ( strUid )
+		{
+			m_pD->m_uid = strUid;
+		}
+		else
+		{
+			m_pD->m_uid = "";
+		}
+	}
+
+
+	//---------------------------------------------------------------------------------------------
+	int NodeScalarEnumParameter::GetDefaultValueIndex() const
+	{
+		return m_pD->m_defaultValue;
+	}
+
+
+	//---------------------------------------------------------------------------------------------
+	void NodeScalarEnumParameter::SetDefaultValueIndex( int i )
+	{
+		m_pD->m_defaultValue = i;
+	}
+
+
+	//---------------------------------------------------------------------------------------------
+	void NodeScalarEnumParameter::SetValueCount( int i )
+	{
+		m_pD->m_options.resize(i);
+	}
+
+
+	//---------------------------------------------------------------------------------------------
+	int NodeScalarEnumParameter::GetValueCount() const
+	{
+		return (int)m_pD->m_options.size();
+	}
+
+
+	//---------------------------------------------------------------------------------------------
+	void NodeScalarEnumParameter::SetValue( int i, float value, const char* strName )
+	{
+		check( i>=0 && i<(int)m_pD->m_options.size() );
+		m_pD->m_options[i].name = strName;
+		m_pD->m_options[i].value = value;
+	}
+
+
+    //---------------------------------------------------------------------------------------------
+    void NodeScalarEnumParameter::SetRangeCount( int i )
+    {
+        check(i>=0);
+        m_pD->m_ranges.resize(i);
+    }
+
+
+    //---------------------------------------------------------------------------------------------
+    void NodeScalarEnumParameter::SetRange( int i, NodeRangePtr pRange )
+    {
+        check( i>=0 && i<int(m_pD->m_ranges.size()) );
+        if ( i>=0 && i<int(m_pD->m_ranges.size()) )
+        {
+            m_pD->m_ranges[i] = pRange;
+        }
+    }
+
+}
+
+
diff --git a/Engine/Plugins/Experimental/Mutable/Source/MutableTools/Private/NodeScalarEnumParameterPrivate.h b/Engine/Plugins/Experimental/Mutable/Source/MutableTools/Private/NodeScalarEnumParameterPrivate.h
new file mode 100644
--- /dev/null
+++ b/Engine/Plugins/Experimental/Mutable/Source/MutableTools/Private/NodeScalarEnumParameterPrivate.h
@@ -0,0 +1,94 @@
+// Copyright Epic Games, Inc. All Rights Reserved.
+
+#pragma once
+
+#include "NodeScalarPrivate.h"
+
+#include "NodeScalarEnumParameter.h"
+#include "NodeImage.h"
+#include "MemoryPrivate.h"
+#include "ParametersPrivate.h"
+
+
+namespace mu
+{
+
+
+	class NodeScalarEnumParameter::Private : public NodeScalar::Private
+	{
+	public:
+
+		MUTABLE_DEFINE_CONST_VISITABLE()
+
+	public:
+
+		static NODE_TYPE s_type;
+
+		string m_name;
+		string m_uid;
+		int m_defaultValue = 0;
+		PARAMETER_DETAILED_TYPE m_detailedType = PARAMETER_DETAILED_TYPE::UNKNOWN;
+
+		struct OPTION
+		{
+			string name;
+			float value;
+
+			//!
+			void Serialise( OutputArchive& arch ) const
+			{
+                uint32_t ver = 0;
+				arch << ver;
+
+				arch << name;
+				arch << value;
+			}
+
+			//!
+			void Unserialise( InputArchive& arch )
+			{
+                uint32_t ver;
+				arch >> ver;
+                check(ver<=0);
+
+				arch >> name;
+				arch >> value;
+			}
+		};
+
+		vector<OPTION> m_options;
+
+        vector<Ptr<NodeRange>> m_ranges;
+
+		//!
+		void Serialise( OutputArchive& arch ) const
+		{
+            uint32_t ver = 2;
+			arch << ver;
+
+			arch << m_name;
+			arch << m_uid;
+			arch << m_defaultValue;
+			arch << m_detailedType;
+            arch << m_options;
+            arch << m_ranges;
+        }
+
+		//!
+		void Unserialise( InputArchive& arch )
+		{
+            uint32_t ver;
+			arch >> ver;
+            check(ver==2);
+
+			arch >> m_name;
+            arch >> m_uid;
+			arch >> m_defaultValue;
+			arch >> m_detailedType;
+			arch >> m_options;
+            arch >> m_ranges;
+        }
+	};
+
+}
+
diff --git a/Engine/Plugins/Experimental/Mutable/Source/MutableTools/Private/NodeScalarParameter.cpp b/Engine/Plugins/Experimental/Mutable/Source/MutableTools/Private/NodeScalarParameter.cpp
new file mode 100644
--- /dev/null
+++ b/Engine/Plugins/Experimental/Mutable/Source/MutableTools/Private/NodeScalarParameter.cpp
@@ -0,0 +1,181 @@
+// Copyright Epic Games, Inc. All Rights Reserved.
+
+
+#include "NodeScalarParameter.h"
+#include "NodeScalarParameterPrivate.h"
+
+#include "MutableMath.h"
+#include "MemoryPrivate.h"
+
+
+namespace mu
+{
+
+	//---------------------------------------------------------------------------------------------
+	// Static initialisation
+	//---------------------------------------------------------------------------------------------
+	NODE_TYPE NodeScalarParameter::Private::s_type =
+			NODE_TYPE( "ScalarParameter", NodeScalar::GetStaticType() );
+
+
+	//---------------------------------------------------------------------------------------------
+	//!
+	//---------------------------------------------------------------------------------------------
+
+	MUTABLE_IMPLEMENT_NODE( NodeScalarParameter, EType::Parameter, Node, Node::EType::Scalar)
+
+
+	//---------------------------------------------------------------------------------------------
+	int NodeScalarParameter::GetInputCount() const
+	{
+        return int( m_pD->m_additionalImages.size()
+                    +
+                    m_pD->m_ranges.size() );
+	}
+
+
+	//---------------------------------------------------------------------------------------------
+    Node* NodeScalarParameter::GetInputNode( int i ) const
+	{
+        check( i<GetInputCount() );
+
+        int imageCount = int(m_pD->m_additionalImages.size());
+        int rangeCount = int(m_pD->m_ranges.size());
+        if (i<imageCount)
+        {
+            return m_pD->m_additionalImages[i].get();
+        }
+        else if ( i < imageCount + rangeCount )
+        {
+            int r = i - imageCount;
+            return m_pD->m_ranges[r].get();
+        }
+        return nullptr;
+	}
+
+
+	//---------------------------------------------------------------------------------------------
+    void NodeScalarParameter::SetInputNode( int i, NodePtr n )
+	{
+        check( i<GetInputCount() );
+        int imageCount = int(m_pD->m_additionalImages.size());
+        int rangeCount = int(m_pD->m_ranges.size());
+        if (i<imageCount)
+        {
+            m_pD->m_additionalImages[i] = dynamic_cast<NodeImage*>(n.get());
+        }
+        else if ( i < imageCount + rangeCount )
+        {
+            int r = i - imageCount;
+            m_pD->m_ranges[r] = dynamic_cast<NodeRange*>(n.get());
+        }
+    }
+
+
+	//---------------------------------------------------------------------------------------------
+	const char* NodeScalarParameter::GetName() const
+	{
+		return m_pD->m_name.c_str();
+	}
+
+
+	//---------------------------------------------------------------------------------------------
+	void NodeScalarParameter::SetName( const char* strName )
+	{
+		if ( strName )
+		{
+			m_pD->m_name = strName;
+		}
+		else
+		{
+			m_pD->m_name = "";
+		}
+	}
+
+
+	const char* NodeScalarParameter::GetUid() const
+	{
+		return m_pD->m_uid.c_str();
+	}
+
+
+	//---------------------------------------------------------------------------------------------
+	void NodeScalarParameter::SetUid( const char* strUid )
+	{
+		if ( strUid )
+		{
+			m_pD->m_uid = strUid;
+		}
+		else
+		{
+			m_pD->m_uid = "";
+		}
+	}
+
+
+	//---------------------------------------------------------------------------------------------
+	float NodeScalarParameter::GetDefaultValue() const
+	{
+		return m_pD->m_defaultValue;
+	}
+
+
+	//---------------------------------------------------------------------------------------------
+	void NodeScalarParameter::SetDefaultValue( float v )
+	{
+		m_pD->m_defaultValue = v;
+	}
+
+
+	//---------------------------------------------------------------------------------------------
+	PARAMETER_DETAILED_TYPE NodeScalarParameter::GetDetailedType() const
+	{
+		return m_pD->m_detailedType;
+	}
+
+
+	//---------------------------------------------------------------------------------------------
+	void NodeScalarParameter::SetDetailedType( PARAMETER_DETAILED_TYPE t )
+	{
+		m_pD->m_detailedType = t;
+	}
+
+
+	//---------------------------------------------------------------------------------------------
+	void NodeScalarParameter::SetDescImageCount( int i )
+	{
+        check(i>=0);
+        m_pD->m_additionalImages.resize(i);
+	}
+
+
+	//---------------------------------------------------------------------------------------------
+	void NodeScalarParameter::SetDescImage( int i, NodeImagePtr pImage )
+	{
+		check( i>=0 && i<(int)m_pD->m_additionalImages.size() );
+		m_pD->m_additionalImages[i] = pImage;
+	}
+
+
+    //---------------------------------------------------------------------------------------------
+    void NodeScalarParameter::SetRangeCount( int i )
+    {
+        check(i>=0);
+        m_pD->m_ranges.resize(i);
+    }
+
+
+    //---------------------------------------------------------------------------------------------
+    void NodeScalarParameter::SetRange( int i, NodeRangePtr pRange )
+    {
+        check( i>=0 && i<int(m_pD->m_ranges.size()) );
+        if ( i>=0 && i<int(m_pD->m_ranges.size()) )
+        {
+            m_pD->m_ranges[i] = pRange;
+        }
+    }
+
+
+}
+
+
diff --git a/Engine/Plugins/Experimental/Mutable/Source/MutableTools/Private/NodeScalarParameterPrivate.h b/Engine/Plugins/Experimental/Mutable/Source/MutableTools/Private/NodeScalarParameterPrivate.h
new file mode 100644
--- /dev/null
+++ b/Engine/Plugins/Experimental/Mutable/Source/MutableTools/Private/NodeScalarParameterPrivate.h
@@ -0,0 +1,67 @@
+// Copyright Epic Games, Inc. All Rights Reserved.
+
+#pragma once
+
+#include "NodeScalarPrivate.h"
+
+#include "NodeScalarParameter.h"
+#include "NodeImage.h"
+#include "NodeRange.h"
+#include "MemoryPrivate.h"
+#include "ParametersPrivate.h"
+
+
+namespace mu
+{
+
+
+	class NodeScalarParameter::Private : public NodeScalar::Private
+	{
+	public:
+
+		MUTABLE_DEFINE_CONST_VISITABLE()
+
+	public:
+
+		static NODE_TYPE s_type;
+
+		float m_defaultValue = 0.0f;
+		string m_name;
+		string m_uid;
+		PARAMETER_DETAILED_TYPE m_detailedType = PARAMETER_DETAILED_TYPE::UNKNOWN;
+
+        vector<Ptr<NodeImage>> m_additionalImages;
+
+        vector<Ptr<NodeRange>> m_ranges;
+
+		//!
+		void Serialise( OutputArchive& arch ) const
+		{
+            uint32_t ver = 4;
+			arch << ver;
+
+			arch << m_defaultValue;
+			arch << m_name;
+			arch << m_uid;
+			arch << m_additionalImages;
+			arch << m_detailedType;
+            arch << m_ranges;
+        }
+
+		//!
+		void Unserialise( InputArchive& arch )
+		{
+            uint32_t ver;
+			arch >> ver;
+            check(ver==4);
+
+			arch >> m_defaultValue;
+			arch >> m_name;
+            arch >> m_uid;
+			arch >> m_additionalImages;
+            arch >> m_detailedType;
+            arch >> m_ranges;
+        }
+	};
+
+}
diff --git a/Engine/Plugins/Experimental/Mutable/Source/MutableTools/Private/NodeScalarPrivate.h b/Engine/Plugins/Experimental/Mutable/Source/MutableTools/Private/NodeScalarPrivate.h
new file mode 100644
--- /dev/null
+++ b/Engine/Plugins/Experimental/Mutable/Source/MutableTools/Private/NodeScalarPrivate.h
@@ -0,0 +1,21 @@
+// Copyright Epic Games, Inc. All Rights Reserved.
+
+#pragma once
+
+#include "NodePrivate.h"
+
+#include "NodeScalar.h"
+#include "MemoryPrivate.h"
+#include "ParametersPrivate.h"
+
+
+namespace mu
+{
+
+
+	class NodeScalar::Private : public Node::Private
+	{
+	};
+
+
+}
diff --git a/Engine/Plugins/Experimental/Mutable/Source/MutableTools/Private/NodeScalarSwitch.cpp b/Engine/Plugins/Experimental/Mutable/Source/MutableTools/Private/NodeScalarSwitch.cpp
new file mode 100644
--- /dev/null
+++ b/Engine/Plugins/Experimental/Mutable/Source/MutableTools/Private/NodeScalarSwitch.cpp
@@ -0,0 +1,118 @@
+// Copyright Epic Games, Inc. All Rights Reserved.
+
+
+#include "NodeScalarSwitch.h"
+#include "NodeScalarSwitchPrivate.h"
+
+#include "NodeScalar.h"
+
+
+namespace mu
+{
+
+	//---------------------------------------------------------------------------------------------
+	// Static initialisation
+	//---------------------------------------------------------------------------------------------
+    NODE_TYPE NodeScalarSwitch::Private::s_type =
+            NODE_TYPE( "ScalarSwitch", NodeScalar::GetStaticType() );
+
+
+	//---------------------------------------------------------------------------------------------
+	//!
+	//---------------------------------------------------------------------------------------------
+
+    MUTABLE_IMPLEMENT_NODE( NodeScalarSwitch, EType::Switch, Node, Node::EType::Scalar)
+
+
+	//---------------------------------------------------------------------------------------------
+	// Node Interface
+	//---------------------------------------------------------------------------------------------
+    int NodeScalarSwitch::GetInputCount() const
+	{
+		return 1 + (int)m_pD->m_options.size();
+	}
+
+
+	//---------------------------------------------------------------------------------------------
+    Node* NodeScalarSwitch::GetInputNode( int i ) const
+	{
+		check( i>=0 && i<GetInputCount() );
+
+		Node* pResult = 0;
+
+		switch (i)
+		{
+		case 0:
+			pResult = m_pD->m_pParameter.get();
+			break;
+
+		default:
+			pResult = m_pD->m_options[i-1].get();
+			break;
+		}
+
+		return pResult;
+	}
+
+
+	//---------------------------------------------------------------------------------------------
+    void NodeScalarSwitch::SetInputNode( int i, NodePtr pNode )
+	{
+		check( i>=0 && i<GetInputCount() );
+
+		switch (i)
+		{
+		case 0:
+			m_pD->m_pParameter = dynamic_cast<NodeScalar*>(pNode.get());
+			break;
+
+		default:
+            m_pD->m_options[i-1] = dynamic_cast<NodeScalar*>(pNode.get());
+			break;
+		}
+	}
+
+
+	//---------------------------------------------------------------------------------------------
+	// Own Interface
+	//---------------------------------------------------------------------------------------------
+    NodeScalarPtr NodeScalarSwitch::GetParameter() const
+	{
+		return m_pD->m_pParameter.get();
+	}
+
+
+	//---------------------------------------------------------------------------------------------
+    void NodeScalarSwitch::SetParameter( NodeScalarPtr pNode )
+	{
+		m_pD->m_pParameter = pNode;
+	}
+
+
+	//---------------------------------------------------------------------------------------------
+    void NodeScalarSwitch::SetOptionCount( int t )
+	{
+		m_pD->m_options.resize(t);
+	}
+
+
+	//---------------------------------------------------------------------------------------------
+    NodeScalarPtr NodeScalarSwitch::GetOption( int t ) const
+	{
+		check( t>=0 && t<(int)m_pD->m_options.size() );
+		return m_pD->m_options[t].get();
+	}
+
+
+	//---------------------------------------------------------------------------------------------
+    void NodeScalarSwitch::SetOption( int t, NodeScalarPtr pNode )
+	{
+		check( t>=0 && t<(int)m_pD->m_options.size() );
+		m_pD->m_options[t] = pNode;
+	}
+
+
+
+}
+
+
diff --git a/Engine/Plugins/Experimental/Mutable/Source/MutableTools/Private/NodeScalarSwitchPrivate.h b/Engine/Plugins/Experimental/Mutable/Source/MutableTools/Private/NodeScalarSwitchPrivate.h
new file mode 100644
--- /dev/null
+++ b/Engine/Plugins/Experimental/Mutable/Source/MutableTools/Private/NodeScalarSwitchPrivate.h
@@ -0,0 +1,51 @@
+// Copyright Epic Games, Inc. All Rights Reserved.
+
+#pragma once
+
+#include "NodeScalarSwitch.h"
+#include "NodeScalarPrivate.h"
+#include "NodeScalar.h"
+#include "AST.h"
+
+
+namespace mu
+{
+
+
+    class NodeScalarSwitch::Private : public NodeScalar::Private
+	{
+	public:
+
+		MUTABLE_DEFINE_CONST_VISITABLE()
+
+	public:
+
+		static NODE_TYPE s_type;
+
+		NodeScalarPtr m_pParameter;
+        vector<NodeScalarPtr> m_options;
+
+		//!
+		void Serialise( OutputArchive& arch ) const
+		{
+            uint32_t ver = 0;
+			arch << ver;
+
+			arch << m_pParameter;
+			arch << m_options;
+		}
+
+		//!
+		void Unserialise( InputArchive& arch )
+		{
+            uint32_t ver;
+			arch >> ver;
+			check(ver==0);
+
+			arch >> m_pParameter;
+			arch >> m_options;
+		}
+	};
+
+
+}
diff --git a/Engine/Plugins/Experimental/Mutable/Source/MutableTools/Private/NodeScalarTable.cpp b/Engine/Plugins/Experimental/Mutable/Source/MutableTools/Private/NodeScalarTable.cpp
new file mode 100644
--- /dev/null
+++ b/Engine/Plugins/Experimental/Mutable/Source/MutableTools/Private/NodeScalarTable.cpp
@@ -0,0 +1,100 @@
+// Copyright Epic Games, Inc. All Rights Reserved.
+
+
+#include "NodeScalarTable.h"
+#include "NodeScalarTablePrivate.h"
+
+
+
+namespace mu
+{
+
+
+	//---------------------------------------------------------------------------------------------
+	// Static initialisation
+	//---------------------------------------------------------------------------------------------
+	NODE_TYPE NodeScalarTable::Private::s_type =
+			NODE_TYPE( "TableScalar", NodeScalar::GetStaticType() );
+
+
+	//---------------------------------------------------------------------------------------------
+	//!
+	//---------------------------------------------------------------------------------------------
+
+	MUTABLE_IMPLEMENT_NODE(NodeScalarTable, EType::Table, Node, Node::EType::Scalar)
+
+
+	//---------------------------------------------------------------------------------------------
+	// Node Interface
+	//---------------------------------------------------------------------------------------------
+	int NodeScalarTable::GetInputCount() const
+	{
+		return 0;
+	}
+
+
+	//---------------------------------------------------------------------------------------------
+	Node* NodeScalarTable::GetInputNode( int i ) const
+	{
+		check( i >=0 && i < 0 );
+        (void)i;
+        return 0;
+	}
+
+
+	//---------------------------------------------------------------------------------------------
+    void NodeScalarTable::SetInputNode( int i, NodePtr )
+	{
+		check( i >=0 && i < 0 );
+        (void)i;
+        //m_pD->m_pObject = dynamic_cast<NodeObject*>( pNode.get() );
+	}
+
+
+	//---------------------------------------------------------------------------------------------
+	// Own Interface
+	//---------------------------------------------------------------------------------------------
+	const char* NodeScalarTable::GetColumn() const
+	{
+		return m_pD->m_columnName.c_str();
+	}
+
+
+	//---------------------------------------------------------------------------------------------
+	void NodeScalarTable::SetColumn( const char* strName )
+	{
+		if (strName)
+		{
+			m_pD->m_columnName = strName;
+		}
+		else
+		{
+			m_pD->m_columnName = "";
+		}
+	}
+
+
+	//---------------------------------------------------------------------------------------------
+	void NodeScalarTable::SetParameterName( const char* strName )
+	{
+		m_pD->m_parameterName = strName;
+	}
+
+
+	//---------------------------------------------------------------------------------------------
+	void NodeScalarTable::SetTable( TablePtr pTable )
+	{
+		m_pD->m_pTable = pTable;
+	}
+
+
+	//---------------------------------------------------------------------------------------------
+	TablePtr NodeScalarTable::GetTable() const
+	{
+		return m_pD->m_pTable;
+	}
+
+
+}
+
+
diff --git a/Engine/Plugins/Experimental/Mutable/Source/MutableTools/Private/NodeScalarTablePrivate.h b/Engine/Plugins/Experimental/Mutable/Source/MutableTools/Private/NodeScalarTablePrivate.h
new file mode 100644
--- /dev/null
+++ b/Engine/Plugins/Experimental/Mutable/Source/MutableTools/Private/NodeScalarTablePrivate.h
@@ -0,0 +1,56 @@
+// Copyright Epic Games, Inc. All Rights Reserved.
+
+#pragma once
+
+#include "NodePrivate.h"
+#include "NodeScalarTable.h"
+#include "Table.h"
+#include "AST.h"
+
+#include "MemoryPrivate.h"
+
+
+namespace mu
+{
+
+
+	class NodeScalarTable::Private : public Node::Private
+	{
+	public:
+
+		MUTABLE_DEFINE_CONST_VISITABLE()
+
+	public:
+
+		static NODE_TYPE s_type;
+
+		string m_parameterName;
+		TablePtr m_pTable;
+		string m_columnName;
+
+		//!
+		void Serialise( OutputArchive& arch ) const
+		{
+            uint32_t ver = 0;
+			arch << ver;
+
+			arch << m_parameterName;
+			arch << m_pTable;
+			arch << m_columnName;
+		}
+
+		//!
+		void Unserialise( InputArchive& arch )
+		{
+            uint32_t ver;
+			arch >> ver;
+            check(ver==0);
+
+			arch >> m_parameterName;
+			arch >> m_pTable;
+			arch >> m_columnName;
+		}
+
+	};
+
+}
