[mutable] Move Mutable plugin from the restricted folder to the experimental plugins folder.

diff --git a/Engine/Plugins/Experimental/Mutable/Source/MutableTools/Private/NodeScalarVariation.cpp b/Engine/Plugins/Experimental/Mutable/Source/MutableTools/Private/NodeScalarVariation.cpp
new file mode 100644
--- /dev/null
+++ b/Engine/Plugins/Experimental/Mutable/Source/MutableTools/Private/NodeScalarVariation.cpp
@@ -0,0 +1,115 @@
+// Copyright Epic Games, Inc. All Rights Reserved.
+
+
+#include "NodeScalarVariation.h"
+#include "NodeScalarVariationPrivate.h"
+
+
+namespace mu
+{
+
+    //---------------------------------------------------------------------------------------------
+    // Static initialisation
+    //---------------------------------------------------------------------------------------------
+    NODE_TYPE NodeScalarVariation::Private::s_type =
+        NODE_TYPE( "ScalarVariation", NodeScalar::GetStaticType() );
+
+
+    //---------------------------------------------------------------------------------------------
+    //---------------------------------------------------------------------------------------------
+    //---------------------------------------------------------------------------------------------
+
+    MUTABLE_IMPLEMENT_NODE( NodeScalarVariation, EType::Variation, Node, Node::EType::Scalar)
+
+
+    //---------------------------------------------------------------------------------------------
+    int NodeScalarVariation::GetInputCount() const { return 1 + int( m_pD->m_variations.size() ); }
+
+
+    //---------------------------------------------------------------------------------------------
+    Node* NodeScalarVariation::GetInputNode( int i ) const
+    {
+        check( i >= 0 && i < GetInputCount() );
+
+        if ( i == 0 )
+        {
+            return m_pD->m_defaultScalar.get();
+        }
+        i -= 1;
+
+        if ( i < int( m_pD->m_variations.size() ) )
+        {
+            return m_pD->m_variations[i].m_scalar.get();
+        }
+        i -= int( m_pD->m_variations.size() );
+
+        return nullptr;
+    }
+
+
+    //---------------------------------------------------------------------------------------------
+    void NodeScalarVariation::SetInputNode( int i, NodePtr pNode )
+    {
+        check( i >= 0 && i < GetInputCount() );
+
+        if ( i == 0 )
+        {
+            m_pD->m_defaultScalar = dynamic_cast<NodeScalar*>( pNode.get() );
+            return;
+        }
+
+        i -= 1;
+        if ( i < int( m_pD->m_variations.size() ) )
+        {
+
+            m_pD->m_variations[i].m_scalar = dynamic_cast<NodeScalar*>( pNode.get() );
+            return;
+        }
+        i -= (int)m_pD->m_variations.size();
+    }
+
+
+    //---------------------------------------------------------------------------------------------
+    // Own Interface
+    //---------------------------------------------------------------------------------------------
+    void NodeScalarVariation::SetDefaultScalar( NodeScalar* p ) { m_pD->m_defaultScalar = p; }
+
+
+    //---------------------------------------------------------------------------------------------
+    int NodeScalarVariation::GetVariationCount() const { return int( m_pD->m_variations.size() ); }
+
+
+    //---------------------------------------------------------------------------------------------
+    void NodeScalarVariation::SetVariationCount( int num )
+    {
+        check( num >= 0 );
+        m_pD->m_variations.resize( num );
+    }
+
+    //---------------------------------------------------------------------------------------------
+    void NodeScalarVariation::SetVariationTag( int index, const char* strTag )
+    {
+        check( index >= 0 && index < (int)m_pD->m_variations.size() );
+        check( strTag );
+
+        if ( strTag )
+        {
+            m_pD->m_variations[index].m_tag = strTag;
+        }
+        else
+        {
+            m_pD->m_variations[index].m_tag = "";
+        }
+    }
+
+
+    //---------------------------------------------------------------------------------------------
+    void NodeScalarVariation::SetVariationScalar( int index, NodeScalar* pNode )
+    {
+        check( index >= 0 && index < (int)m_pD->m_variations.size() );
+
+        m_pD->m_variations[index].m_scalar = pNode;
+    }
+
+
+} // namespace mu
diff --git a/Engine/Plugins/Experimental/Mutable/Source/MutableTools/Private/NodeScalarVariationPrivate.h b/Engine/Plugins/Experimental/Mutable/Source/MutableTools/Private/NodeScalarVariationPrivate.h
new file mode 100644
--- /dev/null
+++ b/Engine/Plugins/Experimental/Mutable/Source/MutableTools/Private/NodeScalarVariationPrivate.h
@@ -0,0 +1,78 @@
+// Copyright Epic Games, Inc. All Rights Reserved.
+
+#pragma once
+
+#include "NodePrivate.h"
+#include "NodeScalarVariation.h"
+
+#include "MemoryPrivate.h"
+
+
+namespace mu
+{
+
+    //---------------------------------------------------------------------------------------------
+    //!
+    //---------------------------------------------------------------------------------------------
+    class NodeScalarVariation::Private : public Node::Private
+    {
+    public:
+        MUTABLE_DEFINE_CONST_VISITABLE()
+
+    public:
+
+        static NODE_TYPE s_type;
+
+        NodeScalarPtr m_defaultScalar;
+
+        struct VARIATION
+        {
+            NodeScalarPtr m_scalar;
+            string m_tag;
+
+            //!
+            void Serialise( OutputArchive& arch ) const
+            {
+                uint32_t ver = 1;
+                arch << ver;
+
+                arch << m_tag;
+                arch << m_scalar;
+            }
+
+            void Unserialise( InputArchive& arch )
+            {
+                uint32_t ver = 0;
+                arch >> ver;
+                check( ver == 1 );
+
+                arch >> m_tag;
+                arch >> m_scalar;
+            }
+        };
+
+        vector<VARIATION> m_variations;
+
+        //!
+        void Serialise( OutputArchive& arch ) const
+        {
+            uint32_t ver = 1;
+            arch << ver;
+
+            arch << m_defaultScalar;
+            arch << m_variations;
+        }
+
+        //!
+        void Unserialise( InputArchive& arch )
+        {
+            uint32_t ver;
+            arch >> ver;
+			check(ver == 1);
+
+            arch >> m_defaultScalar;
+            arch >> m_variations;
+        }
+    };
+
+} // namespace mu
diff --git a/Engine/Plugins/Experimental/Mutable/Source/MutableTools/Private/NodeString.cpp b/Engine/Plugins/Experimental/Mutable/Source/MutableTools/Private/NodeString.cpp
new file mode 100644
--- /dev/null
+++ b/Engine/Plugins/Experimental/Mutable/Source/MutableTools/Private/NodeString.cpp
@@ -0,0 +1,80 @@
+// Copyright Epic Games, Inc. All Rights Reserved.
+
+
+#include "MutableTools/Public/NodeString.h"
+
+#include "MutableTools/Public/Node.h"
+#include "MutableTools/Private/NodePrivate.h"
+
+#include "MutableTools/Public/NodeStringConstant.h"
+#include "MutableTools/Public/NodeStringParameter.h"
+
+
+namespace mu
+{
+
+	//---------------------------------------------------------------------------------------------
+	// Static initialisation
+	//---------------------------------------------------------------------------------------------
+	static NODE_TYPE s_nodeStringType = NODE_TYPE( "NodeString", Node::GetStaticType() );
+
+
+	//---------------------------------------------------------------------------------------------
+	void NodeString::Serialise( const NodeString* p, OutputArchive& arch )
+	{
+        uint32_t ver = 0;
+		arch << ver;
+
+
+#define SERIALISE_CHILDREN( C, ID ) \
+        ( const C* pTyped##ID = dynamic_cast<const C*>(p) )			\
+        {                                                           \
+            arch << (uint32_t)ID;                                   \
+            C::Serialise( pTyped##ID, arch );						\
+        }                                                           \
+
+        if SERIALISE_CHILDREN( NodeStringConstant                   , 0 )
+        else if SERIALISE_CHILDREN( NodeStringParameter             , 1 )
+        else check(false);
+
+#undef SERIALISE_CHILDREN
+    }
+
+    //---------------------------------------------------------------------------------------------
+	NodeStringPtr NodeString::StaticUnserialise( InputArchive& arch )
+	{
+        uint32_t ver;
+		arch >> ver;
+		check( ver == 0 );
+
+        uint32_t id;
+		arch >> id;
+
+		switch (id)
+		{
+		case 0 :  return NodeStringConstant::StaticUnserialise( arch ); break;
+		case 1 :  return NodeStringParameter::StaticUnserialise( arch ); break;
+        default : check(false);
+		}
+
+		return 0;
+	}
+
+
+	//---------------------------------------------------------------------------------------------
+	const NODE_TYPE* NodeString::GetType() const
+	{
+		return GetStaticType();
+	}
+
+
+	//---------------------------------------------------------------------------------------------
+	const NODE_TYPE* NodeString::GetStaticType()
+	{
+		return &s_nodeStringType;
+	}
+
+
+}
+
+
diff --git a/Engine/Plugins/Experimental/Mutable/Source/MutableTools/Private/NodeStringConstant.cpp b/Engine/Plugins/Experimental/Mutable/Source/MutableTools/Private/NodeStringConstant.cpp
new file mode 100644
--- /dev/null
+++ b/Engine/Plugins/Experimental/Mutable/Source/MutableTools/Private/NodeStringConstant.cpp
@@ -0,0 +1,72 @@
+// Copyright Epic Games, Inc. All Rights Reserved.
+
+
+#include "NodeStringConstant.h"
+#include "NodeStringConstantPrivate.h"
+
+#include "MutableMath.h"
+
+
+#define NODE_INPUT_COUNT 	0
+
+
+namespace mu
+{
+
+	//---------------------------------------------------------------------------------------------
+	// Static initialisation
+	//---------------------------------------------------------------------------------------------
+	NODE_TYPE NodeStringConstant::Private::s_type =
+			NODE_TYPE( "StringConstant", NodeString::GetStaticType() );
+
+
+	//---------------------------------------------------------------------------------------------
+	//!
+	//---------------------------------------------------------------------------------------------
+
+	MUTABLE_IMPLEMENT_NODE( NodeStringConstant, EType::Constant, Node, Node::EType::String)
+
+
+	//---------------------------------------------------------------------------------------------
+	// Node Interface
+	//---------------------------------------------------------------------------------------------
+	int NodeStringConstant::GetInputCount() const
+	{
+		return NODE_INPUT_COUNT;
+	}
+
+
+	//---------------------------------------------------------------------------------------------
+    Node* NodeStringConstant::GetInputNode( int ) const
+	{
+		check( false );
+		return 0;
+	}
+
+
+	//---------------------------------------------------------------------------------------------
+    void NodeStringConstant::SetInputNode( int, NodePtr )
+	{
+		check( false );
+	}
+
+
+	//---------------------------------------------------------------------------------------------
+	// Own Interface
+	//---------------------------------------------------------------------------------------------
+	const char* NodeStringConstant::GetValue() const
+	{
+		return m_pD->m_value.c_str();
+	}
+
+	//---------------------------------------------------------------------------------------------
+	void NodeStringConstant::SetValue( const char* v )
+	{
+		m_pD->m_value = v?v:"";
+	}
+
+
+}
+
+#undef NODE_INPUT_COUNT
+
diff --git a/Engine/Plugins/Experimental/Mutable/Source/MutableTools/Private/NodeStringConstantPrivate.h b/Engine/Plugins/Experimental/Mutable/Source/MutableTools/Private/NodeStringConstantPrivate.h
new file mode 100644
--- /dev/null
+++ b/Engine/Plugins/Experimental/Mutable/Source/MutableTools/Private/NodeStringConstantPrivate.h
@@ -0,0 +1,51 @@
+// Copyright Epic Games, Inc. All Rights Reserved.
+
+#pragma once
+
+#include "NodeStringPrivate.h"
+
+#include "NodeStringConstant.h"
+#include "NodeImage.h"
+#include "AST.h"
+#include "MemoryPrivate.h"
+#include "ParametersPrivate.h"
+
+
+namespace mu
+{
+
+
+	class NodeStringConstant::Private : public NodeString::Private
+	{
+	public:
+
+		MUTABLE_DEFINE_CONST_VISITABLE()
+
+	public:
+
+		static NODE_TYPE s_type;
+
+		string m_value;
+
+		//!
+		void Serialise( OutputArchive& arch ) const
+		{
+            uint32_t ver = 0;
+			arch << ver;
+
+			arch << m_value;
+		}
+
+		//!
+		void Unserialise( InputArchive& arch )
+		{
+            uint32_t ver;
+			arch >> ver;
+			check(ver==0);
+
+			arch >> m_value;
+		}
+	};
+
+
+}
diff --git a/Engine/Plugins/Experimental/Mutable/Source/MutableTools/Private/NodeStringParameter.cpp b/Engine/Plugins/Experimental/Mutable/Source/MutableTools/Private/NodeStringParameter.cpp
new file mode 100644
--- /dev/null
+++ b/Engine/Plugins/Experimental/Mutable/Source/MutableTools/Private/NodeStringParameter.cpp
@@ -0,0 +1,165 @@
+// Copyright Epic Games, Inc. All Rights Reserved.
+
+
+#include "MutableTools/Public/NodeStringParameter.h"
+#include "MutableTools/Private/NodeStringParameterPrivate.h"
+
+#include "MutableMath.h"
+#include "MemoryPrivate.h"
+
+
+namespace mu
+{
+
+	//---------------------------------------------------------------------------------------------
+	// Static initialisation
+	//---------------------------------------------------------------------------------------------
+	NODE_TYPE NodeStringParameter::Private::s_type =
+			NODE_TYPE( "StringParameter", NodeString::GetStaticType() );
+
+
+	//---------------------------------------------------------------------------------------------
+	//!
+	//---------------------------------------------------------------------------------------------
+
+	MUTABLE_IMPLEMENT_NODE( NodeStringParameter, EType::Parameter, Node, Node::EType::String)
+
+
+	//---------------------------------------------------------------------------------------------
+	int NodeStringParameter::GetInputCount() const
+	{
+        return int( m_pD->m_additionalImages.size()
+                    +
+                    m_pD->m_ranges.size() );
+	}
+
+
+	//---------------------------------------------------------------------------------------------
+    Node* NodeStringParameter::GetInputNode( int i ) const
+	{
+        check( i<GetInputCount() );
+
+        int imageCount = int(m_pD->m_additionalImages.size());
+        int rangeCount = int(m_pD->m_ranges.size());
+        if (i<imageCount)
+        {
+            return m_pD->m_additionalImages[i].get();
+        }
+        else if ( i < imageCount + rangeCount )
+        {
+            int r = i - imageCount;
+            return m_pD->m_ranges[r].get();
+        }
+        return nullptr;
+	}
+
+
+	//---------------------------------------------------------------------------------------------
+    void NodeStringParameter::SetInputNode( int i, NodePtr n )
+	{
+        check( i<GetInputCount() );
+        int imageCount = int(m_pD->m_additionalImages.size());
+        int rangeCount = int(m_pD->m_ranges.size());
+        if (i<imageCount)
+        {
+            m_pD->m_additionalImages[i] = dynamic_cast<NodeImage*>(n.get());
+        }
+        else if ( i < imageCount + rangeCount )
+        {
+            int r = i - imageCount;
+            m_pD->m_ranges[r] = dynamic_cast<NodeRange*>(n.get());
+        }
+    }
+
+
+	//---------------------------------------------------------------------------------------------
+	const char* NodeStringParameter::GetName() const
+	{
+		return m_pD->m_name.c_str();
+	}
+
+
+	//---------------------------------------------------------------------------------------------
+	void NodeStringParameter::SetName( const char* strName )
+	{
+		if ( strName )
+		{
+			m_pD->m_name = strName;
+		}
+		else
+		{
+			m_pD->m_name = "";
+		}
+	}
+
+
+	const char* NodeStringParameter::GetUid() const
+	{
+		return m_pD->m_uid.c_str();
+	}
+
+
+	//---------------------------------------------------------------------------------------------
+	void NodeStringParameter::SetUid( const char* strUid )
+	{
+		if ( strUid )
+		{
+			m_pD->m_uid = strUid;
+		}
+		else
+		{
+			m_pD->m_uid = "";
+		}
+	}
+
+
+	//---------------------------------------------------------------------------------------------
+	const char* NodeStringParameter::GetDefaultValue() const
+	{
+		return m_pD->m_defaultValue.c_str();
+	}
+
+
+	//---------------------------------------------------------------------------------------------
+	void NodeStringParameter::SetDefaultValue( const char* v )
+	{
+		m_pD->m_defaultValue = v?v:"";
+	}
+
+
+	//---------------------------------------------------------------------------------------------
+	PARAMETER_DETAILED_TYPE NodeStringParameter::GetDetailedType() const
+	{
+		return m_pD->m_detailedType;
+	}
+
+
+	//---------------------------------------------------------------------------------------------
+	void NodeStringParameter::SetDetailedType( PARAMETER_DETAILED_TYPE t )
+	{
+		m_pD->m_detailedType = t;
+	}
+
+
+    //---------------------------------------------------------------------------------------------
+    void NodeStringParameter::SetRangeCount( int i )
+    {
+        check(i>=0);
+        m_pD->m_ranges.resize(i);
+    }
+
+
+    //---------------------------------------------------------------------------------------------
+    void NodeStringParameter::SetRange( int i, NodeRangePtr pRange )
+    {
+        check( i>=0 && i<int(m_pD->m_ranges.size()) );
+        if ( i>=0 && i<int(m_pD->m_ranges.size()) )
+        {
+            m_pD->m_ranges[i] = pRange;
+        }
+    }
+
+
+}
+
+
diff --git a/Engine/Plugins/Experimental/Mutable/Source/MutableTools/Private/NodeStringParameterPrivate.h b/Engine/Plugins/Experimental/Mutable/Source/MutableTools/Private/NodeStringParameterPrivate.h
new file mode 100644
--- /dev/null
+++ b/Engine/Plugins/Experimental/Mutable/Source/MutableTools/Private/NodeStringParameterPrivate.h
@@ -0,0 +1,67 @@
+// Copyright Epic Games, Inc. All Rights Reserved.
+
+#pragma once
+
+#include "NodeStringPrivate.h"
+
+#include "NodeStringParameter.h"
+#include "NodeImage.h"
+#include "NodeRange.h"
+#include "MemoryPrivate.h"
+#include "ParametersPrivate.h"
+
+
+namespace mu
+{
+
+
+	class NodeStringParameter::Private : public NodeString::Private
+	{
+	public:
+
+		MUTABLE_DEFINE_CONST_VISITABLE()
+
+	public:
+
+		static NODE_TYPE s_type;
+
+		string m_defaultValue;
+		string m_name;
+		string m_uid;
+		PARAMETER_DETAILED_TYPE m_detailedType = PARAMETER_DETAILED_TYPE::UNKNOWN;
+
+        vector<Ptr<NodeImage>> m_additionalImages;
+
+        vector<Ptr<NodeRange>> m_ranges;
+
+		//!
+		void Serialise( OutputArchive& arch ) const
+		{
+            uint32_t ver = 4;
+			arch << ver;
+
+			arch << m_defaultValue;
+			arch << m_name;
+			arch << m_uid;
+			arch << m_additionalImages;
+			arch << m_detailedType;
+            arch << m_ranges;
+        }
+
+		//!
+		void Unserialise( InputArchive& arch )
+		{
+            uint32_t ver;
+			arch >> ver;
+            check(ver==4);
+
+			arch >> m_defaultValue;
+			arch >> m_name;
+            arch >> m_uid;
+			arch >> m_additionalImages;
+            arch >> m_detailedType;
+            arch >> m_ranges;
+        }
+	};
+
+}
diff --git a/Engine/Plugins/Experimental/Mutable/Source/MutableTools/Private/NodeStringPrivate.h b/Engine/Plugins/Experimental/Mutable/Source/MutableTools/Private/NodeStringPrivate.h
new file mode 100644
--- /dev/null
+++ b/Engine/Plugins/Experimental/Mutable/Source/MutableTools/Private/NodeStringPrivate.h
@@ -0,0 +1,21 @@
+// Copyright Epic Games, Inc. All Rights Reserved.
+
+#pragma once
+
+#include "NodePrivate.h"
+
+#include "NodeString.h"
+#include "MemoryPrivate.h"
+#include "ParametersPrivate.h"
+
+
+namespace mu
+{
+
+
+	class NodeString::Private : public Node::Private
+	{
+	};
+
+
+}
diff --git a/Engine/Plugins/Experimental/Mutable/Source/MutableTools/Private/NodeSurface.cpp b/Engine/Plugins/Experimental/Mutable/Source/MutableTools/Private/NodeSurface.cpp
new file mode 100644
--- /dev/null
+++ b/Engine/Plugins/Experimental/Mutable/Source/MutableTools/Private/NodeSurface.cpp
@@ -0,0 +1,78 @@
+// Copyright Epic Games, Inc. All Rights Reserved.
+
+
+#include "NodeSurface.h"
+
+#include "NodeMesh.h"
+#include "NodeImage.h"
+
+#include "NodeSurfaceNew.h"
+#include "NodeSurfaceEdit.h"
+#include "NodeSurfaceVariation.h"
+
+#include "SerialisationPrivate.h"
+
+
+namespace mu
+{
+
+
+	//---------------------------------------------------------------------------------------------
+	// Static initialisation
+	//---------------------------------------------------------------------------------------------
+    static NODE_TYPE s_nodeSurfaceType =
+            NODE_TYPE( "NodeSurface", Node::GetStaticType() );
+
+
+	//---------------------------------------------------------------------------------------------
+	//---------------------------------------------------------------------------------------------
+	//---------------------------------------------------------------------------------------------
+    const NODE_TYPE* NodeSurface::GetType() const
+	{
+		return GetStaticType();
+	}
+
+
+	//---------------------------------------------------------------------------------------------
+    const NODE_TYPE* NodeSurface::GetStaticType()
+	{
+        return &s_nodeSurfaceType;
+	}
+
+
+	//---------------------------------------------------------------------------------------------
+    void NodeSurface::Serialise( const NodeSurface* p, OutputArchive& arch )
+	{
+        uint32_t ver = 0;
+		arch << ver;
+
+		arch << uint32_t(p->Type);
+		p->SerialiseWrapper(arch);
+	}
+
+
+    //---------------------------------------------------------------------------------------------
+    NodeSurfacePtr NodeSurface::StaticUnserialise( InputArchive& arch )
+	{
+        uint32_t ver;
+		arch >> ver;
+		check( ver == 0 );
+
+        uint32_t id;
+		arch >> id;
+
+		switch (id)
+		{
+        case 0: return NodeSurfaceNew::StaticUnserialise( arch ); break;
+        case 1: return NodeSurfaceEdit::StaticUnserialise( arch ); break;
+		//case 2: return NodeSelectSurface::StaticUnserialise(arch); break;
+		case 3: return NodeSurfaceVariation::StaticUnserialise(arch); break;
+		default : check(false);
+		}
+
+		return 0;
+	}
+
+}
+
+
diff --git a/Engine/Plugins/Experimental/Mutable/Source/MutableTools/Private/NodeSurfaceEdit.cpp b/Engine/Plugins/Experimental/Mutable/Source/MutableTools/Private/NodeSurfaceEdit.cpp
new file mode 100644
--- /dev/null
+++ b/Engine/Plugins/Experimental/Mutable/Source/MutableTools/Private/NodeSurfaceEdit.cpp
@@ -0,0 +1,247 @@
+// Copyright Epic Games, Inc. All Rights Reserved.
+
+
+#include "MutableTools/Public/NodeSurfaceEdit.h"
+#include "MutableTools/Private/NodeSurfaceEditPrivate.h"
+
+#include "MutableTools/Public/NodeMesh.h"
+#include "MutableTools/Public/NodeImage.h"
+#include "MutableTools/Public/NodeScalar.h"
+
+
+namespace mu
+{
+
+
+	//---------------------------------------------------------------------------------------------
+	// Static initialisation
+	//---------------------------------------------------------------------------------------------
+    NODE_TYPE NodeSurfaceEdit::Private::s_type =
+            NODE_TYPE( "EditSurface", NodeSurface::GetStaticType() );
+
+
+	//---------------------------------------------------------------------------------------------
+	//---------------------------------------------------------------------------------------------
+	//---------------------------------------------------------------------------------------------
+
+    MUTABLE_IMPLEMENT_NODE( NodeSurfaceEdit, EType::Edit, Node, Node::EType::Surface)
+
+
+	//---------------------------------------------------------------------------------------------
+	// Node Interface
+	//---------------------------------------------------------------------------------------------
+    int NodeSurfaceEdit::GetInputCount() const
+	{
+		return
+			3 	// Extend mesh, morph, factor
+			+
+			(int)m_pD->m_textures.size()*2
+			;
+	}
+
+
+	//---------------------------------------------------------------------------------------------
+    Node* NodeSurfaceEdit::GetInputNode( int i ) const
+	{
+		check( i >=0 && i < GetInputCount() );
+
+		Node* pResult = 0;
+
+		if ( i==0 )
+		{
+			pResult = m_pD->m_pMesh.get();
+		}
+		else if ( i==1 )
+		{
+			pResult = m_pD->m_pMorph.get();
+		}
+		else if (i == 2)
+		{
+			pResult = m_pD->m_pFactor.get();
+		}
+		else
+		{
+			int idx = i - 3;
+
+			if (idx%2 == 0)
+			{
+				pResult = m_pD->m_textures[idx/2].m_pExtend.get();
+			}
+			else
+			{
+				pResult = m_pD->m_textures[idx/2].m_pPatch.get();
+			}
+		}
+
+		return pResult;
+	}
+
+
+	//---------------------------------------------------------------------------------------------
+    void NodeSurfaceEdit::SetInputNode( int i, NodePtr pNode )
+	{
+		check( i >=0 && i < GetInputCount());
+
+		if ( i==0 )
+		{
+			m_pD->m_pMesh = dynamic_cast<NodePatchMesh*>( pNode.get() );
+		}
+		else if ( i==1 )
+		{
+			m_pD->m_pMorph = dynamic_cast<NodeMesh*>( pNode.get() );
+		}
+		else if ( i==2 )
+		{
+			m_pD->m_pFactor = dynamic_cast<NodeScalar*>( pNode.get() );
+		}
+		else
+		{
+			int idx = i - 3;
+			if (idx % 2==0)
+			{
+				m_pD->m_textures[idx/2].m_pExtend = dynamic_cast<NodeImage*>( pNode.get() );
+			}
+			else
+			{
+				m_pD->m_textures[idx/2].m_pPatch = dynamic_cast<NodePatchImage*>( pNode.get() );
+			}
+		}
+	}
+
+
+	//---------------------------------------------------------------------------------------------
+	// Own Interface
+	//---------------------------------------------------------------------------------------------
+    void NodeSurfaceEdit::SetParent( NodeSurface* p )
+	{
+		m_pD->m_pParent = p;
+	}
+
+
+	//---------------------------------------------------------------------------------------------
+    NodeSurface* NodeSurfaceEdit::GetParent() const
+	{
+        return m_pD->m_pParent.get();
+	}
+
+
+    //---------------------------------------------------------------------------------------------
+    void NodeSurfaceEdit::AddTag(const char* tagName)
+    {
+        m_pD->m_tags.push_back(tagName);
+    }
+
+
+    //---------------------------------------------------------------------------------------------
+    int NodeSurfaceEdit::GetTagCount() const
+    {
+        return int( m_pD->m_tags.size());
+    }
+
+
+    //---------------------------------------------------------------------------------------------
+    const char* NodeSurfaceEdit::GetTag( int i ) const
+    {
+        if (i>=0 && i<GetTagCount())
+        {
+            return m_pD->m_tags[i].c_str();
+        }
+        return nullptr;
+    }
+
+
+	//---------------------------------------------------------------------------------------------
+    void NodeSurfaceEdit::SetMesh(NodePatchMesh* pNode )
+	{
+		m_pD->m_pMesh = pNode;
+	}
+
+
+	//---------------------------------------------------------------------------------------------
+    NodePatchMesh* NodeSurfaceEdit::GetMesh() const
+	{
+		return m_pD->m_pMesh.get();
+	}
+
+
+	//---------------------------------------------------------------------------------------------
+    void NodeSurfaceEdit::SetMorph( NodeMeshPtr pNode )
+	{
+		m_pD->m_pMorph = pNode;
+	}
+
+
+	//---------------------------------------------------------------------------------------------
+    NodeMeshPtr NodeSurfaceEdit::GetMorph() const
+	{
+		return m_pD->m_pMorph.get();
+	}
+
+
+	//---------------------------------------------------------------------------------------------
+	void NodeSurfaceEdit::SetFactor(NodeScalarPtr pNode)
+	{
+		m_pD->m_pFactor = pNode;
+	}
+
+
+	//---------------------------------------------------------------------------------------------
+	NodeScalarPtr NodeSurfaceEdit::GetFactor() const
+	{
+		return m_pD->m_pFactor.get();
+	}
+
+
+	//---------------------------------------------------------------------------------------------
+    int NodeSurfaceEdit::GetImageCount() const
+	{
+		return (int)m_pD->m_textures.size();
+	}
+
+
+	//---------------------------------------------------------------------------------------------
+    void NodeSurfaceEdit::SetImageCount( int num )
+	{
+		check( num >=0 );
+		m_pD->m_textures.resize( num );
+	}
+
+
+	//---------------------------------------------------------------------------------------------
+    NodeImagePtr NodeSurfaceEdit::GetImage( int index ) const
+	{
+		check( index >=0 && index < (int)m_pD->m_textures.size() );
+
+		return m_pD->m_textures[ index ].m_pExtend.get();
+	}
+
+
+	//---------------------------------------------------------------------------------------------
+    void NodeSurfaceEdit::SetImage( int index, NodeImagePtr pNode )
+	{
+		check( index >=0 && index < (int)m_pD->m_textures.size() );
+
+		m_pD->m_textures[ index ].m_pExtend = pNode;
+	}
+
+
+	//---------------------------------------------------------------------------------------------
+    NodePatchImagePtr NodeSurfaceEdit::GetPatch( int index ) const
+	{
+		check( index >=0 && index < (int)m_pD->m_textures.size() );
+
+		return m_pD->m_textures[ index ].m_pPatch.get();
+	}
+
+
+	//---------------------------------------------------------------------------------------------
+    void NodeSurfaceEdit::SetPatch( int index, NodePatchImagePtr pNode )
+	{
+		check( index >=0 && index < (int)m_pD->m_textures.size() );
+
+		m_pD->m_textures[ index ].m_pPatch = pNode;
+	}
+
+}
+
+
diff --git a/Engine/Plugins/Experimental/Mutable/Source/MutableTools/Private/NodeSurfaceEditPrivate.h b/Engine/Plugins/Experimental/Mutable/Source/MutableTools/Private/NodeSurfaceEditPrivate.h
new file mode 100644
--- /dev/null
+++ b/Engine/Plugins/Experimental/Mutable/Source/MutableTools/Private/NodeSurfaceEditPrivate.h
@@ -0,0 +1,102 @@
+// Copyright Epic Games, Inc. All Rights Reserved.
+
+#pragma once
+
+#include "MutableTools/Private/NodePrivate.h"
+#include "NodeSurfaceEdit.h"
+#include "NodePatchImagePrivate.h"
+
+#include "NodeMesh.h"
+#include "NodePatchMesh.h"
+#include "NodeImage.h"
+#include "NodeScalar.h"
+
+#include "MemoryPrivate.h"
+
+
+namespace mu
+{
+
+
+	//---------------------------------------------------------------------------------------------
+	//!
+	//---------------------------------------------------------------------------------------------
+    class NodeSurfaceEdit::Private : public Node::Private
+	{
+	public:
+
+		MUTABLE_DEFINE_CONST_VISITABLE()
+
+	public:
+
+		static NODE_TYPE s_type;
+
+        NodeSurfacePtr m_pParent;
+		NodePatchMeshPtr m_pMesh;
+		NodeMeshPtr m_pMorph;
+
+        //! This flag indicates that the mesh in the m_pMorph connection is a target mesh, and not
+        //! morph information. This means that the morph information needs to be generated.
+        bool m_morphIsTarget = true;
+
+		struct TEXTURE
+		{
+			NodeImagePtr m_pExtend;
+			NodePatchImagePtr m_pPatch;
+
+			//!
+			void Serialise( OutputArchive& arch ) const
+			{
+				arch << m_pExtend;
+				arch << m_pPatch;
+			}
+
+			void Unserialise( InputArchive& arch )
+			{
+				arch >> m_pExtend;
+				arch >> m_pPatch;
+			}
+		};
+
+		vector<TEXTURE> m_textures;
+
+        //! Tags in this surface edit
+        vector<mu::string> m_tags;
+
+		//! Factor of the morph
+		NodeScalarPtr m_pFactor;
+
+		//!
+		void Serialise( OutputArchive& arch ) const
+		{
+            uint32_t ver = 4;
+			arch << ver;
+
+			arch << m_pParent;
+			arch << m_pMesh;
+            arch << m_pMorph;
+            arch << m_morphIsTarget;
+            arch << m_textures;
+            arch << m_tags;
+			arch << m_pFactor;
+        }
+
+		//!
+		void Unserialise( InputArchive& arch )
+		{
+            uint32_t ver;
+			arch >> ver;
+            check( ver==4 );
+
+			arch >> m_pParent;
+			arch >> m_pMesh;
+			arch >> m_pMorph;
+            arch >> m_morphIsTarget;
+			arch >> m_textures;
+            arch >> m_tags;
+			arch >> m_pFactor;
+        }
+	};
+
+}
+
diff --git a/Engine/Plugins/Experimental/Mutable/Source/MutableTools/Private/NodeSurfaceNew.cpp b/Engine/Plugins/Experimental/Mutable/Source/MutableTools/Private/NodeSurfaceNew.cpp
new file mode 100644
--- /dev/null
+++ b/Engine/Plugins/Experimental/Mutable/Source/MutableTools/Private/NodeSurfaceNew.cpp
@@ -0,0 +1,633 @@
+// Copyright Epic Games, Inc. All Rights Reserved.
+
+
+#include "MutableTools/Public/NodeSurfaceNew.h"
+#include "MutableTools/Private/NodeSurfaceNewPrivate.h"
+
+#include "MutableTools/Public/NodeMesh.h"
+#include "MutableTools/Public/NodeImage.h"
+#include "MutableTools/Public/NodeScalar.h"
+#include "MutableTools/Public/NodeString.h"
+#include "MutableTools/Public/NodeColour.h"
+
+namespace mu
+{
+
+
+	//---------------------------------------------------------------------------------------------
+	// Static initialisation
+	//---------------------------------------------------------------------------------------------
+    NODE_TYPE NodeSurfaceNew::Private::s_type =
+            NODE_TYPE( "NewSurface", NodeSurface::GetStaticType() );
+
+
+	//---------------------------------------------------------------------------------------------
+	//---------------------------------------------------------------------------------------------
+	//---------------------------------------------------------------------------------------------
+
+    MUTABLE_IMPLEMENT_NODE( NodeSurfaceNew, EType::New, Node, Node::EType::Surface)
+
+
+	//---------------------------------------------------------------------------------------------
+	// Node Interface
+	//---------------------------------------------------------------------------------------------
+    int NodeSurfaceNew::GetInputCount() const
+	{
+        return (int)( m_pD->m_meshes.size() + m_pD->m_images.size() + m_pD->m_vectors.size() +
+                      m_pD->m_scalars.size() + m_pD->m_strings.size() );
+    }
+
+
+	//---------------------------------------------------------------------------------------------
+    mu::Node* NodeSurfaceNew::GetInputNode( int i ) const
+	{
+		check( i >=0 && i < GetInputCount() );
+
+        std::size_t index = size_t(i);
+
+        if ( index < m_pD->m_meshes.size() )
+        {
+            return m_pD->m_meshes[index].m_pMesh.get();
+        }
+		else
+		{
+            index -= m_pD->m_meshes.size();
+        }
+
+        if ( index < m_pD->m_images.size() )
+        {
+            return m_pD->m_images[index].m_pImage.get();
+        }
+        else 
+		{
+            index -= m_pD->m_images.size();
+        }
+
+        if ( index < m_pD->m_vectors.size() )
+        {
+            return m_pD->m_vectors[index].m_pVector.get();
+        }
+        else
+        {
+            index -= m_pD->m_vectors.size();
+        }
+
+        if ( index < m_pD->m_scalars.size() )
+        {
+            return m_pD->m_scalars[index].m_pScalar.get();
+        }
+        else
+        {
+            index -= m_pD->m_scalars.size();
+        }
+
+        if ( index < m_pD->m_strings.size() )
+        {
+            return m_pD->m_strings[index].m_pString.get();
+        }
+        else
+        {
+            index -= m_pD->m_strings.size();
+        }
+
+        return nullptr;
+    }
+
+
+	//---------------------------------------------------------------------------------------------
+    void NodeSurfaceNew::SetInputNode( int i, NodePtr pNode )
+	{
+		check( i >=0 && i < GetInputCount());
+
+        std::size_t index = size_t( i );
+
+        if ( index<m_pD->m_meshes.size() )
+		{
+			m_pD->m_meshes[ index ].m_pMesh = dynamic_cast<NodeMesh*>(pNode.get());
+            return;
+        }
+		else
+		{
+            index -= m_pD->m_meshes.size();
+        }
+
+
+        if ( index < m_pD->m_images.size())
+        {
+			m_pD->m_images[index].m_pImage = dynamic_cast<NodeImage*>(pNode.get());
+            return;
+        }
+        else
+        {
+            index -= m_pD->m_images.size();
+        }
+        
+		if ( index < m_pD->m_vectors.size() )
+        {
+            m_pD->m_vectors[index].m_pVector = dynamic_cast<NodeColour*>(pNode.get());
+            return;
+        }
+        else
+        {
+            index -= m_pD->m_vectors.size();
+        }
+
+        if ( index < m_pD->m_scalars.size() )
+        {
+            m_pD->m_scalars[index].m_pScalar = dynamic_cast<NodeScalar*>( pNode.get() );
+            return;
+        }
+        else
+        {
+            index -= m_pD->m_scalars.size();
+        }
+
+        if ( index < m_pD->m_strings.size() )
+        {
+            m_pD->m_strings[index].m_pString = dynamic_cast<NodeString*>( pNode.get() );
+            return;
+        }
+        else
+        {
+            index -= m_pD->m_strings.size();
+        }
+    }
+
+
+	//---------------------------------------------------------------------------------------------
+	// Own Interface
+	//---------------------------------------------------------------------------------------------
+    const char* NodeSurfaceNew::GetName() const
+	{
+		const char* strResult = m_pD->m_name.c_str();
+
+		return strResult;
+	}
+
+
+    //---------------------------------------------------------------------------------------------
+    void NodeSurfaceNew::SetName( const char* strName )
+    {
+        if (strName)
+        {
+            m_pD->m_name = strName;
+        }
+        else
+        {
+            m_pD->m_name = "";
+        }
+    }
+
+
+    //---------------------------------------------------------------------------------------------
+    void NodeSurfaceNew::SetCustomID( uint32_t id )
+    {
+        m_pD->m_customID = id;
+    }
+
+
+	//---------------------------------------------------------------------------------------------
+    int NodeSurfaceNew::GetMeshCount() const
+	{
+		return (int)m_pD->m_meshes.size();
+	}
+
+
+	//---------------------------------------------------------------------------------------------
+    void NodeSurfaceNew::SetMeshCount( int num )
+	{
+		check( num >=0 );
+		m_pD->m_meshes.resize( num );
+	}
+
+
+	//---------------------------------------------------------------------------------------------
+    NodeMeshPtr NodeSurfaceNew::GetMesh( int index ) const
+	{
+		check( index >=0 && index < (int)m_pD->m_meshes.size() );
+
+		return m_pD->m_meshes[ index ].m_pMesh.get();
+	}
+
+
+	//---------------------------------------------------------------------------------------------
+    void NodeSurfaceNew::SetMesh( int index, NodeMeshPtr pNode )
+	{
+		check( index >=0 && index < (int)m_pD->m_meshes.size() );
+
+		m_pD->m_meshes[ index ].m_pMesh = pNode;
+	}
+
+
+	//---------------------------------------------------------------------------------------------
+    const char* NodeSurfaceNew::GetMeshName( int index ) const
+	{
+		check( index >=0 && index < (int)m_pD->m_meshes.size() );
+
+		const char* strResult = m_pD->m_meshes[ index ].m_name.c_str();
+
+		return strResult;
+	}
+
+
+	//---------------------------------------------------------------------------------------------
+    void NodeSurfaceNew::SetMeshName( int index, const char* strName )
+	{
+		check( index >=0 && index < (int)m_pD->m_meshes.size() );
+
+		if (strName)
+		{
+			m_pD->m_meshes[ index ].m_name = strName;
+		}
+		else
+		{
+			m_pD->m_meshes[ index ].m_name = "";
+		}
+	}
+
+
+	//---------------------------------------------------------------------------------------------
+    int NodeSurfaceNew::GetImageCount() const
+	{
+		return (int)m_pD->m_images.size();
+	}
+
+
+	//---------------------------------------------------------------------------------------------
+    void NodeSurfaceNew::SetImageCount( int num )
+	{
+		check( num >=0 );
+		m_pD->m_images.resize( num );
+	}
+
+
+	//---------------------------------------------------------------------------------------------
+    NodeImagePtr NodeSurfaceNew::GetImage( int index ) const
+	{
+		check( index >=0 && index < (int)m_pD->m_images.size() );
+
+		return m_pD->m_images[ index ].m_pImage.get();
+	}
+
+
+	//---------------------------------------------------------------------------------------------
+    void NodeSurfaceNew::SetImage( int index, NodeImagePtr pNode )
+	{
+		check( index >=0 && index < (int)m_pD->m_images.size() );
+
+		m_pD->m_images[ index ].m_pImage = pNode;
+	}
+
+
+	//---------------------------------------------------------------------------------------------
+    const char* NodeSurfaceNew::GetImageName( int index ) const
+	{
+		check( index >=0 && index < (int)m_pD->m_images.size() );
+
+		const char* strResult = m_pD->m_images[ index ].m_name.c_str();
+
+		return strResult;
+	}
+
+
+	//---------------------------------------------------------------------------------------------
+    void NodeSurfaceNew::SetImageName( int index, const char* strName )
+	{
+		check( index >=0 && index < (int)m_pD->m_images.size() );
+
+		if (strName)
+		{
+			m_pD->m_images[ index ].m_name = strName;
+		}
+		else
+		{
+			m_pD->m_images[ index ].m_name = "";
+		}
+	}
+
+
+    //---------------------------------------------------------------------------------------------
+    int NodeSurfaceNew::GetImageLayoutIndex( int index ) const
+    {
+        check( index >= 0 && index < (int)m_pD->m_images.size() );
+
+        return int(m_pD->m_images[index].m_layoutIndex);
+    }
+
+
+	//---------------------------------------------------------------------------------------------
+	void NodeSurfaceNew::SetImageLayoutIndex(int index, int layoutIndex)
+	{
+		check(index >= 0 && index < (int)m_pD->m_images.size());
+
+		m_pD->m_images[index].m_layoutIndex = int8_t(layoutIndex);
+	}
+
+
+	//---------------------------------------------------------------------------------------------
+	void NodeSurfaceNew::SetImageAdditionalNames(int index, const char* strMaterialName, const char* strMaterialParameterName )
+	{
+		if (index >= 0 && index < (int)m_pD->m_images.size())
+		{
+			m_pD->m_images[index].m_materialName = strMaterialName;
+			m_pD->m_images[index].m_materialParameterName = strMaterialParameterName;
+		}
+		else
+		{
+			check(false);
+		}
+	}
+
+
+    //---------------------------------------------------------------------------------------------
+	int NodeSurfaceNew::GetVectorCount() const
+	{
+		return (int)m_pD->m_vectors.size();
+	}
+
+
+	//---------------------------------------------------------------------------------------------
+	void NodeSurfaceNew::SetVectorCount(int num)
+	{
+		check(num >= 0);
+		m_pD->m_vectors.resize(num);
+	}
+
+
+	//---------------------------------------------------------------------------------------------
+	NodeColourPtr NodeSurfaceNew::GetVector(int index) const
+	{
+		check(index >= 0 && index < (int)m_pD->m_vectors.size());
+
+		return m_pD->m_vectors[index].m_pVector.get();
+	}
+
+
+	//---------------------------------------------------------------------------------------------
+	void NodeSurfaceNew::SetVector(int index, NodeColourPtr pNode)
+	{
+		check(index >= 0 && index < (int)m_pD->m_vectors.size());
+
+		m_pD->m_vectors[index].m_pVector = pNode;
+	}
+
+
+	//---------------------------------------------------------------------------------------------
+	const char* NodeSurfaceNew::GetVectorName(int index) const
+	{
+        check(index >= 0 && index < (int)m_pD->m_vectors.size());
+
+		const char* strResult = m_pD->m_vectors[index].m_name.c_str();
+
+		return strResult;
+	}
+
+
+	//---------------------------------------------------------------------------------------------
+	void NodeSurfaceNew::SetVectorName(int index, const char* strName)
+	{
+        check(index >= 0 && index < (int)m_pD->m_vectors.size());
+
+		if (strName)
+		{
+			m_pD->m_vectors[index].m_name = strName;
+		}
+		else
+		{
+			m_pD->m_vectors[index].m_name = "";
+		}
+	}
+
+
+    //---------------------------------------------------------------------------------------------
+    int NodeSurfaceNew::GetScalarCount() const
+    {
+        return (int)m_pD->m_scalars.size();
+    }
+
+
+    //---------------------------------------------------------------------------------------------
+    void NodeSurfaceNew::SetScalarCount(int num)
+    {
+        check(num >= 0);
+        m_pD->m_scalars.resize(num);
+    }
+
+
+    //---------------------------------------------------------------------------------------------
+    NodeScalarPtr NodeSurfaceNew::GetScalar(int index) const
+    {
+        check(index >= 0 && index < (int)m_pD->m_scalars.size());
+
+        return m_pD->m_scalars[index].m_pScalar.get();
+    }
+
+
+    //---------------------------------------------------------------------------------------------
+    void NodeSurfaceNew::SetScalar(int index, NodeScalarPtr pNode)
+    {
+        check(index >= 0 && index < (int)m_pD->m_scalars.size());
+
+        m_pD->m_scalars[index].m_pScalar = pNode;
+    }
+
+
+    //---------------------------------------------------------------------------------------------
+    const char* NodeSurfaceNew::GetScalarName(int index) const
+    {
+        check(index >= 0 && index < (int)m_pD->m_scalars.size());
+
+        const char* strResult = m_pD->m_scalars[index].m_name.c_str();
+
+        return strResult;
+    }
+
+
+    //---------------------------------------------------------------------------------------------
+    void NodeSurfaceNew::SetScalarName(int index, const char* strName)
+    {
+        check(index >= 0 && index < (int)m_pD->m_scalars.size());
+
+        if (strName)
+        {
+            m_pD->m_scalars[index].m_name = strName;
+        }
+        else
+        {
+            m_pD->m_scalars[index].m_name = "";
+        }
+    }
+
+
+    //---------------------------------------------------------------------------------------------
+    int NodeSurfaceNew::GetStringCount() const { return (int)m_pD->m_strings.size(); }
+
+
+    //---------------------------------------------------------------------------------------------
+    void NodeSurfaceNew::SetStringCount( int num )
+    {
+        check( num >= 0 );
+        m_pD->m_strings.resize( num );
+    }
+
+
+    //---------------------------------------------------------------------------------------------
+    NodeStringPtr NodeSurfaceNew::GetString( int index ) const
+    {
+        check( index >= 0 && index < (int)m_pD->m_strings.size() );
+
+        return m_pD->m_strings[index].m_pString.get();
+    }
+
+
+    //---------------------------------------------------------------------------------------------
+    void NodeSurfaceNew::SetString( int index, NodeStringPtr pNode )
+    {
+        check( index >= 0 && index < (int)m_pD->m_strings.size() );
+
+        m_pD->m_strings[index].m_pString = pNode;
+    }
+
+
+    //---------------------------------------------------------------------------------------------
+    const char* NodeSurfaceNew::GetStringName( int index ) const
+    {
+        check( index >= 0 && index < (int)m_pD->m_strings.size() );
+
+        const char* strResult = m_pD->m_strings[index].m_name.c_str();
+
+        return strResult;
+    }
+
+
+    //---------------------------------------------------------------------------------------------
+    void NodeSurfaceNew::SetStringName( int index, const char* strName )
+    {
+        check( index >= 0 && index < (int)m_pD->m_strings.size() );
+
+        if ( strName )
+        {
+            m_pD->m_strings[index].m_name = strName;
+        }
+        else
+        {
+            m_pD->m_strings[index].m_name = "";
+        }
+    }
+
+
+    //---------------------------------------------------------------------------------------------
+    int NodeSurfaceNew::Private::FindImage( const char* strName ) const
+	{
+		int i = 0;
+		for ( vector<IMAGE>::const_iterator it = m_images.begin()
+			; it!=m_images.end()
+			; ++it, ++i )
+		{
+			if ( it->m_name == strName )
+			{
+				return i;
+			}
+		}
+
+		return -1;
+	}
+
+
+	//---------------------------------------------------------------------------------------------
+	int NodeSurfaceNew::Private::FindMesh(const char* strName) const
+	{
+		int i = 0;
+		for (vector<MESH>::const_iterator it = m_meshes.begin()
+			; it != m_meshes.end()
+			; ++it, ++i)
+		{
+			if (it->m_name == strName)
+			{
+				return i;
+			}
+		}
+
+		return -1;
+	}
+
+
+    //---------------------------------------------------------------------------------------------
+    int NodeSurfaceNew::Private::FindVector(const char* strName) const
+    {
+        int i = 0;
+        for (vector<VECTOR>::const_iterator it = m_vectors.begin()
+            ; it != m_vectors.end()
+            ; ++it, ++i)
+        {
+            if (it->m_name == strName)
+            {
+                return i;
+            }
+        }
+
+        return -1;
+    }
+
+
+    //---------------------------------------------------------------------------------------------
+    int NodeSurfaceNew::Private::FindScalar( const char* strName ) const
+    {
+        int i = 0;
+        for ( vector<SCALAR>::const_iterator it = m_scalars.begin(); it != m_scalars.end();
+              ++it, ++i )
+        {
+            if ( it->m_name == strName )
+            {
+                return i;
+            }
+        }
+
+        return -1;
+    }
+
+
+    //---------------------------------------------------------------------------------------------
+    int NodeSurfaceNew::Private::FindString( const char* strName ) const
+    {
+        int i = 0;
+        for ( vector<STRING>::const_iterator it = m_strings.begin(); it != m_strings.end();
+              ++it, ++i )
+        {
+            if ( it->m_name == strName )
+            {
+                return i;
+            }
+        }
+
+        return -1;
+    }
+
+
+    //---------------------------------------------------------------------------------------------
+	void NodeSurfaceNew::AddTag(const char* tagName)
+	{
+		m_pD->m_tags.push_back(tagName);
+	}
+
+
+    //---------------------------------------------------------------------------------------------
+    int NodeSurfaceNew::GetTagCount() const
+    {
+        return int( m_pD->m_tags.size());
+    }
+
+
+    //---------------------------------------------------------------------------------------------
+    const char* NodeSurfaceNew::GetTag( int i ) const
+    {
+        if (i>=0 && i<GetTagCount())
+        {
+            return m_pD->m_tags[i].c_str();
+        }
+        return nullptr;
+    }
+
+}
+
+
diff --git a/Engine/Plugins/Experimental/Mutable/Source/MutableTools/Private/NodeSurfaceNewPrivate.h b/Engine/Plugins/Experimental/Mutable/Source/MutableTools/Private/NodeSurfaceNewPrivate.h
new file mode 100644
--- /dev/null
+++ b/Engine/Plugins/Experimental/Mutable/Source/MutableTools/Private/NodeSurfaceNewPrivate.h
@@ -0,0 +1,212 @@
+// Copyright Epic Games, Inc. All Rights Reserved.
+
+#pragma once
+
+#include "NodePrivate.h"
+#include "NodeSurfaceNew.h"
+#include "NodePatchImagePrivate.h"
+
+#include "NodeMesh.h"
+#include "NodePatchMesh.h"
+#include "NodeImage.h"
+#include "NodeScalar.h"
+#include "NodeString.h"
+#include "NodeColour.h"
+
+#include "MemoryPrivate.h"
+
+
+namespace mu
+{
+
+	//---------------------------------------------------------------------------------------------
+	//!
+	//---------------------------------------------------------------------------------------------
+    class NodeSurfaceNew::Private : public Node::Private
+	{
+	public:
+
+		MUTABLE_DEFINE_CONST_VISITABLE()
+
+	public:
+
+		static NODE_TYPE s_type;
+
+		string m_name;
+        uint32_t m_customID=0;
+
+		struct MESH
+		{
+			string m_name;
+			NodeMeshPtr m_pMesh;
+
+			//!
+			void Serialise( OutputArchive& arch ) const
+			{
+				arch << m_name;
+				arch << m_pMesh;
+			}
+
+			void Unserialise( InputArchive& arch )
+			{
+				arch >> m_name;
+				arch >> m_pMesh;
+			}
+		};
+
+		vector<MESH> m_meshes;
+
+		struct IMAGE
+		{
+			string m_name;
+			string m_materialName;
+			string m_materialParameterName;
+			NodeImagePtr m_pImage;
+
+			// It could be negative, to indicate no layout.
+            int8_t m_layoutIndex = 0;
+
+            //!
+			void Serialise(OutputArchive& arch) const
+			{
+				arch << m_name;
+				arch << m_materialName;
+				arch << m_materialParameterName;
+				arch << m_pImage;
+                arch << m_layoutIndex;
+            }
+
+			void Unserialise(InputArchive& arch)
+			{
+				arch >> m_name;
+				arch >> m_materialName;
+				arch >> m_materialParameterName;
+				arch >> m_pImage;
+                arch >> m_layoutIndex;
+            }
+        };
+
+		vector<IMAGE> m_images;
+
+		struct VECTOR
+		{
+			string m_name;
+			NodeColourPtr m_pVector;
+
+			//!
+			void Serialise(OutputArchive& arch) const
+			{
+				arch << m_name;
+				arch << m_pVector;
+			}
+
+			void Unserialise(InputArchive& arch)
+			{
+				arch >> m_name;
+				arch >> m_pVector;
+			}
+		};
+
+		vector<VECTOR> m_vectors;
+
+        struct SCALAR
+        {
+            string m_name;
+            NodeScalarPtr m_pScalar;
+
+            //!
+            void Serialise(OutputArchive& arch) const
+            {
+                arch << m_name;
+                arch << m_pScalar;
+            }
+
+            void Unserialise(InputArchive& arch)
+            {
+                arch >> m_name;
+                arch >> m_pScalar;
+            }
+        };
+
+        vector<SCALAR> m_scalars;
+
+        struct STRING
+        {
+            string m_name;
+            NodeStringPtr m_pString;
+
+            //!
+            void Serialise( OutputArchive& arch ) const
+            {
+                arch << m_name;
+                arch << m_pString;
+            }
+
+            void Unserialise( InputArchive& arch )
+            {
+                arch >> m_name;
+                arch >> m_pString;
+            }
+        };
+
+        vector<STRING> m_strings;
+
+        //! Tags in this surface
+		vector<mu::string> m_tags;
+
+		//! Find an image node index by name or return -1
+		int FindImage( const char* strName ) const;
+
+		//! Find a mesh node index by name or return -1
+		int FindMesh(const char* strName) const;
+
+        //! Find a vector node index by name or return -1
+        int FindVector(const char* strName) const;
+
+        //! Find a scalar node index by name or return -1
+        int FindScalar( const char* strName ) const;
+
+        //! Find a string node index by name or return -1
+        int FindString( const char* strName ) const;
+
+        //!
+		void Serialise( OutputArchive& arch ) const
+		{
+            uint32_t ver = 7;
+			arch << ver;
+
+			arch << m_name;
+            arch << m_customID;
+
+			arch << m_meshes;
+			arch << m_images;
+			arch << m_tags;
+            arch << m_vectors;
+            arch << m_scalars;
+            arch << m_strings;
+        }
+
+        //!
+		void Unserialise( InputArchive& arch )
+		{
+            uint32_t ver;
+			arch >> ver;
+            check(ver==7);
+
+			arch >> m_name;
+            arch >> m_customID;
+
+			arch >> m_meshes;
+			arch >> m_images;
+            arch >> m_tags;
+            arch >> m_vectors;
+            arch >> m_scalars;
+            arch >> m_strings;
+        }
+    };
+
+
+
+
+}
+
diff --git a/Engine/Plugins/Experimental/Mutable/Source/MutableTools/Private/NodeSurfaceVariation.cpp b/Engine/Plugins/Experimental/Mutable/Source/MutableTools/Private/NodeSurfaceVariation.cpp
new file mode 100644
--- /dev/null
+++ b/Engine/Plugins/Experimental/Mutable/Source/MutableTools/Private/NodeSurfaceVariation.cpp
@@ -0,0 +1,195 @@
+// Copyright Epic Games, Inc. All Rights Reserved.
+
+
+#include "MutableTools/Public/NodeSurfaceVariation.h"
+#include "MutableTools/Private/NodeSurfaceVariationPrivate.h"
+
+#include "MutableTools/Public/NodeMesh.h"
+#include "MutableTools/Public/NodeImage.h"
+
+
+namespace mu
+{
+
+
+	//---------------------------------------------------------------------------------------------
+	// Static initialisation
+	//---------------------------------------------------------------------------------------------
+    NODE_TYPE NodeSurfaceVariation::Private::s_type =
+            NODE_TYPE( "SurfaceVariation", NodeSurface::GetStaticType() );
+
+
+	//---------------------------------------------------------------------------------------------
+	//---------------------------------------------------------------------------------------------
+	//---------------------------------------------------------------------------------------------
+
+    MUTABLE_IMPLEMENT_NODE( NodeSurfaceVariation, EType::Variation, Node, Node::EType::Surface)
+
+
+	//---------------------------------------------------------------------------------------------
+	// Node Interface
+	//---------------------------------------------------------------------------------------------
+    int NodeSurfaceVariation::GetInputCount() const
+	{
+        size_t c = m_pD->m_defaultSurfaces.size();
+        c += m_pD->m_defaultModifiers.size();
+        for (const auto& v : m_pD->m_variations)
+		{
+            c += v.m_surfaces.size();
+            c += v.m_modifiers.size();
+        }
+
+		return (int)c;
+	}
+
+
+	//---------------------------------------------------------------------------------------------
+    Node* NodeSurfaceVariation::GetInputNode( int i ) const
+	{
+		check( i >=0 && i < GetInputCount() );
+
+        if ( i<(int)m_pD->m_defaultSurfaces.size())
+		{
+            return m_pD->m_defaultSurfaces[i].get();
+		}
+        i -= (int)m_pD->m_defaultSurfaces.size();
+
+        if ( i<(int)m_pD->m_defaultModifiers.size())
+        {
+            return m_pD->m_defaultModifiers[i].get();
+        }
+        i -= (int)m_pD->m_defaultModifiers.size();
+
+        for (const auto& v : m_pD->m_variations)
+        {
+            if (i < (int)v.m_surfaces.size())
+            {
+                return v.m_surfaces[i].get();
+            }
+            i -= (int)v.m_surfaces.size();
+
+            if (i < (int)v.m_modifiers.size())
+            {
+                return v.m_modifiers[i].get();
+            }
+            i -= (int)v.m_modifiers.size();
+        }
+
+		return nullptr;
+	}
+
+
+	//---------------------------------------------------------------------------------------------
+    void NodeSurfaceVariation::SetInputNode( int i, NodePtr pNode )
+	{
+		check( i >=0 && i <  GetInputCount());
+
+        if (i<(int)m_pD->m_defaultSurfaces.size())
+		{
+            m_pD->m_defaultSurfaces[i] = dynamic_cast<NodeSurface*>(pNode.get());
+			return;
+		}
+
+        i -= (int)m_pD->m_defaultSurfaces.size();
+        if (i<(int)m_pD->m_defaultModifiers.size())
+        {
+
+            m_pD->m_defaultModifiers[i] = dynamic_cast<NodeModifier*>(pNode.get());
+            return;
+        }
+        i -= (int)m_pD->m_defaultModifiers.size();
+
+        for (auto& v : m_pD->m_variations)
+        {
+            if (i < (int)v.m_surfaces.size())
+            {
+                v.m_surfaces[i] = dynamic_cast<NodeSurface*>(pNode.get());
+                return;
+            }
+            i -= (int)v.m_surfaces.size();
+
+            if (i < (int)v.m_modifiers.size())
+            {
+                v.m_modifiers[i] = dynamic_cast<NodeModifier*>(pNode.get());
+                return;
+            }
+            i -= (int)v.m_modifiers.size();
+        }
+	}
+
+
+	//---------------------------------------------------------------------------------------------
+	// Own Interface
+	//---------------------------------------------------------------------------------------------
+    void NodeSurfaceVariation::AddDefaultSurface( NodeSurface* p )
+    {
+        m_pD->m_defaultSurfaces.push_back(p);
+    }
+
+
+    //---------------------------------------------------------------------------------------------
+    void NodeSurfaceVariation::AddDefaultModifier( NodeModifier* p )
+    {
+        m_pD->m_defaultModifiers.push_back(p);
+    }
+
+
+    //---------------------------------------------------------------------------------------------
+    int NodeSurfaceVariation::GetVariationCount() const
+	{
+		return (int)m_pD->m_variations.size();
+	}
+
+
+	//---------------------------------------------------------------------------------------------
+    void NodeSurfaceVariation::SetVariationCount( int num )
+	{
+		check( num >=0 );
+		m_pD->m_variations.resize( num );
+	}
+
+
+    //---------------------------------------------------------------------------------------------
+    void NodeSurfaceVariation::SetVariationType(VariationType type )
+    {
+        m_pD->m_type = type;
+    }
+
+
+	//---------------------------------------------------------------------------------------------
+	void NodeSurfaceVariation::SetVariationTag(int index, const char* strTag)
+	{
+		check(index >= 0 && index < (int)m_pD->m_variations.size());
+		check(strTag);
+
+        if (strTag)
+        {
+            m_pD->m_variations[index].m_tag = strTag;
+        }
+        else
+        {
+            m_pD->m_variations[index].m_tag = "";
+        }
+	}
+
+
+	//---------------------------------------------------------------------------------------------
+	void NodeSurfaceVariation::AddVariationSurface(int index, NodeSurface* pNode)
+	{
+		check(index >= 0 && index < (int)m_pD->m_variations.size());
+
+		m_pD->m_variations[index].m_surfaces.push_back(pNode);
+	}
+
+
+    //---------------------------------------------------------------------------------------------
+    void NodeSurfaceVariation::AddVariationModifier(int index, NodeModifier* pModifier)
+    {
+        check(index >= 0 && index < (int)m_pD->m_variations.size());
+
+        m_pD->m_variations[index].m_modifiers.push_back(pModifier);
+    }
+
+}
+
+
diff --git a/Engine/Plugins/Experimental/Mutable/Source/MutableTools/Private/NodeSurfaceVariationPrivate.h b/Engine/Plugins/Experimental/Mutable/Source/MutableTools/Private/NodeSurfaceVariationPrivate.h
new file mode 100644
--- /dev/null
+++ b/Engine/Plugins/Experimental/Mutable/Source/MutableTools/Private/NodeSurfaceVariationPrivate.h
@@ -0,0 +1,95 @@
+// Copyright Epic Games, Inc. All Rights Reserved.
+
+#pragma once
+
+#include "NodePrivate.h"
+#include "NodeSurfaceVariation.h"
+#include "NodeSurfaceNewPrivate.h"
+#include "NodeModifier.h"
+
+#include "MemoryPrivate.h"
+
+
+namespace mu
+{
+
+	//---------------------------------------------------------------------------------------------
+	//!
+	//---------------------------------------------------------------------------------------------
+    class NodeSurfaceVariation::Private : public NodeSurface::Private
+	{
+	public:
+
+		MUTABLE_DEFINE_CONST_VISITABLE()
+
+	public:
+
+
+		static NODE_TYPE s_type;
+
+        vector<NodeSurfacePtr> m_defaultSurfaces;
+        vector<NodeModifierPtr> m_defaultModifiers;
+
+		struct VARIATION
+		{
+            vector<NodeSurfacePtr> m_surfaces;
+            vector<NodeModifierPtr> m_modifiers;
+            string m_tag;
+
+			//!
+			void Serialise( OutputArchive& arch ) const
+			{
+                uint32_t ver = 1;
+				arch << ver;
+
+				arch << m_tag;
+                arch << m_surfaces;
+                arch << m_modifiers;
+            }
+
+			void Unserialise( InputArchive& arch )
+			{
+                uint32_t ver = 0;
+				arch >> ver;
+                check(ver==1);
+
+				arch >> m_tag;
+				arch >> m_surfaces;
+                arch >> m_modifiers;
+			}
+		};
+
+        NodeSurfaceVariation::VariationType m_type = NodeSurfaceVariation::VariationType::Tag;
+
+		vector<VARIATION> m_variations;
+
+		//!
+		void Serialise( OutputArchive& arch ) const
+		{
+            uint32_t ver = 3;
+			arch << ver;
+
+            arch << uint32_t(m_type);
+            arch << m_defaultSurfaces;
+            arch << m_defaultModifiers;
+            arch << m_variations;
+		}
+
+		//!
+		void Unserialise( InputArchive& arch )
+		{
+            uint32_t ver;
+			arch >> ver;
+            check( ver==3 );
+
+            uint32_t temp;
+            arch >> temp;
+            m_type = NodeSurfaceVariation::VariationType(temp);
+
+            arch >> m_defaultSurfaces;
+            arch >> m_defaultModifiers;
+            arch >> m_variations;
+		}
+	};
+
+}
diff --git a/Engine/Plugins/Experimental/Mutable/Source/MutableTools/Private/Platform.h b/Engine/Plugins/Experimental/Mutable/Source/MutableTools/Private/Platform.h
new file mode 100644
--- /dev/null
+++ b/Engine/Plugins/Experimental/Mutable/Source/MutableTools/Private/Platform.h
@@ -0,0 +1,89 @@
+// Copyright Epic Games, Inc. All Rights Reserved.
+
+#pragma once
+
+#include "MutableRuntime/Private/Platform.h"
+#include "MemoryPrivate.h"
+
+#include <sys/stat.h>
+
+//#if defined(_MSC_VER) || defined(__MINGW32__)
+    //#include <windef.h>
+    //#include <winbase.h>
+//#endif
+
+// This file contains some hacks to solve differences between platforms necessary for the tools
+// library and not present in the run-time library.
+
+
+namespace mu
+{
+
+    //-------------------------------------------------------------------------------------------------
+    inline FILE* mutable_fopen
+		(
+			const char *filename,
+			const char *mode
+		)
+	{
+		FILE* pFile = 0;
+
+#ifdef _MSC_VER
+        errno_t error = fopen_s( &pFile, filename, mode );
+        if (error!=0)
+        {
+            // It's valid to fail, just return null.
+            //check(false);
+            pFile = nullptr;
+        }
+#else
+		pFile = fopen( filename, mode );
+#endif
+
+		return pFile;
+	}
+
+
+    //-------------------------------------------------------------------------------------------------
+    inline int64_t mutable_ftell( FILE* f )
+    {
+#ifdef _MSC_VER
+        // Windows normal ftell works only with 32 bit int sizes
+        return _ftelli64(f);
+#else
+        return ftell(f);
+#endif
+    }
+
+
+    //-------------------------------------------------------------------------------------------------
+    inline int64_t mutable_fseek( FILE* f, int64_t pos, int origin )
+    {
+#ifdef _MSC_VER
+        // Windows normal fseek works only with 32 bit int sizes
+        return _fseeki64(f,pos,origin);
+#else
+        return fseek(f,pos,origin);
+#endif
+    }
+
+
+    //-------------------------------------------------------------------------------------------------
+    inline string mutable_path_separator()
+    {
+    #ifdef _WIN32
+        return "\\";
+    #else
+        return "/";
+    #endif
+    }
+
+
+    //-------------------------------------------------------------------------------------------------
+    inline bool mutable_file_exists( const string& path )
+    {
+        struct stat buffer;
+        return (stat (path.c_str(), &buffer) == 0);
+    }
+
+}
diff --git a/Engine/Plugins/Experimental/Mutable/Source/MutableTools/Private/Streams.cpp b/Engine/Plugins/Experimental/Mutable/Source/MutableTools/Private/Streams.cpp
new file mode 100644
--- /dev/null
+++ b/Engine/Plugins/Experimental/Mutable/Source/MutableTools/Private/Streams.cpp
@@ -0,0 +1,425 @@
+// Copyright Epic Games, Inc. All Rights Reserved.
+
+
+#include "Streams.h"
+#include "StreamsPrivate.h"
+#include "Platform.h"
+#include "NodeImageConstant.h"
+#include "Image.h"
+#include "MutableTrace.h"
+
+#define MUTABLE_STREAM_BUFFER_SIZE 0x10000
+
+namespace mu
+{
+	InputFileStream::InputFileStream( const char* strFile )
+	{
+		m_pD = new Private();
+
+		m_pD->m_pFile = mutable_fopen( strFile, "rb" );
+		m_pD->m_size = 0;
+		m_pD->m_pos = 0;
+        m_pD->m_readBytes = 0;
+
+        m_pD->m_buffer = (uint8_t*)mutable_malloc_aligned(MUTABLE_STREAM_BUFFER_SIZE,16);
+	}
+
+
+	//---------------------------------------------------------------------------------------------
+	InputFileStream::~InputFileStream()
+	{
+        if ( m_pD->m_pFile )
+        {
+            fclose( m_pD->m_pFile );
+        }
+
+        mutable_free_aligned( m_pD->m_buffer, MUTABLE_STREAM_BUFFER_SIZE );
+        m_pD->m_buffer = nullptr;
+
+        delete m_pD;
+		m_pD = nullptr;
+	}
+    
+    
+	//---------------------------------------------------------------------------------------------
+	bool InputFileStream::IsOpen() const
+    {
+        return m_pD->m_pFile!=0;
+    }
+    
+
+    //---------------------------------------------------------------------------------------------
+	uint64 InputFileStream::GetReadBytes() const
+    {
+        return m_pD->m_readBytes;
+    }
+    
+
+    //---------------------------------------------------------------------------------------------
+	uint64 InputFileStream::Tell() const
+    {
+        size_t bufferOffset = m_pD->m_size - m_pD->m_pos;
+        int64_t filePos = mutable_ftell( m_pD->m_pFile );
+        check(filePos>0);
+        size_t pos = size_t(filePos) - bufferOffset;
+        return pos;
+    }
+
+
+    //---------------------------------------------------------------------------------------------
+    void InputFileStream::Seek( uint64 position )
+    {
+        int64_t result = mutable_fseek( m_pD->m_pFile, int64_t(position), SEEK_SET );
+        check(result==0);
+        (void)result;
+
+        m_pD->m_size = 0;
+        m_pD->m_pos = 0;
+    }
+
+
+	//---------------------------------------------------------------------------------------------
+    void InputFileStream::Read( void* pData, uint64 size )
+	{
+        uint8_t* pDest = (uint8_t*) pData;
+        while (size && m_pD->m_pFile )
+		{			
+			int available = std::min( int(m_pD->m_size) - int(m_pD->m_pos), int(size) );
+
+			// Copy from what is already loaded
+			if ( available>0 )
+			{
+                FMemory::Memcpy( pDest, m_pD->m_buffer+m_pD->m_pos, available );
+				m_pD->m_pos += available;
+				size -= available;
+				pDest += available;
+			}
+
+			// Load more data if necessary
+			if ( m_pD->m_pos == m_pD->m_size && size )
+			{
+                m_pD->m_size = fread( m_pD->m_buffer, 1, MUTABLE_STREAM_BUFFER_SIZE, m_pD->m_pFile );
+                check( !ferror( m_pD->m_pFile ) );
+                check( m_pD->m_size );
+				m_pD->m_pos = 0;
+                m_pD->m_readBytes += m_pD->m_size;
+            }
+		}
+
+	}
+
+
+	//---------------------------------------------------------------------------------------------
+	//---------------------------------------------------------------------------------------------
+	//---------------------------------------------------------------------------------------------
+	OutputFileStream::OutputFileStream( const char* strFile )
+	{
+		m_pD = new Private();
+		m_pD->m_pFile = mutable_fopen( strFile, "wb" );
+		m_pD->m_pos = 0;
+        m_pD->m_buffer = (uint8_t*)mutable_malloc_aligned(MUTABLE_STREAM_BUFFER_SIZE,16);
+
+		check( m_pD->m_pFile );
+	}
+
+
+	//---------------------------------------------------------------------------------------------
+    void OutputFileStream::Write( const void* pData, uint64 size )
+	{
+        const uint8_t* pSrc = (const uint8_t*) pData;
+
+		while (size)
+		{
+            int available = std::min( int(MUTABLE_STREAM_BUFFER_SIZE) - int(m_pD->m_pos), int(size) );
+
+			// Copy from what fits in the current buffer
+			if ( available>0 )
+			{
+                memcpy( m_pD->m_buffer+m_pD->m_pos, pSrc, available );
+				m_pD->m_pos += available;
+				size -= available;
+				pSrc += available;
+			}
+
+			// Flush the data if necessary
+            if ( m_pD->m_pos == MUTABLE_STREAM_BUFFER_SIZE )
+			{
+				Flush();
+			}
+		}
+	}
+
+
+	//---------------------------------------------------------------------------------------------
+	void OutputFileStream::Flush()
+	{
+		if ( m_pD->m_pos )
+		{
+            std::size_t count = fwrite( m_pD->m_buffer,
+										1,
+										m_pD->m_pos,
+										m_pD->m_pFile );
+
+			check( count == m_pD->m_pos );
+            (void)count;
+        }
+
+		m_pD->m_pos = 0;
+	}
+
+
+	//---------------------------------------------------------------------------------------------
+	OutputFileStream::~OutputFileStream()
+	{
+		Flush();
+
+		fclose( m_pD->m_pFile );
+
+        mutable_free_aligned( m_pD->m_buffer, MUTABLE_STREAM_BUFFER_SIZE );
+        m_pD->m_buffer = nullptr;
+
+        delete m_pD;
+		m_pD = 0;
+	}
+
+
+    //---------------------------------------------------------------------------------------------
+    //---------------------------------------------------------------------------------------------
+    //---------------------------------------------------------------------------------------------
+    class FileModelStreamer::Private : public Base
+    {
+    public:
+        OPERATION_ID m_lastOperationID = 0;
+
+        IFileHandle* m_writeFile = nullptr;
+
+        // For debug
+        size_t m_totalWritten = 0;
+
+        static inline FString MakeDataFileName(
+                const FString& strModelName,
+                uint64_t key0
+                )
+        {
+			FString name = strModelName;
+			int extensionPos = 0;
+			if (name.FindLastChar('.', extensionPos))
+			{
+				name = name.Left(extensionPos);
+			}
+
+            FString temp = FString::Printf(TEXT(".%08x-%08x.mutable_data"),
+                               uint32_t((key0>>32) & 0xffffffff), uint32(key0 & 0xffffffff) );
+            name += temp;
+
+            return name;
+        }
+
+
+    };
+
+
+    //---------------------------------------------------------------------------------------------
+    FileModelStreamer::FileModelStreamer()
+    {
+        m_pD = new Private();
+    }
+
+
+    //---------------------------------------------------------------------------------------------
+    FileModelStreamer::~FileModelStreamer()
+    {
+        delete m_pD;
+    }
+
+
+    //---------------------------------------------------------------------------------------------
+    void FileModelStreamer::GetFileName( const FString& strModelName, uint64 key0, FString& outFileName )
+    {
+		outFileName = Private::MakeDataFileName( strModelName, key0 );
+    }
+
+
+    //---------------------------------------------------------------------------------------------
+    ModelStreamer::OPERATION_ID FileModelStreamer::BeginReadBlock(const mu::Model* Model, uint64 key0, void* pBuffer, uint64 size)
+    {
+        //UE_LOG(LogMutableCore,Log,"Streamer OpenRead");
+
+		check(Model);
+
+		OPERATION_ID id = ++m_pD->m_lastOperationID;
+
+        IFileHandle* pFile = nullptr;
+
+        if ( !key0 )
+        {
+            // It is the main model file
+            pFile = FPlatformFileManager::Get().GetPlatformFile().OpenRead( *FString(Model->GetLocation()) );
+
+            // Read the file identifying tag.
+            char tag[4];
+            size_t read = pFile->Read((uint8*)&tag[0], 4);
+            check( read==4 );
+            check( !memcmp(tag,MUTABLE_COMPILED_MODEL_FILETAG,4) );
+
+            uint32_t codeVersion = 0;
+            read = pFile->Read( (uint8*)&codeVersion, sizeof(uint32_t) );
+            check( read==4 );
+            check( codeVersion == MUTABLE_COMPILED_MODEL_CODE_VERSION  );
+
+        }
+        else
+        {
+            // It is a data file from a split model
+            FString name = m_pD->MakeDataFileName( Model->GetLocation(), key0);
+            pFile = FPlatformFileManager::Get().GetPlatformFile().OpenRead( *name );
+        }
+
+        if( !pFile )
+		{
+			return -1;
+		}
+
+        size_t read = 1;
+        size_t totalRead = 0;
+        while (read && totalRead<size)
+        {
+            read = pFile->Read( ((uint8*)pBuffer)+totalRead, (size_t)size-totalRead  );
+            totalRead += read;
+        }
+
+		delete pFile;
+
+        check( totalRead==size );
+
+        return id;
+    }
+
+
+    //---------------------------------------------------------------------------------------------
+    bool FileModelStreamer::IsReadCompleted( OPERATION_ID )
+    {
+        return true;
+    }
+
+
+    //---------------------------------------------------------------------------------------------
+    void FileModelStreamer::EndRead( OPERATION_ID )
+    {
+	}
+
+
+    //---------------------------------------------------------------------------------------------
+    void FileModelStreamer::OpenWriteFile( const char* strModelName, uint64 key0 )
+    {
+        check( !m_pD->m_writeFile );
+
+        IFileHandle* pFile = nullptr;
+        m_pD->m_totalWritten = 0;
+
+        if ( !key0 )
+        {
+            // It is the main model file
+            pFile = FPlatformFileManager::Get().GetPlatformFile().OpenWrite( *FString(strModelName) );
+			check(pFile);
+            if (!pFile)
+                return;
+
+            // Write the file identifying tag.
+            bool success = pFile->Write( (const uint8*)MUTABLE_COMPILED_MODEL_FILETAG, 4 );
+            check(success);
+
+            uint32_t codeVersion = MUTABLE_COMPILED_MODEL_CODE_VERSION;
+			success = pFile->Write( (const uint8*)&codeVersion, sizeof(uint32_t) );
+            check(success);
+        }
+        else
+        {
+            // It is a data file from a split model
+            FString name = m_pD->MakeDataFileName( strModelName, key0 );
+            pFile = FPlatformFileManager::Get().GetPlatformFile().OpenWrite( *name );
+        }
+
+        check( pFile );
+
+        m_pD->m_writeFile = pFile;
+    }
+
+
+    //---------------------------------------------------------------------------------------------
+    void FileModelStreamer::Write( const void* pBuffer, uint64 size )
+    {
+        check( m_pD->m_writeFile );
+
+        m_pD->m_writeFile->Write( ((const uint8*)pBuffer), size );
+        m_pD->m_totalWritten += size;
+    }
+
+
+    //---------------------------------------------------------------------------------------------
+    void FileModelStreamer::CloseWriteFile()
+    {
+        check( m_pD->m_writeFile );
+
+        delete m_pD->m_writeFile;
+        m_pD->m_writeFile = 0;
+    }
+
+
+    //---------------------------------------------------------------------------------------------
+    //---------------------------------------------------------------------------------------------
+    //---------------------------------------------------------------------------------------------
+    Ptr<ResourceProxy<Image>> ProxyFactoryFiles::NewImageProxy(InputArchive& arch)
+    {        
+        Ptr<Image> t = Image::StaticUnserialise( arch );
+        return new ResourceProxyTempFile<Image>(t.get());
+    }
+
+
+    //---------------------------------------------------------------------------------------------
+    //---------------------------------------------------------------------------------------------
+    //---------------------------------------------------------------------------------------------
+    class ProxyFactoryMutableSourceFile::Private : public Base
+    {
+    public:
+        // \todo: pimpl this out
+        InputFileStream* m_pStream = nullptr;
+        std::string m_filename;
+    };
+
+
+    //---------------------------------------------------------------------------------------------
+    ProxyFactoryMutableSourceFile::ProxyFactoryMutableSourceFile( const char* strFileName, 
+                                                                  InputFileStream* pStream )
+    {
+        m_pD = new Private();
+
+        m_pD->m_pStream = pStream;
+        m_pD->m_filename = strFileName;
+    }
+
+
+    //---------------------------------------------------------------------------------------------
+    ProxyFactoryMutableSourceFile::~ProxyFactoryMutableSourceFile()
+    {
+        delete m_pD;
+        m_pD = nullptr;
+    }
+
+
+    //---------------------------------------------------------------------------------------------
+    Ptr<ResourceProxy<Image>> ProxyFactoryMutableSourceFile::NewImageProxy(InputArchive& arch)
+    {
+        // assert: the archive should beusing m_pStream
+        size_t position = m_pD->m_pStream->Tell();
+
+        // we are unserialising just to skip the image. It would be much better if we had the size
+        // and we could skip it.
+        Ptr<Image> t = Image::StaticUnserialise( arch );
+
+        return new ResourceProxyFile<Image>( m_pD->m_filename, position );
+    }
+
+
+}
+
diff --git a/Engine/Plugins/Experimental/Mutable/Source/MutableTools/Private/StreamsPrivate.h b/Engine/Plugins/Experimental/Mutable/Source/MutableTools/Private/StreamsPrivate.h
new file mode 100644
--- /dev/null
+++ b/Engine/Plugins/Experimental/Mutable/Source/MutableTools/Private/StreamsPrivate.h
@@ -0,0 +1,204 @@
+// Copyright Epic Games, Inc. All Rights Reserved.
+
+#pragma once
+
+#include "HAL/PlatformFileManager.h"
+#include "HAL/PlatformFile.h"
+#include "HAL/PlatformProcess.h"
+#include "Math/RandomStream.h"
+#include "Misc/ScopeLock.h"
+
+#include "Streams.h"
+#include "Platform.h"
+#include "Operations.h"
+#include "MemoryPrivate.h"
+
+namespace mu
+{
+
+
+	//---------------------------------------------------------------------------------------------
+	class MUTABLETOOLS_API InputFileStream::Private : public Base
+	{
+	public:
+
+		FILE* m_pFile;
+
+		uint64 m_size;
+		uint64 m_pos;
+
+		uint64 m_readBytes;
+
+    private:
+        uint8* m_buffer;
+
+        friend class InputFileStream;
+    };
+
+
+	//---------------------------------------------------------------------------------------------
+	class MUTABLETOOLS_API OutputFileStream::Private : public Base
+	{
+	public:
+
+		FILE* m_pFile;
+
+		uint64 m_pos;
+
+    private:
+        uint8* m_buffer;
+
+        friend class OutputFileStream;
+    };
+
+    
+    //---------------------------------------------------------------------------------------------
+    //---------------------------------------------------------------------------------------------
+    //---------------------------------------------------------------------------------------------
+    template<class R>
+    class MUTABLETOOLS_API ResourceProxyTempFile : public ResourceProxy<R>
+    {
+    private:
+        Ptr<const R> m_resource;
+        FString m_fileName;
+		uint64 m_fileSize = 0;
+		FCriticalSection Mutex;
+
+    public:
+        ResourceProxyTempFile( const R* resource )
+        {
+			IPlatformFile& PlatformFile = FPlatformFileManager::Get().GetPlatformFile();
+
+            if (!resource)
+            {
+                return;
+            }
+
+			FScopeLock Lock(&Mutex);
+
+            OutputMemoryStream stream;
+            OutputArchive arch(&stream);
+            R::Serialise( resource, arch );
+
+            if (stream.GetBufferSize()<=256*1024)
+            {
+                m_resource = resource;
+            }
+            else
+            {
+                m_fileSize = stream.GetBufferSize();
+
+				FString TempPath = FPlatformProcess::UserTempDir();
+				TempPath += TEXT("mut.temp");
+
+				FRandomStream RandomStream;
+				RandomStream.GenerateNewSeed();
+				
+				auto prefix = TempPath;
+				
+				FString FinalTempPath;
+				bool found = true;
+                while(found)
+                {
+					FinalTempPath = TempPath+FString::Printf(TEXT(".%.8x.%.8x.%.8x.%.8x"), RandomStream.GetUnsignedInt(), RandomStream.GetUnsignedInt(), RandomStream.GetUnsignedInt(), RandomStream.GetUnsignedInt());
+					found = PlatformFile.FileExists(*FinalTempPath);
+                }
+				IFileHandle* resourceFile = PlatformFile.OpenWrite(*FinalTempPath);
+
+                check(resourceFile);
+				resourceFile->Write( (const uint8*) stream.GetBuffer(), stream.GetBufferSize() );
+				delete resourceFile;
+
+                m_fileName = FinalTempPath;
+            }
+        }
+
+        ~ResourceProxyTempFile()
+        {
+			FScopeLock Lock(&Mutex);
+
+            if (!m_fileName.IsEmpty())
+            {
+                // Delete temp file
+				FPlatformFileManager::Get().GetPlatformFile().DeleteFile(*m_fileName);
+                m_fileName.Empty();
+            }
+        }
+
+        Ptr<const R> Get() override
+        {
+			FScopeLock Lock(&Mutex);
+
+            Ptr<const R> r;
+            if (m_resource)
+            {
+                // cached
+                r = m_resource;
+            }
+            else if (!m_fileName.IsEmpty())
+            {
+                std::vector<char> buf(m_fileSize);
+                auto resourceFile = FPlatformFileManager::Get().GetPlatformFile().OpenRead( *m_fileName );
+                check(resourceFile);
+                resourceFile->Read((uint8*)buf.data(), m_fileSize);
+                delete resourceFile;
+
+                InputMemoryStream stream( buf.data(), m_fileSize );
+                InputArchive arch(&stream);
+                r =  R::StaticUnserialise( arch );
+            }
+            return r;
+        }
+
+    };
+
+
+    //---------------------------------------------------------------------------------------------
+    //! Resource proxy for resources that are located in mutable_source files and have deferred
+    //! access: they are not loaded immeditely, but only when needed.
+    //---------------------------------------------------------------------------------------------
+    template<class R>
+    class MUTABLETOOLS_API ResourceProxyFile : public ResourceProxy<R>
+    {
+    private:
+        std::string m_fileName;
+        uint64 m_filePos = 0;
+		FCriticalSection Mutex;
+
+    public:
+        ResourceProxyFile( const std::string& fileName, uint64 filePos )
+        {
+			FScopeLock Lock(&Mutex);
+			m_fileName = fileName;
+            m_filePos = filePos;
+        }
+
+        Ptr<const R> Get() override
+        {
+			FScopeLock Lock(&Mutex);
+
+            Ptr<const R> r;
+            if (!m_fileName.empty())
+            {
+                InputFileStream stream( m_fileName.c_str() );
+                stream.Seek( m_filePos );
+                InputArchive arch(&stream);
+                r =  R::StaticUnserialise( arch );
+            }
+            return r;
+        }
+
+    };
+
+
+}
+
+//! Utility function
+namespace {
+template <class T> inline void hash_combine(uint64 &seed, const T &v) {
+  std::hash<T> hasher;
+  seed ^= hasher(v) + 0x9e3779b9 + (seed << 6) + (seed >> 2);
+}
+} // namespace
+
+
diff --git a/Engine/Plugins/Experimental/Mutable/Source/MutableTools/Private/Table.cpp b/Engine/Plugins/Experimental/Mutable/Source/MutableTools/Private/Table.cpp
new file mode 100644
--- /dev/null
+++ b/Engine/Plugins/Experimental/Mutable/Source/MutableTools/Private/Table.cpp
@@ -0,0 +1,197 @@
+// Copyright Epic Games, Inc. All Rights Reserved.
+
+
+#include "Table.h"
+#include "TablePrivate.h"
+
+#include "Platform.h"
+#include "ErrorLogPrivate.h"
+
+#include "Image.h"
+#include "Mesh.h"
+
+
+namespace mu
+{
+
+	//---------------------------------------------------------------------------------------------
+	Table::Table()
+	{
+		m_pD = new Private();
+	}
+
+
+	//---------------------------------------------------------------------------------------------
+	Table::~Table()
+	{
+        check( m_pD );
+		delete m_pD;
+		m_pD = 0;
+	}
+
+
+	//---------------------------------------------------------------------------------------------
+	void Table::Serialise( const Table* p, OutputArchive& arch )
+	{
+		arch << *p->m_pD;
+	}
+
+
+	//---------------------------------------------------------------------------------------------
+	TablePtr Table::StaticUnserialise( InputArchive& arch )
+	{
+		TablePtr pResult = new Table();
+		arch >> *pResult->m_pD;
+		return pResult;
+	}
+
+
+	//---------------------------------------------------------------------------------------------
+	Table::Private* Table::GetPrivate() const
+	{
+		return m_pD;
+	}
+
+
+	//---------------------------------------------------------------------------------------------
+	void Table::SetName( const char* strName )
+	{
+		m_pD->m_name = strName;
+	}
+
+
+	//---------------------------------------------------------------------------------------------
+	const char* Table::GetName() const
+	{
+		return m_pD->m_name.c_str();
+	}
+
+
+	//---------------------------------------------------------------------------------------------
+	int Table::AddColumn( const char* strName, TABLE_COLUMN_TYPE type )
+	{
+		int result = (int)m_pD->m_columns.size();
+
+		TABLE_COLUMN c;
+
+		if (strName)
+		{
+			c.m_name = strName;
+		}
+		c.m_type = type;
+
+		m_pD->m_columns.push_back(c);
+
+		// Add it to all rows
+		for ( std::size_t r=0; r<m_pD->m_rows.size(); ++r )
+		{
+			m_pD->m_rows[r].m_values.push_back( TABLE_VALUE() );
+		}
+
+		return result;
+	}
+
+
+	//---------------------------------------------------------------------------------------------
+	int Table::FindColumn( const char* strName ) const
+	{
+		int res = -1;
+
+		for ( std::size_t c=0; c<m_pD->m_columns.size(); ++c )
+		{
+			if ( m_pD->m_columns[c].m_name == strName )
+			{
+				res = (int)c;
+			}
+		}
+
+		return res;
+	}
+
+
+	//---------------------------------------------------------------------------------------------
+    void Table::AddRow( uint32_t id )
+	{
+		check( m_pD->FindRow(id)<0 );
+
+		TABLE_ROW r;
+		r.m_id = id;
+		r.m_values.resize( m_pD->m_columns.size() );
+		m_pD->m_rows.push_back( r );
+	}
+
+
+	//---------------------------------------------------------------------------------------------
+	void Table::SetNoneOption(bool bAddOption)
+	{
+		m_pD->m_NoneOption = bAddOption;
+	}
+
+
+	//---------------------------------------------------------------------------------------------
+    void Table::SetCell( int column, uint32_t rowId, float value )
+	{
+		int row = m_pD->FindRow(rowId);
+		check( row>=0 );
+		check( column < (int)m_pD->m_columns.size() );
+		check( column < (int)m_pD->m_rows[row].m_values.size() );
+		check( m_pD->m_columns[column].m_type == TCT_SCALAR );
+
+		m_pD->m_rows[ row ].m_values[column].m_scalar = value;
+	}
+
+
+	//---------------------------------------------------------------------------------------------
+    void Table::SetCell( int column, uint32_t rowId, float r, float g, float b )
+	{
+		int row = m_pD->FindRow(rowId);
+		check( row>=0 );
+		check( column < (int)m_pD->m_columns.size() );
+		check( column < (int)m_pD->m_rows[row].m_values.size() );
+		check( m_pD->m_columns[column].m_type == TCT_COLOUR );
+
+		m_pD->m_rows[ row ].m_values[column].m_colour = vec3<float>( r, g, b );
+	}
+
+
+	//---------------------------------------------------------------------------------------------
+    void Table::SetCell( int column, uint32_t rowId, Image* pImage )
+	{
+		int row = m_pD->FindRow(rowId);
+		check( row>=0 );
+		check( column < (int)m_pD->m_columns.size() );
+		check( column < (int)m_pD->m_rows[row].m_values.size() );
+		check( m_pD->m_columns[column].m_type == TCT_IMAGE );
+
+		m_pD->m_rows[ row ].m_values[column].m_pProxyImage = new ResourceProxyMemory<Image>(pImage);
+	}
+
+
+	//---------------------------------------------------------------------------------------------
+    void Table::SetCell( int column, uint32_t rowId, Mesh* pMesh )
+	{
+		int row = m_pD->FindRow(rowId);
+		check( row>=0 );
+		check( column < (int)m_pD->m_columns.size() );
+		check( column < (int)m_pD->m_rows[row].m_values.size() );
+		check( m_pD->m_columns[column].m_type == TCT_MESH );
+
+		m_pD->m_rows[ row ].m_values[column].m_pMesh = pMesh;
+	}
+
+
+	//---------------------------------------------------------------------------------------------
+    void Table::SetCell( int column, uint32_t rowId, const char* strValue )
+	{
+		int row = m_pD->FindRow(rowId);
+		check( row>=0 );
+		check( column < (int)m_pD->m_columns.size() );
+		check( column < (int)m_pD->m_rows[row].m_values.size() );
+		check( m_pD->m_columns[column].m_type == TCT_STRING );
+
+		m_pD->m_rows[ row ].m_values[column].m_string = strValue;
+	}
+
+
+}
+
diff --git a/Engine/Plugins/Experimental/Mutable/Source/MutableTools/Private/TablePrivate.h b/Engine/Plugins/Experimental/Mutable/Source/MutableTools/Private/TablePrivate.h
new file mode 100644
--- /dev/null
+++ b/Engine/Plugins/Experimental/Mutable/Source/MutableTools/Private/TablePrivate.h
@@ -0,0 +1,183 @@
+// Copyright Epic Games, Inc. All Rights Reserved.
+
+#pragma once
+
+#include "Table.h"
+
+#include "Visitor.h"
+#include "ErrorLogPrivate.h"
+
+#include "MutableMath.h"
+#include "MemoryPrivate.h"
+#include "SerialisationPrivate.h"
+
+
+namespace mu
+{
+
+	MUTABLE_DEFINE_ENUM_SERIALISABLE( TABLE_COLUMN_TYPE )
+
+
+	struct TABLE_COLUMN
+	{
+		string m_name;
+
+		TABLE_COLUMN_TYPE m_type;
+	};
+
+
+	struct TABLE_VALUE
+	{
+		// TODO: Union
+		float m_scalar;
+		vec3<float> m_colour;
+		Ptr<ResourceProxy<Image>> m_pProxyImage;
+		MeshPtr m_pMesh;
+		string m_string;
+	};
+
+
+	struct TABLE_ROW
+	{
+        uint32_t m_id;
+		vector<TABLE_VALUE> m_values;
+	};
+
+
+	//!
+	class Table::Private : public Base
+	{
+	public:
+
+		string m_name;
+		vector<TABLE_COLUMN> m_columns;
+		vector<TABLE_ROW> m_rows;
+		bool m_NoneOption = false;
+
+
+		//!
+		void Serialise( OutputArchive& arch ) const
+		{
+            uint32_t ver = 1;
+			arch << ver;
+
+            arch << (uint32_t)m_columns.size();
+			for ( std::size_t c=0; c<m_columns.size(); ++c )
+			{
+				arch << m_columns[c].m_name;
+				arch << m_columns[c].m_type;
+			}
+
+            arch << (uint32_t)m_rows.size();
+			for ( std::size_t r=0; r<m_rows.size(); ++r )
+			{
+				arch << m_rows[r].m_id;
+
+				for ( std::size_t c=0; c<m_columns.size(); ++c )
+				{
+					const TABLE_VALUE& v = m_rows[r].m_values[c];
+
+					switch (m_columns[c].m_type)
+					{
+					case TCT_SCALAR: 	arch << v.m_scalar; break;
+					case TCT_COLOUR: 	arch << v.m_colour; break;
+					case TCT_MESH: 		arch << v.m_pMesh; break;
+					case TCT_IMAGE:		
+					{
+						Ptr<const Image> image;
+						if (v.m_pProxyImage)
+						{
+							image = v.m_pProxyImage->Get();
+						}
+						arch << image;
+						break;
+					}
+					case TCT_STRING:	arch << v.m_string; break;
+					default: check( false);
+					}
+				}
+			}
+
+			arch << m_NoneOption;
+
+		}
+
+
+		//!
+		void Unserialise( InputArchive& arch )
+		{
+            uint32_t ver;
+			arch >> ver;
+			check(ver<=1);
+
+            uint32_t columnCount;
+			arch >> columnCount;
+			m_columns.resize( columnCount );
+            for (uint32_t c=0; c<columnCount; ++c )
+			{
+				arch >> m_columns[c].m_name;
+				arch >> m_columns[c].m_type;
+			}
+
+            uint32_t rowCount;
+			arch >> rowCount;
+			m_rows.resize( rowCount );
+            for ( uint32_t r=0; r<rowCount; ++r )
+			{
+				arch >> m_rows[r].m_id;
+				m_rows[r].m_values.resize( columnCount );
+
+                for (uint32_t c=0; c<columnCount; ++c )
+				{
+					TABLE_VALUE& v = m_rows[r].m_values[c];
+
+					switch (m_columns[c].m_type)
+					{
+					case TCT_SCALAR:	arch >> v.m_scalar; break;
+					case TCT_COLOUR:	arch >> v.m_colour; break;
+					case TCT_MESH:		arch >> v.m_pMesh; break;
+					case TCT_IMAGE:		
+					{
+						// Are we using proxies?
+						v.m_pProxyImage = arch.NewImageProxy();
+						if (!v.m_pProxyImage)
+						{
+							// Normal serialisation
+							ImagePtr image;
+							arch >> image;
+							v.m_pProxyImage = new ResourceProxyMemory<Image>(image.get());
+						}
+						break;
+					}
+					case TCT_STRING:	arch >> v.m_string; break;
+					default: check( false);
+					}
+				}
+			}
+
+			if (ver >= 1)
+			{
+				arch >> m_NoneOption;
+			}
+		}
+
+
+		//! Find a row in the table by id. Return -1 if not found.
+        int FindRow( uint32_t id ) const
+		{
+			int res = -1;
+
+			for ( std::size_t r=0; res<0 && r<m_rows.size(); ++r )
+			{
+				if ( m_rows[r].m_id==id )
+				{
+					res = (int)r;
+				}
+			}
+
+			return res;
+		}
+
+	};
+
+}
diff --git a/Engine/Plugins/Experimental/Mutable/Source/MutableTools/Private/TypeInfo.cpp b/Engine/Plugins/Experimental/Mutable/Source/MutableTools/Private/TypeInfo.cpp
new file mode 100644
--- /dev/null
+++ b/Engine/Plugins/Experimental/Mutable/Source/MutableTools/Private/TypeInfo.cpp
@@ -0,0 +1,128 @@
+// Copyright Epic Games, Inc. All Rights Reserved.
+
+#include "TypeInfo.h"
+
+using namespace mu;
+
+
+const char* TypeInfo::s_imageFormatName[ size_t(EImageFormat::IF_COUNT) ] =
+{
+	"None",
+	"RGB U888",
+	"RGBA U8888",
+	"L U8",
+
+    "PVRTC2",
+    "PVRTC4",
+    "ETC1",
+    "ETC2",
+
+	"L U8 RLE",
+	"RGB U888 RLE",
+	"RGBA U8888 RLE",
+	"L U1 RLE",
+
+    "BC1",
+    "BC2",
+    "BC3",
+    "BC4",
+    "BC5",
+    "BC6",
+    "BC7",
+
+    "BGRA U8888",
+
+    "ASTC_4x4_RGB_LDR",
+    "ASTC_4x4_RGBA_LDR",
+    "ASTC_4x4_RG_LDR",
+
+};
+
+
+const char* TypeInfo::s_meshBufferSemanticName[ MBS_COUNT ] =
+{
+	"None",
+
+	"VertexIndex",
+
+	"Position",
+	"Normal",
+	"Tangent",
+	"Binormal",
+	"Tex Coords",
+	"Colour",
+	"Bone Weights",
+	"Bone Indices",
+
+	"Layout Block",
+
+	"Chart",
+
+    "Other",
+
+    "Tangent Space Sign",
+
+	"TriangleIndex",
+	"BarycentricCoords",
+	"Distance"
+
+};
+
+const char* TypeInfo::s_meshBufferFormatName[ MBF_COUNT ] =
+{
+	"None",
+	"Float16",
+	"Float32",
+
+	"UInt8",
+	"UInt16",
+	"UInt32",
+	"Int8",
+	"Int16",
+	"Int32",
+
+	"NUInt8",
+	"NUInt16",
+	"NUInt32",
+	"NInt8",
+	"NInt16",
+	"NInt32",
+
+    "PackedDir8",
+    "PackedDir8_WTangentSign",
+    "PackedDirS8",
+    "PackedDirS8_WTangentSign",
+
+	"Float64",
+};
+
+const char* TypeInfo::s_projectorTypeName[  static_cast<uint32>(mu::PROJECTOR_TYPE::COUNT) ] =
+{
+	"Planar",
+	"Cylindrical",
+	"Wrapping",
+};
+
+const char* TypeInfo::s_curveInterpolationModeName [ static_cast<uint8_t>(mu::CurveKeyFrame::InterpMode::Count) ] =
+{
+	"Linear",
+	"Constant",
+	"Cubic",
+	"None",
+};
+
+const char* TypeInfo::s_curveTangentModeName [static_cast<uint8_t> (mu::CurveKeyFrame::TangentMode::Count)] =
+{
+	"Auto",
+	"Manual",
+	"Broken",
+	"None",
+};
+
+const char* TypeInfo::s_curveTangentWeightModeName [ static_cast<uint8_t> (mu::CurveKeyFrame::TangentWeightMode::Count)] =
+{
+	"Auto",
+	"Manual",
+	"Broken",
+	"None",
+};
\ No newline at end of file
