[mutable] Move Mutable plugin from the restricted folder to the experimental plugins folder.

diff --git a/Engine/Plugins/Experimental/Mutable/Source/MutableTools/Private/Visitor.h b/Engine/Plugins/Experimental/Mutable/Source/MutableTools/Private/Visitor.h
new file mode 100644
--- /dev/null
+++ b/Engine/Plugins/Experimental/Mutable/Source/MutableTools/Private/Visitor.h
@@ -0,0 +1,140 @@
+// Copyright Epic Games, Inc. All Rights Reserved.
+
+//! Partially based on the visitor description of the Loki library
+
+
+#ifndef _MUTABLE_VISITOR_H_
+#define _MUTABLE_VISITOR_H_
+
+
+#include "MutableTools/Private/Platform.h"
+
+
+namespace mu
+{
+
+	//---------------------------------------------------------------------------------------------
+	//!
+	//---------------------------------------------------------------------------------------------
+    class BaseVisitor
+    {
+    public:
+        virtual ~BaseVisitor() {}
+    };
+
+    template <class T, typename R = void, bool ConstVisit = false>
+    class Visitor;
+
+    template <class T, typename R>
+    class Visitor<T, R, false>
+    {
+    public:
+        typedef R ReturnType;
+        typedef T ParamType;
+        virtual ~Visitor() {}
+        virtual ReturnType Visit(ParamType&) = 0;
+    };
+
+    template <class T, typename R>
+    class Visitor<T, R, true>
+    {
+    public:
+        typedef R ReturnType;
+        typedef const T ParamType;
+        virtual ~Visitor() {}
+        virtual ReturnType Visit(ParamType&) = 0;
+    };
+
+
+    //---------------------------------------------------------------------------------------------
+    //!
+    //---------------------------------------------------------------------------------------------
+	template <typename R, typename Visited>
+	struct DefaultCatchAll
+	{
+		static R OnUnknownVisitor( Visited&, BaseVisitor& )
+		{
+			check( false );
+			return R();
+		}
+	};
+
+
+    //---------------------------------------------------------------------------------------------
+    //!
+    //---------------------------------------------------------------------------------------------
+	template
+    <
+        typename R = void,
+        template <typename, class> class CatchAll = DefaultCatchAll,
+        bool ConstVisitable = false
+    >
+    class BaseVisitable;
+
+    template<typename R,template <typename, class> class CatchAll>
+    class BaseVisitable<R, CatchAll, false>
+    {
+    public:
+        typedef R ReturnType;
+        virtual ~BaseVisitable() {}
+        virtual ReturnType Accept(BaseVisitor&) = 0;
+
+    protected: // give access only to the hierarchy
+        template <class T>
+        static ReturnType AcceptImpl(T& visited, BaseVisitor& guest)
+        {
+            // Apply the Acyclic Visitor
+            if (Visitor<T,R>* p = dynamic_cast<Visitor<T,R>*>(&guest))
+            {
+                return p->Visit(visited);
+            }
+            return CatchAll<R, T>::OnUnknownVisitor(visited, guest);
+        }
+    };
+
+
+    //---------------------------------------------------------------------------------------------
+    //!
+    //---------------------------------------------------------------------------------------------
+    template<typename R,template <typename, class> class CatchAll>
+    class BaseVisitable<R, CatchAll, true>
+    {
+    public:
+        typedef R ReturnType;
+        virtual ~BaseVisitable() {}
+        virtual ReturnType Accept(BaseVisitor&) const = 0;
+
+    protected: // give access only to the hierarchy
+        template <class T>
+        static ReturnType AcceptImpl(const T& visited, BaseVisitor& guest)
+        {
+            // Apply the Acyclic Visitor
+            if (Visitor<T,R,true>* p = dynamic_cast<Visitor<T,R,true>*>(&guest))
+            {
+                return p->Visit(visited);
+            }
+            return CatchAll<R, T>::OnUnknownVisitor(const_cast<T&>(visited), guest);
+        }
+    };
+
+
+   //---------------------------------------------------------------------------------------------
+   //!
+   //---------------------------------------------------------------------------------------------
+	#define MUTABLE_DEFINE_VISITABLE() 								\
+        ReturnType Accept(::mu::BaseVisitor& guest) override	\
+			{ return AcceptImpl(*this, guest); }
+
+
+   //---------------------------------------------------------------------------------------------
+   //!
+   //---------------------------------------------------------------------------------------------
+	#define MUTABLE_DEFINE_CONST_VISITABLE() 							\
+        ReturnType Accept(::mu::BaseVisitor& guest) const override	\
+			{ return AcceptImpl(*this, guest); }
+
+
+}
+
+
+#endif
diff --git a/Engine/Plugins/Experimental/Mutable/Source/MutableTools/Public/Compiler.h b/Engine/Plugins/Experimental/Mutable/Source/MutableTools/Public/Compiler.h
new file mode 100644
--- /dev/null
+++ b/Engine/Plugins/Experimental/Mutable/Source/MutableTools/Public/Compiler.h
@@ -0,0 +1,181 @@
+// Copyright Epic Games, Inc. All Rights Reserved.
+
+#pragma once
+
+#include "RefCounted.h"
+#include "Ptr.h"
+#include "Types.h"
+
+
+namespace mu
+{
+    // Forward declarations
+    class Compiler;
+    typedef Ptr<Compiler> CompilerPtr;
+    typedef Ptr<const Compiler> CompilerPtrConst;
+
+    class CompilerOptions;
+    using CompilerOptionsPtr=Ptr<CompilerOptions>;
+    using CompilerOptionsPtrConst=Ptr<const CompilerOptions>;
+
+    class Model;
+    using ModelPtr=Ptr<Model>;
+    using ModelPtrConst=Ptr<const Model>;
+
+    class Transform;
+    using TransformPtr=Ptr<Transform>;
+    using TransformPtrConst=Ptr<const Transform>;
+
+    class ErrorLog;
+    using ErrorLogPtr=Ptr<ErrorLog>;
+    using ErrorLogPtrConst=Ptr<const ErrorLog>;
+
+    class ModelReport;
+    using ModelReportPtr=Ptr<ModelReport>;
+    using ModelReportPtrConst=Ptr<const ModelReport>;
+
+    class Node;
+    class NodeTransformedObject;
+
+    //! Available GPU families
+    typedef enum
+    {
+        GPU_NONE			= 0,
+        GPU_GL4,
+        GPU_GLES2,
+
+        GPU_COUNT
+    } GPU_TYPE;
+
+    //! Return a readable string for a GPU.
+    extern const char* GetGPUName( GPU_TYPE );
+
+    //! \brief Options used to compile the models with a compiler.
+    class MUTABLETOOLS_API CompilerOptions : public RefCounted
+    {
+    public:
+
+        //! Create new settings with the default values as explained in each method below.
+        CompilerOptions();
+
+        //!
+        void SetLogEnabled( bool bEnabled );
+
+        //!
+        void SetOptimisationEnabled( bool bEnabled );
+
+        //!
+        void SetConstReductionEnabled( bool bEnabled );
+
+        //!
+        void SetEnablePartialOptimisation( bool bEnabled );
+
+		//!
+		void SetEnableConcurrency(bool bEnabled);
+
+        //!
+        void SetOptimisationMaxIteration( int maxIterations );
+
+        //!
+        void SetIgnoreStates( bool ignore );
+
+        //! If enabled, the disk will be used as temporary memory. This will make the compilation
+        //! process very slow, but will be able to compile very large models.
+        void SetUseDiskCache( bool enabled );
+
+        //! Set the quality for the image compression algorithms. The level value is used internally
+        //! with System::SetImagecompressionQuality
+        void SetImageCompressionQuality( int quality );
+
+        //! 
+        void SetDataPackingStrategy( int32 minRomSize, int32 MinTextureResidentMipCount );
+
+        //! Different data packing strategies
+        enum class TextureLayoutStrategy : std::uint8_t
+        {
+            //! Pack texture layouts without changing any scale
+            Pack,
+
+            //! Do not touch mesh or image texture layouts
+            None,
+
+            //! Helper value, not really a strategy.
+            Count
+        };
+
+        //! Return a readable string for a GPU.
+        static const char* GetTextureLayoutStrategyName( TextureLayoutStrategy s );
+
+        //! Enable or disable the process that splits data into related sets to serialise in
+        //! separate files. By default it is disabled. In order to get multiple files, this has to
+        //! be enabled AND the appropiate serialisation function needs to be called (see
+        //! rutime/Model.h).
+        void SetTextureLayoutStrategy( TextureLayoutStrategy strategy );
+
+        //! Set the GPU type to optimise for, if any. Defaults to none.
+        void SetGPUType( GPU_TYPE );
+
+        //! If GPU optimsiation is enabled, set the maximum number of fragment textures the gpu can
+        //! use.
+        void SetGPUMaxFragmentTextures( int );
+
+        //-----------------------------------------------------------------------------------------
+        // Interface pattern
+        //-----------------------------------------------------------------------------------------
+        class Private;
+        Private* GetPrivate() const;
+
+        CompilerOptions( const CompilerOptions& ) = delete;
+        CompilerOptions( CompilerOptions&& ) = delete;
+        CompilerOptions& operator=(const CompilerOptions&) = delete;
+        CompilerOptions& operator=(CompilerOptions&&) = delete;
+
+    protected:
+
+        //! Forbidden. Manage with the Ptr<> template.
+        ~CompilerOptions() override;
+
+    private:
+
+        Private* m_pD;
+
+    };
+
+
+    //! This utility class compiles two types of expressions: models and transforms.
+    //! Model expressions are compiled into run-time usable objects.
+    //! Transform expressions are compiled into Transform objects than can be applied to Model
+    //! expressions.
+    class MUTABLETOOLS_API Compiler : public RefCounted
+    {
+    public:
+
+        Compiler( CompilerOptionsPtr options=nullptr );
+
+        //! Compile the expression into a run-time model.
+        ModelPtr Compile( const Ptr<Node>& pNode );
+
+        //! After a call to compile a model, this method will return the additional information on
+        //! the result. See the ModelReport class for more information.
+        //! If called without calling Compile first, a null pointer is returned.
+        ModelReportPtrConst GetModelReport() const;
+
+        //! Compile the transform object into a Transform that can be applied to a model expression.
+        TransformPtr Compile( const NodeTransformedObject* pNode );
+
+        //! Return the log of messages of all the compile operations executed so far.
+        ErrorLogPtrConst GetLog() const;
+
+    protected:
+
+        //! Forbidden. Manage with the Ptr<> template.
+        ~Compiler();
+
+    private:
+
+        class Private;
+        Private* m_pD;
+
+    };
+
+}
diff --git a/Engine/Plugins/Experimental/Mutable/Source/MutableTools/Public/ErrorLog.h b/Engine/Plugins/Experimental/Mutable/Source/MutableTools/Public/ErrorLog.h
new file mode 100644
--- /dev/null
+++ b/Engine/Plugins/Experimental/Mutable/Source/MutableTools/Public/ErrorLog.h
@@ -0,0 +1,92 @@
+// Copyright Epic Games, Inc. All Rights Reserved.
+
+#pragma once
+
+#include "RefCounted.h"
+#include "Ptr.h"
+
+
+namespace mu
+{
+
+	// Forward declarations
+	class ErrorLog;
+	typedef Ptr<ErrorLog> ErrorLogPtr;
+	typedef Ptr<const ErrorLog> ErrorLogPtrConst;
+
+
+	//! Types of message stored in the log
+	//! \ingroup tools
+	typedef enum
+	{
+        ELMT_NONE=0,
+		ELMT_ERROR,
+		ELMT_WARNING,
+		ELMT_INFO
+	} ErrorLogMessageType;
+
+	struct ErrorLogMessageAttachedDataView 
+	{
+		const float* m_unassignedUVs = nullptr;
+		size_t m_unassignedUVsSize = 0;
+	};
+
+	//! Class used to store the error, warning and information messages from several processes
+	//! performed by the tools library, like model transformation or compilation.
+	//! \ingroup tools
+	class MUTABLETOOLS_API ErrorLog : public RefCounted
+	{
+	public:
+
+		//-----------------------------------------------------------------------------------------
+		// Life cycle
+		//-----------------------------------------------------------------------------------------
+		ErrorLog();
+
+		//-----------------------------------------------------------------------------------------
+		// Own interface
+		//-----------------------------------------------------------------------------------------
+
+		//! Get the number of messages.
+		//! If a message type is provided, return the number of that type of message only.
+		int GetMessageCount( ErrorLogMessageType type = ELMT_NONE ) const;
+
+		//! Get the text of a message.
+		//! \param index index of the message from 0 to GetMessageCount(ELMT_NONE)-1
+		const char* GetMessageText( int index ) const;
+
+		//! Get the opaque context of a message.
+		//! \param index index of the message from 0 to GetMessageCount(ELMT_NONE)-1
+		const void* GetMessageContext( int index ) const;
+
+		//!
+		ErrorLogMessageType GetMessageType( int index ) const;
+
+		//! Get the attached data of a message.
+		//! \param index index of message data from 0 to GetMessageCount(ELMT_NONE)-1
+		ErrorLogMessageAttachedDataView GetMessageAttachedData( int index ) const;
+
+		//!
+		void Log() const;
+
+		//!
+		void Merge( const ErrorLog* pOther );
+
+		//-----------------------------------------------------------------------------------------
+		// Interface pattern
+		//-----------------------------------------------------------------------------------------
+		class Private;
+		Private* GetPrivate() const;
+
+	protected:
+
+		//! Forbidden. Manage with the Ptr<> template.
+		~ErrorLog();
+
+	private:
+
+		Private* m_pD;
+
+	};
+
+}
diff --git a/Engine/Plugins/Experimental/Mutable/Source/MutableTools/Public/ModelReport.h b/Engine/Plugins/Experimental/Mutable/Source/MutableTools/Public/ModelReport.h
new file mode 100644
--- /dev/null
+++ b/Engine/Plugins/Experimental/Mutable/Source/MutableTools/Public/ModelReport.h
@@ -0,0 +1,88 @@
+// Copyright Epic Games, Inc. All Rights Reserved.
+
+#pragma once
+
+#include "RefCounted.h"
+#include "Ptr.h"
+#include "Types.h"
+
+
+namespace mu
+{
+	// Forward declarations
+	class ModelReport;
+	typedef Ptr<ModelReport> ModelReportPtr;
+	typedef Ptr<const ModelReport> ModelReportPtrConst;
+
+	//! Class tha stores the additional information of a compiled model.
+	class MUTABLETOOLS_API ModelReport : public RefCounted
+	{
+	public:
+
+		//!
+		ModelReport();
+
+		//! Return the name of an image of a particular state.
+		//! If the state didn't target render-time generation, a null pointer is returned.
+		int GetStateCount() const;
+
+		//! Return the name of an image of a particular state.
+		//! If the state didn't target render-time generation, a null pointer is returned.
+		int GetLODCount( int state ) const;
+
+		//! Return the name of an image of a particular state.
+		//! If the state didn't target render-time generation, a null pointer is returned.
+		int GetComponentCount( int state, int lod ) const;
+
+		//! Return the name of an image of a particular state.
+		//! If the state didn't target render-time generation, a null pointer is returned.
+		const char* GetComponentName( int state, int lod, int comp ) const;
+
+		//! Return the name of an image of a particular state.
+		//! If the state didn't target render-time generation, a null pointer is returned.
+		int GetComponentImageCount( int state, int lod, int comp ) const;
+
+		//! Return the name of an image of a particular state.
+		//! If the state didn't target render-time generation, a null pointer is returned.
+		const char* GetImageName( int state, int lod, int comp, int image ) const;
+
+		//! Return the gpu source code that generates the value for an image in a particular state.
+		//! If the state didn't target render-time generation, a null pointer is returned.
+		const char* GetImageFragmentCode( int state, int lod, int comp, int image ) const;
+
+		//!
+		int GetFragmentSourceImageCount( int state, int lod, int comp, int image ) const;
+
+		//!
+		const char* GetFragmentSourceImageName( int state, int lod, int comp, int image, int source ) const;
+
+		//!
+		int GetFragmentSourceVectorCount( int state, int lod, int comp, int image ) const;
+
+		//!
+		const char* GetFragmentSourceVectorName( int state, int lod, int comp, int image, int source ) const;
+
+		//!
+		int GetFragmentSourceScalarCount( int state, int lod, int comp, int image ) const;
+
+		//!
+		const char* GetFragmentSourceScalarName( int state, int lod, int comp, int image, int source ) const;
+
+		//-----------------------------------------------------------------------------------------
+		// Interface pattern
+		//-----------------------------------------------------------------------------------------
+		class Private;
+		Private* GetPrivate() const;
+
+	protected:
+
+		//! Forbidden. Manage with the Ptr<> template.
+		~ModelReport();
+
+	private:
+
+		Private* m_pD;
+
+	};
+
+}
diff --git a/Engine/Plugins/Experimental/Mutable/Source/MutableTools/Public/Node.h b/Engine/Plugins/Experimental/Mutable/Source/MutableTools/Public/Node.h
new file mode 100644
--- /dev/null
+++ b/Engine/Plugins/Experimental/Mutable/Source/MutableTools/Public/Node.h
@@ -0,0 +1,167 @@
+// Copyright Epic Games, Inc. All Rights Reserved.
+
+#pragma once
+
+#include "RefCounted.h"
+#include "Ptr.h"
+#include "Types.h"
+#include "Serialisation.h"
+
+
+//! This tag is used to identify files containing serialised Node hierarchies. The tag is not added
+//! or checked by the Node serialisation methods, but the involved tools should take care of it.
+#define MUTABLE_SOURCE_MODEL_FILETAG	"amsm"
+
+
+namespace mu
+{
+	//! \defgroup tools Tools library classes
+	//! Tools library classes.
+
+	//! \defgroup model Model nodes
+	//! \ingroup tools
+	//! This group contains the nodes that can be used to compose models.
+
+	//! \defgroup transform Transform nodes
+	//! \ingroup tools
+	//! This group contains the nodes that can be used to compose model transformations.
+
+
+	// Forward declarations
+	class Node;
+	typedef Ptr<Node> NodePtr;
+	typedef Ptr<const Node> NodePtrConst;
+
+	class NodeMap;
+	typedef Ptr<NodeMap> NodeMapPtr;
+	typedef Ptr<const NodeMap> NodeMapPtrConst;
+
+
+	//! Information about the type of a node, to provide some means to the tools to deal generically
+	//! with nodes.
+	//! \ingroup tools
+	struct NODE_TYPE
+	{
+		NODE_TYPE();
+		NODE_TYPE( const char* strName, const NODE_TYPE* pParent );
+
+		const char* m_strName;
+		const NODE_TYPE* m_pParent;
+	};
+
+
+
+    //! %Base class for all graphs used in the source data to define models and transforms.
+	class MUTABLETOOLS_API Node : public RefCounted
+	{
+	public:
+
+		// Possible subclasses
+		enum class EType : uint8
+		{
+			Colour = 0,
+			Component = 1,
+			Image = 2,
+			Layout = 3,
+			LOD = 4,
+			Mesh = 5,
+			Object = 6,
+			PatchImage = 8,
+			Scalar = 9,
+			PatchMesh = 13,
+			Volume_Deprecated = 14,
+			Projector = 15,
+			Surface = 16,
+			Modifier = 18,
+			Range = 19,
+			String = 20,
+			Bool = 21,
+
+			None
+		};
+
+		//-----------------------------------------------------------------------------------------
+		// Life cycle
+		//-----------------------------------------------------------------------------------------
+
+		static void Serialise( const Node* pNode, OutputArchive& arch );
+		virtual void SerialiseWrapper(OutputArchive& arch) const = 0;
+		static NodePtr StaticUnserialise( InputArchive& arch );
+
+
+		//-----------------------------------------------------------------------------------------
+		// Own Interface
+		//-----------------------------------------------------------------------------------------
+
+		//! Clone the node tree.
+		//! \param pMap is an optional parameter that will return a map from the source tree nodes
+		//! to the cloned tree nodes.
+        virtual NodePtr Clone( NodeMapPtr pMap = nullptr ) const = 0;
+
+		//! Node type hierarchy data.
+        virtual const NODE_TYPE* GetType() const;
+		static const NODE_TYPE* GetStaticType();
+
+		//! Generic access to graph connections
+        virtual int GetInputCount() const = 0;
+		virtual Node* GetInputNode( int i ) const = 0;
+        virtual void SetInputNode( int i, NodePtr pNode ) = 0;
+
+		//! Set the opaque context returned in messages in the compiler log.
+		void SetMessageContext( const void* context );
+
+		//-----------------------------------------------------------------------------------------
+        // Interface pattern
+		//-----------------------------------------------------------------------------------------
+		class Private;
+		virtual Private* GetBasePrivate() const = 0;
+
+	protected:
+
+		inline ~Node() {}
+
+		//!
+		EType Type = EType::None;
+
+	};
+
+
+	//! Container mapping nodes to nodes.
+	class MUTABLETOOLS_API NodeMap : public RefCounted
+	{
+	public:
+
+		NodeMap();
+
+		//-----------------------------------------------------------------------------------------
+		// Own Interface
+		//-----------------------------------------------------------------------------------------
+
+		//! Return the number of elements in the map
+		int GetSize() const;
+
+		//! Insert a new key-value pair or overwrite the previews value with that key
+		void Add( const void* key, NodePtr value );
+
+		//! Get the value for a particular key. Returns 0 if the key is not in the map.
+		NodePtr Get( const void* key ) const;
+
+		//-----------------------------------------------------------------------------------------
+		// Interface pattern
+		//-----------------------------------------------------------------------------------------
+		class Private;
+		virtual Private* GetPrivate() const;
+
+	protected:
+
+		//! Forbidden. Manage with the Ptr<> template.
+		~NodeMap();
+
+	private:
+
+		Private* m_pD;
+
+	};
+
+
+}
diff --git a/Engine/Plugins/Experimental/Mutable/Source/MutableTools/Public/NodeBool.h b/Engine/Plugins/Experimental/Mutable/Source/MutableTools/Public/NodeBool.h
new file mode 100644
--- /dev/null
+++ b/Engine/Plugins/Experimental/Mutable/Source/MutableTools/Public/NodeBool.h
@@ -0,0 +1,403 @@
+// Copyright Epic Games, Inc. All Rights Reserved.
+
+#pragma once
+
+#include "MutableTools/Public/Node.h"
+
+
+namespace mu
+{
+
+	// Forward definitions
+	class NodeBool;
+    using NodeBoolPtr = Ptr<NodeBool>;
+    using NodeBoolPtrConst = Ptr<const NodeBool>;
+
+	class NodeBoolConstant;
+    using NodeBoolConstantPtr = Ptr<NodeBoolConstant>;
+    using NodeBoolConstantPtrConst = Ptr<const NodeBoolConstant>;
+
+	class NodeBoolParameter;
+    using NodeBoolParameterPtr = Ptr<NodeBoolParameter>;
+    using NodeBoolParameterPtrConst = Ptr<const NodeBoolParameter>;
+
+	class NodeBoolIsNull;
+    using NodeBoolIsNullPtr = Ptr<NodeBoolIsNull>;
+    using NodeBoolIsNullPtrConst = Ptr<const NodeBoolIsNull>;
+
+	class NodeBoolNot;
+    using NodeBoolNotPtr = Ptr<NodeBoolNot>;
+    using NodeBoolNotPtrConst = Ptr<const NodeBoolNot>;
+
+	class NodeBoolAnd;
+    using NodeBoolAndPtr = Ptr<NodeBoolAnd>;
+    using NodeBoolAndPtrConst = Ptr<const NodeBoolAnd>;
+
+    class NodeRange;
+    using NodeRangePtr = Ptr<NodeRange>;
+    using NodeRangePtrConst = Ptr<const NodeRange>;
+
+
+	//---------------------------------------------------------------------------------------------
+    //! %Base class of any node that outputs a Bool value.
+	//! \ingroup model
+	//---------------------------------------------------------------------------------------------
+	class MUTABLETOOLS_API NodeBool : public Node
+	{
+	public:
+
+		// Possible subclasses
+		enum class EType : uint8
+		{
+			Constant = 0,
+			Parameter = 1,
+			IsNull = 2,
+			Not = 3,
+			And = 4,
+
+			None
+		};
+
+		//-----------------------------------------------------------------------------------------
+		// Life cycle
+		//-----------------------------------------------------------------------------------------
+
+		static void Serialise( const NodeBool* pNode, OutputArchive& arch );
+		static NodeBoolPtr StaticUnserialise( InputArchive& arch );
+
+
+		//-----------------------------------------------------------------------------------------
+		// Node Interface
+		//-----------------------------------------------------------------------------------------
+
+        const NODE_TYPE* GetType() const override;
+		static const NODE_TYPE* GetStaticType();
+
+
+		//-----------------------------------------------------------------------------------------
+		// Interface pattern
+		//-----------------------------------------------------------------------------------------
+		class Private;
+
+	protected:
+
+		//! Forbidden. Manage with the Ptr<> template.
+		inline ~NodeBool() {}
+
+		//!
+		EType Type = EType::None;
+
+	};
+
+
+	//---------------------------------------------------------------------------------------------
+	//! Node returning a Bool constant value.
+	//! \ingroup model
+	//---------------------------------------------------------------------------------------------
+	class MUTABLETOOLS_API NodeBoolConstant : public NodeBool
+	{
+	public:
+
+		//-----------------------------------------------------------------------------------------
+		// Life cycle
+		//-----------------------------------------------------------------------------------------
+
+		NodeBoolConstant();
+
+		static void Serialise( const NodeBoolConstant* pNode, OutputArchive& arch );
+		void SerialiseWrapper(OutputArchive& arch) const override;
+		static NodeBoolConstantPtr StaticUnserialise( InputArchive& arch );
+
+
+		//-----------------------------------------------------------------------------------------
+		// Node Interface
+		//-----------------------------------------------------------------------------------------
+
+        NodePtr Clone( NodeMapPtr pMap = 0 ) const override;
+
+        const NODE_TYPE* GetType() const override;
+		static const NODE_TYPE* GetStaticType();
+
+        virtual int GetInputCount() const override;
+        virtual Node* GetInputNode( int i ) const override;
+        void SetInputNode( int i, NodePtr pNode ) override;
+
+		//-----------------------------------------------------------------------------------------
+		// Own Interface
+		//-----------------------------------------------------------------------------------------
+
+		//! Get the value to be returned by the node.
+		bool GetValue() const;
+
+		//! Set the value to be returned by the node.
+		void SetValue( bool v );
+
+		//-----------------------------------------------------------------------------------------
+		// Interface pattern
+		//-----------------------------------------------------------------------------------------
+		class Private;
+		Private* GetPrivate() const;
+        Node::Private* GetBasePrivate() const override;
+
+	protected:
+
+		//! Forbidden.
+		~NodeBoolConstant();
+
+	private:
+
+		Private* m_pD;
+
+	};
+
+
+	//---------------------------------------------------------------------------------------------
+	//! Node that defines a Bool model parameter.
+	//! \ingroup model
+	//---------------------------------------------------------------------------------------------
+	class MUTABLETOOLS_API NodeBoolParameter : public NodeBool
+	{
+	public:
+
+		//-----------------------------------------------------------------------------------------
+		// Life cycle
+		//-----------------------------------------------------------------------------------------
+
+		NodeBoolParameter();
+
+		void SerialiseWrapper(OutputArchive& arch) const override;
+		static void Serialise( const NodeBoolParameter* pNode, OutputArchive& arch );
+		static NodeBoolParameterPtr StaticUnserialise( InputArchive& arch );
+
+
+		//-----------------------------------------------------------------------------------------
+		// Node Interface
+		//-----------------------------------------------------------------------------------------
+
+        NodePtr Clone( NodeMapPtr pMap = 0 ) const override;
+
+        const NODE_TYPE* GetType() const override;
+		static const NODE_TYPE* GetStaticType();
+
+        int GetInputCount() const override;
+        Node* GetInputNode( int i ) const override;
+        void SetInputNode( int i, NodePtr pNode ) override;
+
+		//-----------------------------------------------------------------------------------------
+		// Own Interface
+		//-----------------------------------------------------------------------------------------
+
+		//! Get the name of the parameter. It will be exposed in the final compiled data.
+		const char* GetName() const;
+		void SetName( const char* );
+
+		//! Get the uid of the parameter. It will be exposed in the final compiled data.
+		const char* GetUid() const;
+		void SetUid( const char* );
+
+		//! Get the default value of the parameter.
+		bool GetDefaultValue() const;
+		void SetDefaultValue( bool v );
+
+        //! Set the number of ranges (dimensions) for this parameter.
+        //! By default a parameter has 0 ranges, meaning it only has one value.
+        void SetRangeCount( int i );
+        void SetRange( int i, NodeRangePtr pRange );
+
+		//-----------------------------------------------------------------------------------------
+		// Interface pattern
+		//-----------------------------------------------------------------------------------------
+		class Private;
+		Private* GetPrivate() const;
+        Node::Private* GetBasePrivate() const override;
+
+	protected:
+
+		//! Forbidden. 
+		~NodeBoolParameter();
+
+	private:
+
+		Private* m_pD;
+
+	};
+
+
+	//---------------------------------------------------------------------------------------------
+	//! Node that returns true if there is an input or false if there is nothing connected.
+	//! This node is mostly useful when used in the model Transform's conditions.
+	//! \ingroup model
+	//---------------------------------------------------------------------------------------------
+	class MUTABLETOOLS_API NodeBoolIsNull : public NodeBool
+	{
+	public:
+
+		//-----------------------------------------------------------------------------------------
+		// Life cycle
+		//-----------------------------------------------------------------------------------------
+
+		NodeBoolIsNull();
+
+		void SerialiseWrapper(OutputArchive& arch) const override;
+		static void Serialise( const NodeBoolIsNull* pNode, OutputArchive& arch );
+		static NodeBoolIsNullPtr StaticUnserialise( InputArchive& arch );
+
+
+		//-----------------------------------------------------------------------------------------
+		// Node Interface
+		//-----------------------------------------------------------------------------------------
+
+        NodePtr Clone( NodeMapPtr pMap = 0 ) const override;
+
+        const NODE_TYPE* GetType() const override;
+		static const NODE_TYPE* GetStaticType();
+
+        int GetInputCount() const override;
+        Node* GetInputNode( int i ) const override;
+        void SetInputNode( int i, NodePtr pNode ) override;
+
+		//-----------------------------------------------------------------------------------------
+		// Own Interface
+		//-----------------------------------------------------------------------------------------
+
+		//-----------------------------------------------------------------------------------------
+		// Interface pattern
+		//-----------------------------------------------------------------------------------------
+		class Private;
+		Private* GetPrivate() const;
+        Node::Private* GetBasePrivate() const override;
+
+	protected:
+
+		//! Forbidden. 
+		~NodeBoolIsNull();
+
+	private:
+
+		Private* m_pD;
+
+	};
+
+
+	//---------------------------------------------------------------------------------------------
+	//! Node that returns the oposite of the input value.
+	//! \ingroup model
+	//---------------------------------------------------------------------------------------------
+	class MUTABLETOOLS_API NodeBoolNot : public NodeBool
+	{
+	public:
+
+		//-----------------------------------------------------------------------------------------
+		// Life cycle
+		//-----------------------------------------------------------------------------------------
+
+		NodeBoolNot();
+
+		void SerialiseWrapper(OutputArchive& arch) const override;
+		static void Serialise( const NodeBoolNot* pNode, OutputArchive& arch );
+		static NodeBoolNotPtr StaticUnserialise( InputArchive& arch );
+
+
+		//-----------------------------------------------------------------------------------------
+		// Node Interface
+		//-----------------------------------------------------------------------------------------
+
+        NodePtr Clone( NodeMapPtr pMap = 0 ) const override;
+
+        const NODE_TYPE* GetType() const override;
+		static const NODE_TYPE* GetStaticType();
+
+        int GetInputCount() const override;
+        Node* GetInputNode( int i ) const override;
+        void SetInputNode( int i, NodePtr pNode ) override;
+
+		//-----------------------------------------------------------------------------------------
+		// Own Interface
+		//-----------------------------------------------------------------------------------------
+		
+		//! Input
+		NodeBoolPtr GetInput() const;
+		void SetInput( NodeBoolPtr );
+
+		//-----------------------------------------------------------------------------------------
+		// Interface pattern
+		//-----------------------------------------------------------------------------------------
+		class Private;
+		Private* GetPrivate() const;
+        Node::Private* GetBasePrivate() const override;
+
+	protected:
+
+		//! Forbidden. 
+		~NodeBoolNot();
+
+	private:
+
+		Private* m_pD;
+
+	};
+
+
+	//---------------------------------------------------------------------------------------------
+	//! 
+	//! \ingroup model
+	//---------------------------------------------------------------------------------------------
+	class MUTABLETOOLS_API NodeBoolAnd : public NodeBool
+	{
+	public:
+
+		//-----------------------------------------------------------------------------------------
+		// Life cycle
+		//-----------------------------------------------------------------------------------------
+
+		NodeBoolAnd();
+
+		void SerialiseWrapper(OutputArchive& arch) const override;
+		static void Serialise( const NodeBoolAnd* pNode, OutputArchive& arch );
+		static NodeBoolAndPtr StaticUnserialise( InputArchive& arch );
+
+
+		//-----------------------------------------------------------------------------------------
+		// Node Interface
+		//-----------------------------------------------------------------------------------------
+
+        NodePtr Clone( NodeMapPtr pMap = 0 ) const override;
+
+        const NODE_TYPE* GetType() const override;
+		static const NODE_TYPE* GetStaticType();
+
+        int GetInputCount() const override;
+        Node* GetInputNode( int i ) const override;
+        void SetInputNode( int i, NodePtr pNode ) override;
+
+		//-----------------------------------------------------------------------------------------
+		// Own Interface
+		//-----------------------------------------------------------------------------------------
+		
+		//! Inputs
+		NodeBoolPtr GetA() const;
+		void SetA( NodeBoolPtr );
+
+		NodeBoolPtr GetB() const;
+		void SetB( NodeBoolPtr );
+
+		//-----------------------------------------------------------------------------------------
+		// Interface pattern
+		//-----------------------------------------------------------------------------------------
+		class Private;
+		Private* GetPrivate() const;
+        Node::Private* GetBasePrivate() const override;
+
+	protected:
+
+		//! Forbidden. 
+		~NodeBoolAnd();
+
+	private:
+
+		Private* m_pD;
+
+	};
+
+
+}
diff --git a/Engine/Plugins/Experimental/Mutable/Source/MutableTools/Public/NodeColour.h b/Engine/Plugins/Experimental/Mutable/Source/MutableTools/Public/NodeColour.h
new file mode 100644
--- /dev/null
+++ b/Engine/Plugins/Experimental/Mutable/Source/MutableTools/Public/NodeColour.h
@@ -0,0 +1,70 @@
+// Copyright Epic Games, Inc. All Rights Reserved.
+
+#pragma once
+
+#include "Node.h"
+
+
+namespace mu
+{
+
+	class NodeColour;
+	typedef Ptr<NodeColour> NodeColourPtr;
+	typedef Ptr<const NodeColour> NodeColourPtrConst;
+
+
+    //! %Base class of any node that outputs a colour.
+    //! \ingroup model
+    class MUTABLETOOLS_API NodeColour : public Node
+    {
+    public:
+
+		// Possible subclasses
+		enum class EType : uint8
+		{
+			Constant = 0,
+			Parameter = 1,
+			SampleImage = 2,
+			Table = 3,
+			ImageSize = 4,
+			FromScalars = 5,
+			ArithmeticOperation = 6,
+			Switch = 7,
+			Variation = 8,
+
+			None
+		};
+
+        //-----------------------------------------------------------------------------------------
+        // Life cycle
+        //-----------------------------------------------------------------------------------------
+
+        static void Serialise( const NodeColour* pNode, OutputArchive& arch );
+        static NodeColourPtr StaticUnserialise( InputArchive& arch );
+
+
+        //-----------------------------------------------------------------------------------------
+        // Node Interface
+        //-----------------------------------------------------------------------------------------
+
+        const NODE_TYPE* GetType() const override;
+        static const NODE_TYPE* GetStaticType();
+
+
+        //-----------------------------------------------------------------------------------------
+        // Interface pattern
+        //-----------------------------------------------------------------------------------------
+        class Private;
+
+    protected:
+
+        //! Forbidden. Manage with the Ptr<> template.
+        inline ~NodeColour() {}
+
+		//!
+		EType Type = EType::None;
+
+    };
+
+}
+
diff --git a/Engine/Plugins/Experimental/Mutable/Source/MutableTools/Public/NodeColourArithmeticOperation.h b/Engine/Plugins/Experimental/Mutable/Source/MutableTools/Public/NodeColourArithmeticOperation.h
new file mode 100644
--- /dev/null
+++ b/Engine/Plugins/Experimental/Mutable/Source/MutableTools/Public/NodeColourArithmeticOperation.h
@@ -0,0 +1,98 @@
+// Copyright Epic Games, Inc. All Rights Reserved.
+
+#pragma once
+
+#include "NodeColour.h"
+
+
+namespace mu
+{
+
+	// Forward definitions
+	class NodeColourArithmeticOperation;
+	typedef Ptr<NodeColourArithmeticOperation> NodeColourArithmeticOperationPtr;
+	typedef Ptr<const NodeColourArithmeticOperation> NodeColourArithmeticOperationPtrConst;
+
+
+    //! Perform a per-component arithmetic operation between two colours.
+	//! \ingroup model
+	class MUTABLETOOLS_API NodeColourArithmeticOperation : public NodeColour
+	{
+	public:
+
+		//-----------------------------------------------------------------------------------------
+		// Life cycle
+		//-----------------------------------------------------------------------------------------
+
+		NodeColourArithmeticOperation();
+
+		void SerialiseWrapper(OutputArchive& arch) const override;
+		static void Serialise( const NodeColourArithmeticOperation* pNode, OutputArchive& arch );
+		static NodeColourArithmeticOperationPtr StaticUnserialise( InputArchive& arch );
+
+
+		//-----------------------------------------------------------------------------------------
+		// Node Interface
+		//-----------------------------------------------------------------------------------------
+
+		virtual NodePtr Clone( NodeMapPtr pMap = 0 ) const override;
+
+		virtual const NODE_TYPE* GetType() const override;
+		static const NODE_TYPE* GetStaticType();
+
+		virtual int GetInputCount() const override;
+		virtual Node* GetInputNode( int i ) const override;
+		virtual void SetInputNode( int i, NodePtr pNode ) override;
+
+
+		//-----------------------------------------------------------------------------------------
+		// Own Interface
+		//-----------------------------------------------------------------------------------------
+
+		//! Possible arithmetic operations
+		typedef enum
+		{
+			AO_ADD,
+			AO_SUBTRACT,
+			AO_MULTIPLY,
+            AO_DIVIDE,
+            _AO_COUNT
+		} OPERATION;
+
+        static const char* s_opTypeName[ _AO_COUNT ];
+
+		//! Get child selection type
+		OPERATION GetOperation() const;
+
+		//! Set the child selection type
+		void SetOperation(OPERATION);
+
+
+		//! Get the first operand
+		NodeColourPtr GetA() const;
+		void SetA(NodeColourPtr);
+
+		// Get the second operand
+		NodeColourPtr GetB() const;
+		void SetB(NodeColourPtr);
+
+		//-----------------------------------------------------------------------------------------
+		// Interface pattern
+		//-----------------------------------------------------------------------------------------
+		class Private;
+		Private* GetPrivate() const;
+        Node::Private* GetBasePrivate() const override;
+
+	protected:
+
+		//! Forbidden. Manage with the Ptr<> template.
+		~NodeColourArithmeticOperation();
+
+	private:
+
+		Private* m_pD;
+
+	};
+
+
+}
diff --git a/Engine/Plugins/Experimental/Mutable/Source/MutableTools/Public/NodeColourConstant.h b/Engine/Plugins/Experimental/Mutable/Source/MutableTools/Public/NodeColourConstant.h
new file mode 100644
--- /dev/null
+++ b/Engine/Plugins/Experimental/Mutable/Source/MutableTools/Public/NodeColourConstant.h
@@ -0,0 +1,74 @@
+// Copyright Epic Games, Inc. All Rights Reserved.
+
+#pragma once
+
+#include "NodeColour.h"
+
+
+namespace mu
+{
+
+	class NodeColourConstant;
+	typedef Ptr<NodeColourConstant> NodeColourConstantPtr;
+	typedef Ptr<const NodeColourConstant> NodeColourConstantPtrConst;
+
+
+	//! This node outputs a predefined colour value.
+	//! \ingroup model
+	class MUTABLETOOLS_API NodeColourConstant : public NodeColour
+	{
+	public:
+
+		//-----------------------------------------------------------------------------------------
+		// Life cycle
+		//-----------------------------------------------------------------------------------------
+
+		NodeColourConstant();
+
+		void SerialiseWrapper(OutputArchive& arch) const override;
+		static void Serialise( const NodeColourConstant* pNode, OutputArchive& arch );
+		static NodeColourConstantPtr StaticUnserialise( InputArchive& arch );
+
+		//-----------------------------------------------------------------------------------------
+		// Node Interface
+		//-----------------------------------------------------------------------------------------
+
+        NodePtr Clone( NodeMapPtr pMap = 0 ) const override;
+
+        const NODE_TYPE* GetType() const override;
+		static const NODE_TYPE* GetStaticType();
+
+        virtual int GetInputCount() const override;
+        virtual Node* GetInputNode( int i ) const override;
+        void SetInputNode( int i, NodePtr pNode ) override;
+
+		//-----------------------------------------------------------------------------------------
+		// Own Interface
+		//-----------------------------------------------------------------------------------------
+
+		//! Get the value that this node returns
+		void GetValue( float* pR, float* pG, float* pB ) const;
+
+		//! Set the value to be returned by this node
+		void SetValue( float r, float g, float b );
+
+		//-----------------------------------------------------------------------------------------
+		// Interface pattern
+		//-----------------------------------------------------------------------------------------
+		class Private;
+		Private* GetPrivate() const;
+        Node::Private* GetBasePrivate() const override;
+
+	protected:
+
+		//! Forbidden. Manage with the Ptr<> template.
+		~NodeColourConstant();
+
+	private:
+
+		Private* m_pD;
+
+	};
+
+
+}
diff --git a/Engine/Plugins/Experimental/Mutable/Source/MutableTools/Public/NodeColourFromScalars.h b/Engine/Plugins/Experimental/Mutable/Source/MutableTools/Public/NodeColourFromScalars.h
new file mode 100644
--- /dev/null
+++ b/Engine/Plugins/Experimental/Mutable/Source/MutableTools/Public/NodeColourFromScalars.h
@@ -0,0 +1,91 @@
+// Copyright Epic Games, Inc. All Rights Reserved.
+
+#pragma once
+
+#include "NodeColour.h"
+
+#include "Image.h"
+
+
+namespace mu
+{
+
+	// Forward definitions
+	class NodeScalar;
+	typedef Ptr<NodeScalar> NodeScalarPtr;
+	typedef Ptr<const NodeScalar> NodeScalarPtrConst;
+
+	class NodeColourFromScalars;
+	typedef Ptr<NodeColourFromScalars> NodeColourFromScalarsPtr;
+	typedef Ptr<const NodeColourFromScalars> NodeColourFromScalarsPtrConst;
+
+
+	//! Obtain a colour by sampling an image at specific homogeneous coordinates.
+	//! \ingroup model
+	class MUTABLETOOLS_API NodeColourFromScalars : public NodeColour
+	{
+	public:
+
+		//-----------------------------------------------------------------------------------------
+		// Life cycle
+		//-----------------------------------------------------------------------------------------
+
+		NodeColourFromScalars();
+
+		void SerialiseWrapper(OutputArchive& arch) const override;
+		static void Serialise( const NodeColourFromScalars* pNode, OutputArchive& arch );
+		static NodeColourFromScalarsPtr StaticUnserialise( InputArchive& arch );
+
+
+		//-----------------------------------------------------------------------------------------
+		// Node Interface
+		//-----------------------------------------------------------------------------------------
+
+        NodePtr Clone( NodeMapPtr pMap = nullptr ) const override;
+
+		const NODE_TYPE* GetType() const override;
+		static const NODE_TYPE* GetStaticType();
+
+        int GetInputCount() const override;
+        Node* GetInputNode( int i ) const override;
+		void SetInputNode( int i, NodePtr pNode ) override;
+
+
+		//-----------------------------------------------------------------------------------------
+		// Own Interface
+		//-----------------------------------------------------------------------------------------
+
+		//! Get the node generating the x coordinate to sample, in the range of 0.0 to 1.0.
+		//! If it is not specified, a value of 1 will be used.
+		NodeScalarPtr GetX() const;
+		void SetX( NodeScalarPtr );
+
+		NodeScalarPtr GetY() const;
+		void SetY( NodeScalarPtr );
+
+		NodeScalarPtr GetZ() const;
+		void SetZ( NodeScalarPtr );
+
+		NodeScalarPtr GetW() const;
+		void SetW( NodeScalarPtr );
+
+		//-----------------------------------------------------------------------------------------
+		// Interface pattern
+		//-----------------------------------------------------------------------------------------
+		class Private;
+		Private* GetPrivate() const;
+        Node::Private* GetBasePrivate() const override;
+
+	protected:
+
+		//! Forbidden. Manage with the Ptr<> template.
+		~NodeColourFromScalars();
+
+	private:
+
+		Private* m_pD;
+
+	};
+
+
+}
diff --git a/Engine/Plugins/Experimental/Mutable/Source/MutableTools/Public/NodeColourParameter.h b/Engine/Plugins/Experimental/Mutable/Source/MutableTools/Public/NodeColourParameter.h
new file mode 100644
--- /dev/null
+++ b/Engine/Plugins/Experimental/Mutable/Source/MutableTools/Public/NodeColourParameter.h
@@ -0,0 +1,96 @@
+// Copyright Epic Games, Inc. All Rights Reserved.
+
+#pragma once
+
+#include "NodeColour.h"
+
+
+namespace mu
+{
+
+	class NodeColourParameter;
+    using NodeColourParameterPtr = Ptr<NodeColourParameter>;
+    using NodeColourParameterPtrConst = Ptr<const NodeColourParameter>;
+
+    class NodeRange;
+    using NodeRangePtr = Ptr<NodeRange>;
+    using NodeRangePtrConst = Ptr<const NodeRange>;
+
+
+	//! Node that defines a colour model parameter.
+	//! \ingroup model
+	class MUTABLETOOLS_API NodeColourParameter : public NodeColour
+	{
+	public:
+
+		//-----------------------------------------------------------------------------------------
+		// Life cycle
+		//-----------------------------------------------------------------------------------------
+
+		NodeColourParameter();
+
+		void SerialiseWrapper(OutputArchive& arch) const override;
+		static void Serialise( const NodeColourParameter* pNode, OutputArchive& arch );
+		static NodeColourParameterPtr StaticUnserialise( InputArchive& arch );
+
+
+		//-----------------------------------------------------------------------------------------
+		// Node Interface
+		//-----------------------------------------------------------------------------------------
+
+		NodePtr Clone( NodeMapPtr pMap = 0 ) const override;
+
+		const NODE_TYPE* GetType() const override;
+		static const NODE_TYPE* GetStaticType();
+
+		int GetInputCount() const override;
+		Node* GetInputNode( int i ) const override;
+		void SetInputNode( int i, NodePtr pNode ) override;
+
+		//-----------------------------------------------------------------------------------------
+		// Own Interface
+		//-----------------------------------------------------------------------------------------
+
+		//! Get the name of the parameter. It will be exposed in the final compiled data.
+		const char* GetName() const;
+
+		//! Set the name of the parameter.
+		void SetName( const char* );
+
+		//! Get the uid of the parameter. It will be exposed in the final compiled data.
+		const char* GetUid() const;
+
+		//! Set the uid of the parameter.
+		void SetUid( const char* );
+
+		//! Get the default value of the parameter.
+		void GetDefaultValue( float* pR, float* pG, float* pB ) const;
+
+		//! Set the default value of the parameter.
+		void SetDefaultValue( float r, float g, float b );
+
+        //! Set the number of ranges (dimensions) for this parameter.
+        //! By default a parameter has 0 ranges, meaning it only has one value.
+        void SetRangeCount( int i );
+        void SetRange( int i, NodeRangePtr pRange );
+
+		//-----------------------------------------------------------------------------------------
+		// Interface pattern
+		//-----------------------------------------------------------------------------------------
+		class Private;
+		Private* GetPrivate() const;
+        Node::Private* GetBasePrivate() const override;
+
+	protected:
+
+		//! Forbidden. Manage with the Ptr<> template.
+		~NodeColourParameter();
+
+	private:
+
+		Private* m_pD;
+
+	};
+
+
+}
diff --git a/Engine/Plugins/Experimental/Mutable/Source/MutableTools/Public/NodeColourSampleImage.h b/Engine/Plugins/Experimental/Mutable/Source/MutableTools/Public/NodeColourSampleImage.h
new file mode 100644
--- /dev/null
+++ b/Engine/Plugins/Experimental/Mutable/Source/MutableTools/Public/NodeColourSampleImage.h
@@ -0,0 +1,95 @@
+// Copyright Epic Games, Inc. All Rights Reserved.
+
+#pragma once
+
+#include "NodeColour.h"
+
+#include "Image.h"
+
+
+namespace mu
+{
+
+	// Forward definitions
+	class NodeScalar;
+	typedef Ptr<NodeScalar> NodeScalarPtr;
+	typedef Ptr<const NodeScalar> NodeScalarPtrConst;
+
+	class NodeImage;
+	typedef Ptr<NodeImage> NodeImagePtr;
+	typedef Ptr<const NodeImage> NodeImagePtrConst;
+
+	class NodeColourSampleImage;
+	typedef Ptr<NodeColourSampleImage> NodeColourSampleImagePtr;
+	typedef Ptr<const NodeColourSampleImage> NodeColourSampleImagePtrConst;
+
+
+	//! Obtain a colour by sampling an image at specific homogeneous coordinates.
+	//! \ingroup model
+	class MUTABLETOOLS_API NodeColourSampleImage : public NodeColour
+	{
+	public:
+
+		//-----------------------------------------------------------------------------------------
+		// Life cycle
+		//-----------------------------------------------------------------------------------------
+
+		NodeColourSampleImage();
+
+		void SerialiseWrapper(OutputArchive& arch) const override;
+		static void Serialise( const NodeColourSampleImage* pNode, OutputArchive& arch );
+		static NodeColourSampleImagePtr StaticUnserialise( InputArchive& arch );
+
+
+		//-----------------------------------------------------------------------------------------
+		// Node Interface
+		//-----------------------------------------------------------------------------------------
+
+		NodePtr Clone( NodeMapPtr pMap = 0 ) const override;
+
+		const NODE_TYPE* GetType() const override;
+		static const NODE_TYPE* GetStaticType();
+
+		int GetInputCount() const override;
+		Node* GetInputNode( int i ) const override;
+		void SetInputNode( int i, NodePtr pNode ) override;
+
+
+		//-----------------------------------------------------------------------------------------
+		// Own Interface
+		//-----------------------------------------------------------------------------------------
+
+		//! Get the node generating the x coordinate to sample, in the range of 0.0 to 1.0.
+		//! If it is not specified, a value of x =0.5 will be used.
+		NodeScalarPtr GetX() const;
+		void SetX( NodeScalarPtr );
+
+		//! Get the node generating the y coordinate to sample, in the range of 0.0 to 1.0.
+		//! If it is not specified, a value of y =0.5 will be used.
+		NodeScalarPtr GetY() const;
+		void SetY( NodeScalarPtr );
+
+		//! Get the node generating the source image to be sampled.
+		NodeImagePtr GetImage() const;
+		void SetImage( NodeImagePtr );
+
+		//-----------------------------------------------------------------------------------------
+		// Interface pattern
+		//-----------------------------------------------------------------------------------------
+		class Private;
+		Private* GetPrivate() const;
+        Node::Private* GetBasePrivate() const override;
+
+	protected:
+
+		//! Forbidden. Manage with the Ptr<> template.
+		~NodeColourSampleImage();
+
+	private:
+
+		Private* m_pD;
+
+	};
+
+
+}
diff --git a/Engine/Plugins/Experimental/Mutable/Source/MutableTools/Public/NodeColourSwitch.h b/Engine/Plugins/Experimental/Mutable/Source/MutableTools/Public/NodeColourSwitch.h
new file mode 100644
--- /dev/null
+++ b/Engine/Plugins/Experimental/Mutable/Source/MutableTools/Public/NodeColourSwitch.h
@@ -0,0 +1,87 @@
+// Copyright Epic Games, Inc. All Rights Reserved.
+
+#pragma once
+
+#include "NodeColour.h"
+
+
+namespace mu
+{
+
+	// Forward definitions
+	class NodeScalar;
+	typedef Ptr<NodeScalar> NodeScalarPtr;
+	typedef Ptr<const NodeScalar> NodeScalarPtrConst;
+
+    class NodeColourSwitch;
+    typedef Ptr<NodeColourSwitch> NodeColourSwitchPtr;
+    typedef Ptr<const NodeColourSwitch> NodeColourSwitchPtrConst;
+
+
+    //! This node selects an output Colour from a set of input Colours based on a parameter.
+	//! \ingroup model
+    class MUTABLETOOLS_API NodeColourSwitch : public NodeColour
+	{
+	public:
+
+		//-----------------------------------------------------------------------------------------
+		// Life cycle
+		//-----------------------------------------------------------------------------------------
+
+        NodeColourSwitch();
+
+		void SerialiseWrapper(OutputArchive& arch) const override;
+		static void Serialise( const NodeColourSwitch* pNode, OutputArchive& arch );
+        static NodeColourSwitchPtr StaticUnserialise( InputArchive& arch );
+
+
+		//-----------------------------------------------------------------------------------------
+		// Node Interface
+		//-----------------------------------------------------------------------------------------
+
+        NodePtr Clone( NodeMapPtr pMap = 0 ) const override;
+
+        const NODE_TYPE* GetType() const override;
+		static const NODE_TYPE* GetStaticType();
+
+        virtual int GetInputCount() const override;
+        virtual Node* GetInputNode( int i ) const override;
+        void SetInputNode( int i, NodePtr pNode ) override;
+
+		//-----------------------------------------------------------------------------------------
+		// Own Interface
+		//-----------------------------------------------------------------------------------------
+
+		//! Get the node generating the parameter used to select the option.
+		NodeScalarPtr GetParameter() const;
+		void SetParameter( NodeScalarPtr );
+
+        //! Set the number of option Colours. It will keep the currently set targets and initialise
+		//! the new ones as null.
+		void SetOptionCount( int );
+
+        //! Get the node generating the t-th option Colour.
+        NodeColourPtr GetOption( int t ) const;
+        void SetOption( int t, NodeColourPtr );
+
+
+		//-----------------------------------------------------------------------------------------
+		// Interface pattern
+		//-----------------------------------------------------------------------------------------
+		class Private;
+		Private* GetPrivate() const;
+        Node::Private* GetBasePrivate() const override;
+
+	protected:
+
+		//! Forbidden. Manage with the Ptr<> template.
+        ~NodeColourSwitch();
+
+	private:
+
+		Private* m_pD;
+
+	};
+
+
+}
diff --git a/Engine/Plugins/Experimental/Mutable/Source/MutableTools/Public/NodeColourTable.h b/Engine/Plugins/Experimental/Mutable/Source/MutableTools/Public/NodeColourTable.h
new file mode 100644
--- /dev/null
+++ b/Engine/Plugins/Experimental/Mutable/Source/MutableTools/Public/NodeColourTable.h
@@ -0,0 +1,86 @@
+// Copyright Epic Games, Inc. All Rights Reserved.
+
+#pragma once
+
+#include "NodeColour.h"
+
+
+namespace mu
+{
+
+	// Forward definitions
+	class NodeColourTable;
+	typedef Ptr<NodeColourTable> NodeColourTablePtr;
+	typedef Ptr<const NodeColourTable> NodeColourTablePtrConst;
+
+	class Table;
+	typedef Ptr<Table> TablePtr;
+	typedef Ptr<const Table> TablePtrConst;
+
+
+	//! This node provides the meshes stored in the column of a table.
+	//! \ingroup transform
+	class MUTABLETOOLS_API NodeColourTable : public NodeColour
+	{
+	public:
+
+		//-----------------------------------------------------------------------------------------
+		// Life cycle
+		//-----------------------------------------------------------------------------------------
+
+		NodeColourTable();
+
+		void SerialiseWrapper(OutputArchive& arch) const override;
+		static void Serialise( const NodeColourTable* pNode, OutputArchive& arch );
+		static NodeColourTablePtr StaticUnserialise( InputArchive& arch );
+
+
+		//-----------------------------------------------------------------------------------------
+		// Node Interface
+		//-----------------------------------------------------------------------------------------
+
+		NodePtr Clone( NodeMapPtr pMap = 0 ) const override;
+
+		const NODE_TYPE* GetType() const override;
+		static const NODE_TYPE* GetStaticType();
+
+		int GetInputCount() const override;
+		Node* GetInputNode( int i ) const override;
+		void SetInputNode( int i, NodePtr pNode ) override;
+
+		//-----------------------------------------------------------------------------------------
+		// Own Interface
+		//-----------------------------------------------------------------------------------------
+
+		//! Set the name of the implicit table parameter.
+		void SetParameterName( const char* strName );
+
+		//!
+		TablePtr GetTable() const;
+		void SetTable( TablePtr );
+
+		//!
+		const char* GetColumn() const;
+		void SetColumn( const char* strName );
+
+		//-----------------------------------------------------------------------------------------
+		// Interface pattern
+		//-----------------------------------------------------------------------------------------
+		class Private;
+		Private* GetPrivate() const;
+        Node::Private* GetBasePrivate() const override;
+
+
+	protected:
+
+		//! Forbidden. Manage with the Ptr<> template.
+		~NodeColourTable();
+
+	private:
+
+		Private* m_pD;
+
+	};
+
+
+}
diff --git a/Engine/Plugins/Experimental/Mutable/Source/MutableTools/Public/NodeColourVariation.h b/Engine/Plugins/Experimental/Mutable/Source/MutableTools/Public/NodeColourVariation.h
new file mode 100644
--- /dev/null
+++ b/Engine/Plugins/Experimental/Mutable/Source/MutableTools/Public/NodeColourVariation.h
@@ -0,0 +1,83 @@
+// Copyright Epic Games, Inc. All Rights Reserved.
+
+#pragma once
+
+#include "NodeColour.h"
+
+
+namespace mu
+{
+
+    // Forward definitions
+    class NodeColourVariation;
+    typedef Ptr<NodeColourVariation> NodeColourVariationPtr;
+    typedef Ptr<const NodeColourVariation> NodeColourVariationPtrConst;
+
+
+    //!
+    //! \ingroup model
+    class MUTABLETOOLS_API NodeColourVariation : public NodeColour
+    {
+    public:
+        //-----------------------------------------------------------------------------------------
+        // Life cycle
+        //-----------------------------------------------------------------------------------------
+
+        NodeColourVariation();
+
+		void SerialiseWrapper(OutputArchive& arch) const override;
+		static void Serialise( const NodeColourVariation* pNode, OutputArchive& arch );
+        static NodeColourVariationPtr StaticUnserialise( InputArchive& arch );
+
+        //-----------------------------------------------------------------------------------------
+        // Node Interface
+        //-----------------------------------------------------------------------------------------
+
+        NodePtr Clone( NodeMapPtr pMap = 0 ) const override;
+
+        const NODE_TYPE* GetType() const override;
+        static const NODE_TYPE* GetStaticType();
+
+        virtual int GetInputCount() const override;
+        virtual Node* GetInputNode( int i ) const override;
+        void SetInputNode( int i, NodePtr pNode ) override;
+
+        //-----------------------------------------------------------------------------------------
+        // Own Interface
+        //-----------------------------------------------------------------------------------------
+
+        //!
+        void SetDefaultColour( NodeColour* Colour );
+
+        //! Set the number of tags to consider in this variation
+        void SetVariationCount( int count );
+
+        //!
+        int GetVariationCount() const;
+
+        //! Set the tag or state name that will enable a specific vartiation
+        void SetVariationTag( int index, const char* strTag );
+
+        //!
+        void SetVariationColour( int index, NodeColour* Colour );
+
+        //!}
+
+
+        //-----------------------------------------------------------------------------------------
+        // Interface pattern
+        //-----------------------------------------------------------------------------------------
+        class Private;
+        Private* GetPrivate() const;
+        Node::Private* GetBasePrivate() const override;
+
+    protected:
+        //! Forbidden. Manage with the Ptr<> template.
+        ~NodeColourVariation();
+
+    private:
+        Private* m_pD;
+    };
+
+
+} // namespace mu
diff --git a/Engine/Plugins/Experimental/Mutable/Source/MutableTools/Public/NodeComponent.h b/Engine/Plugins/Experimental/Mutable/Source/MutableTools/Public/NodeComponent.h
new file mode 100644
--- /dev/null
+++ b/Engine/Plugins/Experimental/Mutable/Source/MutableTools/Public/NodeComponent.h
@@ -0,0 +1,87 @@
+// Copyright Epic Games, Inc. All Rights Reserved.
+
+#pragma once
+
+#include "Node.h"
+
+
+namespace mu
+{
+
+	// Forward definitions
+	class NodeComponent;
+	typedef Ptr<NodeComponent> NodeComponentPtr;
+	typedef Ptr<const NodeComponent> NodeComponentPtrConst;
+
+    class NodeSurface;
+
+
+	//! This class is the parent of all nodes that output a component.
+	//! \ingroup model
+	class MUTABLETOOLS_API NodeComponent : public Node
+	{
+	public:
+
+		// Possible subclasses
+		enum class EType : uint8
+		{
+			New = 0,
+			Edit = 1,
+
+			None
+		};
+
+		//-----------------------------------------------------------------------------------------
+		// Life cycle
+		//-----------------------------------------------------------------------------------------
+
+		static void Serialise( const NodeComponent* pNode, OutputArchive& arch );
+		static NodeComponentPtr StaticUnserialise( InputArchive& arch );
+
+		//-----------------------------------------------------------------------------------------
+        // Node interface
+		//-----------------------------------------------------------------------------------------
+
+        const NODE_TYPE* GetType() const override;
+		static const NODE_TYPE* GetStaticType();
+
+        //-----------------------------------------------------------------------------------------
+        // Own interface
+        //-----------------------------------------------------------------------------------------
+
+        //! \name Surfaces
+        //! \{
+
+        //! Get the number of meshes in the component.
+        virtual int GetSurfaceCount() const = 0;
+
+        //! Set the number of meshes in the component.
+        virtual void SetSurfaceCount( int ) = 0;
+
+        //! Get the node generating one of the meshes in the component.
+        //! \param index index of the mesh, from 0 to GetSurfaceCount()-1
+        virtual NodeSurface* GetSurface( int index ) const = 0;
+
+        //! Set the node generating one of the meshes in the component.
+        //! \param index index of the mesh, from 0 to GetSurfaceCount()-1
+        virtual void SetSurface( int index, NodeSurface* ) = 0;
+
+        //! \}
+
+		//-----------------------------------------------------------------------------------------
+		// Interface pattern
+		//-----------------------------------------------------------------------------------------
+		class Private;
+
+	protected:
+
+		//! Forbidden. Manage with the Ptr<> template.
+		inline ~NodeComponent() {}
+
+		//!
+		EType Type = EType::None;
+
+	};
+
+
+}
diff --git a/Engine/Plugins/Experimental/Mutable/Source/MutableTools/Public/NodeComponentEdit.h b/Engine/Plugins/Experimental/Mutable/Source/MutableTools/Public/NodeComponentEdit.h
new file mode 100644
--- /dev/null
+++ b/Engine/Plugins/Experimental/Mutable/Source/MutableTools/Public/NodeComponentEdit.h
@@ -0,0 +1,90 @@
+// Copyright Epic Games, Inc. All Rights Reserved.
+
+#pragma once
+
+#include "NodeComponent.h"
+
+
+namespace mu
+{
+
+	// Forward definitions
+	class NodeComponentEdit;
+	typedef Ptr<NodeComponentEdit> NodeComponentEditPtr;
+	typedef Ptr<const NodeComponentEdit> NodeComponentEditPtrConst;
+
+    class NodeSurface;
+    typedef Ptr<NodeSurface> NodeSurfacePtr;
+    typedef Ptr<const NodeSurface> NodeSurfacePtrConst;
+
+
+
+	//! This node modifies a node of the parent object of the object that this node belongs to.
+	//! It allows to extend, cut and morph the parent component's meshes.
+	//! It also allows to patch the parent component's textures.
+	//! \ingroup model
+	class MUTABLETOOLS_API NodeComponentEdit : public NodeComponent
+	{
+	public:
+
+		//-----------------------------------------------------------------------------------------
+		// Life cycle
+		//-----------------------------------------------------------------------------------------
+
+		NodeComponentEdit();
+
+		void SerialiseWrapper(OutputArchive& arch) const override;
+		static void Serialise( const NodeComponentEdit* pNode, OutputArchive& arch );
+		static NodeComponentEditPtr StaticUnserialise( InputArchive& arch );
+
+		//-----------------------------------------------------------------------------------------
+        // Node interface
+		//-----------------------------------------------------------------------------------------
+		NodePtr Clone( NodeMapPtr pMap = 0 ) const override;
+
+		const NODE_TYPE* GetType() const override;
+		static const NODE_TYPE* GetStaticType();
+
+		int GetInputCount() const override;
+		Node* GetInputNode( int i ) const override;
+		void SetInputNode( int i, NodePtr pNode ) override;
+
+        //-----------------------------------------------------------------------------------------
+        // NodeComponent interface
+        //-----------------------------------------------------------------------------------------
+        int GetSurfaceCount() const override;
+        void SetSurfaceCount( int ) override;
+        NodeSurface* GetSurface( int index ) const override;
+        void SetSurface( int index, NodeSurface* ) override;
+
+		//-----------------------------------------------------------------------------------------
+        // Own interface
+		//-----------------------------------------------------------------------------------------
+
+		//! Set the parent component to modify. It should be a node from the same LOD of one of
+		//! the parent objects of this node's object.
+        void SetParent( NodeComponent* );
+        NodeComponent* GetParent() const;
+
+
+		//-----------------------------------------------------------------------------------------
+		// Interface pattern
+		//-----------------------------------------------------------------------------------------
+		class Private;
+		Private* GetPrivate() const;
+        Node::Private* GetBasePrivate() const override;
+
+	protected:
+
+		//! Forbidden. Manage with the Ptr<> template.
+		~NodeComponentEdit();
+
+	private:
+
+		Private* m_pD;
+
+	};
+
+
+
+}
diff --git a/Engine/Plugins/Experimental/Mutable/Source/MutableTools/Public/NodeComponentNew.h b/Engine/Plugins/Experimental/Mutable/Source/MutableTools/Public/NodeComponentNew.h
new file mode 100644
--- /dev/null
+++ b/Engine/Plugins/Experimental/Mutable/Source/MutableTools/Public/NodeComponentNew.h
@@ -0,0 +1,96 @@
+// Copyright Epic Games, Inc. All Rights Reserved.
+
+#pragma once
+
+#include "NodeComponent.h"
+
+
+namespace mu
+{
+
+	// Forward definitions
+	class NodeComponent;
+	typedef Ptr<NodeComponent> NodeComponentPtr;
+	typedef Ptr<const NodeComponent> NodeComponentPtrConst;
+
+	class NodeComponentNew;
+	typedef Ptr<NodeComponentNew> NodeComponentNewPtr;
+	typedef Ptr<const NodeComponentNew> NodeComponentNewPtrConst;
+
+    class NodeSurface;
+    typedef Ptr<NodeSurface> NodeSurfacePtr;
+    typedef Ptr<const NodeSurface> NodeSurfacePtrConst;
+
+
+	//! This node makes a new component from several meshes and images.
+	//! \ingroup model
+	class MUTABLETOOLS_API NodeComponentNew : public NodeComponent
+	{
+	public:
+
+		//-----------------------------------------------------------------------------------------
+		// Life cycle
+		//-----------------------------------------------------------------------------------------
+
+		NodeComponentNew();
+
+		void SerialiseWrapper(OutputArchive& arch) const override;
+		static void Serialise( const NodeComponentNew* pNode, OutputArchive& arch );
+		static NodeComponentNewPtr StaticUnserialise( InputArchive& arch );
+
+		//-----------------------------------------------------------------------------------------
+        // Node interface
+		//-----------------------------------------------------------------------------------------
+		NodePtr Clone( NodeMapPtr pMap = 0 ) const override;
+
+        const NODE_TYPE* GetType() const override;
+		static const NODE_TYPE* GetStaticType();
+
+		int GetInputCount() const override;
+		Node* GetInputNode( int i ) const override;
+		void SetInputNode( int i, NodePtr pNode ) override;
+
+        //-----------------------------------------------------------------------------------------
+        // NodeComponent interface
+        //-----------------------------------------------------------------------------------------
+        int GetSurfaceCount() const override;
+        void SetSurfaceCount( int ) override;
+        NodeSurface* GetSurface( int index ) const override;
+        void SetSurface( int index, NodeSurface* ) override;
+
+		//-----------------------------------------------------------------------------------------
+        // Own interface
+		//-----------------------------------------------------------------------------------------
+
+		//! Get the name of the component.
+		const char* GetName() const;
+
+		//! Set the name of the component.
+		void SetName( const char* );
+
+		//! Get the id of the component.
+		uint16 GetId() const;
+
+		//! Set the id of the component.
+		void SetId( uint16 );
+
+		//-----------------------------------------------------------------------------------------
+		// Interface pattern
+		//-----------------------------------------------------------------------------------------
+		class Private;
+		Private* GetPrivate() const;
+        Node::Private* GetBasePrivate() const override;
+
+	protected:
+
+		//! Forbidden. Manage with the Ptr<> template.
+		~NodeComponentNew();
+
+	private:
+
+		Private* m_pD;
+
+	};
+
+
+}
diff --git a/Engine/Plugins/Experimental/Mutable/Source/MutableTools/Public/NodeImage.h b/Engine/Plugins/Experimental/Mutable/Source/MutableTools/Public/NodeImage.h
new file mode 100644
--- /dev/null
+++ b/Engine/Plugins/Experimental/Mutable/Source/MutableTools/Public/NodeImage.h
@@ -0,0 +1,91 @@
+// Copyright Epic Games, Inc. All Rights Reserved.
+
+#pragma once
+
+#include "Node.h"
+
+namespace mu
+{
+
+	// Forward definitions
+	class NodeImage;
+	typedef Ptr<NodeImage> NodeImagePtr;
+	typedef Ptr<const NodeImage> NodeImagePtrConst;
+
+
+    //! %Base class of any node that outputs an image.
+	//! \ingroup model
+	class MUTABLETOOLS_API NodeImage : public Node
+	{
+	public:
+
+		// Possible subclasses
+		enum class EType : uint8
+		{
+			Constant = 1,
+			Difference = 2,
+			Interpolate = 4,
+			Saturate = 6,
+			Table = 9,
+			Swizzle = 10,
+			SelectColour = 11,
+			ColourMap = 12,
+			Gradient = 13,
+			VolumeLayer_Deprecated = 14,
+			Binarise = 17,
+			Luminance = 18,
+			Layer = 19,
+			LayerColour = 20,
+			Resize = 21,
+			PlainColour = 22,
+			Interpolate3 = 23,
+			Project = 24,
+			Mipmap = 25,
+			Switch = 26,
+			Conditional = 27,
+			Format = 28,
+			Parameter = 29,
+			MultiLayer = 30,
+			Invert = 31,
+			Variation = 32,
+			NormalComposite = 33,
+			Transform = 34,
+			None
+		};
+
+		//-----------------------------------------------------------------------------------------
+		// Life cycle
+		//-----------------------------------------------------------------------------------------
+
+		static void Serialise( const NodeImage* pNode, OutputArchive& arch );
+		static NodeImagePtr StaticUnserialise( InputArchive& arch );
+
+
+		//-----------------------------------------------------------------------------------------
+		// Node Interface
+		//-----------------------------------------------------------------------------------------
+
+		const NODE_TYPE* GetType() const override;
+		static const NODE_TYPE* GetStaticType();
+
+		inline EType GetImageNodeType() const { return Type; }
+
+
+		//-----------------------------------------------------------------------------------------
+		// Interface pattern
+		//-----------------------------------------------------------------------------------------
+		class Private;
+
+	protected:
+
+		//! Forbidden. Manage with the Ptr<> template.
+		inline ~NodeImage() {}
+
+		//!
+		EType Type = EType::None;
+
+	};
+
+
+
+}
diff --git a/Engine/Plugins/Experimental/Mutable/Source/MutableTools/Public/NodeImageBinarise.h b/Engine/Plugins/Experimental/Mutable/Source/MutableTools/Public/NodeImageBinarise.h
new file mode 100644
--- /dev/null
+++ b/Engine/Plugins/Experimental/Mutable/Source/MutableTools/Public/NodeImageBinarise.h
@@ -0,0 +1,82 @@
+// Copyright Epic Games, Inc. All Rights Reserved.
+
+#pragma once
+
+#include "NodeImage.h"
+
+
+namespace mu
+{
+
+	// Forward definitions
+	class NodeScalar;
+	typedef Ptr<NodeScalar> NodeScalarPtr;
+	typedef Ptr<const NodeScalar> NodeScalarPtrConst;
+
+	class NodeImageBinarise;
+	typedef Ptr<NodeImageBinarise> NodeImageBinarisePtr;
+	typedef Ptr<const NodeImageBinarise> NodeImageBinarisePtrConst;
+
+
+	//! Node that multiplies the colors of an image, channel by channel.
+	//! \ingroup model
+	class MUTABLETOOLS_API NodeImageBinarise : public NodeImage
+	{
+	public:
+
+		//-----------------------------------------------------------------------------------------
+		// Life cycle
+		//-----------------------------------------------------------------------------------------
+
+		NodeImageBinarise();
+
+		void SerialiseWrapper(OutputArchive& arch) const override;
+		static void Serialise( const NodeImageBinarise* pNode, OutputArchive& arch );
+		static NodeImageBinarisePtr StaticUnserialise( InputArchive& arch );
+
+
+		//-----------------------------------------------------------------------------------------
+		// Node Interface
+		//-----------------------------------------------------------------------------------------
+
+        NodePtr Clone( NodeMapPtr pMap = 0 ) const override;
+
+        const NODE_TYPE* GetType() const override;
+		static const NODE_TYPE* GetStaticType();
+
+        virtual int GetInputCount() const override;
+        virtual Node* GetInputNode( int i ) const override;
+        void SetInputNode( int i, NodePtr pNode ) override;
+
+		//-----------------------------------------------------------------------------------------
+		// Own Interface
+		//-----------------------------------------------------------------------------------------
+
+		//! Base image to multiply.
+		NodeImagePtr GetBase() const;
+		void SetBase( NodeImagePtr );
+
+		//! Colour to multiply by the image.
+		NodeScalarPtr GetThreshold() const;
+		void SetThreshold( NodeScalarPtr );
+
+		//-----------------------------------------------------------------------------------------
+		// Interface pattern
+		//-----------------------------------------------------------------------------------------
+		class Private;
+		Private* GetPrivate() const;
+        Node::Private* GetBasePrivate() const override;
+
+	protected:
+
+		//! Forbidden. Manage with the Ptr<> template.
+		~NodeImageBinarise();
+
+	private:
+
+		Private* m_pD;
+
+	};
+
+
+}
diff --git a/Engine/Plugins/Experimental/Mutable/Source/MutableTools/Public/NodeImageColourMap.h b/Engine/Plugins/Experimental/Mutable/Source/MutableTools/Public/NodeImageColourMap.h
new file mode 100644
--- /dev/null
+++ b/Engine/Plugins/Experimental/Mutable/Source/MutableTools/Public/NodeImageColourMap.h
@@ -0,0 +1,84 @@
+// Copyright Epic Games, Inc. All Rights Reserved.
+
+#pragma once
+
+#include "NodeImage.h"
+
+
+namespace mu
+{
+
+	// Forward definitions
+	class NodeImageColourMap;
+	typedef Ptr<NodeImageColourMap> NodeImageColourMapPtr;
+	typedef Ptr<const NodeImageColourMap> NodeImageColourMapPtrConst;
+
+
+	//! This node changes the colours of a selectd part of the image, applying a colour map from
+	//! conteined in another image.
+	//! \ingroup model
+	class MUTABLETOOLS_API NodeImageColourMap : public NodeImage
+	{
+	public:
+
+		//-----------------------------------------------------------------------------------------
+		// Life cycle
+		//-----------------------------------------------------------------------------------------
+
+		NodeImageColourMap();
+
+		void SerialiseWrapper(OutputArchive& arch) const override;
+		static void Serialise( const NodeImageColourMap* pNode, OutputArchive& arch );
+		static NodeImageColourMapPtr StaticUnserialise( InputArchive& arch );
+
+
+		//-----------------------------------------------------------------------------------------
+		// Node Interface
+		//-----------------------------------------------------------------------------------------
+
+        NodePtr Clone( NodeMapPtr pMap = 0 ) const override;
+
+        const NODE_TYPE* GetType() const override;
+		static const NODE_TYPE* GetStaticType();
+
+        virtual int GetInputCount() const override;
+        virtual Node* GetInputNode( int i ) const override;
+        void SetInputNode( int i, NodePtr pNode ) override;
+
+		//-----------------------------------------------------------------------------------------
+		// Own Interface
+		//-----------------------------------------------------------------------------------------
+
+		//! Get the node generating the base image that will have the other one blended on top.
+		NodeImagePtr GetBase() const;
+		void SetBase( NodeImagePtr );
+
+		//! Get the node generating the mask image controlling the weight of the blend effect.
+		//! \todo: make it optional
+		NodeImagePtr GetMask() const;
+		void SetMask( NodeImagePtr );
+
+		//! Get the node generating the image to blend on the base.
+		NodeImagePtr GetMap() const;
+		void SetMap( NodeImagePtr );
+
+		//-----------------------------------------------------------------------------------------
+		// Interface pattern
+		//-----------------------------------------------------------------------------------------
+		class Private;
+		Private* GetPrivate() const;
+        Node::Private* GetBasePrivate() const override;
+
+	protected:
+
+		//! Forbidden. Manage with the Ptr<> template.
+		~NodeImageColourMap();
+
+	private:
+
+		Private* m_pD;
+
+	};
+
+
+}
diff --git a/Engine/Plugins/Experimental/Mutable/Source/MutableTools/Public/NodeImageConditional.h b/Engine/Plugins/Experimental/Mutable/Source/MutableTools/Public/NodeImageConditional.h
new file mode 100644
--- /dev/null
+++ b/Engine/Plugins/Experimental/Mutable/Source/MutableTools/Public/NodeImageConditional.h
@@ -0,0 +1,88 @@
+// Copyright Epic Games, Inc. All Rights Reserved.
+
+#pragma once
+
+#include "NodeImage.h"
+
+#include "Image.h"
+
+
+namespace mu
+{
+
+	// Forward definitions
+    class NodeBool;
+    typedef Ptr<NodeBool> NodeBoolPtr;
+    typedef Ptr<const NodeBool> NodeBoolPtrConst;
+
+    class NodeImageConditional;
+    typedef Ptr<NodeImageConditional> NodeImageConditionalPtr;
+    typedef Ptr<const NodeImageConditional> NodeImageConditionalPtrConst;
+
+
+	//! This node selects an output image from a set of input images based on a parameter.
+	//! \ingroup model
+    class MUTABLETOOLS_API NodeImageConditional : public NodeImage
+	{
+	public:
+
+		//-----------------------------------------------------------------------------------------
+		// Life cycle
+		//-----------------------------------------------------------------------------------------
+
+        NodeImageConditional();
+
+		void SerialiseWrapper(OutputArchive& arch) const override;
+		static void Serialise( const NodeImageConditional* pNode, OutputArchive& arch );
+        static NodeImageConditionalPtr StaticUnserialise( InputArchive& arch );
+
+
+		//-----------------------------------------------------------------------------------------
+		// Node Interface
+		//-----------------------------------------------------------------------------------------
+
+        NodePtr Clone( NodeMapPtr pMap = 0 ) const override;
+
+        const NODE_TYPE* GetType() const override;
+		static const NODE_TYPE* GetStaticType();
+
+        virtual int GetInputCount() const override;
+        virtual Node* GetInputNode( int i ) const override;
+        void SetInputNode( int i, NodePtr pNode ) override;
+
+		//-----------------------------------------------------------------------------------------
+		// Own Interface
+		//-----------------------------------------------------------------------------------------
+
+		//! Get the node generating the parameter used to select the option.
+        NodeBool* GetParameter() const;
+        void SetParameter( NodeBool* );
+
+        //! Get the node generating the option image when the parameters is true.
+        NodeImage* GetOptionTrue() const;
+        void SetOptionTrue( NodeImage* );
+
+        NodeImage* GetOptionFalse() const;
+        void SetOptionFalse( NodeImage* );
+
+
+		//-----------------------------------------------------------------------------------------
+		// Interface pattern
+		//-----------------------------------------------------------------------------------------
+		class Private;
+		Private* GetPrivate() const;
+        Node::Private* GetBasePrivate() const override;
+
+	protected:
+
+		//! Forbidden. Manage with the Ptr<> template.
+        ~NodeImageConditional();
+
+	private:
+
+		Private* m_pD;
+
+	};
+
+
+}
diff --git a/Engine/Plugins/Experimental/Mutable/Source/MutableTools/Public/NodeImageConstant.h b/Engine/Plugins/Experimental/Mutable/Source/MutableTools/Public/NodeImageConstant.h
new file mode 100644
--- /dev/null
+++ b/Engine/Plugins/Experimental/Mutable/Source/MutableTools/Public/NodeImageConstant.h
@@ -0,0 +1,80 @@
+// Copyright Epic Games, Inc. All Rights Reserved.
+
+#pragma once
+
+#include "NodeImage.h"
+
+#include "Image.h"
+
+
+namespace mu
+{
+
+	// Forward definitions
+    class NodeImageConstant;
+    typedef Ptr<NodeImageConstant> NodeImageConstantPtr;
+    typedef Ptr<const NodeImageConstant> NodeImageConstantPtrConst;
+
+
+	//! Node that outputs a constant image.
+	//! \ingroup model
+	class MUTABLETOOLS_API NodeImageConstant : public NodeImage
+	{
+	public:
+
+		//-----------------------------------------------------------------------------------------
+		// Life cycle
+		//-----------------------------------------------------------------------------------------
+
+		NodeImageConstant();
+
+		void SerialiseWrapper(OutputArchive& arch) const override;
+		static void Serialise( const NodeImageConstant* pNode, OutputArchive& arch );
+		static NodeImageConstantPtr StaticUnserialise( InputArchive& arch );
+
+
+		//-----------------------------------------------------------------------------------------
+		// Node Interface
+		//-----------------------------------------------------------------------------------------
+
+        NodePtr Clone( NodeMapPtr pMap = 0 ) const override;
+
+        const NODE_TYPE* GetType() const override;
+		static const NODE_TYPE* GetStaticType();
+
+        virtual int GetInputCount() const override;
+        virtual Node* GetInputNode( int i ) const override;
+        void SetInputNode( int i, NodePtr pNode ) override;
+
+		//-----------------------------------------------------------------------------------------
+		// Own Interface
+		//-----------------------------------------------------------------------------------------
+
+		//! Get the image that will be output by this node
+        Ptr<const Image> GetValue() const;
+
+        //! Set the image to be output by this node
+        void SetValue( Ptr<const Image> pImage );
+
+        //! Set the image proxy that will provide the image for this node when necessary
+        void SetValue( Ptr<ResourceProxy<Image>> pImage );
+
+		//-----------------------------------------------------------------------------------------
+		// Interface pattern
+		//-----------------------------------------------------------------------------------------
+		class Private;
+		Private* GetPrivate() const;
+        Node::Private* GetBasePrivate() const override;
+
+	protected:
+
+		//! Forbidden. Manage with the Ptr<> template.
+		~NodeImageConstant();
+
+	private:
+
+		Private* m_pD;
+
+	};
+
+}
diff --git a/Engine/Plugins/Experimental/Mutable/Source/MutableTools/Public/NodeImageDifference.h b/Engine/Plugins/Experimental/Mutable/Source/MutableTools/Public/NodeImageDifference.h
new file mode 100644
--- /dev/null
+++ b/Engine/Plugins/Experimental/Mutable/Source/MutableTools/Public/NodeImageDifference.h
@@ -0,0 +1,85 @@
+// Copyright Epic Games, Inc. All Rights Reserved.
+
+#pragma once
+
+#include "NodeImage.h"
+
+
+namespace mu
+{
+
+	// Forward definitions
+	class NodeImageDifference;
+	typedef Ptr<NodeImageDifference> NodeImageDifferencePtr;
+	typedef Ptr<const NodeImageDifference> NodeImageDifferencePtrConst;
+
+
+	//! Node that compares two images and returns a black-and-white image with pixels that are
+	//! "significally different", meaning around 0.03 different. This node is for testing
+	//! purposes.
+	//! \todo rename to NodeImageCompare
+	//! \ingroup model
+	class MUTABLETOOLS_API NodeImageDifference : public NodeImage
+	{
+	public:
+
+		//-----------------------------------------------------------------------------------------
+		// Life cycle
+		//-----------------------------------------------------------------------------------------
+
+		NodeImageDifference();
+
+		void SerialiseWrapper(OutputArchive& arch) const override;
+		static void Serialise( const NodeImageDifference* pNode, OutputArchive& arch );
+		static NodeImageDifferencePtr StaticUnserialise( InputArchive& arch );
+
+
+		//-----------------------------------------------------------------------------------------
+		// Node Interface
+		//-----------------------------------------------------------------------------------------
+
+        NodePtr Clone( NodeMapPtr pMap = 0 ) const override;
+
+        const NODE_TYPE* GetType() const override;
+		static const NODE_TYPE* GetStaticType();
+
+        virtual int GetInputCount() const override;
+        virtual Node* GetInputNode( int i ) const override;
+        void SetInputNode( int i, NodePtr pNode ) override;
+
+		//-----------------------------------------------------------------------------------------
+		// Own Interface
+		//-----------------------------------------------------------------------------------------
+
+		//! Return the first image to compare.
+		NodeImagePtr GetA() const;
+
+		//! Set the first image to compare.
+		void SetA( NodeImagePtr );
+
+		//! Return the second image to compare.
+		NodeImagePtr GetB() const;
+
+		//! Set the second image to compare.
+		void SetB( NodeImagePtr );
+
+		//-----------------------------------------------------------------------------------------
+		// Interface pattern
+		//-----------------------------------------------------------------------------------------
+		class Private;
+		Private* GetPrivate() const;
+        Node::Private* GetBasePrivate() const override;
+
+	protected:
+
+		//! Forbidden. Manage with the Ptr<> template.
+		~NodeImageDifference();
+
+	private:
+
+		Private* m_pD;
+
+	};
+
+
+}
diff --git a/Engine/Plugins/Experimental/Mutable/Source/MutableTools/Public/NodeImageFormat.h b/Engine/Plugins/Experimental/Mutable/Source/MutableTools/Public/NodeImageFormat.h
new file mode 100644
--- /dev/null
+++ b/Engine/Plugins/Experimental/Mutable/Source/MutableTools/Public/NodeImageFormat.h
@@ -0,0 +1,83 @@
+// Copyright Epic Games, Inc. All Rights Reserved.
+
+#pragma once
+
+#include "NodeImage.h"
+
+#include "Image.h"
+
+
+namespace mu
+{
+
+	// Forward definitions
+	class NodeImageFormat;
+	typedef Ptr<NodeImageFormat> NodeImageFormatPtr;
+	typedef Ptr<const NodeImageFormat> NodeImageFormatPtrConst;
+
+
+	//! Node that composes a new image by gathering pixel data from channels in other images.
+	//! \ingroup model
+	class MUTABLETOOLS_API NodeImageFormat : public NodeImage
+	{
+	public:
+
+		//-----------------------------------------------------------------------------------------
+		// Life cycle
+		//-----------------------------------------------------------------------------------------
+
+		NodeImageFormat();
+
+		void SerialiseWrapper(OutputArchive& arch) const override;
+		static void Serialise( const NodeImageFormat* pNode, OutputArchive& arch );
+		static NodeImageFormatPtr StaticUnserialise( InputArchive& arch );
+
+
+		//-----------------------------------------------------------------------------------------
+		// Node Interface
+		//-----------------------------------------------------------------------------------------
+
+		NodePtr Clone( NodeMapPtr pMap = 0 ) const override;
+
+		const NODE_TYPE* GetType() const override;
+		static const NODE_TYPE* GetStaticType();
+
+		int GetInputCount() const override;
+		Node* GetInputNode( int i ) const override;
+		void SetInputNode( int i, NodePtr pNode ) override;
+
+		//-----------------------------------------------------------------------------------------
+		// Own Interface
+		//-----------------------------------------------------------------------------------------
+
+		//!
+		EImageFormat GetFormat() const;
+        void SetFormat(EImageFormat format, EImageFormat formatIfAlpha = EImageFormat::IF_NONE );
+
+		//!
+		NodeImagePtr GetSource() const;
+
+		void SetSource( NodeImagePtr );
+
+
+		//-----------------------------------------------------------------------------------------
+		// Interface pattern
+		//-----------------------------------------------------------------------------------------
+		class Private;
+		Private* GetPrivate() const;
+        Node::Private* GetBasePrivate() const override;
+
+	protected:
+
+		//! Forbidden. Manage with the Ptr<> template.
+		~NodeImageFormat();
+
+	private:
+
+		Private* m_pD;
+
+	};
+
+
+}
+
diff --git a/Engine/Plugins/Experimental/Mutable/Source/MutableTools/Public/NodeImageGradient.h b/Engine/Plugins/Experimental/Mutable/Source/MutableTools/Public/NodeImageGradient.h
new file mode 100644
--- /dev/null
+++ b/Engine/Plugins/Experimental/Mutable/Source/MutableTools/Public/NodeImageGradient.h
@@ -0,0 +1,87 @@
+// Copyright Epic Games, Inc. All Rights Reserved.
+
+#pragma once
+
+#include "NodeImage.h"
+
+
+namespace mu
+{
+
+	// Forward definitions
+	class NodeImageGradient;
+	typedef Ptr<NodeImageGradient> NodeImageGradientPtr;
+	typedef Ptr<const NodeImageGradient> NodeImageGradientPtrConst;
+
+	class NodeColour;
+	typedef Ptr<NodeColour> NodeColourPtr;
+	typedef Ptr<const NodeColour> NodeColourPtrConst;
+
+
+	//! This node generats an horizontal linear gradient between two colours in a new image.
+	//! \ingroup model
+	class MUTABLETOOLS_API NodeImageGradient : public NodeImage
+	{
+	public:
+
+		//-----------------------------------------------------------------------------------------
+		// Life cycle
+		//-----------------------------------------------------------------------------------------
+
+		NodeImageGradient();
+
+		void SerialiseWrapper(OutputArchive& arch) const override;
+		static void Serialise( const NodeImageGradient* pNode, OutputArchive& arch );
+		static NodeImageGradientPtr StaticUnserialise( InputArchive& arch );
+
+
+		//-----------------------------------------------------------------------------------------
+		// Node Interface
+		//-----------------------------------------------------------------------------------------
+
+        NodePtr Clone( NodeMapPtr pMap = 0 ) const override;
+
+        const NODE_TYPE* GetType() const override;
+		static const NODE_TYPE* GetStaticType();
+
+        virtual int GetInputCount() const override;
+        virtual Node* GetInputNode( int i ) const override;
+        void SetInputNode( int i, NodePtr pNode ) override;
+
+		//-----------------------------------------------------------------------------------------
+		// Own Interface
+		//-----------------------------------------------------------------------------------------
+
+		//! First colour of the gradient.
+		NodeColourPtr GetColour0() const;
+		void SetColour0( NodeColourPtr );
+
+		//! Second colour of the gradient.
+		NodeColourPtr GetColour1() const;
+		void SetColour1( NodeColourPtr );
+
+		//! Generated image size.
+		int GetSizeX() const;
+		int GetSizeY() const;
+		void SetSize( int x, int y );
+
+		//-----------------------------------------------------------------------------------------
+		// Interface pattern
+		//-----------------------------------------------------------------------------------------
+		class Private;
+		Private* GetPrivate() const;
+        Node::Private* GetBasePrivate() const override;
+
+	protected:
+
+		//! Forbidden. Manage with the Ptr<> template.
+		~NodeImageGradient();
+
+	private:
+
+		Private* m_pD;
+
+	};
+
+
+}
diff --git a/Engine/Plugins/Experimental/Mutable/Source/MutableTools/Public/NodeImageInterpolate.h b/Engine/Plugins/Experimental/Mutable/Source/MutableTools/Public/NodeImageInterpolate.h
new file mode 100644
--- /dev/null
+++ b/Engine/Plugins/Experimental/Mutable/Source/MutableTools/Public/NodeImageInterpolate.h
@@ -0,0 +1,95 @@
+// Copyright Epic Games, Inc. All Rights Reserved.
+
+#pragma once
+
+#include "NodeImage.h"
+
+#include "Image.h"
+
+
+namespace mu
+{
+
+	// Forward definitions
+	class NodeScalar;
+	typedef Ptr<NodeScalar> NodeScalarPtr;
+	typedef Ptr<const NodeScalar> NodeScalarPtrConst;
+
+	class NodeImageInterpolate;
+	typedef Ptr<NodeImageInterpolate> NodeImageInterpolatePtr;
+	typedef Ptr<const NodeImageInterpolate> NodeImageInterpolatePtrConst;
+
+
+	//! Node that interpolates linearly between several images based on a weight.
+	//! The different input images are uniformly distributed in the 0 to 1 range:
+	//! - if two images are set, the first one is 0.0 and the second one is 1.0
+	//! - if three images are set, the first one is 0.0 the second one is 0.5 and the third one 1.0
+	//! regardless of the input slots used: If B and C are set, but not A, B will be at weight 0
+	//! and C will be at weight 1
+	//! \ingroup model
+	class MUTABLETOOLS_API NodeImageInterpolate : public NodeImage
+	{
+	public:
+
+		//-----------------------------------------------------------------------------------------
+		// Life cycle
+		//-----------------------------------------------------------------------------------------
+
+		NodeImageInterpolate();
+
+		void SerialiseWrapper(OutputArchive& arch) const override;
+		static void Serialise( const NodeImageInterpolate* pNode, OutputArchive& arch );
+		static NodeImageInterpolatePtr StaticUnserialise( InputArchive& arch );
+
+
+		//-----------------------------------------------------------------------------------------
+		// Node Interface
+		//-----------------------------------------------------------------------------------------
+
+        NodePtr Clone( NodeMapPtr pMap = 0 ) const override;
+
+        const NODE_TYPE* GetType() const override;
+		static const NODE_TYPE* GetStaticType();
+
+        virtual int GetInputCount() const override;
+        virtual Node* GetInputNode( int i ) const override;
+        void SetInputNode( int i, NodePtr pNode ) override;
+
+		//-----------------------------------------------------------------------------------------
+		// Own Interface
+		//-----------------------------------------------------------------------------------------
+
+		//! Get the node generating the weight of the interpolation.
+		NodeScalarPtr GetFactor() const;
+		void SetFactor( NodeScalarPtr );
+
+		//! Set the number of target images. It will keep the currently set targets and initialise
+		//! the new ones as null.
+		void SetTargetCount( int );
+		int GetTargetCount() const;
+
+		//! Get the node generating the t-th target image of the interpolation.
+		NodeImagePtr GetTarget( int t ) const;
+		void SetTarget( int t, NodeImagePtr );
+
+
+		//-----------------------------------------------------------------------------------------
+		// Interface pattern
+		//-----------------------------------------------------------------------------------------
+		class Private;
+		Private* GetPrivate() const;
+        Node::Private* GetBasePrivate() const override;
+
+	protected:
+
+		//! Forbidden. Manage with the Ptr<> template.
+		~NodeImageInterpolate();
+
+	private:
+
+		Private* m_pD;
+
+	};
+
+
+}
