[mutable] Move Mutable plugin from the restricted folder to the experimental plugins folder.

diff --git a/Engine/Plugins/Experimental/Mutable/Source/MutableTools/Public/NodeImageInterpolate3.h b/Engine/Plugins/Experimental/Mutable/Source/MutableTools/Public/NodeImageInterpolate3.h
new file mode 100644
--- /dev/null
+++ b/Engine/Plugins/Experimental/Mutable/Source/MutableTools/Public/NodeImageInterpolate3.h
@@ -0,0 +1,100 @@
+// Copyright Epic Games, Inc. All Rights Reserved.
+
+#pragma once
+
+#include "NodeImage.h"
+
+#include "Image.h"
+
+
+namespace mu
+{
+
+	// Forward definitions
+	class NodeScalar;
+	typedef Ptr<NodeScalar> NodeScalarPtr;
+	typedef Ptr<const NodeScalar> NodeScalarPtrConst;
+
+	class NodeImageInterpolate3;
+	typedef Ptr<NodeImageInterpolate3> NodeImageInterpolate3Ptr;
+	typedef Ptr<const NodeImageInterpolate3> NodeImageInterpolate3PtrConst;
+
+
+	//! Node that Interpolates linearly between 3 images. It has two weights and the third one is
+	//! deduced as 1-w0-w1.
+	//! \ingroup model
+	class MUTABLETOOLS_API NodeImageInterpolate3 : public NodeImage
+	{
+	public:
+
+		//-----------------------------------------------------------------------------------------
+		// Life cycle
+		//-----------------------------------------------------------------------------------------
+
+		NodeImageInterpolate3();
+
+		void SerialiseWrapper(OutputArchive& arch) const override;
+		static void Serialise( const NodeImageInterpolate3* pNode, OutputArchive& arch );
+		static NodeImageInterpolate3Ptr StaticUnserialise( InputArchive& arch );
+
+
+		//-----------------------------------------------------------------------------------------
+		// Node Interface
+		//-----------------------------------------------------------------------------------------
+
+        NodePtr Clone( NodeMapPtr pMap = 0 ) const override;
+
+        const NODE_TYPE* GetType() const override;
+		static const NODE_TYPE* GetStaticType();
+
+        virtual int GetInputCount() const override;
+        virtual Node* GetInputNode( int i ) const override;
+        void SetInputNode( int i, NodePtr pNode ) override;
+
+		//-----------------------------------------------------------------------------------------
+		// Own Interface
+		//-----------------------------------------------------------------------------------------
+
+		//! Get the weight of the first target. If this is 1 (so the other is 0) the resulting
+		//! image is the target 1. If both weights are 0, the result is target 0.
+		NodeScalarPtr GetFactor1() const;
+		void SetFactor1( NodeScalarPtr );
+
+		//! Get the weight of the second target. If this is 1 (so the other is 1) the resulting
+		//! image is the target 2. If both weights are 0, the result is target 0.
+		NodeScalarPtr GetFactor2() const;
+		void SetFactor2( NodeScalarPtr );
+
+		//! Get the node generating the first target image of the interpolation.
+		NodeImagePtr GetTarget0() const;
+		void SetTarget0( NodeImagePtr );
+
+		//! Get the node generating the second target image of the interpolation.
+		NodeImagePtr GetTarget1() const;
+		void SetTarget1( NodeImagePtr );
+
+		//! Get the node generating the third target image of the interpolation.
+		NodeImagePtr GetTarget2() const;
+		void SetTarget2( NodeImagePtr );
+
+
+		//-----------------------------------------------------------------------------------------
+		// Interface pattern
+		//-----------------------------------------------------------------------------------------
+		class Private;
+		Private* GetPrivate() const;
+        Node::Private* GetBasePrivate() const override;
+
+	protected:
+
+		//! Forbidden. Manage with the Ptr<> template.
+		~NodeImageInterpolate3();
+
+	private:
+
+		Private* m_pD;
+
+	};
+
+
+}
diff --git a/Engine/Plugins/Experimental/Mutable/Source/MutableTools/Public/NodeImageInvert.h b/Engine/Plugins/Experimental/Mutable/Source/MutableTools/Public/NodeImageInvert.h
new file mode 100644
--- /dev/null
+++ b/Engine/Plugins/Experimental/Mutable/Source/MutableTools/Public/NodeImageInvert.h
@@ -0,0 +1,71 @@
+// Copyright Epic Games, Inc. All Rights Reserved.
+
+#pragma once
+
+#include "NodeImage.h"
+
+
+namespace mu
+{
+
+	// Forward definitions
+	class NodeImageInvert;
+	typedef Ptr<NodeImageInvert> NodeImageInvertPtr;
+	typedef Ptr<const NodeImageInvert> NodeImageInvertPtrConst;
+
+	//! Node that inverts the colors of an image, channel by channel
+	//! \ingroup model
+
+	class MUTABLETOOLS_API NodeImageInvert : public NodeImage
+	{
+	public:
+
+		//-----------------------------------------------------------------------------------------
+		// Life cycle
+		//-----------------------------------------------------------------------------------------
+
+		NodeImageInvert();
+
+		static void Serialise(const NodeImageInvert* pNode, OutputArchive& arch);
+		void SerialiseWrapper(OutputArchive& arch) const override;
+		static NodeImageInvertPtr StaticUnserialise(InputArchive& arch);
+
+
+		//-----------------------------------------------------------------------------------------
+		// Node Interface
+		//-----------------------------------------------------------------------------------------
+
+		NodePtr Clone(NodeMapPtr pMap = 0) const override;
+
+		const NODE_TYPE* GetType() const override;
+		static const NODE_TYPE* GetStaticType();
+
+		virtual int GetInputCount() const override;
+		virtual Node* GetInputNode(int i)const override;
+		void SetInputNode(int i, NodePtr pNode) override;
+
+		//-----------------------------------------------------------------------------------------
+		// Own Interface
+		//-----------------------------------------------------------------------------------------
+
+		//! Base image to invert.
+		NodeImagePtr GetBase() const;
+		void SetBase(NodeImagePtr);
+
+		//-----------------------------------------------------------------------------------------
+		// Interface pattern
+		//-----------------------------------------------------------------------------------------
+		class Private;
+		Private* GetPrivate() const;
+		Node::Private*GetBasePrivate()const override;
+
+	protected:
+
+		//! Forbidden. Manage with the Ptr<> template
+		~NodeImageInvert();
+
+	private:
+
+		Private* m_pD;
+	};
+}
diff --git a/Engine/Plugins/Experimental/Mutable/Source/MutableTools/Public/NodeImageLayer.h b/Engine/Plugins/Experimental/Mutable/Source/MutableTools/Public/NodeImageLayer.h
new file mode 100644
--- /dev/null
+++ b/Engine/Plugins/Experimental/Mutable/Source/MutableTools/Public/NodeImageLayer.h
@@ -0,0 +1,89 @@
+// Copyright Epic Games, Inc. All Rights Reserved.
+
+#pragma once
+
+#include "NodeImage.h"
+#include "MutableRuntime/Public/Image.h"
+
+
+namespace mu
+{
+
+	// Forward definitions
+	class NodeImageLayer;
+	typedef Ptr<NodeImageLayer> NodeImageLayerPtr;
+	typedef Ptr<const NodeImageLayer> NodeImageLayerPtrConst;
+
+
+	//! This node applies a layer blending effect on a base image using a mask and a blended image.
+	//! \ingroup model
+	class MUTABLETOOLS_API NodeImageLayer : public NodeImage
+	{
+	public:
+
+		//-----------------------------------------------------------------------------------------
+		// Life cycle
+		//-----------------------------------------------------------------------------------------
+
+		NodeImageLayer();
+
+		void SerialiseWrapper(OutputArchive& arch) const override;
+		static void Serialise( const NodeImageLayer* pNode, OutputArchive& arch );
+		static NodeImageLayerPtr StaticUnserialise( InputArchive& arch );
+
+
+		//-----------------------------------------------------------------------------------------
+		// Node Interface
+		//-----------------------------------------------------------------------------------------
+
+        NodePtr Clone( NodeMapPtr pMap = 0 ) const override;
+
+        const NODE_TYPE* GetType() const override;
+		static const NODE_TYPE* GetStaticType();
+
+        virtual int GetInputCount() const override;
+        virtual Node* GetInputNode( int i ) const override;
+        void SetInputNode( int i, NodePtr pNode ) override;
+
+		//-----------------------------------------------------------------------------------------
+		// Own Interface
+		//-----------------------------------------------------------------------------------------
+
+		//! Get the node generating the base image that will have the blendeing effect applied.
+		NodeImagePtr GetBase() const;
+		void SetBase( NodeImagePtr );
+
+		//! Get the node generating the mask image controlling the weight of the effect.
+		NodeImagePtr GetMask() const;
+		void SetMask( NodeImagePtr );
+
+		//! Get the image blended on top of the base.
+		NodeImagePtr GetBlended() const;
+		void SetBlended( NodeImagePtr );
+
+		static const char* s_blendTypeName[int32(EBlendType::_BT_COUNT)];
+
+		//!
+		EBlendType GetBlendType() const;
+		void SetBlendType(EBlendType);
+
+		//-----------------------------------------------------------------------------------------
+		// Interface pattern
+		//-----------------------------------------------------------------------------------------
+		class Private;
+		Private* GetPrivate() const;
+        Node::Private* GetBasePrivate() const override;
+
+	protected:
+
+		//! Forbidden. Manage with the Ptr<> template.
+		~NodeImageLayer();
+
+	private:
+
+		Private* m_pD;
+
+	};
+
+
+}
diff --git a/Engine/Plugins/Experimental/Mutable/Source/MutableTools/Public/NodeImageLayerColour.h b/Engine/Plugins/Experimental/Mutable/Source/MutableTools/Public/NodeImageLayerColour.h
new file mode 100644
--- /dev/null
+++ b/Engine/Plugins/Experimental/Mutable/Source/MutableTools/Public/NodeImageLayerColour.h
@@ -0,0 +1,92 @@
+// Copyright Epic Games, Inc. All Rights Reserved.
+
+#pragma once
+
+#include "NodeImageLayer.h"
+#include "MutableRuntime/Public/Image.h"
+
+namespace mu
+{
+
+	// Forward definitions
+	class NodeColour;
+	typedef Ptr<NodeColour> NodeColourPtr;
+	typedef Ptr<const NodeColour> NodeColourPtrConst;
+
+	class NodeImageLayerColour;
+	typedef Ptr<NodeImageLayerColour> NodeImageLayerColourPtr;
+	typedef Ptr<const NodeImageLayerColour> NodeImageLayerColourPtrConst;
+
+
+	//! This node applies a layer blending effect on a base image using a mask and a colour.
+	//! \ingroup model
+	class MUTABLETOOLS_API NodeImageLayerColour : public NodeImage
+	{
+	public:
+
+		//-----------------------------------------------------------------------------------------
+		// Life cycle
+		//-----------------------------------------------------------------------------------------
+
+		NodeImageLayerColour();
+
+		void SerialiseWrapper(OutputArchive& arch) const override;
+		static void Serialise( const NodeImageLayerColour* pNode, OutputArchive& arch );
+		static NodeImageLayerColourPtr StaticUnserialise( InputArchive& arch );
+
+
+		//-----------------------------------------------------------------------------------------
+		// Node Interface
+		//-----------------------------------------------------------------------------------------
+
+        NodePtr Clone( NodeMapPtr pMap = 0 ) const override;
+
+        const NODE_TYPE* GetType() const override;
+		static const NODE_TYPE* GetStaticType();
+
+        virtual int GetInputCount() const override;
+        virtual Node* GetInputNode( int i ) const override;
+        void SetInputNode( int i, NodePtr pNode ) override;
+
+		//-----------------------------------------------------------------------------------------
+		// Own Interface
+		//-----------------------------------------------------------------------------------------
+
+		//! Get the node generating the base image that will have the blending effect applied.
+		NodeImagePtr GetBase() const;
+		void SetBase( NodeImagePtr );
+
+		//! Get the node generating the mask image controlling the weight of the effect.
+		NodeImagePtr GetMask() const;
+		void SetMask( NodeImagePtr );
+
+		//! Get the node generating the color to blend on the base.
+		NodeColourPtr GetColour() const;
+		void SetColour( NodeColourPtr );
+
+		//!
+		EBlendType GetBlendType() const;
+
+		//!
+		void SetBlendType(EBlendType);
+
+		//-----------------------------------------------------------------------------------------
+		// Interface pattern
+		//-----------------------------------------------------------------------------------------
+		class Private;
+		Private* GetPrivate() const;
+        Node::Private* GetBasePrivate() const override;
+
+	protected:
+
+		//! Forbidden. Manage with the Ptr<> template.
+		~NodeImageLayerColour();
+
+	private:
+
+		Private* m_pD;
+
+	};
+
+
+}
diff --git a/Engine/Plugins/Experimental/Mutable/Source/MutableTools/Public/NodeImageLuminance.h b/Engine/Plugins/Experimental/Mutable/Source/MutableTools/Public/NodeImageLuminance.h
new file mode 100644
--- /dev/null
+++ b/Engine/Plugins/Experimental/Mutable/Source/MutableTools/Public/NodeImageLuminance.h
@@ -0,0 +1,76 @@
+// Copyright Epic Games, Inc. All Rights Reserved.
+
+#pragma once
+
+#include "NodeImage.h"
+
+#include "Image.h"
+
+
+namespace mu
+{
+
+	// Forward definitions
+	class NodeImageLuminance;
+	typedef Ptr<NodeImageLuminance> NodeImageLuminancePtr;
+	typedef Ptr<const NodeImageLuminance> NodeImageLuminancePtrConst;
+
+
+	//! Calculate the luminance of an image into a new single-channel image..
+	//! \ingroup model
+	class MUTABLETOOLS_API NodeImageLuminance : public NodeImage
+	{
+	public:
+
+		//-----------------------------------------------------------------------------------------
+		// Life cycle
+		//-----------------------------------------------------------------------------------------
+
+		NodeImageLuminance();
+
+		void SerialiseWrapper(OutputArchive& arch) const override;
+		static void Serialise( const NodeImageLuminance* pNode, OutputArchive& arch );
+		static NodeImageLuminancePtr StaticUnserialise( InputArchive& arch );
+
+
+		//-----------------------------------------------------------------------------------------
+		// Node Interface
+		//-----------------------------------------------------------------------------------------
+
+        NodePtr Clone( NodeMapPtr pMap = 0 ) const override;
+
+        const NODE_TYPE* GetType() const override;
+		static const NODE_TYPE* GetStaticType();
+
+        virtual int GetInputCount() const override;
+        virtual Node* GetInputNode( int i ) const override;
+        void SetInputNode( int i, NodePtr pNode ) override;
+
+		//-----------------------------------------------------------------------------------------
+		// Own Interface
+		//-----------------------------------------------------------------------------------------
+
+		//! Get the node generating the source image calculate the luminance from.
+		NodeImagePtr GetSource() const;
+		void SetSource( NodeImagePtr );
+
+		//-----------------------------------------------------------------------------------------
+		// Interface pattern
+		//-----------------------------------------------------------------------------------------
+		class Private;
+		Private* GetPrivate() const;
+        Node::Private* GetBasePrivate() const override;
+
+	protected:
+
+		//! Forbidden. Manage with the Ptr<> template.
+		~NodeImageLuminance();
+
+	private:
+
+		Private* m_pD;
+
+	};
+
+
+}
diff --git a/Engine/Plugins/Experimental/Mutable/Source/MutableTools/Public/NodeImageMipmap.h b/Engine/Plugins/Experimental/Mutable/Source/MutableTools/Public/NodeImageMipmap.h
new file mode 100644
--- /dev/null
+++ b/Engine/Plugins/Experimental/Mutable/Source/MutableTools/Public/NodeImageMipmap.h
@@ -0,0 +1,85 @@
+// Copyright Epic Games, Inc. All Rights Reserved.
+
+#pragma once
+
+#include "NodeImage.h"
+
+#include "Image.h"
+
+
+namespace mu
+{
+
+	// Forward definitions
+	class NodeScalar;
+	typedef Ptr<NodeScalar> NodeScalarPtr;
+	typedef Ptr<const NodeScalar> NodeScalarPtrConst;
+
+    class NodeImageMipmap;
+    typedef Ptr<NodeImageMipmap> NodeImageMipmapPtr;
+    typedef Ptr<const NodeImageMipmap> NodeImageMipmapPtrConst;
+
+
+	//! Generate mimaps for an image.
+	//! \ingroup model
+    class MUTABLETOOLS_API NodeImageMipmap : public NodeImage
+	{
+	public:
+
+		//-----------------------------------------------------------------------------------------
+		// Life cycle
+		//-----------------------------------------------------------------------------------------
+
+        NodeImageMipmap();
+
+		void SerialiseWrapper(OutputArchive& arch) const override;
+		static void Serialise( const NodeImageMipmap* pNode, OutputArchive& arch );
+        static NodeImageMipmapPtr StaticUnserialise( InputArchive& arch );
+
+
+		//-----------------------------------------------------------------------------------------
+		// Node Interface
+		//-----------------------------------------------------------------------------------------
+
+        NodePtr Clone( NodeMapPtr pMap = 0 ) const override;
+
+        const NODE_TYPE* GetType() const override;
+		static const NODE_TYPE* GetStaticType();
+
+        virtual int GetInputCount() const override;
+        virtual Node* GetInputNode( int i ) const override;
+        void SetInputNode( int i, NodePtr pNode ) override;
+
+		//-----------------------------------------------------------------------------------------
+		// Own Interface
+		//-----------------------------------------------------------------------------------------
+
+        //! Get the node generating the source image to be mipmapped.
+		NodeImagePtr GetSource() const;
+		void SetSource( NodeImagePtr );
+
+		void SetMipmapGenerationSettings( EMipmapFilterType filterType,
+										  EAddressMode addressMode,
+										  float sharpenFactor,  
+										  bool bDitherMipAlpha );
+
+		//-----------------------------------------------------------------------------------------
+		// Interface pattern
+		//-----------------------------------------------------------------------------------------
+		class Private;
+		Private* GetPrivate() const;
+        Node::Private* GetBasePrivate() const override;
+
+	protected:
+
+		//! Forbidden. Manage with the Ptr<> template.
+        ~NodeImageMipmap();
+
+	private:
+
+		Private* m_pD;
+
+	};
+
+
+}
diff --git a/Engine/Plugins/Experimental/Mutable/Source/MutableTools/Public/NodeImageMultiLayer.h b/Engine/Plugins/Experimental/Mutable/Source/MutableTools/Public/NodeImageMultiLayer.h
new file mode 100644
--- /dev/null
+++ b/Engine/Plugins/Experimental/Mutable/Source/MutableTools/Public/NodeImageMultiLayer.h
@@ -0,0 +1,94 @@
+// Copyright Epic Games, Inc. All Rights Reserved.
+
+#pragma once
+
+#include "NodeImage.h"
+#include "NodeImageLayer.h"
+#include "NodeRange.h"
+#include "MutableRuntime/Public/Image.h"
+
+
+namespace mu
+{
+
+	// Forward definitions
+    class NodeImageMultiLayer;
+    typedef Ptr<NodeImageMultiLayer> NodeImageMultiLayerPtr;
+    typedef Ptr<const NodeImageMultiLayer> NodeImageMultiLayerPtrConst;
+
+
+    //! This node applies any numbre of layer blending effect on a base image using a mask and a
+    //! blended image. The number of layers depends on a scalar input of this node.
+	//! \ingroup model
+    class MUTABLETOOLS_API NodeImageMultiLayer : public NodeImage
+	{
+	public:
+
+		//-----------------------------------------------------------------------------------------
+		// Life cycle
+		//-----------------------------------------------------------------------------------------
+
+        NodeImageMultiLayer();
+
+		void SerialiseWrapper(OutputArchive& arch) const override;
+		static void Serialise( const NodeImageMultiLayer* pNode, OutputArchive& arch );
+        static NodeImageMultiLayerPtr StaticUnserialise( InputArchive& arch );
+
+
+		//-----------------------------------------------------------------------------------------
+		// Node Interface
+		//-----------------------------------------------------------------------------------------
+
+        NodePtr Clone( NodeMapPtr pMap = 0 ) const override;
+
+        const NODE_TYPE* GetType() const override;
+		static const NODE_TYPE* GetStaticType();
+
+        virtual int GetInputCount() const override;
+        virtual Node* GetInputNode( int i ) const override;
+        void SetInputNode( int i, NodePtr pNode ) override;
+
+		//-----------------------------------------------------------------------------------------
+		// Own Interface
+		//-----------------------------------------------------------------------------------------
+
+		//! Get the node generating the base image that will have the blendeing effect applied.
+		NodeImagePtr GetBase() const;
+		void SetBase( NodeImagePtr );
+
+		//! Get the node generating the mask image controlling the weight of the effect.
+		NodeImagePtr GetMask() const;
+		void SetMask( NodeImagePtr );
+
+        //! Get the image blended on top of the base.
+        NodeImagePtr GetBlended() const;
+        void SetBlended( NodeImagePtr );
+
+        //! Get the number of layers to apply.
+        NodeRangePtr GetRange() const;
+        void SetRange( NodeRangePtr );
+
+		//!
+		EBlendType GetBlendType() const;
+        void SetBlendType(EBlendType);
+
+		//-----------------------------------------------------------------------------------------
+		// Interface pattern
+		//-----------------------------------------------------------------------------------------
+		class Private;
+		Private* GetPrivate() const;
+        Node::Private* GetBasePrivate() const override;
+
+	protected:
+
+		//! Forbidden. Manage with the Ptr<> template.
+        ~NodeImageMultiLayer();
+
+	private:
+
+		Private* m_pD;
+
+	};
+
+
+}
diff --git a/Engine/Plugins/Experimental/Mutable/Source/MutableTools/Public/NodeImageNormalComposite.h b/Engine/Plugins/Experimental/Mutable/Source/MutableTools/Public/NodeImageNormalComposite.h
new file mode 100644
--- /dev/null
+++ b/Engine/Plugins/Experimental/Mutable/Source/MutableTools/Public/NodeImageNormalComposite.h
@@ -0,0 +1,82 @@
+// Copyright Epic Games, Inc. All Rights Reserved.
+
+#pragma once
+
+#include "NodeImage.h"
+#include "Image.h"
+
+namespace mu
+{
+
+	// Forward definitions
+	class NodeImageNormalComposite;
+	typedef Ptr<NodeImageNormalComposite> NodeImageNormalCompositePtr;
+	typedef Ptr<const NodeImageNormalComposite> NodeImageNormalCompositePtrConst;
+
+	//! Node that inverts the colors of an image, channel by channel
+	//! \ingroup model
+
+	class MUTABLETOOLS_API NodeImageNormalComposite : public NodeImage
+	{
+	public:
+
+		//MUTABLE_DEFINE_CONST_VISITABLE()
+		
+		//-----------------------------------------------------------------------------------------
+		// Life cycle
+		//-----------------------------------------------------------------------------------------
+
+		NodeImageNormalComposite();
+
+		void SerialiseWrapper(OutputArchive& arch) const override;
+		static void Serialise(const NodeImageNormalComposite* pNode, OutputArchive& arch);
+		static NodeImageNormalCompositePtr StaticUnserialise(InputArchive& arch);
+
+
+		//-----------------------------------------------------------------------------------------
+		// Node Interface
+		//-----------------------------------------------------------------------------------------
+
+		NodePtr Clone(NodeMapPtr pMap = 0) const override;
+
+		const NODE_TYPE* GetType() const override;
+		static const NODE_TYPE* GetStaticType();
+
+		virtual int GetInputCount() const override;
+		virtual Node* GetInputNode(int i)const override;
+		void SetInputNode(int i, NodePtr pNode) override;
+
+		//-----------------------------------------------------------------------------------------
+		// Own Interface
+		//-----------------------------------------------------------------------------------------
+
+		NodeImagePtr GetBase() const;
+		void SetBase(NodeImagePtr);
+
+        NodeImagePtr GetNormal() const;
+        void SetNormal(NodeImagePtr);
+
+		float GetPower() const;
+		void SetPower(float); 
+
+        ECompositeImageMode GetMode() const;
+        void SetMode(ECompositeImageMode);
+
+		//-----------------------------------------------------------------------------------------
+		// Interface pattern
+		//-----------------------------------------------------------------------------------------
+		class Private;
+		Private* GetPrivate() const;
+		Node::Private*GetBasePrivate()const override;
+
+	protected:
+
+		//! Forbidden. Manage with the Ptr<> template
+		~NodeImageNormalComposite();
+
+	private:
+
+		Private* m_pD;
+	};
+
+}
diff --git a/Engine/Plugins/Experimental/Mutable/Source/MutableTools/Public/NodeImageParameter.h b/Engine/Plugins/Experimental/Mutable/Source/MutableTools/Public/NodeImageParameter.h
new file mode 100644
--- /dev/null
+++ b/Engine/Plugins/Experimental/Mutable/Source/MutableTools/Public/NodeImageParameter.h
@@ -0,0 +1,88 @@
+// Copyright Epic Games, Inc. All Rights Reserved.
+
+#pragma once
+
+#include "NodeImage.h"
+#include "Image.h"
+
+
+namespace mu
+{
+
+    class NodeImageParameter;
+    typedef Ptr<NodeImageParameter> NodeImageParameterPtr;
+    typedef Ptr<const NodeImageParameter> NodeImageParameterPtrConst;
+
+
+    //! Node that defines a Image model parameter.
+	//! \ingroup model
+    class MUTABLETOOLS_API NodeImageParameter : public NodeImage
+	{
+	public:
+
+		//-----------------------------------------------------------------------------------------
+		// Life cycle
+		//-----------------------------------------------------------------------------------------
+
+        NodeImageParameter();
+
+		void SerialiseWrapper(OutputArchive& arch) const override;
+		static void Serialise( const NodeImageParameter* pNode, OutputArchive& arch );
+        static NodeImageParameterPtr StaticUnserialise( InputArchive& arch );
+
+
+		//-----------------------------------------------------------------------------------------
+		// Node Interface
+		//-----------------------------------------------------------------------------------------
+
+		NodePtr Clone( NodeMapPtr pMap = 0 ) const override;
+
+		const NODE_TYPE* GetType() const override;
+		static const NODE_TYPE* GetStaticType();
+
+		int GetInputCount() const override;
+		Node* GetInputNode( int i ) const override;
+		void SetInputNode( int i, NodePtr pNode ) override;
+
+		//-----------------------------------------------------------------------------------------
+		// Own Interface
+		//-----------------------------------------------------------------------------------------
+
+		//! Get the name of the parameter. It will be exposed in the final compiled data.
+		const char* GetName() const;
+
+		//! Set the name of the parameter.
+		void SetName( const char* );
+
+		//! Get the uid of the parameter. It will be exposed in the final compiled data.
+		const char* GetUid() const;
+
+		//! Set the uid of the parameter.
+		void SetUid( const char* );
+
+		//! Get the default value of the parameter.
+        //ImagePtrConst GetDefaultValue() const;
+
+		//! Set the default value of the parameter.
+        //void SetDefaultValue( ImagePtr );
+
+		//-----------------------------------------------------------------------------------------
+		// Interface pattern
+		//-----------------------------------------------------------------------------------------
+		class Private;
+		Private* GetPrivate() const;
+        Node::Private* GetBasePrivate() const override;
+
+	protected:
+
+		//! Forbidden. Manage with the Ptr<> template.
+        ~NodeImageParameter();
+
+	private:
+
+		Private* m_pD;
+
+	};
+
+
+}
diff --git a/Engine/Plugins/Experimental/Mutable/Source/MutableTools/Public/NodeImagePlainColour.h b/Engine/Plugins/Experimental/Mutable/Source/MutableTools/Public/NodeImagePlainColour.h
new file mode 100644
--- /dev/null
+++ b/Engine/Plugins/Experimental/Mutable/Source/MutableTools/Public/NodeImagePlainColour.h
@@ -0,0 +1,81 @@
+// Copyright Epic Games, Inc. All Rights Reserved.
+
+#pragma once
+
+#include "NodeImage.h"
+
+
+namespace mu
+{
+
+	// Forward definitions
+	class NodeColour;
+	typedef Ptr<NodeColour> NodeColourPtr;
+	typedef Ptr<const NodeColour> NodeColourPtrConst;
+
+	class NodeImagePlainColour;
+	typedef Ptr<NodeImagePlainColour> NodeImagePlainColourPtr;
+	typedef Ptr<const NodeImagePlainColour> NodeImagePlainColourPtrConst;
+
+
+	//! Node that multiplies the colors of an image, channel by channel.
+	//! \ingroup model
+	class MUTABLETOOLS_API NodeImagePlainColour : public NodeImage
+	{
+	public:
+
+		//-----------------------------------------------------------------------------------------
+		// Life cycle
+		//-----------------------------------------------------------------------------------------
+
+		NodeImagePlainColour();
+
+		void SerialiseWrapper(OutputArchive& arch) const override;
+		static void Serialise( const NodeImagePlainColour* pNode, OutputArchive& arch );
+		static NodeImagePlainColourPtr StaticUnserialise( InputArchive& arch );
+
+
+		//-----------------------------------------------------------------------------------------
+		// Node Interface
+		//-----------------------------------------------------------------------------------------
+
+        NodePtr Clone( NodeMapPtr pMap = 0 ) const override;
+
+        const NODE_TYPE* GetType() const override;
+		static const NODE_TYPE* GetStaticType();
+
+        virtual int GetInputCount() const override;
+        virtual Node* GetInputNode( int i ) const override;
+        void SetInputNode( int i, NodePtr pNode ) override;
+
+		//-----------------------------------------------------------------------------------------
+		// Own Interface
+		//-----------------------------------------------------------------------------------------
+
+		//! Colour of the image.
+		NodeColourPtr GetColour() const;
+		void SetColour( NodeColourPtr );
+
+		//! New size or relative factor
+		void SetSize( int x, int y );
+
+		//-----------------------------------------------------------------------------------------
+		// Interface pattern
+		//-----------------------------------------------------------------------------------------
+		class Private;
+		Private* GetPrivate() const;
+        Node::Private* GetBasePrivate() const override;
+
+	protected:
+
+		//! Forbidden. Manage with the Ptr<> template.
+		~NodeImagePlainColour();
+
+	private:
+
+		Private* m_pD;
+
+	};
+
+
+}
diff --git a/Engine/Plugins/Experimental/Mutable/Source/MutableTools/Public/NodeImageProject.h b/Engine/Plugins/Experimental/Mutable/Source/MutableTools/Public/NodeImageProject.h
new file mode 100644
--- /dev/null
+++ b/Engine/Plugins/Experimental/Mutable/Source/MutableTools/Public/NodeImageProject.h
@@ -0,0 +1,121 @@
+// Copyright Epic Games, Inc. All Rights Reserved.
+
+
+#pragma once
+
+#include "MutableTools/Public/NodeImage.h"
+#include "MutableRuntime/Public/Image.h"
+#include "Math/IntVector.h"
+
+namespace mu
+{
+
+	// Forward definitions
+	class NodeScalar;
+	typedef Ptr<NodeScalar> NodeScalarPtr;
+	typedef Ptr<const NodeScalar> NodeScalarPtrConst;
+
+	class NodeColour;
+	typedef Ptr<NodeColour> NodeColourPtr;
+	typedef Ptr<const NodeColour> NodeColourPtrConst;
+
+	class NodeMesh;
+	typedef Ptr<NodeMesh> NodeMeshPtr;
+	typedef Ptr<const NodeMesh> NodeMeshPtrConst;
+
+	class NodeProjector;
+	typedef Ptr<NodeProjector> NodeProjectorPtr;
+	typedef Ptr<const NodeProjector> NodeProjectorPtrConst;
+
+	class NodeImageProject;
+	typedef Ptr<NodeImageProject> NodeImageProjectPtr;
+	typedef Ptr<const NodeImageProject> NodeImageProjectPtrConst;
+
+
+	//!
+	//! \ingroup model
+	class MUTABLETOOLS_API NodeImageProject : public NodeImage
+	{
+	public:
+
+		//-----------------------------------------------------------------------------------------
+		// Life cycle
+		//-----------------------------------------------------------------------------------------
+
+		NodeImageProject();
+
+		void SerialiseWrapper(OutputArchive& arch) const override;
+		static void Serialise( const NodeImageProject* pNode, OutputArchive& arch );
+		static NodeImageProjectPtr StaticUnserialise( InputArchive& arch );
+
+
+		//-----------------------------------------------------------------------------------------
+		// Node Interface
+		//-----------------------------------------------------------------------------------------
+
+        NodePtr Clone( NodeMapPtr pMap = 0 ) const override;
+
+        const NODE_TYPE* GetType() const override;
+		static const NODE_TYPE* GetStaticType();
+
+        virtual int GetInputCount() const override;
+        virtual Node* GetInputNode( int i ) const override;
+        void SetInputNode( int i, NodePtr pNode ) override;
+
+		//-----------------------------------------------------------------------------------------
+		// Own Interface
+		//-----------------------------------------------------------------------------------------
+
+		//! Get the node generating the mesh to be used for the projection.
+		NodeMeshPtr GetMesh() const;
+		void SetMesh( NodeMeshPtr );
+
+		//! Get the node generating the projector to be used for the projection.
+		NodeProjectorPtr GetProjector() const;
+		void SetProjector( NodeProjectorPtr );
+
+		//! Get the node generating a mask to filter out the projected images. It is optional.
+		NodeImagePtr GetTargetMask() const;
+		void SetTargetMask( NodeImagePtr );
+
+        //! Get the node generating the threshold angle of the projection. In degrees.
+		NodeScalarPtr GetAngleFadeStart() const;
+		void SetAngleFadeStart( NodeScalarPtr );
+
+        //! Get the node generating the threshold angle of the projection. In degrees.
+		NodeScalarPtr GetAngleFadeEnd() const;
+		void SetAngleFadeEnd( NodeScalarPtr );
+
+        //! Get the node generating the image to project.
+        NodeImagePtr GetImage() const;
+        void SetImage( NodeImagePtr );
+
+        //! UV layout of the mesh to use for the generated image. Defaults to 0.
+        uint8_t GetLayout() const;
+        void SetLayout( uint8_t  );
+
+		//! Set the size of the image to generate with the projection. If set to 0 (default) a size
+		//! that matches how this node is used will try to be guessed. 
+		const FUintVector2& GetImageSize() const;
+		void SetImageSize( const FUintVector2& size );
+
+		//-----------------------------------------------------------------------------------------
+		// Interface pattern
+		//-----------------------------------------------------------------------------------------
+		class Private;
+		Private* GetPrivate() const;
+        Node::Private* GetBasePrivate() const override;
+
+	protected:
+
+		//! Forbidden. Manage with the Ptr<> template.
+		~NodeImageProject();
+
+	private:
+
+		Private* m_pD;
+
+	};
+
+
+}
diff --git a/Engine/Plugins/Experimental/Mutable/Source/MutableTools/Public/NodeImageResize.h b/Engine/Plugins/Experimental/Mutable/Source/MutableTools/Public/NodeImageResize.h
new file mode 100644
--- /dev/null
+++ b/Engine/Plugins/Experimental/Mutable/Source/MutableTools/Public/NodeImageResize.h
@@ -0,0 +1,84 @@
+// Copyright Epic Games, Inc. All Rights Reserved.
+
+#pragma once
+
+#include "NodeImage.h"
+
+
+namespace mu
+{
+
+	// Forward definitions
+	class NodeScalar;
+	typedef Ptr<NodeScalar> NodeScalarPtr;
+	typedef Ptr<const NodeScalar> NodeScalarPtrConst;
+
+	class NodeImageResize;
+	typedef Ptr<NodeImageResize> NodeImageResizePtr;
+	typedef Ptr<const NodeImageResize> NodeImageResizePtrConst;
+
+
+	//! Node that multiplies the colors of an image, channel by channel.
+	//! \ingroup model
+	class MUTABLETOOLS_API NodeImageResize : public NodeImage
+	{
+	public:
+
+		//-----------------------------------------------------------------------------------------
+		// Life cycle
+		//-----------------------------------------------------------------------------------------
+
+		NodeImageResize();
+
+		void SerialiseWrapper(OutputArchive& arch) const override;
+		static void Serialise( const NodeImageResize* pNode, OutputArchive& arch );
+		static NodeImageResizePtr StaticUnserialise( InputArchive& arch );
+
+
+		//-----------------------------------------------------------------------------------------
+		// Node Interface
+		//-----------------------------------------------------------------------------------------
+
+        NodePtr Clone( NodeMapPtr pMap = 0 ) const override;
+
+        const NODE_TYPE* GetType() const override;
+		static const NODE_TYPE* GetStaticType();
+
+        virtual int GetInputCount() const override;
+        virtual Node* GetInputNode( int i ) const override;
+        void SetInputNode( int i, NodePtr pNode ) override;
+
+		//-----------------------------------------------------------------------------------------
+		// Own Interface
+		//-----------------------------------------------------------------------------------------
+
+		//! Base image to resize.
+		NodeImagePtr GetBase() const;
+		void SetBase( NodeImagePtr );
+
+		//! Is the size a relative factor or an absolute size?
+		void SetRelative( bool );
+
+		//! New size or relative factor
+		void SetSize( float x, float y );
+
+		//-----------------------------------------------------------------------------------------
+		// Interface pattern
+		//-----------------------------------------------------------------------------------------
+		class Private;
+		Private* GetPrivate() const;
+        Node::Private* GetBasePrivate() const override;
+
+	protected:
+
+		//! Forbidden. Manage with the Ptr<> template.
+		~NodeImageResize();
+
+	private:
+
+		Private* m_pD;
+
+	};
+
+
+}
diff --git a/Engine/Plugins/Experimental/Mutable/Source/MutableTools/Public/NodeImageSaturate.h b/Engine/Plugins/Experimental/Mutable/Source/MutableTools/Public/NodeImageSaturate.h
new file mode 100644
--- /dev/null
+++ b/Engine/Plugins/Experimental/Mutable/Source/MutableTools/Public/NodeImageSaturate.h
@@ -0,0 +1,88 @@
+// Copyright Epic Games, Inc. All Rights Reserved.
+
+#pragma once
+
+#include "NodeImage.h"
+
+#include "Image.h"
+
+
+namespace mu
+{
+
+	// Forward definitions
+	class NodeScalar;
+	typedef Ptr<NodeScalar> NodeScalarPtr;
+	typedef Ptr<const NodeScalar> NodeScalarPtrConst;
+
+	class NodeImageSaturate;
+	typedef Ptr<NodeImageSaturate> NodeImageSaturatePtr;
+	typedef Ptr<const NodeImageSaturate> NodeImageSaturatePtrConst;
+
+
+	//! Change the saturation of an image. This node can be used to increase the saturation with a
+	//! factor bigger than 1, or to decrease it or desaturate it completely with a factor smaller
+	//! than 1 or 0.
+	//! \ingroup model
+	class MUTABLETOOLS_API NodeImageSaturate : public NodeImage
+	{
+	public:
+
+		//-----------------------------------------------------------------------------------------
+		// Life cycle
+		//-----------------------------------------------------------------------------------------
+
+		NodeImageSaturate();
+
+		void SerialiseWrapper(OutputArchive& arch) const override;
+		static void Serialise( const NodeImageSaturate* pNode, OutputArchive& arch );
+		static NodeImageSaturatePtr StaticUnserialise( InputArchive& arch );
+
+
+		//-----------------------------------------------------------------------------------------
+		// Node Interface
+		//-----------------------------------------------------------------------------------------
+
+        NodePtr Clone( NodeMapPtr pMap = 0 ) const override;
+
+        const NODE_TYPE* GetType() const override;
+		static const NODE_TYPE* GetStaticType();
+
+        virtual int GetInputCount() const override;
+        virtual Node* GetInputNode( int i ) const override;
+        void SetInputNode( int i, NodePtr pNode ) override;
+
+		//-----------------------------------------------------------------------------------------
+		// Own Interface
+		//-----------------------------------------------------------------------------------------
+
+		//! Get the node generating the saturation factor.
+		//! A value of 0 completely desaturates. 1 will leave the same saturation and bigger than 1
+		//! will increase it.
+		NodeScalarPtr GetFactor() const;
+		void SetFactor( NodeScalarPtr );
+
+		//! Get the node generating the source image to be [de]saturated.
+		NodeImagePtr GetSource() const;
+		void SetSource( NodeImagePtr );
+
+		//-----------------------------------------------------------------------------------------
+		// Interface pattern
+		//-----------------------------------------------------------------------------------------
+		class Private;
+		Private* GetPrivate() const;
+        Node::Private* GetBasePrivate() const override;
+
+	protected:
+
+		//! Forbidden. Manage with the Ptr<> template.
+		~NodeImageSaturate();
+
+	private:
+
+		Private* m_pD;
+
+	};
+
+
+}
diff --git a/Engine/Plugins/Experimental/Mutable/Source/MutableTools/Public/NodeImageSelectColour.h b/Engine/Plugins/Experimental/Mutable/Source/MutableTools/Public/NodeImageSelectColour.h
new file mode 100644
--- /dev/null
+++ b/Engine/Plugins/Experimental/Mutable/Source/MutableTools/Public/NodeImageSelectColour.h
@@ -0,0 +1,85 @@
+// Copyright Epic Games, Inc. All Rights Reserved.
+
+#pragma once
+
+#include "NodeImage.h"
+
+#include "Image.h"
+
+
+namespace mu
+{
+
+	// Forward definitions
+	class NodeColour;
+	typedef Ptr<NodeColour> NodeColourPtr;
+	typedef Ptr<const NodeColour> NodeColourPtrConst;
+
+	class NodeImageSelectColour;
+	typedef Ptr<NodeImageSelectColour> NodeImageSelectColourPtr;
+	typedef Ptr<const NodeImageSelectColour> NodeImageSelectColourPtrConst;
+
+
+	//! Create a black and white mask image by selecting the pixels from a source image that match
+	//! a colour.
+	//! \ingroup model
+	class MUTABLETOOLS_API NodeImageSelectColour : public NodeImage
+	{
+	public:
+
+		//-----------------------------------------------------------------------------------------
+		// Life cycle
+		//-----------------------------------------------------------------------------------------
+
+		NodeImageSelectColour();
+
+		void SerialiseWrapper(OutputArchive& arch) const override;
+		static void Serialise( const NodeImageSelectColour* pNode, OutputArchive& arch );
+		static NodeImageSelectColourPtr StaticUnserialise( InputArchive& arch );
+
+
+		//-----------------------------------------------------------------------------------------
+		// Node Interface
+		//-----------------------------------------------------------------------------------------
+
+        NodePtr Clone( NodeMapPtr pMap = 0 ) const override;
+
+        const NODE_TYPE* GetType() const override;
+		static const NODE_TYPE* GetStaticType();
+
+        virtual int GetInputCount() const override;
+        virtual Node* GetInputNode( int i ) const override;
+        void SetInputNode( int i, NodePtr pNode ) override;
+
+		//-----------------------------------------------------------------------------------------
+		// Own Interface
+		//-----------------------------------------------------------------------------------------
+
+		//! Get the node generating the colour to select.
+		NodeColourPtr GetColour() const;
+		void SetColour( NodeColourPtr );
+
+		//! Get the node generating the source image to select pixels from.
+		NodeImagePtr GetSource() const;
+		void SetSource( NodeImagePtr );
+
+		//-----------------------------------------------------------------------------------------
+		// Interface pattern
+		//-----------------------------------------------------------------------------------------
+		class Private;
+		Private* GetPrivate() const;
+        Node::Private* GetBasePrivate() const override;
+
+	protected:
+
+		//! Forbidden. Manage with the Ptr<> template.
+		~NodeImageSelectColour();
+
+	private:
+
+		Private* m_pD;
+
+	};
+
+
+}
diff --git a/Engine/Plugins/Experimental/Mutable/Source/MutableTools/Public/NodeImageSwitch.h b/Engine/Plugins/Experimental/Mutable/Source/MutableTools/Public/NodeImageSwitch.h
new file mode 100644
--- /dev/null
+++ b/Engine/Plugins/Experimental/Mutable/Source/MutableTools/Public/NodeImageSwitch.h
@@ -0,0 +1,90 @@
+// Copyright Epic Games, Inc. All Rights Reserved.
+
+#pragma once
+
+#include "NodeImage.h"
+
+#include "Image.h"
+
+
+namespace mu
+{
+
+	// Forward definitions
+	class NodeScalar;
+	typedef Ptr<NodeScalar> NodeScalarPtr;
+	typedef Ptr<const NodeScalar> NodeScalarPtrConst;
+
+	class NodeImageSwitch;
+	typedef Ptr<NodeImageSwitch> NodeImageSwitchPtr;
+	typedef Ptr<const NodeImageSwitch> NodeImageSwitchPtrConst;
+
+
+	//! This node selects an output image from a set of input images based on a parameter.
+	//! \ingroup model
+	class MUTABLETOOLS_API NodeImageSwitch : public NodeImage
+	{
+	public:
+
+		//-----------------------------------------------------------------------------------------
+		// Life cycle
+		//-----------------------------------------------------------------------------------------
+
+		NodeImageSwitch();
+
+		void SerialiseWrapper(OutputArchive& arch) const override;
+		static void Serialise( const NodeImageSwitch* pNode, OutputArchive& arch );
+		static NodeImageSwitchPtr StaticUnserialise( InputArchive& arch );
+
+
+		//-----------------------------------------------------------------------------------------
+		// Node Interface
+		//-----------------------------------------------------------------------------------------
+
+        NodePtr Clone( NodeMapPtr pMap = 0 ) const override;
+
+        const NODE_TYPE* GetType() const override;
+		static const NODE_TYPE* GetStaticType();
+
+        virtual int GetInputCount() const override;
+        virtual Node* GetInputNode( int i ) const override;
+        void SetInputNode( int i, NodePtr pNode ) override;
+
+		//-----------------------------------------------------------------------------------------
+		// Own Interface
+		//-----------------------------------------------------------------------------------------
+
+		//! Get the node generating the parameter used to select the option.
+		NodeScalarPtr GetParameter() const;
+		void SetParameter( NodeScalarPtr );
+
+		//! Set the number of option images. It will keep the currently set targets and initialise
+		//! the new ones as null.
+		void SetOptionCount( int );
+		int GetOptionCount() const;
+
+		//! Get the node generating the t-th option image.
+		NodeImagePtr GetOption( int t ) const;
+		void SetOption( int t, NodeImagePtr );
+
+
+		//-----------------------------------------------------------------------------------------
+		// Interface pattern
+		//-----------------------------------------------------------------------------------------
+		class Private;
+		Private* GetPrivate() const;
+        Node::Private* GetBasePrivate() const override;
+
+	protected:
+
+		//! Forbidden. Manage with the Ptr<> template.
+		~NodeImageSwitch();
+
+	private:
+
+		Private* m_pD;
+
+	};
+
+
+}
diff --git a/Engine/Plugins/Experimental/Mutable/Source/MutableTools/Public/NodeImageSwizzle.h b/Engine/Plugins/Experimental/Mutable/Source/MutableTools/Public/NodeImageSwizzle.h
new file mode 100644
--- /dev/null
+++ b/Engine/Plugins/Experimental/Mutable/Source/MutableTools/Public/NodeImageSwizzle.h
@@ -0,0 +1,84 @@
+// Copyright Epic Games, Inc. All Rights Reserved.
+
+#pragma once
+
+#include "NodeImage.h"
+
+#include "Image.h"
+
+
+namespace mu
+{
+
+	// Forward definitions
+	class NodeImageSwizzle;
+	typedef Ptr<NodeImageSwizzle> NodeImageSwizzlePtr;
+	typedef Ptr<const NodeImageSwizzle> NodeImageSwizzlePtrConst;
+
+
+	//! Node that composes a new image by gathering pixel data from channels in other images.
+	//! \ingroup model
+	class MUTABLETOOLS_API NodeImageSwizzle : public NodeImage
+	{
+	public:
+
+		//-----------------------------------------------------------------------------------------
+		// Life cycle
+		//-----------------------------------------------------------------------------------------
+
+		NodeImageSwizzle();
+
+		void SerialiseWrapper(OutputArchive& arch) const override;
+		static void Serialise( const NodeImageSwizzle* pNode, OutputArchive& arch );
+		static NodeImageSwizzlePtr StaticUnserialise( InputArchive& arch );
+
+
+		//-----------------------------------------------------------------------------------------
+		// Node Interface
+		//-----------------------------------------------------------------------------------------
+
+        NodePtr Clone( NodeMapPtr pMap = 0 ) const override;
+
+        const NODE_TYPE* GetType() const override;
+		static const NODE_TYPE* GetStaticType();
+
+        virtual int GetInputCount() const override;
+        virtual Node* GetInputNode( int i ) const override;
+        void SetInputNode( int i, NodePtr pNode ) override;
+
+		//-----------------------------------------------------------------------------------------
+		// Own Interface
+		//-----------------------------------------------------------------------------------------
+
+		//!
+		EImageFormat GetFormat() const;
+		void SetFormat(EImageFormat);
+
+		//!
+		NodeImagePtr GetSource( int ) const;
+		int GetSourceChannel( int ) const;
+
+		void SetSource( int, NodeImagePtr );
+		void SetSourceChannel( int, int );
+
+
+		//-----------------------------------------------------------------------------------------
+		// Interface pattern
+		//-----------------------------------------------------------------------------------------
+		class Private;
+		Private* GetPrivate() const;
+        Node::Private* GetBasePrivate() const override;
+
+	protected:
+
+		//! Forbidden. Manage with the Ptr<> template.
+		~NodeImageSwizzle();
+
+	private:
+
+		Private* m_pD;
+
+	};
+
+
+}
diff --git a/Engine/Plugins/Experimental/Mutable/Source/MutableTools/Public/NodeImageTable.h b/Engine/Plugins/Experimental/Mutable/Source/MutableTools/Public/NodeImageTable.h
new file mode 100644
--- /dev/null
+++ b/Engine/Plugins/Experimental/Mutable/Source/MutableTools/Public/NodeImageTable.h
@@ -0,0 +1,86 @@
+// Copyright Epic Games, Inc. All Rights Reserved.
+
+#pragma once
+
+#include "NodeImage.h"
+
+
+namespace mu
+{
+
+	// Forward definitions
+	class NodeImageTable;
+	typedef Ptr<NodeImageTable> NodeImageTablePtr;
+	typedef Ptr<const NodeImageTable> NodeImageTablePtrConst;
+
+	class Table;
+	typedef Ptr<Table> TablePtr;
+	typedef Ptr<const Table> TablePtrConst;
+
+
+	//! This node provides the meshes stored in the column of a table.
+	//! \ingroup transform
+	class MUTABLETOOLS_API NodeImageTable : public NodeImage
+	{
+	public:
+
+		//-----------------------------------------------------------------------------------------
+		// Life cycle
+		//-----------------------------------------------------------------------------------------
+
+		NodeImageTable();
+
+		void SerialiseWrapper(OutputArchive& arch) const override;
+		static void Serialise( const NodeImageTable* pNode, OutputArchive& arch );
+		static NodeImageTablePtr StaticUnserialise( InputArchive& arch );
+
+
+		//-----------------------------------------------------------------------------------------
+		// Node Interface
+		//-----------------------------------------------------------------------------------------
+
+        NodePtr Clone( NodeMapPtr pMap = 0 ) const override;
+
+        const NODE_TYPE* GetType() const override;
+		static const NODE_TYPE* GetStaticType();
+
+        virtual int GetInputCount() const override;
+        virtual Node* GetInputNode( int i ) const override;
+        void SetInputNode( int i, NodePtr pNode ) override;
+
+		//-----------------------------------------------------------------------------------------
+		// Own Interface
+		//-----------------------------------------------------------------------------------------
+
+		//! Set the name of the implicit table parameter.
+		void SetParameterName( const char* strName );
+
+		//!
+		TablePtr GetTable() const;
+		void SetTable( TablePtr );
+
+		//!
+		const char* GetColumn() const;
+		void SetColumn( const char* strName );
+
+		//-----------------------------------------------------------------------------------------
+		// Interface pattern
+		//-----------------------------------------------------------------------------------------
+		class Private;
+		Private* GetPrivate() const;
+        Node::Private* GetBasePrivate() const override;
+
+
+	protected:
+
+		//! Forbidden. Manage with the Ptr<> template.
+		~NodeImageTable();
+
+	private:
+
+		Private* m_pD;
+
+	};
+
+
+}
diff --git a/Engine/Plugins/Experimental/Mutable/Source/MutableTools/Public/NodeImageTransform.h b/Engine/Plugins/Experimental/Mutable/Source/MutableTools/Public/NodeImageTransform.h
new file mode 100644
--- /dev/null
+++ b/Engine/Plugins/Experimental/Mutable/Source/MutableTools/Public/NodeImageTransform.h
@@ -0,0 +1,88 @@
+// Copyright Epic Games, Inc. All Rights Reserved.
+
+#pragma once
+
+#include "NodeImage.h"
+
+
+namespace mu
+{
+
+	// Forward definitions
+	class NodeScalar;
+	typedef Ptr<NodeScalar> NodeScalarPtr;
+	typedef Ptr<const NodeScalar> NodeScalarPtrConst;
+
+	class NodeImageTransform;
+	typedef Ptr<NodeImageTransform> NodeImageTransformPtr;
+	typedef Ptr<const NodeImageTransform> NodeImageTransformPtrConst;
+
+
+	//! Node that multiplies the colors of an image, channel by channel.
+	//! \ingroup model
+	class MUTABLETOOLS_API NodeImageTransform : public NodeImage
+	{
+	public:
+
+		//-----------------------------------------------------------------------------------------
+		// Life cycle
+		//-----------------------------------------------------------------------------------------
+
+		NodeImageTransform();
+
+		void SerialiseWrapper(OutputArchive& arch) const override;
+		static void Serialise( const NodeImageTransform* pNode, OutputArchive& arch );
+		static NodeImageTransformPtr StaticUnserialise( InputArchive& arch );
+
+
+		//-----------------------------------------------------------------------------------------
+		// Node Interface
+		//-----------------------------------------------------------------------------------------
+
+        NodePtr Clone( NodeMapPtr pMap = 0 ) const override;
+
+        const NODE_TYPE* GetType() const override;
+		static const NODE_TYPE* GetStaticType();
+
+        virtual int GetInputCount() const override;
+        virtual Node* GetInputNode( int i ) const override;
+        void SetInputNode( int i, NodePtr pNode ) override;
+
+		//-----------------------------------------------------------------------------------------
+		// Own Interface
+		//-----------------------------------------------------------------------------------------
+
+		//! Base image to multiply.
+		NodeImagePtr GetBase() const;
+		void SetBase( NodeImagePtr );
+	
+		NodeScalarPtr GetOffsetX() const;
+		NodeScalarPtr GetOffsetY() const;
+		NodeScalarPtr GetScaleX() const;
+		NodeScalarPtr GetScaleY() const;
+		NodeScalarPtr GetRotation() const;
+		void SetOffsetX( NodeScalarPtr pNode );
+		void SetOffsetY( NodeScalarPtr pNode );
+		void SetScaleX( NodeScalarPtr pNode );
+		void SetScaleY( NodeScalarPtr pNode );
+		void SetRotation( NodeScalarPtr pNode );
+		//-----------------------------------------------------------------------------------------
+		// Interface pattern
+		//-----------------------------------------------------------------------------------------
+		class Private;
+		Private* GetPrivate() const;
+        Node::Private* GetBasePrivate() const override;
+
+	protected:
+
+		//! Forbidden. Manage with the Ptr<> template.
+		~NodeImageTransform();
+
+	private:
+
+		Private* m_pD;
+
+	};
+
+
+}
diff --git a/Engine/Plugins/Experimental/Mutable/Source/MutableTools/Public/NodeImageVariation.h b/Engine/Plugins/Experimental/Mutable/Source/MutableTools/Public/NodeImageVariation.h
new file mode 100644
--- /dev/null
+++ b/Engine/Plugins/Experimental/Mutable/Source/MutableTools/Public/NodeImageVariation.h
@@ -0,0 +1,83 @@
+// Copyright Epic Games, Inc. All Rights Reserved.
+
+#pragma once
+
+#include "NodeImage.h"
+
+
+namespace mu
+{
+
+    // Forward definitions
+    class NodeImageVariation;
+    typedef Ptr<NodeImageVariation> NodeImageVariationPtr;
+    typedef Ptr<const NodeImageVariation> NodeImageVariationPtrConst;
+
+
+    //!
+    //! \ingroup model
+    class MUTABLETOOLS_API NodeImageVariation : public NodeImage
+    {
+    public:
+        //-----------------------------------------------------------------------------------------
+        // Life cycle
+        //-----------------------------------------------------------------------------------------
+
+        NodeImageVariation();
+
+		void SerialiseWrapper(OutputArchive& arch) const override;
+		static void Serialise( const NodeImageVariation* pNode, OutputArchive& arch );
+        static NodeImageVariationPtr StaticUnserialise( InputArchive& arch );
+
+        //-----------------------------------------------------------------------------------------
+        // Node Interface
+        //-----------------------------------------------------------------------------------------
+
+        NodePtr Clone( NodeMapPtr pMap = 0 ) const override;
+
+        const NODE_TYPE* GetType() const override;
+        static const NODE_TYPE* GetStaticType();
+
+        virtual int GetInputCount() const override;
+        virtual Node* GetInputNode( int i ) const override;
+        void SetInputNode( int i, NodePtr pNode ) override;
+
+        //-----------------------------------------------------------------------------------------
+        // Own Interface
+        //-----------------------------------------------------------------------------------------
+
+        //!
+        void SetDefaultImage( NodeImage* Image );
+
+        //! Set the number of tags to consider in this variation
+        void SetVariationCount( int count );
+
+        //!
+        int GetVariationCount() const;
+
+        //! Set the tag or state name that will enable a specific vartiation
+        void SetVariationTag( int index, const char* strTag );
+
+        //!
+        void SetVariationImage( int index, NodeImage* Image );
+
+        //!}
+
+
+        //-----------------------------------------------------------------------------------------
+        // Interface pattern
+        //-----------------------------------------------------------------------------------------
+        class Private;
+        Private* GetPrivate() const;
+        Node::Private* GetBasePrivate() const override;
+
+    protected:
+        //! Forbidden. Manage with the Ptr<> template.
+        ~NodeImageVariation();
+
+    private:
+        Private* m_pD;
+    };
+
+
+} // namespace mu
diff --git a/Engine/Plugins/Experimental/Mutable/Source/MutableTools/Public/NodeLOD.h b/Engine/Plugins/Experimental/Mutable/Source/MutableTools/Public/NodeLOD.h
new file mode 100644
--- /dev/null
+++ b/Engine/Plugins/Experimental/Mutable/Source/MutableTools/Public/NodeLOD.h
@@ -0,0 +1,104 @@
+// Copyright Epic Games, Inc. All Rights Reserved.
+
+
+#pragma once
+
+#include "MutableTools/Public/Node.h"
+
+
+namespace mu
+{
+
+	// Forward definitions
+	class NodeComponent;
+	typedef Ptr<NodeComponent> NodeComponentPtr;
+	typedef Ptr<const NodeComponent> NodeComponentPtrConst;
+
+	class NodeModifier;
+	typedef Ptr<NodeModifier> NodeModifierPtr;
+	typedef Ptr<const NodeModifier> NodeModifierPtrConst;
+
+	class NodeLOD;
+	typedef Ptr<NodeLOD> NodeLODPtr;
+	typedef Ptr<const NodeLOD> NodeLODPtrConst;
+
+
+	//! Node that creates a new level of detail by assembling several object components.
+	//! \ingroup model
+	class MUTABLETOOLS_API NodeLOD : public Node
+	{
+	public:
+
+		//-----------------------------------------------------------------------------------------
+		// Life cycle
+		//-----------------------------------------------------------------------------------------
+
+		NodeLOD();
+
+		void SerialiseWrapper(OutputArchive& arch) const override;
+		static void Serialise( const NodeLOD* pNode, OutputArchive& arch );
+		static NodeLODPtr StaticUnserialise( InputArchive& arch );
+
+
+		//-----------------------------------------------------------------------------------------
+		// Node Interface
+		//-----------------------------------------------------------------------------------------
+
+		NodePtr Clone( NodeMapPtr pMap = 0 ) const override;
+
+		const NODE_TYPE* GetType() const override;
+		static const NODE_TYPE* GetStaticType();
+
+		int GetInputCount() const override;
+		Node* GetInputNode( int i ) const override;
+		void SetInputNode( int i, NodePtr pNode ) override;
+
+		//-----------------------------------------------------------------------------------------
+		// Own Interface
+		//-----------------------------------------------------------------------------------------
+
+		//! Get the number of components in the LOD to create
+		int GetComponentCount() const;
+
+		//! Set the number of components in the LOD to create
+		void SetComponentCount( int );
+
+		//! Get the node generating one of the components of the LOD.
+		//! \param index index of the component from 0 to GetComponentCount()-1
+		NodeComponentPtr GetComponent( int index ) const;
+
+		//! Set the node generating one of the components of the LOD.
+		//! \param index index of the component from 0 to GetComponentCount()-1
+		void SetComponent( int index, NodeComponentPtr );
+
+		//! Set the number of modifiers in the LOD
+		void SetModifierCount(int);
+
+		//! Get the number of modifiers in the LOD
+		int GetModifierCount() const;
+
+		//! Set the node generating one of the modifiers of the LOD.
+		//! \param index index of the component from 0 to GetModifierCount()-1
+		void SetModifier(int index, NodeModifierPtr);
+
+		//-----------------------------------------------------------------------------------------
+		// Interface pattern
+		//-----------------------------------------------------------------------------------------
+		class Private;
+		Private* GetPrivate() const;
+        Node::Private* GetBasePrivate() const override;
+
+	protected:
+
+		//! Forbidden. Manage with the Ptr<> template.
+		~NodeLOD();
+
+	private:
+
+		Private* m_pD;
+
+	};
+
+
+
+}
diff --git a/Engine/Plugins/Experimental/Mutable/Source/MutableTools/Public/NodeLayout.h b/Engine/Plugins/Experimental/Mutable/Source/MutableTools/Public/NodeLayout.h
new file mode 100644
--- /dev/null
+++ b/Engine/Plugins/Experimental/Mutable/Source/MutableTools/Public/NodeLayout.h
@@ -0,0 +1,168 @@
+// Copyright Epic Games, Inc. All Rights Reserved.
+
+#pragma once
+
+#include "MutableTools/Public/Node.h"
+#include "MutableRuntime/Public/Mesh.h"
+#include "MutableRuntime/Public/Layout.h"
+
+
+namespace mu
+{
+
+	//
+	class NodeLayout;
+	typedef Ptr<NodeLayout> NodeLayoutPtr;
+	typedef Ptr<const NodeLayout> NodeLayoutPtrConst;
+
+	//
+	class NodeLayoutBlocks;
+	typedef Ptr<NodeLayoutBlocks> NodeLayoutBlocksPtr;
+	typedef Ptr<const NodeLayoutBlocks> NodeLayoutBlocksPtrConst;
+
+
+	//! This node is used to define the texture layout for a texture coordinates channel of a
+	//! constant mesh.
+	//! \ingroup model
+	class MUTABLETOOLS_API NodeLayout : public Node
+	{
+	public:
+
+		// Possible subclasses
+		enum class EType : uint8
+		{
+			Blocks = 0,
+
+			None
+		};
+
+		//-----------------------------------------------------------------------------------------
+		// Life cycle
+		//-----------------------------------------------------------------------------------------
+
+		static void Serialise( const NodeLayout* pNode, OutputArchive& arch );
+		static NodeLayoutPtr StaticUnserialise( InputArchive& arch );
+
+
+		//-----------------------------------------------------------------------------------------
+		// Node Interface
+		//-----------------------------------------------------------------------------------------
+
+		const NODE_TYPE* GetType() const override;
+		static const NODE_TYPE* GetStaticType();
+
+		//-----------------------------------------------------------------------------------------
+		// Own Interface
+		//-----------------------------------------------------------------------------------------
+
+
+		//-----------------------------------------------------------------------------------------
+		// Interface pattern
+		//-----------------------------------------------------------------------------------------
+		class Private;
+
+	protected:
+
+		//! Forbidden. Manage with the Ptr<> template.
+		~NodeLayout() {}
+
+		//!
+		EType Type = EType::None;
+
+	};
+
+
+	//! This node is used to define the texture layout for a texture coordinates channel of a
+	//! constant mesh.
+	//! The blocks defined here will be set to the channel of the mesh were this node is connected.
+	//! \ingroup model
+	class MUTABLETOOLS_API NodeLayoutBlocks : public NodeLayout
+	{
+	public:
+
+		//-----------------------------------------------------------------------------------------
+		// Life cycle
+		//-----------------------------------------------------------------------------------------
+
+		NodeLayoutBlocks();
+
+		void SerialiseWrapper(OutputArchive& arch) const override;
+		static void Serialise( const NodeLayoutBlocks* pNode, OutputArchive& arch );
+		static NodeLayoutBlocksPtr StaticUnserialise( InputArchive& arch );
+
+
+		//-----------------------------------------------------------------------------------------
+		// Node Interface
+		//-----------------------------------------------------------------------------------------
+
+		NodePtr Clone( NodeMapPtr pMap = 0 ) const override;
+
+		const NODE_TYPE* GetType() const override;
+		static const NODE_TYPE* GetStaticType();
+
+		int GetInputCount() const override;
+		Node* GetInputNode( int i ) const override;
+		void SetInputNode( int i, NodePtr pNode ) override;
+
+		//-----------------------------------------------------------------------------------------
+		// Own Interface
+		//-----------------------------------------------------------------------------------------
+
+		//! Set the size of the grid where the blocks will be defined
+		void SetGridSize( int x, int y );
+
+		//! Set the maximum size of the grid where the blocks can be defined
+		void SetMaxGridSize(int x, int y);
+
+		//! Get the size of the grid where the blocks will be defined
+		void GetGridSize( int* pX, int* pY ) const;
+
+		//! Get the maximum size of the grid where the blocks can be defined
+		void GetMaxGridSize(int* pX, int* pY) const;
+
+		//! Set the number of blocks in the layout.
+		//! It keeps the current data as much as possible, and the new data is undefined.
+		void SetBlockCount( int );
+
+		//! Get the number of blocks in the layout
+		int GetBlockCount();
+
+		//! Set a block of the layout.
+		//! minx and miny refer to the lowest left corner of the block.
+        void SetBlock( int index, int minx, int miny, int sizex, int sizey );
+
+		//! Get a block of the layout.
+		//! minx and miny refer to the lowest left corner of the block.
+		void GetBlock(int index, int* minx, int* miny, int* sizex, int* sizey);
+
+		//! Set priority of a block.
+		void SetBlockPriority(int index, int priority);
+
+		//! Set the texture layout packing strategy 
+		void SetLayoutPackingStrategy(EPackStrategy strategy);
+
+		//! Generate the blocks of a layout using the UV of the meshes
+		static NodeLayoutBlocksPtr GenerateLayoutBlocks(const MeshPtr pMesh, int layoutIndex, int gridSizeX, int gridSizeY);
+
+
+		//-----------------------------------------------------------------------------------------
+		// Interface pattern
+		//-----------------------------------------------------------------------------------------
+		class Private;
+		Private* GetPrivate() const;
+        Node::Private* GetBasePrivate() const override;
+
+	protected:
+
+		//! Forbidden. Manage with the Ptr<> template.
+		~NodeLayoutBlocks();
+
+	private:
+
+		Private* m_pD;
+
+	};
+
+
+
+}
diff --git a/Engine/Plugins/Experimental/Mutable/Source/MutableTools/Public/NodeMesh.h b/Engine/Plugins/Experimental/Mutable/Source/MutableTools/Public/NodeMesh.h
new file mode 100644
--- /dev/null
+++ b/Engine/Plugins/Experimental/Mutable/Source/MutableTools/Public/NodeMesh.h
@@ -0,0 +1,81 @@
+// Copyright Epic Games, Inc. All Rights Reserved.
+
+#pragma once
+
+#include "Node.h"
+
+namespace mu
+{
+
+	// Forward definitions
+	class NodeMesh;
+	typedef Ptr<NodeMesh> NodeMeshPtr;
+	typedef Ptr<const NodeMesh> NodeMeshPtrConst;
+
+
+    //! %Base class of any node that outputs a mesh.
+	//! \ingroup model
+	class MUTABLETOOLS_API NodeMesh : public Node
+	{
+	public:
+
+		// Possible subclasses
+		enum class EType : uint8
+		{
+			Constant = 0,
+			Interpolate = 2,
+			Table = 4,
+			Subtract = 6,
+			Format = 7,
+			Tangents = 8,
+			Morph = 9,
+			MakeMorph = 10,
+			Switch = 11,
+			Fragment = 12,
+			Transform = 13,
+			ClipMorphPlane = 14,
+			ClipWithMesh = 15,
+			ApplyPose = 16,
+			Variation = 17,
+			GeometryOperation = 18,
+			Reshape = 19,
+			ClipDeform = 20,
+
+			None
+		};
+
+		//-----------------------------------------------------------------------------------------
+		// Life cycle
+		//-----------------------------------------------------------------------------------------
+
+		static void Serialise( const NodeMesh* pNode, OutputArchive& arch );
+		static NodeMeshPtr StaticUnserialise( InputArchive& arch );
+
+
+		//-----------------------------------------------------------------------------------------
+		// Node Interface
+		//-----------------------------------------------------------------------------------------
+
+        const NODE_TYPE* GetType() const override;
+		static const NODE_TYPE* GetStaticType();
+
+		inline EType GetMeshNodeType() const { return Type; }
+
+
+		//-----------------------------------------------------------------------------------------
+		// Interface pattern
+		//-----------------------------------------------------------------------------------------
+		class Private;
+
+	protected:
+
+		//! Forbidden. Manage with the Ptr<> template.
+		inline ~NodeMesh() {}
+
+		//!
+		EType Type = EType::None;
+	};
+
+
+}
+
diff --git a/Engine/Plugins/Experimental/Mutable/Source/MutableTools/Public/NodeMeshApplyPose.h b/Engine/Plugins/Experimental/Mutable/Source/MutableTools/Public/NodeMeshApplyPose.h
new file mode 100644
--- /dev/null
+++ b/Engine/Plugins/Experimental/Mutable/Source/MutableTools/Public/NodeMeshApplyPose.h
@@ -0,0 +1,81 @@
+// Copyright Epic Games, Inc. All Rights Reserved.
+
+#pragma once
+
+#include "NodeMesh.h"
+
+#include "Mesh.h"
+
+
+namespace mu
+{
+
+	// Forward definitions
+    class NodeMeshApplyPose;
+    using NodeMeshApplyPosePtr = Ptr<NodeMeshApplyPose> ;
+    using NodeMeshApplyPosePtrConst = Ptr<const NodeMeshApplyPose>;
+
+
+    //! Node that applies a pose to a mesh, baking it into the vertex data
+	//! \ingroup model
+    class MUTABLETOOLS_API NodeMeshApplyPose : public NodeMesh
+	{
+	public:
+
+		//-----------------------------------------------------------------------------------------
+		// Life cycle
+		//-----------------------------------------------------------------------------------------
+
+        NodeMeshApplyPose();
+
+		void SerialiseWrapper(OutputArchive& arch) const override;
+		static void Serialise( const NodeMeshApplyPose* pNode, OutputArchive& arch );
+        static NodeMeshApplyPosePtr StaticUnserialise( InputArchive& arch );
+
+
+		//-----------------------------------------------------------------------------------------
+		// Node Interface
+		//-----------------------------------------------------------------------------------------
+
+		NodePtr Clone( NodeMapPtr pMap = 0 ) const override;
+
+		const NODE_TYPE* GetType() const override;
+		static const NODE_TYPE* GetStaticType();
+
+		int GetInputCount() const override;
+		Node* GetInputNode( int i ) const override;
+		void SetInputNode( int i, NodePtr pNode ) override;
+
+		//-----------------------------------------------------------------------------------------
+		// Own Interface
+		//-----------------------------------------------------------------------------------------
+
+		//! Get the nodes generating the base mesh to be morphed
+		NodeMeshPtr GetBase() const;
+		void SetBase( NodeMeshPtr );
+
+        //! Get the nodes generating the pose to apply. A pose is represented with a mesh object
+        //! with no geometry: only bone matrices and names are relevant.
+        NodeMeshPtr GetPose() const;
+        void SetPose( NodeMeshPtr );
+
+		//-----------------------------------------------------------------------------------------
+		// Interface pattern
+		//-----------------------------------------------------------------------------------------
+		class Private;
+		Private* GetPrivate() const;
+        Node::Private* GetBasePrivate() const override;
+
+	protected:
+
+		//! Forbidden. Manage with the Ptr<> template.
+        ~NodeMeshApplyPose();
+
+	private:
+
+		Private* m_pD;
+
+	};
+
+
+}
diff --git a/Engine/Plugins/Experimental/Mutable/Source/MutableTools/Public/NodeMeshClipDeform.h b/Engine/Plugins/Experimental/Mutable/Source/MutableTools/Public/NodeMeshClipDeform.h
new file mode 100644
--- /dev/null
+++ b/Engine/Plugins/Experimental/Mutable/Source/MutableTools/Public/NodeMeshClipDeform.h
@@ -0,0 +1,77 @@
+// Copyright Epic Games, Inc. All Rights Reserved.
+
+#pragma once
+
+#include "MutableTools/Public/NodeMesh.h"
+#include "MutableTools/Public/NodeImage.h"
+
+#include "MutableRuntime/Public/Mesh.h"
+#include "MutableRuntime/Public/Image.h"
+
+
+namespace mu
+{
+
+	//! Node that morphs a base mesh with one or two weighted targets from a sequence.
+	//! \ingroup model
+	class MUTABLETOOLS_API NodeMeshClipDeform : public NodeMesh
+	{
+	public:
+
+		//-----------------------------------------------------------------------------------------
+		// Life cycle
+		//-----------------------------------------------------------------------------------------
+
+		NodeMeshClipDeform();
+
+		void SerialiseWrapper(OutputArchive& arch) const override;
+		static void Serialise( const NodeMeshClipDeform* pNode, OutputArchive& arch );
+		static Ptr<NodeMeshClipDeform> StaticUnserialise( InputArchive& arch );
+
+
+		//-----------------------------------------------------------------------------------------
+		// Node Interface
+		//-----------------------------------------------------------------------------------------
+
+		NodePtr Clone( NodeMapPtr pMap = 0 ) const override;
+
+		const NODE_TYPE* GetType() const override;
+		static const NODE_TYPE* GetStaticType();
+
+		int GetInputCount() const override;
+		Node* GetInputNode( int i ) const override;
+		void SetInputNode( int i, NodePtr pNode ) override;
+
+		//-----------------------------------------------------------------------------------------
+		// Own Interface
+		//-----------------------------------------------------------------------------------------
+
+		const NodeMeshPtr& GetBaseMesh() const;
+		void SetBaseMesh( const NodeMeshPtr& );
+
+		const NodeMeshPtr& GetClipShape() const;
+		void SetClipShape(const NodeMeshPtr&);
+
+		const NodeImagePtr& GetShapeWeights() const;
+		void SetShapeWeights(const NodeImagePtr&);
+
+
+        //-----------------------------------------------------------------------------------------
+		// Interface pattern
+		//-----------------------------------------------------------------------------------------
+		class Private;
+		Private* GetPrivate() const;
+        Node::Private* GetBasePrivate() const override;
+
+	protected:
+
+		//! Forbidden. Manage with the Ptr<> template.
+		~NodeMeshClipDeform();
+
+	private:
+
+		Private* m_pD;
+
+	};
+
+}
diff --git a/Engine/Plugins/Experimental/Mutable/Source/MutableTools/Public/NodeMeshClipMorphPlane.h b/Engine/Plugins/Experimental/Mutable/Source/MutableTools/Public/NodeMeshClipMorphPlane.h
new file mode 100644
--- /dev/null
+++ b/Engine/Plugins/Experimental/Mutable/Source/MutableTools/Public/NodeMeshClipMorphPlane.h
@@ -0,0 +1,94 @@
+// Copyright Epic Games, Inc. All Rights Reserved.
+
+#pragma once
+
+#include "NodeMesh.h"
+
+#include "Mesh.h"
+
+
+namespace mu
+{
+
+	// Forward definitions
+    class NodeMeshClipMorphPlane;
+    typedef Ptr<NodeMeshClipMorphPlane> NodeMeshClipMorphPlanePtr;
+    typedef Ptr<const NodeMeshClipMorphPlane> NodeMeshClipMorphPlaneConst;
+
+
+    //! This node applies a geometric transform represented by a 4x4 matrix to a mesh
+	//! \ingroup model
+    class MUTABLETOOLS_API NodeMeshClipMorphPlane : public NodeMesh
+	{
+	public:
+
+		//-----------------------------------------------------------------------------------------
+		// Life cycle
+		//-----------------------------------------------------------------------------------------
+
+		NodeMeshClipMorphPlane();
+
+		void SerialiseWrapper(OutputArchive& arch) const override;
+		static void Serialise( const NodeMeshClipMorphPlane* pNode, OutputArchive& arch );
+        static NodeMeshClipMorphPlanePtr StaticUnserialise( InputArchive& arch );
+
+		//-----------------------------------------------------------------------------------------
+		// Node Interface
+		//-----------------------------------------------------------------------------------------
+
+		virtual NodePtr Clone( NodeMapPtr pMap = 0 ) const;
+
+		virtual const NODE_TYPE* GetType() const;
+		static const NODE_TYPE* GetStaticType();
+
+		virtual int GetInputCount() const;
+		virtual const char* GetInputName( int i ) const;
+		virtual const NODE_TYPE* GetInputType( int i ) const;
+		virtual Node* GetInputNode( int i ) const;
+		virtual void SetInputNode( int i, NodePtr pNode );
+
+		//-----------------------------------------------------------------------------------------
+		// Own Interface
+		//-----------------------------------------------------------------------------------------
+
+		//! Source mesh to be clipped and morphed
+        NodeMeshPtr GetSource() const;
+		void SetSource( NodeMesh* );
+
+        //! \param 
+		void SetPlane(float centerX, float centerY, float centerZ, float normalX, float normalY, float normalZ);
+		void SetParams(float dist, float factor);
+		void SetMorphEllipse(float radius1, float radius2, float rotation);
+
+		//! Define an axis-aligned box that will select the vertices to be morphed.
+		//! If this is not set, or the size of the box is 0, all vertices will be affected.
+		//! Only one of Box or Bone Hierarchy can be used (the last one set)
+		void SetVertexSelectionBox(float centerX, float centerY, float centerZ, float radiusX, float radiusY, float radiusZ);
+
+		//! Define the root bone of the subhierarchy of the mesh that will be affected.
+		//! Only one of Box or Bone Hierarchy can be used (the last one set)
+		void SetVertexSelectionBone(const char* strBoneName, float maxEffectRadius);
+
+		//! Add a tag to the clip morph operation, which will only affect surfaces with the same tag
+		void AddTag(const char* tagName);
+
+		//-----------------------------------------------------------------------------------------
+		// Interface pattern
+		//-----------------------------------------------------------------------------------------
+		class Private;
+		Private* GetPrivate() const;
+		Node::Private* GetBasePrivate() const;
+
+	protected:
+
+		//! Forbidden. Manage with the Ptr<> template.
+        ~NodeMeshClipMorphPlane();
+
+	private:
+
+		Private* m_pD;
+
+	};
+
+
+}
diff --git a/Engine/Plugins/Experimental/Mutable/Source/MutableTools/Public/NodeMeshClipWithMesh.h b/Engine/Plugins/Experimental/Mutable/Source/MutableTools/Public/NodeMeshClipWithMesh.h
new file mode 100644
--- /dev/null
+++ b/Engine/Plugins/Experimental/Mutable/Source/MutableTools/Public/NodeMeshClipWithMesh.h
@@ -0,0 +1,83 @@
+// Copyright Epic Games, Inc. All Rights Reserved.
+
+#pragma once
+
+#include "NodeMesh.h"
+
+#include "Mesh.h"
+
+
+namespace mu
+{
+
+	// Forward definitions
+    class NodeMeshClipWithMesh;
+    typedef Ptr<NodeMeshClipWithMesh> NodeMeshClipWithMeshPtr;
+    typedef Ptr<const NodeMeshClipWithMesh> NodeMeshClipWithMeshConst;
+
+
+    //! This node applies a geometric transform represented by a 4x4 matrix to a mesh
+	//! \ingroup model
+    class MUTABLETOOLS_API NodeMeshClipWithMesh : public NodeMesh
+	{
+	public:
+
+		//-----------------------------------------------------------------------------------------
+		// Life cycle
+		//-----------------------------------------------------------------------------------------
+
+		NodeMeshClipWithMesh();
+
+		void SerialiseWrapper(OutputArchive& arch) const override;
+		static void Serialise( const NodeMeshClipWithMesh* pNode, OutputArchive& arch );
+        static NodeMeshClipWithMeshPtr StaticUnserialise( InputArchive& arch );
+
+		//-----------------------------------------------------------------------------------------
+		// Node Interface
+		//-----------------------------------------------------------------------------------------
+
+		virtual NodePtr Clone( NodeMapPtr pMap = 0 ) const;
+
+		virtual const NODE_TYPE* GetType() const;
+		static const NODE_TYPE* GetStaticType();
+
+		virtual int GetInputCount() const;
+		virtual const char* GetInputName( int i ) const;
+		virtual const NODE_TYPE* GetInputType( int i ) const;
+		virtual Node* GetInputNode( int i ) const;
+		virtual void SetInputNode( int i, NodePtr pNode );
+
+		//-----------------------------------------------------------------------------------------
+		// Own Interface
+		//-----------------------------------------------------------------------------------------
+
+		//! Source mesh to be clipped
+        NodeMeshPtr GetSource() const;
+		void SetSource(NodeMesh*);
+
+        //! \param 
+		void SetClipMesh(NodeMesh*);
+
+		//! Add a tag to the clip morph operation, which will only affect surfaces with the same tag
+		void AddTag(const char* tagName);
+
+		//-----------------------------------------------------------------------------------------
+		// Interface pattern
+		//-----------------------------------------------------------------------------------------
+		class Private;
+		Private* GetPrivate() const;
+		Node::Private* GetBasePrivate() const;
+
+	protected:
+
+		//! Forbidden. Manage with the Ptr<> template.
+        ~NodeMeshClipWithMesh();
+
+	private:
+
+		Private* m_pD;
+
+	};
+
+
+}
diff --git a/Engine/Plugins/Experimental/Mutable/Source/MutableTools/Public/NodeMeshConstant.h b/Engine/Plugins/Experimental/Mutable/Source/MutableTools/Public/NodeMeshConstant.h
new file mode 100644
--- /dev/null
+++ b/Engine/Plugins/Experimental/Mutable/Source/MutableTools/Public/NodeMeshConstant.h
@@ -0,0 +1,92 @@
+// Copyright Epic Games, Inc. All Rights Reserved.
+
+#pragma once
+
+#include "NodeMesh.h"
+
+#include "Mesh.h"
+#include "Types.h"
+
+
+namespace mu
+{
+
+	// Forward definitions
+	class NodeMeshConstant;
+	typedef Ptr<NodeMeshConstant> NodeMeshConstantPtr;
+	typedef Ptr<const NodeMeshConstant> NodeMeshConstantPtrConst;
+
+	class NodeLayout;
+	typedef Ptr<NodeLayout> NodeLayoutPtr;
+	typedef Ptr<const NodeLayout> NodeLayoutPtrConst;
+
+
+	//! Node that outputs a constant mesh.
+	//! It allows to define the layouts for the texture channels of the constant mesh
+	//! \ingroup model
+	class MUTABLETOOLS_API NodeMeshConstant : public NodeMesh
+	{
+	public:
+
+		//-----------------------------------------------------------------------------------------
+		// Life cycle
+		//-----------------------------------------------------------------------------------------
+
+		NodeMeshConstant();
+
+		void SerialiseWrapper(OutputArchive& arch) const override;
+		static void Serialise( const NodeMeshConstant* pNode, OutputArchive& arch );
+		static NodeMeshConstantPtr StaticUnserialise( InputArchive& arch );
+
+
+		//-----------------------------------------------------------------------------------------
+		// Node Interface
+		//-----------------------------------------------------------------------------------------
+
+		NodePtr Clone( NodeMapPtr pMap = 0 ) const override;
+
+		const NODE_TYPE* GetType() const override;
+		static const NODE_TYPE* GetStaticType();
+
+		int GetInputCount() const override;
+		Node* GetInputNode( int i ) const override;
+		void SetInputNode( int i, NodePtr pNode ) override;
+
+		//-----------------------------------------------------------------------------------------
+		// Own Interface
+		//-----------------------------------------------------------------------------------------
+
+		//! Get the constant mesh that will be returned.
+		MeshPtr GetValue() const;
+
+		//! Set the constant mesh that will be returned.
+		void SetValue( MeshPtr );
+
+		//! Get the number of layouts defined in this mesh.
+		int GetLayoutCount() const;
+		void SetLayoutCount( int );
+
+		//! Get the node defining a layout of the returned mesh.
+		NodeLayoutPtr GetLayout( int index ) const;
+		void SetLayout( int index, NodeLayoutPtr );
+
+		//-----------------------------------------------------------------------------------------
+		// Interface pattern
+		//-----------------------------------------------------------------------------------------
+		class Private;
+		Private* GetPrivate() const;
+        Node::Private* GetBasePrivate() const override;
+
+	protected:
+
+		//! Forbidden. Manage with the Ptr<> template.
+		~NodeMeshConstant();
+
+	private:
+
+		Private* m_pD;
+
+	};
+
+
+}
diff --git a/Engine/Plugins/Experimental/Mutable/Source/MutableTools/Public/NodeMeshFormat.h b/Engine/Plugins/Experimental/Mutable/Source/MutableTools/Public/NodeMeshFormat.h
new file mode 100644
--- /dev/null
+++ b/Engine/Plugins/Experimental/Mutable/Source/MutableTools/Public/NodeMeshFormat.h
@@ -0,0 +1,94 @@
+// Copyright Epic Games, Inc. All Rights Reserved.
+
+#pragma once
+
+#include "NodeMesh.h"
+
+#include "Mesh.h"
+
+
+namespace mu
+{
+
+	// Forward definitions
+	class NodeMeshFormat;
+	typedef Ptr<NodeMeshFormat> NodeMeshFormatPtr;
+	typedef Ptr<const NodeMeshFormat> NodeMeshFormatPtrConst;
+
+
+	//! This node can change the buffer formats of a mesh vertices, indices and faces.
+	//! \ingroup model
+	class MUTABLETOOLS_API NodeMeshFormat : public NodeMesh
+	{
+	public:
+
+		//-----------------------------------------------------------------------------------------
+		// Life cycle
+		//-----------------------------------------------------------------------------------------
+
+		NodeMeshFormat();
+
+		void SerialiseWrapper(OutputArchive& arch) const override;
+		static void Serialise( const NodeMeshFormat* pNode, OutputArchive& arch );
+		static NodeMeshFormatPtr StaticUnserialise( InputArchive& arch );
+
+		//-----------------------------------------------------------------------------------------
+		// Node Interface
+		//-----------------------------------------------------------------------------------------
+
+		NodePtr Clone( NodeMapPtr pMap = 0 ) const override;
+
+		const NODE_TYPE* GetType() const override;
+		static const NODE_TYPE* GetStaticType();
+
+		int GetInputCount() const override;
+		Node* GetInputNode( int i ) const override;
+		void SetInputNode( int i, NodePtr pNode ) override;
+
+		//-----------------------------------------------------------------------------------------
+		// Own Interface
+		//-----------------------------------------------------------------------------------------
+
+		//! Source mesh to be re-formatted
+		NodeMeshPtr GetSource() const;
+		void SetSource( NodeMesh* );
+
+		//! Get and set the MeshBufferSet that defines the new format for the mesh vertices. These
+		//! buffers don't really contain any data (they have 0 elements) but they define the
+		//! structure. If this is null, the vertex buffers will not be changed.
+		FMeshBufferSet& GetVertexBuffers();
+
+		//! Get and set the MeshBufferSet that defines the new format for the mesh indices. These
+		//! buffers don't really contain any data (they have 0 elements) but they define the
+		//! structure. If this is null, the vertex buffers will not be changed.
+		FMeshBufferSet& GetIndexBuffers();
+
+		//! Get and set the MeshBufferSet that defines the new format for the mesh faces. These
+		//! buffers don't really contain any data (they have 0 elements) but they define the
+		//! structure. If this is null, the vertex buffers will not be changed.
+		FMeshBufferSet& GetFaceBuffers();
+
+		//! Enable or disable the rebuilding of the tangent space of the mesh.
+		//! By default it is disabled.
+		void SetRebuildTangents( bool );
+
+		//-----------------------------------------------------------------------------------------
+		// Interface pattern
+		//-----------------------------------------------------------------------------------------
+		class Private;
+		Private* GetPrivate() const;
+        Node::Private* GetBasePrivate() const override;
+
+	protected:
+
+		//! Forbidden. Manage with the Ptr<> template.
+		~NodeMeshFormat();
+
+	private:
+
+		Private* m_pD;
+
+	};
+
+
+}
diff --git a/Engine/Plugins/Experimental/Mutable/Source/MutableTools/Public/NodeMeshFragment.h b/Engine/Plugins/Experimental/Mutable/Source/MutableTools/Public/NodeMeshFragment.h
new file mode 100644
--- /dev/null
+++ b/Engine/Plugins/Experimental/Mutable/Source/MutableTools/Public/NodeMeshFragment.h
@@ -0,0 +1,97 @@
+// Copyright Epic Games, Inc. All Rights Reserved.
+
+#pragma once
+
+#include "NodeMesh.h"
+
+#include "Mesh.h"
+
+
+namespace mu
+{
+
+	// Forward definitions
+	class NodeMeshFragment;
+	typedef Ptr<NodeMeshFragment> NodeMeshFragmentPtr;
+	typedef Ptr<const NodeMeshFragment> NodeMeshFragmentPtrConst;
+
+
+	//! This node selects a fragment of a mesh, by selecting some of its layout blocks.
+	//! \ingroup model
+	class MUTABLETOOLS_API NodeMeshFragment : public NodeMesh
+	{
+	public:
+
+		//-----------------------------------------------------------------------------------------
+		// Life cycle
+		//-----------------------------------------------------------------------------------------
+
+		NodeMeshFragment();
+
+		void SerialiseWrapper(OutputArchive& arch) const override;
+		static void Serialise( const NodeMeshFragment* pNode, OutputArchive& arch );
+		static NodeMeshFragmentPtr StaticUnserialise( InputArchive& arch );
+
+
+		//-----------------------------------------------------------------------------------------
+		// Node Interface
+		//-----------------------------------------------------------------------------------------
+
+		NodePtr Clone( NodeMapPtr pMap = 0 ) const override;
+
+		const NODE_TYPE* GetType() const override;
+		static const NODE_TYPE* GetStaticType();
+
+		int GetInputCount() const override;
+		Node* GetInputNode( int i ) const override;
+		void SetInputNode( int i, NodePtr pNode ) override;
+
+		//-----------------------------------------------------------------------------------------
+		// Own Interface
+		//-----------------------------------------------------------------------------------------
+
+        typedef enum {
+            FT_LAYOUT_BLOCKS = 0,
+            FT_FACE_GROUP
+        } FRAGMENT_TYPE;
+
+        //! Get the type of mesh fragment to extract.
+        FRAGMENT_TYPE GetFragmentType() const;
+        void SetFragmentType( FRAGMENT_TYPE type );
+
+		//!
+		NodeMeshPtr GetMesh() const;
+		void SetMesh( NodeMeshPtr );
+
+		//! Get the index of the layout to use to extract blocks.
+        int GetLayoutOrGroup() const;
+        void SetLayoutOrGroup( int layoutIndex );
+
+		//! Set the number of layout blocks to extract.
+		void SetBlockCount( int );
+
+		//! Get one of the layout block indices
+		int GetBlock( int i ) const;
+		void SetBlock( int i, int blockIndex );
+
+
+		//-----------------------------------------------------------------------------------------
+		// Interface pattern
+		//-----------------------------------------------------------------------------------------
+		class Private;
+		Private* GetPrivate() const;
+        Node::Private* GetBasePrivate() const override;
+
+	protected:
+
+		//! Forbidden. Manage with the Ptr<> template.
+		~NodeMeshFragment();
+
+	private:
+
+		Private* m_pD;
+
+	};
+
+
+}
diff --git a/Engine/Plugins/Experimental/Mutable/Source/MutableTools/Public/NodeMeshGeometryOperation.h b/Engine/Plugins/Experimental/Mutable/Source/MutableTools/Public/NodeMeshGeometryOperation.h
new file mode 100644
--- /dev/null
+++ b/Engine/Plugins/Experimental/Mutable/Source/MutableTools/Public/NodeMeshGeometryOperation.h
@@ -0,0 +1,86 @@
+// Copyright Epic Games, Inc. All Rights Reserved.
+
+#pragma once
+
+#include "MutableTools/Public/NodeMesh.h"
+#include "MutableRuntime/Public/Mesh.h"
+
+
+namespace mu
+{
+
+	// Forward definitions
+	class NodeScalar;
+	typedef Ptr<NodeScalar> NodeScalarPtr;
+	typedef Ptr<const NodeScalar> NodeScalarPtrConst;
+
+	class NodeMeshGeometryOperation;
+	typedef Ptr<NodeMeshGeometryOperation> NodeMeshGeometryOperationPtr;
+	typedef Ptr<const NodeMeshGeometryOperation> NodeMeshGeometryOperationPtrConst;
+
+
+	//! Node that morphs a base mesh with one or two weighted targets from a sequence.
+	//! \ingroup model
+	class MUTABLETOOLS_API NodeMeshGeometryOperation : public NodeMesh
+	{
+	public:
+
+		//-----------------------------------------------------------------------------------------
+		// Life cycle
+		//-----------------------------------------------------------------------------------------
+
+		NodeMeshGeometryOperation();
+
+		void SerialiseWrapper(OutputArchive& arch) const override;
+		static void Serialise( const NodeMeshGeometryOperation* pNode, OutputArchive& arch );
+		static NodeMeshGeometryOperationPtr StaticUnserialise( InputArchive& arch );
+
+
+		//-----------------------------------------------------------------------------------------
+		// Node Interface
+		//-----------------------------------------------------------------------------------------
+
+		NodePtr Clone( NodeMapPtr pMap = 0 ) const override;
+
+		const NODE_TYPE* GetType() const override;
+		static const NODE_TYPE* GetStaticType();
+
+		int GetInputCount() const override;
+		Node* GetInputNode( int i ) const override;
+		void SetInputNode( int i, NodePtr pNode ) override;
+
+		//-----------------------------------------------------------------------------------------
+		// Own Interface
+		//-----------------------------------------------------------------------------------------
+
+		const NodeMeshPtr& GetMeshA() const;
+		void SetMeshA( const NodeMeshPtr& );
+
+		const NodeMeshPtr& GetMeshB() const;
+		void SetMeshB(const NodeMeshPtr&);
+
+		const NodeScalarPtr& GetScalarA() const;
+		void SetScalarA(const NodeScalarPtr&);
+
+		const NodeScalarPtr& GetScalarB() const;
+		void SetScalarB(const NodeScalarPtr&);
+
+        //-----------------------------------------------------------------------------------------
+		// Interface pattern
+		//-----------------------------------------------------------------------------------------
+		class Private;
+		Private* GetPrivate() const;
+        Node::Private* GetBasePrivate() const override;
+
+	protected:
+
+		//! Forbidden. Manage with the Ptr<> template.
+		~NodeMeshGeometryOperation();
+
+	private:
+
+		Private* m_pD;
+
+	};
+
+}
diff --git a/Engine/Plugins/Experimental/Mutable/Source/MutableTools/Public/NodeMeshInterpolate.h b/Engine/Plugins/Experimental/Mutable/Source/MutableTools/Public/NodeMeshInterpolate.h
new file mode 100644
--- /dev/null
+++ b/Engine/Plugins/Experimental/Mutable/Source/MutableTools/Public/NodeMeshInterpolate.h
@@ -0,0 +1,102 @@
+// Copyright Epic Games, Inc. All Rights Reserved.
+
+#pragma once
+
+#include "NodeMesh.h"
+
+#include "Mesh.h"
+
+
+namespace mu
+{
+
+	// Forward definitions
+	class NodeScalar;
+	typedef Ptr<NodeScalar> NodeScalarPtr;
+	typedef Ptr<const NodeScalar> NodeScalarPtrConst;
+
+	class NodeMeshInterpolate;
+	typedef Ptr<NodeMeshInterpolate> NodeMeshInterpolatePtr;
+	typedef Ptr<const NodeMeshInterpolate> NodeMeshInterpolatePtrConst;
+
+
+	//! Node that morphs (interpolates linearly) between several meshes based on a weight.
+	//! The different input images are uniformly distributed in the 0 to 1 range:
+	//! - if two meshes are set, the first one is 0.0 and the second one is 1.0
+	//! - if three meshes are set, the first one is 0.0 the second one is 0.5 and the third one 1.0
+	//! regardless of the input slots used: If B and C are set, but not A, B will be at weight 0
+	//! and C will be at weight 1
+	//! \ingroup model
+	class MUTABLETOOLS_API NodeMeshInterpolate : public NodeMesh
+	{
+	public:
+
+		//-----------------------------------------------------------------------------------------
+		// Life cycle
+		//-----------------------------------------------------------------------------------------
+
+		NodeMeshInterpolate();
+
+		void SerialiseWrapper(OutputArchive& arch) const override;
+		static void Serialise( const NodeMeshInterpolate* pNode, OutputArchive& arch );
+		static NodeMeshInterpolatePtr StaticUnserialise( InputArchive& arch );
+
+
+		//-----------------------------------------------------------------------------------------
+		// Node Interface
+		//-----------------------------------------------------------------------------------------
+
+		NodePtr Clone( NodeMapPtr pMap = 0 ) const override;
+
+		const NODE_TYPE* GetType() const override;
+		static const NODE_TYPE* GetStaticType();
+
+		int GetInputCount() const override;
+		Node* GetInputNode( int i ) const override;
+		void SetInputNode( int i, NodePtr pNode ) override;
+
+		//-----------------------------------------------------------------------------------------
+		// Own Interface
+		//-----------------------------------------------------------------------------------------
+
+		//! Get the node generating the weight of the interpolation.
+		NodeScalarPtr GetFactor() const;
+		void SetFactor( NodeScalarPtr );
+
+		//! Set the number of target meshes. It will keep the currently set targets and initialise
+		//! the new ones as 0.
+		void SetTargetCount( int );
+		int GetTargetCount() const;
+
+		//! Get the node generating the several mesh targets of the interpolation.
+		NodeMeshPtr GetTarget( int t ) const;
+		void SetTarget( int t, NodeMeshPtr );
+
+		//! Set the number of channels of the mesh that will be interpolated. It will keep the
+		//! currently set targets and initialise the new ones as null. The channels that are not
+		//! interpolated will be copied from the first target.
+		void SetChannelCount( int );
+
+		//! Set a channel to interpolate from the mesh.
+		void SetChannel( int index, MESH_BUFFER_SEMANTIC semantic, int semanticIndex );
+
+		//-----------------------------------------------------------------------------------------
+		// Interface pattern
+		//-----------------------------------------------------------------------------------------
+		class Private;
+		Private* GetPrivate() const;
+        Node::Private* GetBasePrivate() const override;
+
+	protected:
+
+		//! Forbidden. Manage with the Ptr<> template.
+		~NodeMeshInterpolate();
+
+	private:
+
+		Private* m_pD;
+
+	};
+
+
+}
