[mutable] Move Mutable plugin from the restricted folder to the experimental plugins folder.

diff --git a/Engine/Plugins/Experimental/Mutable/Source/MutableTools/Public/NodeMeshMakeMorph.h b/Engine/Plugins/Experimental/Mutable/Source/MutableTools/Public/NodeMeshMakeMorph.h
new file mode 100644
--- /dev/null
+++ b/Engine/Plugins/Experimental/Mutable/Source/MutableTools/Public/NodeMeshMakeMorph.h
@@ -0,0 +1,77 @@
+// Copyright Epic Games, Inc. All Rights Reserved.
+
+#pragma once
+
+#include "NodeMesh.h"
+
+
+namespace mu
+{
+
+	// Forward definitions
+    class NodeMeshMakeMorph;
+    typedef Ptr<NodeMeshMakeMorph> NodeMeshMakeMorphPtr;
+    typedef Ptr<const NodeMeshMakeMorph> NodeMeshMakeMorphPtrConst;
+
+
+	//! This node removes from a mesh A all the faces that are also part of a mesh B.
+	//! \ingroup model
+    class MUTABLETOOLS_API NodeMeshMakeMorph : public NodeMesh
+	{
+	public:
+
+		//-----------------------------------------------------------------------------------------
+		// Life cycle
+		//-----------------------------------------------------------------------------------------
+
+        NodeMeshMakeMorph();
+
+		void SerialiseWrapper(OutputArchive& arch) const override;
+		static void Serialise( const NodeMeshMakeMorph* pNode, OutputArchive& arch );
+        static NodeMeshMakeMorphPtr StaticUnserialise( InputArchive& arch );
+
+		//-----------------------------------------------------------------------------------------
+		// Node Interface
+		//-----------------------------------------------------------------------------------------
+
+		NodePtr Clone( NodeMapPtr pMap = 0 ) const override;
+
+		const NODE_TYPE* GetType() const override;
+		static const NODE_TYPE* GetStaticType();
+
+		int GetInputCount() const override;
+		Node* GetInputNode( int i ) const override;
+		void SetInputNode( int i, NodePtr pNode ) override;
+
+		//-----------------------------------------------------------------------------------------
+		// Own Interface
+		//-----------------------------------------------------------------------------------------
+
+		//!
+		NodeMeshPtr GetBase() const;
+		void SetBase( NodeMesh* p );
+
+		//!
+		NodeMeshPtr GetTarget() const;
+		void SetTarget( NodeMesh* p );
+
+		//-----------------------------------------------------------------------------------------
+		// Interface pattern
+		//-----------------------------------------------------------------------------------------
+		class Private;
+		Private* GetPrivate() const;
+        Node::Private* GetBasePrivate() const override;
+
+	protected:
+
+		//! Forbidden. Manage with the Ptr<> template.
+        ~NodeMeshMakeMorph();
+
+	private:
+
+		Private* m_pD;
+
+	};
+
+
+}
diff --git a/Engine/Plugins/Experimental/Mutable/Source/MutableTools/Public/NodeMeshMorph.h b/Engine/Plugins/Experimental/Mutable/Source/MutableTools/Public/NodeMeshMorph.h
new file mode 100644
--- /dev/null
+++ b/Engine/Plugins/Experimental/Mutable/Source/MutableTools/Public/NodeMeshMorph.h
@@ -0,0 +1,116 @@
+// Copyright Epic Games, Inc. All Rights Reserved.
+
+#pragma once
+
+#include "NodeMesh.h"
+
+#include "Mesh.h"
+
+
+namespace mu
+{
+
+	// Forward definitions
+	class NodeScalar;
+	typedef Ptr<NodeScalar> NodeScalarPtr;
+	typedef Ptr<const NodeScalar> NodeScalarPtrConst;
+
+	class NodeMeshMorph;
+	typedef Ptr<NodeMeshMorph> NodeMeshMorphPtr;
+	typedef Ptr<const NodeMeshMorph> NodeMeshMorphPtrConst;
+
+
+	//! Node that morphs a base mesh with one or two weighted targets from a sequence.
+	//! \ingroup model
+	class MUTABLETOOLS_API NodeMeshMorph : public NodeMesh
+	{
+	public:
+
+		//-----------------------------------------------------------------------------------------
+		// Life cycle
+		//-----------------------------------------------------------------------------------------
+
+		NodeMeshMorph();
+
+		void SerialiseWrapper(OutputArchive& arch) const override;
+		static void Serialise( const NodeMeshMorph* pNode, OutputArchive& arch );
+		static NodeMeshMorphPtr StaticUnserialise( InputArchive& arch );
+
+
+		//-----------------------------------------------------------------------------------------
+		// Node Interface
+		//-----------------------------------------------------------------------------------------
+
+		NodePtr Clone( NodeMapPtr pMap = 0 ) const override;
+
+		const NODE_TYPE* GetType() const override;
+		static const NODE_TYPE* GetStaticType();
+
+		int GetInputCount() const override;
+		Node* GetInputNode( int i ) const override;
+		void SetInputNode( int i, NodePtr pNode ) override;
+
+		//-----------------------------------------------------------------------------------------
+		// Own Interface
+		//-----------------------------------------------------------------------------------------
+
+		//! Get the node generating the weight used to select (and combine) the morphs to apply.
+		NodeScalarPtr GetFactor() const;
+		void SetFactor( NodeScalarPtr );
+
+		//! Get the nodes generating the base mesh to be morphed
+		NodeMeshPtr GetBase() const;
+		void SetBase( NodeMeshPtr );
+
+		//! Set the number of target meshes. It will keep the currently set targets and initialise
+		//! the new ones as 0.
+		void SetMorphCount( int );
+		int GetMorphCount() const;
+
+		//! Get the nodes generating the several morph targets
+		NodeMeshPtr GetMorph( int t ) const;
+		void SetMorph( int t, NodeMeshPtr );
+
+		//! Set the domain of the vertex indices in the morph connections. By default they are 
+		//! absolute as generated by the MakeMeshMorph node. But if they are provided externally, 
+		//! they may be relative to the main mesh connected in the Base node.
+        void SetMorphIndicesAreRelative( bool relative );
+
+		/** Also deform the mesh skeleton. Disabled by default. */
+		void SetReshapeSkeleton(bool);	
+
+		/** Deform all the bones of the skeleton */
+		void SetDeformAllBones(bool);
+
+		/** Deform all the physiscs bodies */
+		void SetDeformAllPhysics(bool);
+
+		/** Deform Mesh Physics Volumes */
+		void SetReshapePhysicsVolumes(bool);
+
+		/** Sets the number of bones that will be deform */
+		void AddBoneToDeform(const char* BoneName);	
+
+		/** Add a Physics Body to deform */
+		void AddPhysicsBodyToDeform(const char* BoneName);
+
+        //-----------------------------------------------------------------------------------------
+		// Interface pattern
+		//-----------------------------------------------------------------------------------------
+		class Private;
+		Private* GetPrivate() const;
+        Node::Private* GetBasePrivate() const override;
+
+	protected:
+
+		//! Forbidden. Manage with the Ptr<> template.
+		~NodeMeshMorph();
+
+	private:
+
+		Private* m_pD;
+
+	};
+
+
+}
diff --git a/Engine/Plugins/Experimental/Mutable/Source/MutableTools/Public/NodeMeshReshape.h b/Engine/Plugins/Experimental/Mutable/Source/MutableTools/Public/NodeMeshReshape.h
new file mode 100644
--- /dev/null
+++ b/Engine/Plugins/Experimental/Mutable/Source/MutableTools/Public/NodeMeshReshape.h
@@ -0,0 +1,94 @@
+// Copyright Epic Games, Inc. All Rights Reserved.
+
+#pragma once
+
+#include "MutableTools/Public/NodeMesh.h"
+#include "MutableRuntime/Public/Mesh.h"
+
+
+namespace mu
+{
+
+	//! Node that morphs a base mesh with one or two weighted targets from a sequence.
+	//! \ingroup model
+	class MUTABLETOOLS_API NodeMeshReshape : public NodeMesh
+	{
+	public:
+
+		//-----------------------------------------------------------------------------------------
+		// Life cycle
+		//-----------------------------------------------------------------------------------------
+
+		NodeMeshReshape();
+
+		void SerialiseWrapper(OutputArchive& arch) const override;
+		static void Serialise( const NodeMeshReshape* pNode, OutputArchive& arch );
+		static Ptr<NodeMeshReshape> StaticUnserialise( InputArchive& arch );
+
+
+		//-----------------------------------------------------------------------------------------
+		// Node Interface
+		//-----------------------------------------------------------------------------------------
+
+		NodePtr Clone( NodeMapPtr pMap = 0 ) const override;
+
+		const NODE_TYPE* GetType() const override;
+		static const NODE_TYPE* GetStaticType();
+
+		int GetInputCount() const override;
+		Node* GetInputNode( int i ) const override;
+		void SetInputNode( int i, NodePtr pNode ) override;
+
+		//-----------------------------------------------------------------------------------------
+		// Own Interface
+		//-----------------------------------------------------------------------------------------
+
+		const NodeMeshPtr& GetBaseMesh() const;
+		void SetBaseMesh( const NodeMeshPtr& );
+
+		const NodeMeshPtr& GetBaseShape() const;
+		void SetBaseShape(const NodeMeshPtr&);
+
+		const NodeMeshPtr& GetTargetShape() const;
+		void SetTargetShape(const NodeMeshPtr&);
+
+		/** Also deform the mesh skeleton. Disabled by default. */
+		void SetReshapeSkeleton(bool);
+		
+		/** Search for clusters of rigid parts and not deform them. */
+		void SetEnableRigidParts(bool);
+
+		/** Deform all the bones of the skeleton */
+		void SetDeformAllBones(bool);
+
+		/** Deform all the physiscs bodies */
+		void SetDeformAllPhysics(bool);
+
+		/** Deform Mesh Physics Volumes */
+		void SetReshapePhysicsVolumes(bool);
+
+		/** Sets the number of bones that will be deform */
+		void AddBoneToDeform(const char* BoneName);
+	
+		/** Add a Physics Body to deform */
+		void AddPhysicsBodyToDeform(const char* BoneName);
+        
+		//-----------------------------------------------------------------------------------------
+		// Interface pattern
+		//-----------------------------------------------------------------------------------------
+		class Private;
+		Private* GetPrivate() const;
+        Node::Private* GetBasePrivate() const override;
+
+	protected:
+
+		//! Forbidden. Manage with the Ptr<> template.
+		~NodeMeshReshape();
+
+	private:
+
+		Private* m_pD;
+
+	};
+
+}
diff --git a/Engine/Plugins/Experimental/Mutable/Source/MutableTools/Public/NodeMeshSubtract.h b/Engine/Plugins/Experimental/Mutable/Source/MutableTools/Public/NodeMeshSubtract.h
new file mode 100644
--- /dev/null
+++ b/Engine/Plugins/Experimental/Mutable/Source/MutableTools/Public/NodeMeshSubtract.h
@@ -0,0 +1,77 @@
+// Copyright Epic Games, Inc. All Rights Reserved.
+
+#pragma once
+
+#include "NodeMesh.h"
+
+
+namespace mu
+{
+
+	// Forward definitions
+	class NodeMeshSubtract;
+	typedef Ptr<NodeMeshSubtract> NodeMeshSubtractPtr;
+	typedef Ptr<const NodeMeshSubtract> NodeMeshSubtractPtrConst;
+
+
+	//! This node removes from a mesh A all the faces that are also part of a mesh B.
+	//! \ingroup model
+	class MUTABLETOOLS_API NodeMeshSubtract : public NodeMesh
+	{
+	public:
+
+		//-----------------------------------------------------------------------------------------
+		// Life cycle
+		//-----------------------------------------------------------------------------------------
+
+		NodeMeshSubtract();
+
+		void SerialiseWrapper(OutputArchive& arch) const override;
+		static void Serialise( const NodeMeshSubtract* pNode, OutputArchive& arch );
+		static NodeMeshSubtractPtr StaticUnserialise( InputArchive& arch );
+
+		//-----------------------------------------------------------------------------------------
+		// Node Interface
+		//-----------------------------------------------------------------------------------------
+
+        NodePtr Clone( NodeMapPtr pMap = 0 ) const override;
+
+		const NODE_TYPE* GetType() const override;
+		static const NODE_TYPE* GetStaticType();
+
+		int GetInputCount() const override;
+		Node* GetInputNode( int i ) const override;
+		void SetInputNode( int i, NodePtr pNode ) override;
+
+		//-----------------------------------------------------------------------------------------
+		// Own Interface
+		//-----------------------------------------------------------------------------------------
+
+		//!
+		NodeMeshPtr GetA() const;
+		void SetA( NodeMesh* p );
+
+		//!
+		NodeMeshPtr GetB() const;
+		void SetB( NodeMesh* p );
+
+		//-----------------------------------------------------------------------------------------
+		// Interface pattern
+		//-----------------------------------------------------------------------------------------
+		class Private;
+		Private* GetPrivate() const;
+        Node::Private* GetBasePrivate() const override;
+
+	protected:
+
+		//! Forbidden. Manage with the Ptr<> template.
+		~NodeMeshSubtract();
+
+	private:
+
+		Private* m_pD;
+
+	};
+
+
+}
diff --git a/Engine/Plugins/Experimental/Mutable/Source/MutableTools/Public/NodeMeshSwitch.h b/Engine/Plugins/Experimental/Mutable/Source/MutableTools/Public/NodeMeshSwitch.h
new file mode 100644
--- /dev/null
+++ b/Engine/Plugins/Experimental/Mutable/Source/MutableTools/Public/NodeMeshSwitch.h
@@ -0,0 +1,89 @@
+// Copyright Epic Games, Inc. All Rights Reserved.
+
+#pragma once
+
+#include "NodeMesh.h"
+
+#include "Mesh.h"
+
+
+namespace mu
+{
+
+	// Forward definitions
+	class NodeScalar;
+	typedef Ptr<NodeScalar> NodeScalarPtr;
+	typedef Ptr<const NodeScalar> NodeScalarPtrConst;
+
+	class NodeMeshSwitch;
+	typedef Ptr<NodeMeshSwitch> NodeMeshSwitchPtr;
+	typedef Ptr<const NodeMeshSwitch> NodeMeshSwitchPtrConst;
+
+
+	//! This node selects an output Mesh from a set of input Meshs based on a parameter.
+	//! \ingroup model
+	class MUTABLETOOLS_API NodeMeshSwitch : public NodeMesh
+	{
+	public:
+
+		//-----------------------------------------------------------------------------------------
+		// Life cycle
+		//-----------------------------------------------------------------------------------------
+
+		NodeMeshSwitch();
+
+		void SerialiseWrapper(OutputArchive& arch) const override;
+		static void Serialise( const NodeMeshSwitch* pNode, OutputArchive& arch );
+		static NodeMeshSwitchPtr StaticUnserialise( InputArchive& arch );
+
+
+		//-----------------------------------------------------------------------------------------
+		// Node Interface
+		//-----------------------------------------------------------------------------------------
+
+        NodePtr Clone( NodeMapPtr pMap = 0 ) const override;
+
+		const NODE_TYPE* GetType() const override;
+		static const NODE_TYPE* GetStaticType();
+
+		int GetInputCount() const override;
+		Node* GetInputNode( int i ) const override;
+		void SetInputNode( int i, NodePtr pNode ) override;
+
+		//-----------------------------------------------------------------------------------------
+		// Own Interface
+		//-----------------------------------------------------------------------------------------
+
+		//! Get the node generating the parameter used to select the option.
+		NodeScalarPtr GetParameter() const;
+		void SetParameter( NodeScalarPtr );
+
+		//! Set the number of option Meshs. It will keep the currently set targets and initialise
+		//! the new ones as null.
+		void SetOptionCount( int );
+
+		//! Get the node generating the t-th option Mesh.
+		NodeMeshPtr GetOption( int t ) const;
+		void SetOption( int t, NodeMeshPtr );
+
+
+		//-----------------------------------------------------------------------------------------
+		// Interface pattern
+		//-----------------------------------------------------------------------------------------
+		class Private;
+		Private* GetPrivate() const;
+        Node::Private* GetBasePrivate() const override;
+
+	protected:
+
+		//! Forbidden. Manage with the Ptr<> template.
+		~NodeMeshSwitch();
+
+	private:
+
+		Private* m_pD;
+
+	};
+
+
+}
diff --git a/Engine/Plugins/Experimental/Mutable/Source/MutableTools/Public/NodeMeshTable.h b/Engine/Plugins/Experimental/Mutable/Source/MutableTools/Public/NodeMeshTable.h
new file mode 100644
--- /dev/null
+++ b/Engine/Plugins/Experimental/Mutable/Source/MutableTools/Public/NodeMeshTable.h
@@ -0,0 +1,98 @@
+// Copyright Epic Games, Inc. All Rights Reserved.
+
+#pragma once
+
+#include "NodeMesh.h"
+
+
+namespace mu
+{
+
+	// Forward definitions
+	class NodeMeshTable;
+	typedef Ptr<NodeMeshTable> NodeMeshTablePtr;
+	typedef Ptr<const NodeMeshTable> NodeMeshTablePtrConst;
+
+	class NodeLayout;
+	typedef Ptr<NodeLayout> NodeLayoutPtr;
+	typedef Ptr<const NodeLayout> NodeLayoutPtrConst;
+
+	class Table;
+	typedef Ptr<Table> TablePtr;
+	typedef Ptr<const Table> TablePtrConst;
+
+
+	//! This node provides the meshes stored in the column of a table.
+	//! \ingroup transform
+	class MUTABLETOOLS_API NodeMeshTable : public NodeMesh
+	{
+	public:
+
+		//-----------------------------------------------------------------------------------------
+		// Life cycle
+		//-----------------------------------------------------------------------------------------
+
+		NodeMeshTable();
+
+		void SerialiseWrapper(OutputArchive& arch) const override;
+		static void Serialise( const NodeMeshTable* pNode, OutputArchive& arch );
+		static NodeMeshTablePtr StaticUnserialise( InputArchive& arch );
+
+
+		//-----------------------------------------------------------------------------------------
+		// Node Interface
+		//-----------------------------------------------------------------------------------------
+
+        NodePtr Clone( NodeMapPtr pMap = 0 ) const override;
+
+        const NODE_TYPE* GetType() const override;
+		static const NODE_TYPE* GetStaticType();
+
+        int GetInputCount() const override;
+        Node* GetInputNode( int i ) const override;
+        void SetInputNode( int i, NodePtr pNode ) override;
+
+		//-----------------------------------------------------------------------------------------
+		// Own Interface
+		//-----------------------------------------------------------------------------------------
+
+		//! Set the name of the implicit table parameter.
+		void SetParameterName( const char* strName );
+
+		//!
+		TablePtr GetTable() const;
+		void SetTable( TablePtr );
+
+		//!
+		const char* GetColumn() const;
+		void SetColumn( const char* strName );
+
+		//! Get the number of layouts defined in the meshes on this column.
+		int GetLayoutCount() const;
+		void SetLayoutCount( int );
+
+		//! Get the node defining a layout of the meshes on this column.
+		NodeLayoutPtr GetLayout( int index ) const;
+		void SetLayout( int index, NodeLayoutPtr );
+
+		//-----------------------------------------------------------------------------------------
+		// Interface pattern
+		//-----------------------------------------------------------------------------------------
+		class Private;
+		Private* GetPrivate() const;
+        Node::Private* GetBasePrivate() const override;
+
+
+	protected:
+
+		//! Forbidden. Manage with the Ptr<> template.
+		~NodeMeshTable();
+
+	private:
+
+		Private* m_pD;
+
+	};
+
+
+}
diff --git a/Engine/Plugins/Experimental/Mutable/Source/MutableTools/Public/NodeMeshTangents.h b/Engine/Plugins/Experimental/Mutable/Source/MutableTools/Public/NodeMeshTangents.h
new file mode 100644
--- /dev/null
+++ b/Engine/Plugins/Experimental/Mutable/Source/MutableTools/Public/NodeMeshTangents.h
@@ -0,0 +1,73 @@
+// Copyright Epic Games, Inc. All Rights Reserved.
+
+#pragma once
+
+#include "NodeMesh.h"
+
+
+namespace mu
+{
+
+	// Forward definitions
+	class NodeMeshTangents;
+	typedef Ptr<NodeMeshTangents> NodeMeshTangentsPtr;
+	typedef Ptr<const NodeMeshTangents> NodeMeshTangentsPtrConst;
+
+
+	//! This node rebuilds the tangents and binormals of the source mesh.
+	//! \ingroup model
+	class MUTABLETOOLS_API NodeMeshTangents : public NodeMesh
+	{
+	public:
+
+		//-----------------------------------------------------------------------------------------
+		// Life cycle
+		//-----------------------------------------------------------------------------------------
+
+		NodeMeshTangents();
+
+		void SerialiseWrapper(OutputArchive& arch) const override;
+		static void Serialise( const NodeMeshTangents* pNode, OutputArchive& arch );
+		static NodeMeshTangentsPtr StaticUnserialise( InputArchive& arch );
+
+		//-----------------------------------------------------------------------------------------
+		// Node Interface
+		//-----------------------------------------------------------------------------------------
+
+        NodePtr Clone( NodeMapPtr pMap = 0 ) const override;
+
+        const NODE_TYPE* GetType() const override;
+		static const NODE_TYPE* GetStaticType();
+
+        int GetInputCount() const override;
+        Node* GetInputNode( int i ) const override;
+        void SetInputNode( int i, NodePtr pNode ) override;
+
+		//-----------------------------------------------------------------------------------------
+		// Own Interface
+		//-----------------------------------------------------------------------------------------
+
+		//!
+		NodeMeshPtr GetSource() const;
+		void SetSource( NodeMesh* p );
+
+		//-----------------------------------------------------------------------------------------
+		// Interface pattern
+		//-----------------------------------------------------------------------------------------
+		class Private;
+		Private* GetPrivate() const;
+        Node::Private* GetBasePrivate() const override;
+
+	protected:
+
+		//! Forbidden. Manage with the Ptr<> template.
+		~NodeMeshTangents();
+
+	private:
+
+		Private* m_pD;
+
+	};
+
+
+}
diff --git a/Engine/Plugins/Experimental/Mutable/Source/MutableTools/Public/NodeMeshTransform.h b/Engine/Plugins/Experimental/Mutable/Source/MutableTools/Public/NodeMeshTransform.h
new file mode 100644
--- /dev/null
+++ b/Engine/Plugins/Experimental/Mutable/Source/MutableTools/Public/NodeMeshTransform.h
@@ -0,0 +1,79 @@
+// Copyright Epic Games, Inc. All Rights Reserved.
+
+#pragma once
+
+#include "NodeMesh.h"
+
+#include "Mesh.h"
+
+
+namespace mu
+{
+
+	// Forward definitions
+    class NodeMeshTransform;
+    typedef Ptr<NodeMeshTransform> NodeMeshTransformPtr;
+    typedef Ptr<const NodeMeshTransform> NodeMeshTransformPtrConst;
+
+
+    //! This node applies a geometric transform represented by a 4x4 matrix to a mesh
+	//! \ingroup model
+    class MUTABLETOOLS_API NodeMeshTransform : public NodeMesh
+	{
+	public:
+
+		//-----------------------------------------------------------------------------------------
+		// Life cycle
+		//-----------------------------------------------------------------------------------------
+
+        NodeMeshTransform();
+
+		void SerialiseWrapper(OutputArchive& arch) const override;
+		static void Serialise( const NodeMeshTransform* pNode, OutputArchive& arch );
+        static NodeMeshTransformPtr StaticUnserialise( InputArchive& arch );
+
+		//-----------------------------------------------------------------------------------------
+		// Node Interface
+		//-----------------------------------------------------------------------------------------
+
+        NodePtr Clone( NodeMapPtr pMap = 0 ) const override;
+
+        const NODE_TYPE* GetType() const override;
+		static const NODE_TYPE* GetStaticType();
+
+        int GetInputCount() const override;
+        Node* GetInputNode( int i ) const override;
+        void SetInputNode( int i, NodePtr pNode ) override;
+
+		//-----------------------------------------------------------------------------------------
+		// Own Interface
+		//-----------------------------------------------------------------------------------------
+
+		//! Source mesh to be re-formatted
+        NodeMeshPtr GetSource() const;
+		void SetSource( NodeMesh* );
+
+        //! \param pMat4x4 pointer to 16 floats representing the 4x4 matrix to apply
+        void SetTransform( const float* pMat4x4 );
+        void GetTransform( float* pMat4x4 ) const;
+
+		//-----------------------------------------------------------------------------------------
+		// Interface pattern
+		//-----------------------------------------------------------------------------------------
+		class Private;
+		Private* GetPrivate() const;
+        Node::Private* GetBasePrivate() const override;
+
+	protected:
+
+		//! Forbidden. Manage with the Ptr<> template.
+        ~NodeMeshTransform();
+
+	private:
+
+		Private* m_pD;
+
+	};
+
+
+}
diff --git a/Engine/Plugins/Experimental/Mutable/Source/MutableTools/Public/NodeMeshVariation.h b/Engine/Plugins/Experimental/Mutable/Source/MutableTools/Public/NodeMeshVariation.h
new file mode 100644
--- /dev/null
+++ b/Engine/Plugins/Experimental/Mutable/Source/MutableTools/Public/NodeMeshVariation.h
@@ -0,0 +1,83 @@
+// Copyright Epic Games, Inc. All Rights Reserved.
+
+#pragma once
+
+#include "NodeMesh.h"
+
+
+namespace mu
+{
+
+    // Forward definitions
+    class NodeMeshVariation;
+    typedef Ptr<NodeMeshVariation> NodeMeshVariationPtr;
+    typedef Ptr<const NodeMeshVariation> NodeMeshVariationPtrConst;
+
+
+    //!
+    //! \ingroup model
+    class MUTABLETOOLS_API NodeMeshVariation : public NodeMesh
+    {
+    public:
+        //-----------------------------------------------------------------------------------------
+        // Life cycle
+        //-----------------------------------------------------------------------------------------
+
+        NodeMeshVariation();
+
+		void SerialiseWrapper(OutputArchive& arch) const override;
+		static void Serialise( const NodeMeshVariation* pNode, OutputArchive& arch );
+        static NodeMeshVariationPtr StaticUnserialise( InputArchive& arch );
+
+        //-----------------------------------------------------------------------------------------
+        // Node Interface
+        //-----------------------------------------------------------------------------------------
+
+        NodePtr Clone( NodeMapPtr pMap = 0 ) const override;
+
+        const NODE_TYPE* GetType() const override;
+        static const NODE_TYPE* GetStaticType();
+
+        virtual int GetInputCount() const override;
+        virtual Node* GetInputNode( int i ) const override;
+        void SetInputNode( int i, NodePtr pNode ) override;
+
+        //-----------------------------------------------------------------------------------------
+        // Own Interface
+        //-----------------------------------------------------------------------------------------
+
+        //!
+        void SetDefaultMesh( NodeMesh* Mesh );
+
+        //! Set the number of tags to consider in this variation
+        void SetVariationCount( int count );
+
+        //!
+        int GetVariationCount() const;
+
+        //! Set the tag or state name that will enable a specific vartiation
+        void SetVariationTag( int index, const char* strTag );
+
+        //!
+        void SetVariationMesh( int index, NodeMesh* Mesh );
+
+        //!}
+
+
+        //-----------------------------------------------------------------------------------------
+        // Interface pattern
+        //-----------------------------------------------------------------------------------------
+        class Private;
+        Private* GetPrivate() const;
+        Node::Private* GetBasePrivate() const override;
+
+    protected:
+        //! Forbidden. Manage with the Ptr<> template.
+        ~NodeMeshVariation();
+
+    private:
+        Private* m_pD;
+    };
+
+
+} // namespace mu
diff --git a/Engine/Plugins/Experimental/Mutable/Source/MutableTools/Public/NodeModifier.h b/Engine/Plugins/Experimental/Mutable/Source/MutableTools/Public/NodeModifier.h
new file mode 100644
--- /dev/null
+++ b/Engine/Plugins/Experimental/Mutable/Source/MutableTools/Public/NodeModifier.h
@@ -0,0 +1,84 @@
+// Copyright Epic Games, Inc. All Rights Reserved.
+
+#pragma once
+
+#include "Node.h"
+
+
+namespace mu
+{
+
+	// Forward definitions
+	class NodeModifier;
+	typedef Ptr<NodeModifier> NodeModifierPtr;
+	typedef Ptr<const NodeModifier> NodeModifierConst;
+
+
+	//! This class is the parent of all nodes that output a component.
+	//! \ingroup model
+	class MUTABLETOOLS_API NodeModifier : public Node
+	{
+	public:
+
+		// Possible subclasses
+		enum class EType : uint8
+		{
+			MeshClipMorphPlane = 0,
+			MeshClipWithMesh = 1,
+			MeshClipDeform = 2,
+			
+			None
+		};
+
+		//-----------------------------------------------------------------------------------------
+		// Life cycle
+		//-----------------------------------------------------------------------------------------
+
+		static void Serialise( const NodeModifier* pNode, OutputArchive& arch );
+		static NodeModifierPtr StaticUnserialise( InputArchive& arch );
+
+		//-----------------------------------------------------------------------------------------
+        // Node interface
+		//-----------------------------------------------------------------------------------------
+
+		const NODE_TYPE* GetType() const override;
+		static const NODE_TYPE* GetStaticType();
+
+        //-----------------------------------------------------------------------------------------
+        // Own interface
+        //-----------------------------------------------------------------------------------------
+
+        //! \name Tags
+        //! \{
+
+        //! Add a tag to the surface, which will be affected by modifier nodes with the same tag
+        void AddTag(const char* tagName);
+
+        //! Get the number of tags added to the Surface.
+        int GetTagCount() const;
+
+        //! Get a tag string from an index (0 to GetTagCount-1)
+        const char* GetTag( int ) const;
+
+        //! \}
+
+		//! Set the stage to apply this modifier in. Default is before normal operations.
+		void SetStage( bool bBeforeNormalOperation );
+
+		//-----------------------------------------------------------------------------------------
+		// Interface pattern
+		//-----------------------------------------------------------------------------------------
+		class Private;
+
+	protected:
+
+		//! Forbidden. Manage with the Ptr<> template.
+		inline ~NodeModifier() {}
+
+		//!
+		EType Type = EType::None;
+
+	};
+
+
+}
diff --git a/Engine/Plugins/Experimental/Mutable/Source/MutableTools/Public/NodeModifierMeshClipDeform.h b/Engine/Plugins/Experimental/Mutable/Source/MutableTools/Public/NodeModifierMeshClipDeform.h
new file mode 100644
--- /dev/null
+++ b/Engine/Plugins/Experimental/Mutable/Source/MutableTools/Public/NodeModifierMeshClipDeform.h
@@ -0,0 +1,77 @@
+// Copyright Epic Games, Inc. All Rights Reserved.
+
+#pragma once
+
+#include "NodeModifier.h"
+
+namespace mu
+{
+
+	// Forward definitions
+	class NodeModifierMeshClipDeform;
+	typedef Ptr<NodeModifierMeshClipDeform> NodeModifierMeshClipDeformPtr;
+	typedef Ptr<const NodeModifierMeshClipDeform> NodeModifierMeshClipDeformPtrConst;
+
+	class NodeMesh;
+	typedef Ptr<NodeMesh> NodeMeshPtr;
+	typedef Ptr<const NodeMesh> NodeMeshPtrConst;
+
+	enum class EShapeBindingMethod : uint32;
+	
+	//! This node makes a new component from several meshes and images.
+	//! \ingroup model
+	class MUTABLETOOLS_API NodeModifierMeshClipDeform : public NodeModifier
+	{
+	public:
+
+		//-----------------------------------------------------------------------------------------
+		// Life cycle
+		//-----------------------------------------------------------------------------------------
+
+		NodeModifierMeshClipDeform();
+
+		void SerialiseWrapper(OutputArchive& arch) const override;
+		static void Serialise( const NodeModifierMeshClipDeform* pNode, OutputArchive& arch );
+		static NodeModifierMeshClipDeformPtr StaticUnserialise( InputArchive& arch );
+
+		//-----------------------------------------------------------------------------------------
+        // Node interface
+		//-----------------------------------------------------------------------------------------
+
+        NodePtr Clone( NodeMapPtr pMap = 0 ) const override;
+
+        const NODE_TYPE* GetType() const override;
+		static const NODE_TYPE* GetStaticType();
+
+        int GetInputCount() const override;
+        Node* GetInputNode( int i ) const override;
+        void SetInputNode( int i, NodePtr pNode ) override;
+
+		//-----------------------------------------------------------------------------------------
+        // Own interface
+		//-----------------------------------------------------------------------------------------
+
+		//! \param 
+		void SetClipMesh( NodeMesh* InClipMesh);
+		void SetBindingMethod(EShapeBindingMethod BindingMethod);
+	
+		//-----------------------------------------------------------------------------------------
+		// Interface pattern
+		//-----------------------------------------------------------------------------------------
+		class Private;
+		Private* GetPrivate() const;
+        Node::Private* GetBasePrivate() const override;
+
+	protected:
+
+		//! Forbidden. Manage with the Ptr<> template.
+		~NodeModifierMeshClipDeform();
+
+	private:
+
+		Private* m_pD;
+
+	};
+
+
+}
diff --git a/Engine/Plugins/Experimental/Mutable/Source/MutableTools/Public/NodeModifierMeshClipMorphPlane.h b/Engine/Plugins/Experimental/Mutable/Source/MutableTools/Public/NodeModifierMeshClipMorphPlane.h
new file mode 100644
--- /dev/null
+++ b/Engine/Plugins/Experimental/Mutable/Source/MutableTools/Public/NodeModifierMeshClipMorphPlane.h
@@ -0,0 +1,82 @@
+// Copyright Epic Games, Inc. All Rights Reserved.
+
+#pragma once
+
+#include "NodeModifier.h"
+
+
+namespace mu
+{
+
+	// Forward definitions
+	class NodeModifierMeshClipMorphPlane;
+	typedef Ptr<NodeModifierMeshClipMorphPlane> NodeModifierMeshClipMorphPlanePtr;
+	typedef Ptr<const NodeModifierMeshClipMorphPlane> NodeModifierMeshClipMorphPlanePtrConst;
+
+
+	//! This node makes a new component from several meshes and images.
+	//! \ingroup model
+	class MUTABLETOOLS_API NodeModifierMeshClipMorphPlane : public NodeModifier
+	{
+	public:
+
+		//-----------------------------------------------------------------------------------------
+		// Life cycle
+		//-----------------------------------------------------------------------------------------
+
+		NodeModifierMeshClipMorphPlane();
+
+		void SerialiseWrapper(OutputArchive& arch) const override;
+		static void Serialise( const NodeModifierMeshClipMorphPlane* pNode, OutputArchive& arch );
+		static NodeModifierMeshClipMorphPlanePtr StaticUnserialise( InputArchive& arch );
+
+		//-----------------------------------------------------------------------------------------
+        // Node interface
+		//-----------------------------------------------------------------------------------------
+
+        NodePtr Clone( NodeMapPtr pMap = 0 ) const override;
+
+        const NODE_TYPE* GetType() const override;
+		static const NODE_TYPE* GetStaticType();
+
+        int GetInputCount() const override;
+        Node* GetInputNode( int i ) const override;
+        void SetInputNode( int i, NodePtr pNode ) override;
+
+		//-----------------------------------------------------------------------------------------
+        // Own interface
+		//-----------------------------------------------------------------------------------------
+
+		//! \param 
+		void SetPlane(float centerX, float centerY, float centerZ, float normalX, float normalY, float normalZ);
+		void SetParams(float dist, float factor);
+		void SetMorphEllipse(float radius1, float radius2, float rotation);
+
+		//! Define an axis-aligned box that will select the vertices to be morphed.
+		//! Only one of Box or Bone Hierarchy can be used (the last one set)
+		void SetVertexSelectionBox(float centerX, float centerY, float centerZ, float radiusX, float radiusY, float radiusZ);
+
+		//! Define the root bone of the subhierarchy of the mesh that will be affected.
+		//! Only one of Box or Bone Hierarchy can be used (the last one set)
+		void SetVertexSelectionBone(const char* strBoneName, float maxEffectRadius);
+
+		//-----------------------------------------------------------------------------------------
+		// Interface pattern
+		//-----------------------------------------------------------------------------------------
+		class Private;
+		Private* GetPrivate() const;
+        Node::Private* GetBasePrivate() const override;
+
+	protected:
+
+		//! Forbidden. Manage with the Ptr<> template.
+		~NodeModifierMeshClipMorphPlane();
+
+	private:
+
+		Private* m_pD;
+
+	};
+
+
+}
diff --git a/Engine/Plugins/Experimental/Mutable/Source/MutableTools/Public/NodeModifierMeshClipWithMesh.h b/Engine/Plugins/Experimental/Mutable/Source/MutableTools/Public/NodeModifierMeshClipWithMesh.h
new file mode 100644
--- /dev/null
+++ b/Engine/Plugins/Experimental/Mutable/Source/MutableTools/Public/NodeModifierMeshClipWithMesh.h
@@ -0,0 +1,76 @@
+// Copyright Epic Games, Inc. All Rights Reserved.
+
+#pragma once
+
+#include "NodeModifier.h"
+
+
+namespace mu
+{
+
+	// Forward definitions
+	class NodeModifierMeshClipWithMesh;
+	typedef Ptr<NodeModifierMeshClipWithMesh> NodeModifierMeshClipWithMeshPtr;
+	typedef Ptr<const NodeModifierMeshClipWithMesh> NodeModifierMeshClipWithMeshPtrConst;
+
+	class NodeMesh;
+	typedef Ptr<NodeMesh> NodeMeshPtr;
+	typedef Ptr<const NodeMesh> NodeMeshPtrConst;
+
+
+	//! This node makes a new component from several meshes and images.
+	//! \ingroup model
+	class MUTABLETOOLS_API NodeModifierMeshClipWithMesh : public NodeModifier
+	{
+	public:
+
+		//-----------------------------------------------------------------------------------------
+		// Life cycle
+		//-----------------------------------------------------------------------------------------
+
+		NodeModifierMeshClipWithMesh();
+
+		void SerialiseWrapper(OutputArchive& arch) const override;
+		static void Serialise( const NodeModifierMeshClipWithMesh* pNode, OutputArchive& arch );
+		static NodeModifierMeshClipWithMeshPtr StaticUnserialise( InputArchive& arch );
+
+		//-----------------------------------------------------------------------------------------
+        // Node interface
+		//-----------------------------------------------------------------------------------------
+
+        NodePtr Clone( NodeMapPtr pMap = 0 ) const override;
+
+        const NODE_TYPE* GetType() const override;
+		static const NODE_TYPE* GetStaticType();
+
+        int GetInputCount() const override;
+        Node* GetInputNode( int i ) const override;
+        void SetInputNode( int i, NodePtr pNode ) override;
+
+		//-----------------------------------------------------------------------------------------
+        // Own interface
+		//-----------------------------------------------------------------------------------------
+
+		//! \param 
+		void SetClipMesh(NodeMesh*);
+
+		//-----------------------------------------------------------------------------------------
+		// Interface pattern
+		//-----------------------------------------------------------------------------------------
+		class Private;
+		Private* GetPrivate() const;
+        Node::Private* GetBasePrivate() const override;
+
+	protected:
+
+		//! Forbidden. Manage with the Ptr<> template.
+		~NodeModifierMeshClipWithMesh();
+
+	private:
+
+		Private* m_pD;
+
+	};
+
+
+}
diff --git a/Engine/Plugins/Experimental/Mutable/Source/MutableTools/Public/NodeObject.h b/Engine/Plugins/Experimental/Mutable/Source/MutableTools/Public/NodeObject.h
new file mode 100644
--- /dev/null
+++ b/Engine/Plugins/Experimental/Mutable/Source/MutableTools/Public/NodeObject.h
@@ -0,0 +1,80 @@
+// Copyright Epic Games, Inc. All Rights Reserved.
+
+#pragma once
+
+#include "MutableTools/Public/Node.h"
+
+
+namespace mu
+{
+
+	// Forward definitions
+	class NodeObject;
+	typedef Ptr<NodeObject> NodeObjectPtr;
+	typedef Ptr<const NodeObject> NodeObjectPtrConst;
+
+
+    //! %Base class for any node that outputs an object.
+	//! \ingroup model
+	class MUTABLETOOLS_API NodeObject : public Node
+	{
+	public:
+
+		// Possible subclasses
+		enum class EType : uint8
+		{
+			New = 0,
+			Group = 1,
+			State = 2,
+
+			None
+		};
+
+		//-----------------------------------------------------------------------------------------
+		// Life cycle
+		//-----------------------------------------------------------------------------------------
+
+		static void Serialise( const NodeObject* pNode, OutputArchive& arch );
+		static NodeObjectPtr StaticUnserialise( InputArchive& arch );
+
+
+		//-----------------------------------------------------------------------------------------
+		// Node Interface
+		//-----------------------------------------------------------------------------------------
+
+		const NODE_TYPE* GetType() const override;
+		static const NODE_TYPE* GetStaticType();
+
+		//-----------------------------------------------------------------------------------------
+		// Own interface
+		//-----------------------------------------------------------------------------------------
+
+		//! Get the name of the object.
+		virtual const char* GetName() const = 0;
+
+		//! Set the name of the object.
+		virtual void SetName( const char* strName ) = 0;
+
+		//! Get the uid of the object.
+		virtual const char* GetUid() const = 0;
+
+		//! Set the uid of the object.
+		virtual void SetUid( const char* strUid ) = 0;
+
+		//-----------------------------------------------------------------------------------------
+		// Interface pattern
+		//-----------------------------------------------------------------------------------------
+		class Private;
+
+	protected:
+
+		//! Forbidden. Manage with the Ptr<> template.
+		inline ~NodeObject() {}
+
+		//!
+		EType Type = EType::None;
+
+	};
+
+
+}
diff --git a/Engine/Plugins/Experimental/Mutable/Source/MutableTools/Public/NodeObjectGroup.h b/Engine/Plugins/Experimental/Mutable/Source/MutableTools/Public/NodeObjectGroup.h
new file mode 100644
--- /dev/null
+++ b/Engine/Plugins/Experimental/Mutable/Source/MutableTools/Public/NodeObjectGroup.h
@@ -0,0 +1,117 @@
+// Copyright Epic Games, Inc. All Rights Reserved.
+
+#pragma once
+
+#include "NodeObject.h"
+
+
+namespace mu
+{
+
+	// Forward definitions
+	class NodeObjectGroup;
+	typedef Ptr<NodeObjectGroup> NodeObjectGroupPtr;
+	typedef Ptr<const NodeObjectGroup>NodeObjectGroupPtrConst;
+
+
+	//! Node that creates a group of objects and describes how they are selected.
+	//! \ingroup model
+	class MUTABLETOOLS_API NodeObjectGroup : public NodeObject
+	{
+	public:
+
+		//-----------------------------------------------------------------------------------------
+		// Life cycle
+		//-----------------------------------------------------------------------------------------
+
+		NodeObjectGroup();
+
+		void SerialiseWrapper(OutputArchive& arch) const override;
+		static void Serialise( const NodeObjectGroup* pNode, OutputArchive& arch );
+		static NodeObjectGroupPtr StaticUnserialise( InputArchive& arch );
+
+
+		//-----------------------------------------------------------------------------------------
+		// Node Interface
+		//-----------------------------------------------------------------------------------------
+
+        NodePtr Clone( NodeMapPtr pMap = 0 ) const override;
+
+        const NODE_TYPE* GetType() const override;
+		static const NODE_TYPE* GetStaticType();
+
+        int GetInputCount() const override;
+        Node* GetInputNode( int i ) const override;
+        void SetInputNode( int i, NodePtr pNode ) override;
+
+		//-----------------------------------------------------------------------------------------
+		// NodeObject Interface
+		//-----------------------------------------------------------------------------------------
+        const char* GetName() const override;
+        void SetName( const char* strName ) override;
+        const char* GetUid() const override;
+        void SetUid( const char* strUid ) override;
+
+		//-----------------------------------------------------------------------------------------
+		// Own Interface
+		//-----------------------------------------------------------------------------------------
+
+		//! Typed of child selection
+		typedef enum
+		{
+			//! All objects in the group will always be enabled, and no parameter will be generated.
+			CS_ALWAYS_ALL,
+
+			//! Only one children may be selected, but it is allowed to have none.
+			//! An enumeration parameter will be generated but it may have a null value
+			CS_ONE_OR_NONE,
+
+			//! One and only one children has to be selected at all times
+			//! An enumeraation parameter will be generated and it cannot be null
+			CS_ALWAYS_ONE,
+
+			//! Each child in the group can be enabled or disabled individually.
+			//! A boolean parameter will be generated for every child.
+			CS_TOGGLE_EACH
+
+		} CHILD_SELECTION;
+
+		//! Get child selection type
+		CHILD_SELECTION GetSelectionType() const;
+
+		//! Set the child selection type
+		void SetSelectionType( CHILD_SELECTION );
+
+		//! Get the number of child objects
+		int GetChildCount() const;
+
+		//! Set the number of child objects
+		void SetChildCount( int );
+
+		//! Get a child object node
+		NodeObjectPtr GetChild( int index ) const;
+
+		//! Set a child object node
+		void SetChild( int index, NodeObjectPtr );
+
+		//-----------------------------------------------------------------------------------------
+		// Interface pattern
+		//-----------------------------------------------------------------------------------------
+		class Private;
+		Private* GetPrivate() const;
+        Node::Private* GetBasePrivate() const override;
+
+	protected:
+
+		//! Forbidden. Manage with the Ptr<> template.
+		~NodeObjectGroup();
+
+	private:
+
+		Private* m_pD;
+
+	};
+
+
+
+}
diff --git a/Engine/Plugins/Experimental/Mutable/Source/MutableTools/Public/NodeObjectNew.h b/Engine/Plugins/Experimental/Mutable/Source/MutableTools/Public/NodeObjectNew.h
new file mode 100644
--- /dev/null
+++ b/Engine/Plugins/Experimental/Mutable/Source/MutableTools/Public/NodeObjectNew.h
@@ -0,0 +1,124 @@
+// Copyright Epic Games, Inc. All Rights Reserved.
+
+#pragma once
+#include "NodeObject.h"
+
+
+namespace mu
+{
+
+	// Forward definitions
+	class NodeObjectNew;
+	typedef Ptr<NodeObjectNew> NodeObjectNewPtr;
+	typedef Ptr<const NodeObjectNew> NodeObjectNewPtrConst;
+
+	class NodeLOD;
+	typedef Ptr<NodeLOD> NodeLODPtr;
+	typedef Ptr<const NodeLOD> NodeLODPtrConst;
+
+
+	//! Node that creates a new object by setting its levels-of-detail and its children.
+	//! \ingroup model
+	class MUTABLETOOLS_API NodeObjectNew : public NodeObject
+	{
+	public:
+
+		//-----------------------------------------------------------------------------------------
+		// Life cycle
+		//-----------------------------------------------------------------------------------------
+
+		NodeObjectNew();
+
+		void SerialiseWrapper(OutputArchive& arch) const override;
+		static void Serialise( const NodeObjectNew* pNode, OutputArchive& arch );
+		static NodeObjectNewPtr StaticUnserialise( InputArchive& arch );
+
+
+		//-----------------------------------------------------------------------------------------
+		// Node Interface
+		//-----------------------------------------------------------------------------------------
+
+        NodePtr Clone( NodeMapPtr pMap = 0 ) const override;
+
+        const NODE_TYPE* GetType() const override;
+		static const NODE_TYPE* GetStaticType();
+
+        int GetInputCount() const override;
+        Node* GetInputNode( int i ) const override;
+        void SetInputNode( int i, NodePtr pNode ) override;
+
+		//-----------------------------------------------------------------------------------------
+		// NodeObject Interface
+		//-----------------------------------------------------------------------------------------
+        const char* GetName() const override;
+        void SetName( const char* strName ) override;
+        const char* GetUid() const override;
+        void SetUid( const char* strUid ) override;
+
+		//-----------------------------------------------------------------------------------------
+		// Own Interface
+		//-----------------------------------------------------------------------------------------
+
+		//! Get the number of levels of detail in the object.
+		int GetLODCount() const;
+		void SetLODCount( int );
+
+		//! Get a level of detail node from the object.
+		NodeLODPtr GetLOD( int index ) const;
+		void SetLOD( int index, NodeLODPtr );
+
+		//! Get the number of child objects
+		int GetChildCount() const;
+		void SetChildCount( int );
+
+		//! Get a child object node
+		NodeObjectPtr GetChild( int index ) const;
+		void SetChild( int index, NodeObjectPtr );
+
+		//! Set the number of states that the model can be in.
+		int GetStateCount() const;
+		void SetStateCount( int c );
+
+		//! Set the name of a state
+		const char* GetStateName( int s ) const;
+		void SetStateName( int s, const char* n );
+
+		//! State optimization options
+		void SetStateGPUOptimisation( int s, bool internal, bool external, bool returnGPU = false );
+		bool GetStateGPUOptimisationInternal( int s ) const;
+		bool GetStateGPUOptimisationExternal( int s ) const;
+		bool GetStateGPUOptimisationReturnImages( int s ) const;
+
+		//! See if a state has a parameter as runtime.
+		bool HasStateParam( int s, const char* param ) const;
+
+		//! Add a runtime parameter to the state.
+		void AddStateParam( int s, const char* param );
+
+		//! Remove a runtime parameter from a state.
+		void RemoveStateParam( int s, const char* param );
+
+        //! Set the optimisation properties of a state
+        void SetStateProperties( int s, bool avoidRuntimeCompression, bool onlyFirstLOD, int firstLOD );
+
+		//-----------------------------------------------------------------------------------------
+		// Interface pattern
+		//-----------------------------------------------------------------------------------------
+		class Private;
+		Private* GetPrivate() const;
+        Node::Private* GetBasePrivate() const override;
+
+	protected:
+
+		//! Forbidden. Manage with the Ptr<> template.
+		~NodeObjectNew();
+
+	private:
+
+		Private* m_pD;
+
+	};
+
+
+
+}
diff --git a/Engine/Plugins/Experimental/Mutable/Source/MutableTools/Public/NodeObjectState.h b/Engine/Plugins/Experimental/Mutable/Source/MutableTools/Public/NodeObjectState.h
new file mode 100644
--- /dev/null
+++ b/Engine/Plugins/Experimental/Mutable/Source/MutableTools/Public/NodeObjectState.h
@@ -0,0 +1,108 @@
+// Copyright Epic Games, Inc. All Rights Reserved.
+
+#pragma once
+
+#include "NodeObject.h"
+
+
+namespace mu
+{
+
+	//! Forward definitions
+	class NodeObjectState;
+	typedef Ptr<NodeObjectState> NodeObjectStatePtr;
+	typedef Ptr<const NodeObjectState> NodeObjectStatePtrConst;
+
+
+
+	//!
+	//! \ingroup model
+	class MUTABLETOOLS_API NodeObjectState : public NodeObject
+	{
+	public:
+
+		//-----------------------------------------------------------------------------------------
+		// Life cycle
+		//-----------------------------------------------------------------------------------------
+
+		NodeObjectState();
+
+		void SerialiseWrapper(OutputArchive& arch) const override;
+		static void Serialise( const NodeObjectState* pNode, OutputArchive& arch );
+		static NodeObjectStatePtr StaticUnserialise( InputArchive& arch );
+
+
+		//-----------------------------------------------------------------------------------------
+		// Node interface
+		//-----------------------------------------------------------------------------------------
+
+        NodePtr Clone( NodeMapPtr pMap = 0 ) const override;
+
+        const NODE_TYPE* GetType() const override;
+		static const NODE_TYPE* GetStaticType();
+
+        int GetInputCount() const override;
+        Node* GetInputNode( int i ) const override;
+        void SetInputNode( int i, NodePtr pNode ) override;
+
+
+		//-----------------------------------------------------------------------------------------
+		// NodeObject interface
+		//-----------------------------------------------------------------------------------------
+        const char* GetName() const override;
+        void SetName( const char* strName ) override;
+        const char* GetUid() const override;
+        void SetUid( const char* strUid ) override;
+
+		//-----------------------------------------------------------------------------------------
+		// Own interface
+		//-----------------------------------------------------------------------------------------
+
+		//! Source object
+		NodeObjectPtr GetSource() const;
+		void SetSource( NodeObjectPtr );
+
+		//! State diagram to add
+		NodeObjectPtr GetStateRoot() const;
+		void SetStateRoot( NodeObjectPtr );
+
+		//! Set the name of a state
+		const char* GetStateName() const;
+		void SetStateName( const char* n );
+
+		//! State optimization options
+		void SetStateGPUOptimisation( bool internal, bool external, bool returnGPU = false );
+		bool GetStateGPUOptimisationInternal() const;
+		bool GetStateGPUOptimisationExternal() const;
+		bool GetStateGPUOptimisationReturnImages() const;
+
+		//! See if a state has a parameter as runtime.
+		bool HasStateParam( const char* param ) const;
+
+		//! Add a runtime parameter to the state.
+		void AddStateParam( const char* param );
+
+		//! Remove a runtime parameter from a state.
+		void RemoveStateParam( const char* param );
+
+		//-----------------------------------------------------------------------------------------
+		// Interface pattern
+		//-----------------------------------------------------------------------------------------
+		class Private;
+		Private* GetPrivate() const;
+        Node::Private* GetBasePrivate() const override;
+
+	protected:
+
+		//! Forbidden. Manage with the Ptr<> template.
+		~NodeObjectState();
+
+	private:
+
+		Private* m_pD;
+
+	};
+
+
+
+}
diff --git a/Engine/Plugins/Experimental/Mutable/Source/MutableTools/Public/NodePatchMesh.h b/Engine/Plugins/Experimental/Mutable/Source/MutableTools/Public/NodePatchMesh.h
new file mode 100644
--- /dev/null
+++ b/Engine/Plugins/Experimental/Mutable/Source/MutableTools/Public/NodePatchMesh.h
@@ -0,0 +1,77 @@
+// Copyright Epic Games, Inc. All Rights Reserved.
+
+#pragma once
+
+#include "NodeMesh.h"
+
+
+namespace mu
+{
+
+	// Forward definitions
+	class NodePatchMesh;
+	typedef Ptr<NodePatchMesh> NodePatchMeshPtr;
+	typedef Ptr<const NodePatchMesh> NodePatchMeshPtrConst;
+
+
+	//!
+	//! \ingroup model
+	class MUTABLETOOLS_API NodePatchMesh : public Node
+	{
+	public:
+
+		//-----------------------------------------------------------------------------------------
+		// Life cycle
+		//-----------------------------------------------------------------------------------------
+
+		NodePatchMesh();
+
+		void SerialiseWrapper(OutputArchive& arch) const override;
+		static void Serialise( const NodePatchMesh* pNode, OutputArchive& arch );
+		static NodePatchMeshPtr StaticUnserialise( InputArchive& arch );
+
+		//-----------------------------------------------------------------------------------------
+		// Node Interface
+		//-----------------------------------------------------------------------------------------
+
+        NodePtr Clone( NodeMapPtr pMap = 0 ) const override;
+
+        const NODE_TYPE* GetType() const override;
+		static const NODE_TYPE* GetStaticType();
+
+        int GetInputCount() const override;
+        Node* GetInputNode( int i ) const override;
+        void SetInputNode( int i, NodePtr pNode ) override;
+
+		//-----------------------------------------------------------------------------------------
+		// Own Interface
+		//-----------------------------------------------------------------------------------------
+
+		//!
+        NodeMesh* GetRemove() const;
+        void SetRemove( NodeMesh* );
+
+		//!
+        NodeMesh* GetAdd() const;
+        void SetAdd( NodeMesh* );
+
+		//-----------------------------------------------------------------------------------------
+		// Interface pattern
+		//-----------------------------------------------------------------------------------------
+		class Private;
+		Private* GetPrivate() const;
+        Node::Private* GetBasePrivate() const override;
+
+	protected:
+
+		//! Forbidden. Manage with the Ptr<> template.
+		~NodePatchMesh();
+
+	private:
+
+		Private* m_pD;
+
+	};
+
+
+}
diff --git a/Engine/Plugins/Experimental/Mutable/Source/MutableTools/Public/NodeProjector.h b/Engine/Plugins/Experimental/Mutable/Source/MutableTools/Public/NodeProjector.h
new file mode 100644
--- /dev/null
+++ b/Engine/Plugins/Experimental/Mutable/Source/MutableTools/Public/NodeProjector.h
@@ -0,0 +1,230 @@
+// Copyright Epic Games, Inc. All Rights Reserved.
+
+#pragma once
+
+#include "Node.h"
+#include "Parameters.h"
+
+
+namespace mu
+{
+
+	class NodeProjector;
+    using NodeProjectorPtr = Ptr<NodeProjector>;
+    using NodeProjectorPtrConst = Ptr<const NodeProjector>;
+
+	class NodeProjectorConstant;
+    using NodeProjectorConstantPtr = Ptr<NodeProjectorConstant>;
+    using NodeProjectorConstantPtrConst = Ptr<const NodeProjectorConstant>;
+
+	class NodeProjectorParameter;
+    using NodeProjectorParameterPtr = Ptr<NodeProjectorParameter>;
+    using NodeProjectorParameterPtrConst = Ptr<const NodeProjectorParameter>;
+
+    class NodeRange;
+    using NodeRangePtr = Ptr<NodeRange>;
+    using NodeRangePtrConst = Ptr<const NodeRange>;
+
+
+    //! Base class of any node that outputs a Projector.
+	//! \ingroup model
+	class MUTABLETOOLS_API NodeProjector : public Node
+	{
+	public:
+
+		// Possible subclasses
+		enum class EType : uint8
+		{
+			Constant = 0,
+			Parameter = 1,
+
+			None
+		};
+
+		//-----------------------------------------------------------------------------------------
+		// Life cycle
+		//-----------------------------------------------------------------------------------------
+
+		static void Serialise( const NodeProjector* pNode, OutputArchive& arch );
+		static NodeProjectorPtr StaticUnserialise( InputArchive& arch );
+
+
+		//-----------------------------------------------------------------------------------------
+		// Node Interface
+		//-----------------------------------------------------------------------------------------
+
+        const NODE_TYPE* GetType() const override;
+		static const NODE_TYPE* GetStaticType();
+
+
+		//-----------------------------------------------------------------------------------------
+		// Interface pattern
+		//-----------------------------------------------------------------------------------------
+		class Private;
+
+	protected:
+
+		//! Forbidden. Manage with the Ptr<> template.
+		inline ~NodeProjector() {}
+		
+		//!
+		EType Type = EType::None;
+
+	};
+
+
+	//! This node outputs a predefined Projector value.
+	//! \ingroup model
+	class MUTABLETOOLS_API NodeProjectorConstant : public NodeProjector
+	{
+	public:
+
+		//-----------------------------------------------------------------------------------------
+		// Life cycle
+		//-----------------------------------------------------------------------------------------
+
+		NodeProjectorConstant();
+
+		void SerialiseWrapper(OutputArchive& arch) const override;
+		static void Serialise( const NodeProjectorConstant* pNode, OutputArchive& arch );
+		static NodeProjectorConstantPtr StaticUnserialise( InputArchive& arch );
+
+		//-----------------------------------------------------------------------------------------
+		// Node Interface
+		//-----------------------------------------------------------------------------------------
+
+        NodePtr Clone( NodeMapPtr pMap = 0 ) const override;
+
+        const NODE_TYPE* GetType() const override;
+		static const NODE_TYPE* GetStaticType();
+
+        int GetInputCount() const override;
+        Node* GetInputNode( int i ) const override;
+        void SetInputNode( int i, NodePtr pNode ) override;
+
+		//-----------------------------------------------------------------------------------------
+		// Own Interface
+		//-----------------------------------------------------------------------------------------
+
+		//! Get the value that this node returns
+        void GetValue( PROJECTOR_TYPE* pType,
+                       float* pPosX, float* pPosY, float* pPosZ,
+					   float* pDirX, float* pDirY, float* pDirZ,
+					   float* pUpX, float* pUpY, float *pUpZ,
+                       float* pScaleU, float* pScaleV, float* pScaleW,
+                       float* pProjectionAngle ) const;
+
+		//! Set the value to be returned by this node
+        void SetValue( PROJECTOR_TYPE type,
+                       float posX, float posY, float posZ,
+					   float dirX, float dirY, float dirZ,
+					   float upX, float upY, float upZ,
+                       float scaleU, float scaleV, float scaleW,
+                       float projectionAngle );
+
+		//-----------------------------------------------------------------------------------------
+		// Interface pattern
+		//-----------------------------------------------------------------------------------------
+		class Private;
+		Private* GetPrivate() const;
+        Node::Private* GetBasePrivate() const override;
+
+	protected:
+
+		//! Forbidden. Manage with the Ptr<> template.
+		~NodeProjectorConstant();
+
+	private:
+
+		Private* m_pD;
+
+	};
+
+
+	//! Node that defines a Projector model parameter.
+	//! \ingroup model
+	class MUTABLETOOLS_API NodeProjectorParameter : public NodeProjector
+	{
+	public:
+
+		//-----------------------------------------------------------------------------------------
+		// Life cycle
+		//-----------------------------------------------------------------------------------------
+
+		NodeProjectorParameter();
+
+		void SerialiseWrapper(OutputArchive& arch) const override;
+		static void Serialise( const NodeProjectorParameter* pNode, OutputArchive& arch );
+		static NodeProjectorParameterPtr StaticUnserialise( InputArchive& arch );
+
+
+		//-----------------------------------------------------------------------------------------
+		// Node Interface
+		//-----------------------------------------------------------------------------------------
+
+        NodePtr Clone( NodeMapPtr pMap = 0 ) const override;
+
+        const NODE_TYPE* GetType() const override;
+		static const NODE_TYPE* GetStaticType();
+
+        int GetInputCount() const override;
+        Node* GetInputNode( int i ) const override;
+        void SetInputNode( int i, NodePtr pNode ) override;
+
+		//-----------------------------------------------------------------------------------------
+		// Own Interface
+		//-----------------------------------------------------------------------------------------
+
+		//! Get the name of the parameter. It will be exposed in the final compiled data.
+		const char* GetName() const;
+
+		//! Set the name of the parameter.
+		void SetName( const char* );
+
+		//! Get the uid of the parameter. It will be exposed in the final compiled data.
+		const char* GetUid() const;
+
+		//! Set the uid of the parameter.
+		void SetUid( const char* );
+
+		//! Get the default value of the parameter.
+        void GetDefaultValue( PROJECTOR_TYPE* pType,
+                              float* posX, float* posY, float* posZ,
+							  float* pDirX, float* pDirY, float* pDirZ,
+							  float* pUpX, float* pUpY, float *pUpZ,
+                              float *pScaleU, float *pScaleV, float* pScaleW,
+                              float* pProjectionAngle ) const;
+
+		//! Set the default value of the parameter.
+        void SetDefaultValue( PROJECTOR_TYPE type,
+                              float posX, float posY, float posZ,
+							  float dirX, float dirY, float dirZ,
+							  float upX, float upY, float upZ,
+                              float scaleU, float scaleV, float scaleW,
+                              float projectionAngle );
+
+        //! Set the number of ranges (dimensions) for this parameter.
+        //! By default a parameter has 0 ranges, meaning it only has one value.
+        void SetRangeCount( int i );
+        void SetRange( int i, NodeRangePtr pRange );
+
+		//-----------------------------------------------------------------------------------------
+		// Interface pattern
+		//-----------------------------------------------------------------------------------------
+		class Private;
+		Private* GetPrivate() const;
+        Node::Private* GetBasePrivate() const override;
+
+	protected:
+
+		//! Forbidden. Manage with the Ptr<> template.
+		~NodeProjectorParameter();
+
+	private:
+
+		Private* m_pD;
+
+	};
+
+
+}
diff --git a/Engine/Plugins/Experimental/Mutable/Source/MutableTools/Public/NodeRange.h b/Engine/Plugins/Experimental/Mutable/Source/MutableTools/Public/NodeRange.h
new file mode 100644
--- /dev/null
+++ b/Engine/Plugins/Experimental/Mutable/Source/MutableTools/Public/NodeRange.h
@@ -0,0 +1,64 @@
+// Copyright Epic Games, Inc. All Rights Reserved.
+
+#pragma once
+
+#include "MutableTools/Public/Node.h"
+
+
+namespace mu
+{
+
+    // Forward definitions
+    class NodeRange;
+    typedef Ptr<NodeRange> NodeRangePtr;
+    typedef Ptr<const NodeRange> NodeRangePtrConst;
+
+
+    //! %Base class of any node that outputs a range.
+	//! \ingroup model
+    class MUTABLETOOLS_API NodeRange : public Node
+	{
+	public:
+
+		// Possible subclasses
+		enum class EType : uint8
+		{
+			FromScalar = 0,
+
+			None
+		};
+
+		//-----------------------------------------------------------------------------------------
+		// Life cycle
+		//-----------------------------------------------------------------------------------------
+
+        static void Serialise( const NodeRange* pNode, OutputArchive& arch );
+        static NodeRangePtr StaticUnserialise( InputArchive& arch );
+
+
+		//-----------------------------------------------------------------------------------------
+		// Node Interface
+		//-----------------------------------------------------------------------------------------
+
+		const NODE_TYPE* GetType() const override;
+		static const NODE_TYPE* GetStaticType();
+
+
+		//-----------------------------------------------------------------------------------------
+		// Interface pattern
+		//-----------------------------------------------------------------------------------------
+		class Private;
+
+	protected:
+
+		//! Forbidden. Manage with the Ptr<> template.
+        inline ~NodeRange() {}
+
+		//!
+		EType Type = EType::None;
+
+	};
+
+
+}
+
diff --git a/Engine/Plugins/Experimental/Mutable/Source/MutableTools/Public/NodeRangeFromScalar.h b/Engine/Plugins/Experimental/Mutable/Source/MutableTools/Public/NodeRangeFromScalar.h
new file mode 100644
--- /dev/null
+++ b/Engine/Plugins/Experimental/Mutable/Source/MutableTools/Public/NodeRangeFromScalar.h
@@ -0,0 +1,76 @@
+// Copyright Epic Games, Inc. All Rights Reserved.
+
+#pragma once
+
+#include "MutableTools/Public/NodeRange.h"
+
+namespace mu
+{
+
+	// Forward definitions
+    class NodeScalar;
+
+
+	//! 
+	//! \ingroup model
+    class MUTABLETOOLS_API NodeRangeFromScalar : public NodeRange
+	{
+	public:
+
+		//-----------------------------------------------------------------------------------------
+		// Life cycle
+		//-----------------------------------------------------------------------------------------
+
+        NodeRangeFromScalar();
+
+		void SerialiseWrapper(OutputArchive& arch) const override;
+		static void Serialise( const NodeRangeFromScalar* pNode, OutputArchive& arch );
+        static Ptr<NodeRangeFromScalar> StaticUnserialise( InputArchive& arch );
+
+
+		//-----------------------------------------------------------------------------------------
+		// Node Interface
+		//-----------------------------------------------------------------------------------------
+
+        NodePtr Clone( NodeMapPtr pMap = nullptr ) const override;
+
+		const NODE_TYPE* GetType() const override;
+		static const NODE_TYPE* GetStaticType();
+
+        int GetInputCount() const override;
+        Node* GetInputNode( int i ) const override;
+		void SetInputNode( int i, NodePtr pNode ) override;
+
+
+		//-----------------------------------------------------------------------------------------
+		// Own Interface
+		//-----------------------------------------------------------------------------------------
+
+        //!
+        Ptr<NodeScalar> GetSize() const;
+        void SetSize( const Ptr<NodeScalar>& );
+
+        //!
+        const char* GetName() const;
+        void SetName( const char* strName );
+
+		//-----------------------------------------------------------------------------------------
+		// Interface pattern
+		//-----------------------------------------------------------------------------------------
+		class Private;
+		Private* GetPrivate() const;
+        Node::Private* GetBasePrivate() const override;
+
+	protected:
+
+		//! Forbidden. Manage with the Ptr<> template.
+        ~NodeRangeFromScalar();
+
+	private:
+
+		Private* m_pD;
+
+	};
+
+
+}
diff --git a/Engine/Plugins/Experimental/Mutable/Source/MutableTools/Public/NodeScalar.h b/Engine/Plugins/Experimental/Mutable/Source/MutableTools/Public/NodeScalar.h
new file mode 100644
--- /dev/null
+++ b/Engine/Plugins/Experimental/Mutable/Source/MutableTools/Public/NodeScalar.h
@@ -0,0 +1,71 @@
+// Copyright Epic Games, Inc. All Rights Reserved.
+
+#pragma once
+
+#include "Node.h"
+#include "Parameters.h"
+
+
+namespace mu
+{
+
+	// Forward definitions
+	class NodeScalar;
+	typedef Ptr<NodeScalar> NodeScalarPtr;
+	typedef Ptr<const NodeScalar> NodeScalarPtrConst;
+
+
+    //! %Base class of any node that outputs a scalar value.
+	//! \ingroup model
+	class MUTABLETOOLS_API NodeScalar : public Node
+	{
+	public:
+
+		// Possible subclasses
+		enum class EType : uint8
+		{
+			Constant = 0,
+			Parameter = 1,
+			EnumParameter = 2,
+			Curve = 3,
+			Switch = 4,
+			ArithmeticOperation = 5,
+			Variation = 6,
+			Table = 7,
+
+			None
+		};
+
+		//-----------------------------------------------------------------------------------------
+		// Life cycle
+		//-----------------------------------------------------------------------------------------
+
+		static void Serialise( const NodeScalar* pNode, OutputArchive& arch );
+		static NodeScalarPtr StaticUnserialise( InputArchive& arch );
+
+
+		//-----------------------------------------------------------------------------------------
+		// Node Interface
+		//-----------------------------------------------------------------------------------------
+
+		const NODE_TYPE* GetType() const override;
+		static const NODE_TYPE* GetStaticType();
+
+
+		//-----------------------------------------------------------------------------------------
+		// Interface pattern
+		//-----------------------------------------------------------------------------------------
+		class Private;
+
+	protected:
+
+		//! Forbidden. Manage with the Ptr<> template.
+		inline ~NodeScalar() {}
+
+		//!
+		EType Type = EType::None;
+
+	};
+
+
+}
diff --git a/Engine/Plugins/Experimental/Mutable/Source/MutableTools/Public/NodeScalarArithmeticOperation.h b/Engine/Plugins/Experimental/Mutable/Source/MutableTools/Public/NodeScalarArithmeticOperation.h
new file mode 100644
--- /dev/null
+++ b/Engine/Plugins/Experimental/Mutable/Source/MutableTools/Public/NodeScalarArithmeticOperation.h
@@ -0,0 +1,98 @@
+// Copyright Epic Games, Inc. All Rights Reserved.
+
+#pragma once
+
+#include "NodeScalar.h"
+
+
+namespace mu
+{
+
+	// Forward definitions
+    class NodeScalarArithmeticOperation;
+    typedef Ptr<NodeScalarArithmeticOperation> NodeScalarArithmeticOperationPtr;
+    typedef Ptr<const NodeScalarArithmeticOperation> NodeScalarArithmeticOperationPtrConst;
+
+
+    //! Perform an arithmetic operation between two scalars.
+	//! \ingroup model
+    class MUTABLETOOLS_API NodeScalarArithmeticOperation : public NodeScalar
+	{
+	public:
+
+		//-----------------------------------------------------------------------------------------
+		// Life cycle
+		//-----------------------------------------------------------------------------------------
+
+        NodeScalarArithmeticOperation();
+
+		void SerialiseWrapper(OutputArchive& arch) const override;
+		static void Serialise( const NodeScalarArithmeticOperation* pNode, OutputArchive& arch );
+        static NodeScalarArithmeticOperationPtr StaticUnserialise( InputArchive& arch );
+
+
+		//-----------------------------------------------------------------------------------------
+		// Node Interface
+		//-----------------------------------------------------------------------------------------
+
+        NodePtr Clone( NodeMapPtr pMap = 0 ) const override;
+
+        const NODE_TYPE* GetType() const override;
+		static const NODE_TYPE* GetStaticType();
+
+        int GetInputCount() const override;
+        Node* GetInputNode( int i ) const override;
+        void SetInputNode( int i, NodePtr pNode ) override;
+
+
+		//-----------------------------------------------------------------------------------------
+		// Own Interface
+		//-----------------------------------------------------------------------------------------
+
+		//! Possible arithmetic operations
+		typedef enum
+		{
+            AO_ADD = 0,
+			AO_SUBTRACT,
+			AO_MULTIPLY,
+            AO_DIVIDE,
+            _AO_COUNT
+		} OPERATION;
+
+        static const char* s_opTypeName[ _AO_COUNT ];
+
+		//! Get child selection type
+		OPERATION GetOperation() const;
+
+		//! Set the child selection type
+		void SetOperation(OPERATION);
+
+
+		//! Get the first operand
+        NodeScalarPtr GetA() const;
+        void SetA(NodeScalarPtr);
+
+		// Get the second operand
+        NodeScalarPtr GetB() const;
+        void SetB(NodeScalarPtr);
+
+		//-----------------------------------------------------------------------------------------
+		// Interface pattern
+		//-----------------------------------------------------------------------------------------
+		class Private;
+		Private* GetPrivate() const;
+        Node::Private* GetBasePrivate() const override;
+
+	protected:
+
+		//! Forbidden. Manage with the Ptr<> template.
+        ~NodeScalarArithmeticOperation();
+
+	private:
+
+		Private* m_pD;
+
+	};
+
+
+}
diff --git a/Engine/Plugins/Experimental/Mutable/Source/MutableTools/Public/NodeScalarConstant.h b/Engine/Plugins/Experimental/Mutable/Source/MutableTools/Public/NodeScalarConstant.h
new file mode 100644
--- /dev/null
+++ b/Engine/Plugins/Experimental/Mutable/Source/MutableTools/Public/NodeScalarConstant.h
@@ -0,0 +1,78 @@
+// Copyright Epic Games, Inc. All Rights Reserved.
+
+#pragma once
+
+#include "NodeScalar.h"
+#include "Parameters.h"
+
+
+namespace mu
+{
+
+	// Forward definitions
+	class NodeScalarConstant;
+	typedef Ptr<NodeScalarConstant> NodeScalarConstantPtr;
+	typedef Ptr<const NodeScalarConstant> NodeScalarConstantPtrConst;
+
+
+	//! Node returning a scalar constant value.
+	//! \ingroup model
+	class MUTABLETOOLS_API NodeScalarConstant : public NodeScalar
+	{
+	public:
+
+		//-----------------------------------------------------------------------------------------
+		// Life cycle
+		//-----------------------------------------------------------------------------------------
+
+		NodeScalarConstant();
+
+		void SerialiseWrapper(OutputArchive& arch) const override;
+		static void Serialise( const NodeScalarConstant* pNode, OutputArchive& arch );
+		static NodeScalarConstantPtr StaticUnserialise( InputArchive& arch );
+
+
+		//-----------------------------------------------------------------------------------------
+		// Node Interface
+		//-----------------------------------------------------------------------------------------
+
+        NodePtr Clone( NodeMapPtr pMap = 0 ) const override;
+
+        const NODE_TYPE* GetType() const override;
+		static const NODE_TYPE* GetStaticType();
+
+        int GetInputCount() const override;
+        Node* GetInputNode( int i ) const override;
+        void SetInputNode( int i, NodePtr pNode ) override;
+
+		//-----------------------------------------------------------------------------------------
+		// Own Interface
+		//-----------------------------------------------------------------------------------------
+
+		//! Get the value to be returned by the node.
+		float GetValue() const;
+
+		//! Set the value to be returned by the node.
+		void SetValue( float v );
+
+		//-----------------------------------------------------------------------------------------
+		// Interface pattern
+		//-----------------------------------------------------------------------------------------
+		class Private;
+		Private* GetPrivate() const;
+        Node::Private* GetBasePrivate() const override;
+
+	protected:
+
+		//! Forbidden. Manage with the Ptr<> template.
+		~NodeScalarConstant();
+
+	private:
+
+		Private* m_pD;
+
+	};
+
+
+
+}
diff --git a/Engine/Plugins/Experimental/Mutable/Source/MutableTools/Public/NodeScalarCurve.h b/Engine/Plugins/Experimental/Mutable/Source/MutableTools/Public/NodeScalarCurve.h
new file mode 100644
--- /dev/null
+++ b/Engine/Plugins/Experimental/Mutable/Source/MutableTools/Public/NodeScalarCurve.h
@@ -0,0 +1,88 @@
+// Copyright Epic Games, Inc. All Rights Reserved.
+
+
+#pragma once
+
+#include "NodeScalar.h"
+
+namespace mu
+{
+
+	// Forward definitions
+    class NodeScalarCurve;
+    typedef Ptr<NodeScalarCurve> NodeScalarCurvePtr;
+    typedef Ptr<const NodeScalarCurve> NodeScalarCurvePtrConst;
+
+
+    //! This node makes a new scalar value transforming another scalar value with a curve.
+	//! \ingroup model
+    class MUTABLETOOLS_API NodeScalarCurve : public NodeScalar
+	{
+	public:
+
+		//-----------------------------------------------------------------------------------------
+		// Life cycle
+		//-----------------------------------------------------------------------------------------
+
+        NodeScalarCurve();
+
+		void SerialiseWrapper(OutputArchive& arch) const override;
+		static void Serialise( const NodeScalarCurve* pNode, OutputArchive& arch );
+        static NodeScalarCurvePtr StaticUnserialise( InputArchive& arch );
+
+		//-----------------------------------------------------------------------------------------
+		// Node Interface
+		//-----------------------------------------------------------------------------------------
+        NodePtr Clone( NodeMapPtr pMap = nullptr ) const override;
+
+        const NODE_TYPE* GetType() const override;
+		static const NODE_TYPE* GetStaticType();
+
+        int GetInputCount() const override;
+        Node* GetInputNode( int i ) const override;
+        void SetInputNode( int i, NodePtr pNode ) override;
+
+		//-----------------------------------------------------------------------------------------
+		// Own Interface
+		//-----------------------------------------------------------------------------------------
+
+        //! Get the number of key frames in the Curve.
+		int GetKeyFrameCount() const;
+
+        //! Set the number of key frames in the Curve.
+		void SetKeyFrameCount( int );
+
+		//! Set the default value of the curve
+		void SetDefaultValue(float);
+
+		//! \param index index of the curve, from 0 to GetKeyFrameCount()-1
+        void SetKeyFrame(int index, float time, float value,
+                         float in_tangent, float in_tangent_weight,
+                         float out_tangent, float out_tangent_weight,
+                         uint8_t interp_mode,
+                         uint8_t tangent_mode, uint8_t tangent_weight_mode);
+
+        //! Get the curve time parameter
+		NodeScalarPtr GetT() const;
+		void SetT(NodeScalarPtr);
+
+		//-----------------------------------------------------------------------------------------
+		// Interface pattern
+		//-----------------------------------------------------------------------------------------
+		class Private;
+		Private* GetPrivate() const;
+        Node::Private* GetBasePrivate() const override;
+
+	protected:
+
+		//! Forbidden. Manage with the Ptr<> template.
+        ~NodeScalarCurve();
+
+	private:
+
+		Private* m_pD;
+
+	};
+
+
+}
diff --git a/Engine/Plugins/Experimental/Mutable/Source/MutableTools/Public/NodeScalarEnumParameter.h b/Engine/Plugins/Experimental/Mutable/Source/MutableTools/Public/NodeScalarEnumParameter.h
new file mode 100644
--- /dev/null
+++ b/Engine/Plugins/Experimental/Mutable/Source/MutableTools/Public/NodeScalarEnumParameter.h
@@ -0,0 +1,106 @@
+// Copyright Epic Games, Inc. All Rights Reserved.
+
+#pragma once
+
+#include "NodeScalar.h"
+#include "Parameters.h"
+#include "NodeRange.h"
+
+
+namespace mu
+{
+
+	// Forward definitions
+	class NodeImage;
+    using NodeImagePtr = Ptr<NodeImage>;
+    using NodeImagePtrConst = Ptr<const NodeImage>;
+
+	class NodeScalarEnumParameter;
+    using NodeScalarEnumParameterPtr = Ptr<NodeScalarEnumParameter>;
+    using NodeScalarEnumParameterPtrConst = Ptr<const NodeScalarEnumParameter>;
+
+    class NodeRange;
+    using NodeRangePtr = Ptr<NodeRange>;
+    using NodeRangePtrConst = Ptr<const NodeRange>;
+
+
+	//! Node that defines a scalar model parameter to be selected from a set of named values.
+	//! \ingroup model
+	class MUTABLETOOLS_API NodeScalarEnumParameter : public NodeScalar
+	{
+	public:
+
+		//-----------------------------------------------------------------------------------------
+		// Life cycle
+		//-----------------------------------------------------------------------------------------
+
+		NodeScalarEnumParameter();
+
+		void SerialiseWrapper(OutputArchive& arch) const override;
+		static void Serialise( const NodeScalarEnumParameter* pNode, OutputArchive& arch );
+		static NodeScalarEnumParameterPtr StaticUnserialise( InputArchive& arch );
+
+
+		//-----------------------------------------------------------------------------------------
+		// Node Interface
+		//-----------------------------------------------------------------------------------------
+
+        NodePtr Clone( NodeMapPtr pMap = 0 ) const override;
+
+        const NODE_TYPE* GetType() const override;
+		static const NODE_TYPE* GetStaticType();
+
+        int GetInputCount() const override;
+        Node* GetInputNode( int i ) const override;
+        void SetInputNode( int i, NodePtr pNode ) override;
+
+		//-----------------------------------------------------------------------------------------
+		// Own Interface
+		//-----------------------------------------------------------------------------------------
+
+		//! Get the name of the parameter. It will be exposed in the final compiled data.
+		const char* GetName() const;
+		void SetName( const char* );
+
+		//! Get the uid of the parameter. It will be exposed in the final compiled data.
+		const char* GetUid() const;
+		void SetUid( const char* );
+
+		//! Get the index of the default value of the parameter.
+		int GetDefaultValueIndex() const;
+		void SetDefaultValueIndex( int v );
+
+		//! Set the number of possible values.
+		void SetValueCount( int i );
+
+		//! Get the number of possible values.
+		int GetValueCount() const;
+
+		//! Set the data of one of the possible values of the parameter.
+		void SetValue( int i, float value, const char* strName );
+
+        //! Set the number of ranges (dimensions) for this parameter.
+        //! By default a parameter has 0 ranges, meaning it only has one value.
+        void SetRangeCount( int i );
+        void SetRange( int i, NodeRangePtr pRange );
+
+		//-----------------------------------------------------------------------------------------
+		// Interface pattern
+		//-----------------------------------------------------------------------------------------
+		class Private;
+		Private* GetPrivate() const;
+        Node::Private* GetBasePrivate() const override;
+
+	protected:
+
+		//! Forbidden. Manage with the Ptr<> template.
+		~NodeScalarEnumParameter();
+
+	private:
+
+		Private* m_pD;
+
+	};
+
+
+}
diff --git a/Engine/Plugins/Experimental/Mutable/Source/MutableTools/Public/NodeScalarParameter.h b/Engine/Plugins/Experimental/Mutable/Source/MutableTools/Public/NodeScalarParameter.h
new file mode 100644
--- /dev/null
+++ b/Engine/Plugins/Experimental/Mutable/Source/MutableTools/Public/NodeScalarParameter.h
@@ -0,0 +1,106 @@
+// Copyright Epic Games, Inc. All Rights Reserved.
+
+#pragma once
+
+#include "NodeScalar.h"
+#include "Parameters.h"
+
+
+namespace mu
+{
+
+	// Forward definitions
+    class NodeImage;
+    using NodeImagePtr = Ptr<NodeImage>;
+    using NodeImagePtrConst = Ptr<const NodeImage>;
+
+    class NodeRange;
+    using NodeRangePtr = Ptr<NodeRange>;
+    using NodeRangePtrConst = Ptr<const NodeRange>;
+
+    class NodeScalarParameter;
+    using NodeScalarParameterPtr = Ptr<NodeScalarParameter>;
+    using NodeScalarParameterPtrConst = Ptr<const NodeScalarParameter>;
+
+
+	//! Node that defines a scalar model parameter.
+	//! \ingroup model
+	class MUTABLETOOLS_API NodeScalarParameter : public NodeScalar
+	{
+	public:
+
+		//-----------------------------------------------------------------------------------------
+		// Life cycle
+		//-----------------------------------------------------------------------------------------
+
+		NodeScalarParameter();
+
+		void SerialiseWrapper(OutputArchive& arch) const override;
+		static void Serialise( const NodeScalarParameter* pNode, OutputArchive& arch );
+		static NodeScalarParameterPtr StaticUnserialise( InputArchive& arch );
+
+
+		//-----------------------------------------------------------------------------------------
+		// Node Interface
+		//-----------------------------------------------------------------------------------------
+
+        NodePtr Clone( NodeMapPtr pMap = 0 ) const override;
+
+        const NODE_TYPE* GetType() const override;
+		static const NODE_TYPE* GetStaticType();
+
+        int GetInputCount() const override;
+        Node* GetInputNode( int i ) const override;
+        void SetInputNode( int i, NodePtr pNode ) override;
+
+		//-----------------------------------------------------------------------------------------
+		// Own Interface
+		//-----------------------------------------------------------------------------------------
+
+		//! Get the name of the parameter. It will be exposed in the final compiled data.
+		const char* GetName() const;
+		void SetName( const char* );
+
+		//! Get the uid of the parameter. It will be exposed in the final compiled data.
+		const char* GetUid() const;
+		void SetUid( const char* );
+
+		//! Get the default value of the parameter.
+		float GetDefaultValue() const;
+		void SetDefaultValue( float v );
+
+		//! Get the additional information about the type
+		PARAMETER_DETAILED_TYPE GetDetailedType() const ;
+		void SetDetailedType( PARAMETER_DETAILED_TYPE type );
+
+        //! Set the number of additional description images associated to this parameter.
+        //! These images can be built from a parameter and, for example, they can be used to build
+        //! user-interfaces for customisation.
+        void SetDescImageCount( int i );
+        void SetDescImage( int i, NodeImagePtr pImage );
+
+        //! Set the number of ranges (dimensions) for this parameter.
+        //! By default a parameter has 0 ranges, meaning it only has one value.
+        void SetRangeCount( int i );
+        void SetRange( int i, NodeRangePtr pRange );
+
+		//-----------------------------------------------------------------------------------------
+		// Interface pattern
+		//-----------------------------------------------------------------------------------------
+		class Private;
+		Private* GetPrivate() const;
+        Node::Private* GetBasePrivate() const override;
+
+	protected:
+
+		//! Forbidden. Manage with the Ptr<> template.
+		~NodeScalarParameter();
+
+	private:
+
+		Private* m_pD;
+
+	};
+
+
+}
diff --git a/Engine/Plugins/Experimental/Mutable/Source/MutableTools/Public/NodeScalarSwitch.h b/Engine/Plugins/Experimental/Mutable/Source/MutableTools/Public/NodeScalarSwitch.h
new file mode 100644
--- /dev/null
+++ b/Engine/Plugins/Experimental/Mutable/Source/MutableTools/Public/NodeScalarSwitch.h
@@ -0,0 +1,87 @@
+// Copyright Epic Games, Inc. All Rights Reserved.
+
+#pragma once
+
+#include "NodeScalar.h"
+
+
+namespace mu
+{
+
+	// Forward definitions
+	class NodeScalar;
+	typedef Ptr<NodeScalar> NodeScalarPtr;
+	typedef Ptr<const NodeScalar> NodeScalarPtrConst;
+
+    class NodeScalarSwitch;
+    typedef Ptr<NodeScalarSwitch> NodeScalarSwitchPtr;
+    typedef Ptr<const NodeScalarSwitch> NodeScalarSwitchPtrConst;
+
+
+    //! This node selects an output Scalar from a set of input Scalars based on a parameter.
+	//! \ingroup model
+    class MUTABLETOOLS_API NodeScalarSwitch : public NodeScalar
+	{
+	public:
+
+		//-----------------------------------------------------------------------------------------
+		// Life cycle
+		//-----------------------------------------------------------------------------------------
+
+        NodeScalarSwitch();
+
+		void SerialiseWrapper(OutputArchive& arch) const override;
+		static void Serialise( const NodeScalarSwitch* pNode, OutputArchive& arch );
+        static NodeScalarSwitchPtr StaticUnserialise( InputArchive& arch );
+
+
+		//-----------------------------------------------------------------------------------------
+		// Node Interface
+		//-----------------------------------------------------------------------------------------
+
+        NodePtr Clone( NodeMapPtr pMap = nullptr ) const override;
+
+        const NODE_TYPE* GetType() const override;
+		static const NODE_TYPE* GetStaticType();
+
+        virtual int GetInputCount() const override;
+        virtual Node* GetInputNode( int i ) const override;
+        void SetInputNode( int i, NodePtr pNode ) override;
+
+		//-----------------------------------------------------------------------------------------
+		// Own Interface
+		//-----------------------------------------------------------------------------------------
+
+		//! Get the node generating the parameter used to select the option.
+		NodeScalarPtr GetParameter() const;
+		void SetParameter( NodeScalarPtr );
+
+        //! Set the number of option Scalars. It will keep the currently set targets and initialise
+		//! the new ones as null.
+		void SetOptionCount( int );
+
+        //! Get the node generating the t-th option Scalar.
+        NodeScalarPtr GetOption( int t ) const;
+        void SetOption( int t, NodeScalarPtr );
+
+
+		//-----------------------------------------------------------------------------------------
+		// Interface pattern
+		//-----------------------------------------------------------------------------------------
+		class Private;
+		Private* GetPrivate() const;
+        Node::Private* GetBasePrivate() const override;
+
+	protected:
+
+		//! Forbidden. Manage with the Ptr<> template.
+        ~NodeScalarSwitch();
+
+	private:
+
+		Private* m_pD;
+
+	};
+
+
+}
diff --git a/Engine/Plugins/Experimental/Mutable/Source/MutableTools/Public/NodeScalarTable.h b/Engine/Plugins/Experimental/Mutable/Source/MutableTools/Public/NodeScalarTable.h
new file mode 100644
--- /dev/null
+++ b/Engine/Plugins/Experimental/Mutable/Source/MutableTools/Public/NodeScalarTable.h
@@ -0,0 +1,86 @@
+// Copyright Epic Games, Inc. All Rights Reserved.
+
+#pragma once
+
+#include "NodeScalar.h"
+
+
+namespace mu
+{
+
+	// Forward definitions
+	class NodeScalarTable;
+	typedef Ptr<NodeScalarTable> NodeScalarTablePtr;
+	typedef Ptr<const NodeScalarTable> NodeScalarTablePtrConst;
+
+	class Table;
+	typedef Ptr<Table> TablePtr;
+	typedef Ptr<const Table> TablePtrConst;
+
+
+	//! This node provides the meshes stored in the column of a table.
+	//! \ingroup transform
+	class MUTABLETOOLS_API NodeScalarTable : public NodeScalar
+	{
+	public:
+
+		//-----------------------------------------------------------------------------------------
+		// Life cycle
+		//-----------------------------------------------------------------------------------------
+
+		NodeScalarTable();
+
+		void SerialiseWrapper(OutputArchive& arch) const override;
+		static void Serialise( const NodeScalarTable* pNode, OutputArchive& arch );
+		static NodeScalarTablePtr StaticUnserialise( InputArchive& arch );
+
+
+		//-----------------------------------------------------------------------------------------
+		// Node Interface
+		//-----------------------------------------------------------------------------------------
+
+		NodePtr Clone( NodeMapPtr pMap = 0 ) const override;
+
+		const NODE_TYPE* GetType() const override;
+		static const NODE_TYPE* GetStaticType();
+
+		int GetInputCount() const override;
+		Node* GetInputNode( int i ) const override;
+		void SetInputNode( int i, NodePtr pNode ) override;
+
+		//-----------------------------------------------------------------------------------------
+		// Own Interface
+		//-----------------------------------------------------------------------------------------
+
+		//! Set the name of the implicit table parameter.
+		void SetParameterName( const char* strName );
+
+		//!
+		TablePtr GetTable() const;
+		void SetTable( TablePtr );
+
+		//!
+		const char* GetColumn() const;
+		void SetColumn( const char* strName );
+
+		//-----------------------------------------------------------------------------------------
+		// Interface pattern
+		//-----------------------------------------------------------------------------------------
+		class Private;
+		Private* GetPrivate() const;
+        Node::Private* GetBasePrivate() const override;
+
+
+	protected:
+
+		//! Forbidden. Manage with the Ptr<> template.
+		~NodeScalarTable();
+
+	private:
+
+		Private* m_pD;
+
+	};
+
+
+}
diff --git a/Engine/Plugins/Experimental/Mutable/Source/MutableTools/Public/NodeScalarVariation.h b/Engine/Plugins/Experimental/Mutable/Source/MutableTools/Public/NodeScalarVariation.h
new file mode 100644
--- /dev/null
+++ b/Engine/Plugins/Experimental/Mutable/Source/MutableTools/Public/NodeScalarVariation.h
@@ -0,0 +1,83 @@
+// Copyright Epic Games, Inc. All Rights Reserved.
+
+#pragma once
+
+#include "NodeScalar.h"
+
+
+namespace mu
+{
+
+    // Forward definitions
+    class NodeScalarVariation;
+    typedef Ptr<NodeScalarVariation> NodeScalarVariationPtr;
+    typedef Ptr<const NodeScalarVariation> NodeScalarVariationPtrConst;
+
+
+    //!
+    //! \ingroup model
+    class MUTABLETOOLS_API NodeScalarVariation : public NodeScalar
+    {
+    public:
+        //-----------------------------------------------------------------------------------------
+        // Life cycle
+        //-----------------------------------------------------------------------------------------
+
+        NodeScalarVariation();
+
+		void SerialiseWrapper(OutputArchive& arch) const override;
+		static void Serialise( const NodeScalarVariation* pNode, OutputArchive& arch );
+        static NodeScalarVariationPtr StaticUnserialise( InputArchive& arch );
+
+        //-----------------------------------------------------------------------------------------
+        // Node Interface
+        //-----------------------------------------------------------------------------------------
+
+        NodePtr Clone( NodeMapPtr pMap = 0 ) const override;
+
+        const NODE_TYPE* GetType() const override;
+        static const NODE_TYPE* GetStaticType();
+
+        virtual int GetInputCount() const override;
+        virtual Node* GetInputNode( int i ) const override;
+        void SetInputNode( int i, NodePtr pNode ) override;
+
+        //-----------------------------------------------------------------------------------------
+        // Own Interface
+        //-----------------------------------------------------------------------------------------
+
+        //!
+        void SetDefaultScalar( NodeScalar* Scalar );
+
+        //! Set the number of tags to consider in this variation
+        void SetVariationCount( int count );
+
+        //!
+        int GetVariationCount() const;
+
+        //! Set the tag or state name that will enable a specific vartiation
+        void SetVariationTag( int index, const char* strTag );
+
+        //!
+        void SetVariationScalar( int index, NodeScalar* Scalar );
+
+        //!}
+
+
+        //-----------------------------------------------------------------------------------------
+        // Interface pattern
+        //-----------------------------------------------------------------------------------------
+        class Private;
+        Private* GetPrivate() const;
+        Node::Private* GetBasePrivate() const override;
+
+    protected:
+        //! Forbidden. Manage with the Ptr<> template.
+        ~NodeScalarVariation();
+
+    private:
+        Private* m_pD;
+    };
+
+
+} // namespace mu
diff --git a/Engine/Plugins/Experimental/Mutable/Source/MutableTools/Public/NodeString.h b/Engine/Plugins/Experimental/Mutable/Source/MutableTools/Public/NodeString.h
new file mode 100644
--- /dev/null
+++ b/Engine/Plugins/Experimental/Mutable/Source/MutableTools/Public/NodeString.h
@@ -0,0 +1,66 @@
+// Copyright Epic Games, Inc. All Rights Reserved.
+
+#pragma once
+
+#include "MutableTools/Public/Node.h"
+#include "MutableRuntime/Public/Parameters.h"
+
+
+namespace mu
+{
+
+	// Forward definitions
+	class NodeString;
+	typedef Ptr<NodeString> NodeStringPtr;
+	typedef Ptr<const NodeString> NodeStringPtrConst;
+
+
+    //! %Base class of any node that outputs a string value.
+	//! \ingroup model
+	class MUTABLETOOLS_API NodeString : public Node
+	{
+	public:
+
+		// Possible subclasses
+		enum class EType : uint8
+		{
+			Constant = 0,
+			Parameter = 1,
+
+			None
+		};
+
+		//-----------------------------------------------------------------------------------------
+		// Life cycle
+		//-----------------------------------------------------------------------------------------
+
+		static void Serialise( const NodeString* pNode, OutputArchive& arch );
+		static NodeStringPtr StaticUnserialise( InputArchive& arch );
+
+
+		//-----------------------------------------------------------------------------------------
+		// Node Interface
+		//-----------------------------------------------------------------------------------------
+
+		const NODE_TYPE* GetType() const override;
+		static const NODE_TYPE* GetStaticType();
+
+
+		//-----------------------------------------------------------------------------------------
+		// Interface pattern
+		//-----------------------------------------------------------------------------------------
+		class Private;
+
+	protected:
+
+		//! Forbidden. Manage with the Ptr<> template.
+		inline ~NodeString() {}
+
+		//!
+		EType Type = EType::None;
+
+
+	};
+
+
+}
