[mutable] Move Mutable plugin from the restricted folder to the experimental plugins folder.

diff --git a/Engine/Plugins/Experimental/Mutable/Source/MutableTools/Public/NodeStringConstant.h b/Engine/Plugins/Experimental/Mutable/Source/MutableTools/Public/NodeStringConstant.h
new file mode 100644
--- /dev/null
+++ b/Engine/Plugins/Experimental/Mutable/Source/MutableTools/Public/NodeStringConstant.h
@@ -0,0 +1,78 @@
+// Copyright Epic Games, Inc. All Rights Reserved.
+
+#pragma once
+
+#include "MutableTools/Public/NodeString.h"
+#include "MutableRuntime/Public/Parameters.h"
+
+
+namespace mu
+{
+
+	// Forward definitions
+	class NodeStringConstant;
+	typedef Ptr<NodeStringConstant> NodeStringConstantPtr;
+	typedef Ptr<const NodeStringConstant> NodeStringConstantPtrConst;
+
+
+	//! Node returning a string constant value.
+	//! \ingroup model
+	class MUTABLETOOLS_API NodeStringConstant : public NodeString
+	{
+	public:
+
+		//-----------------------------------------------------------------------------------------
+		// Life cycle
+		//-----------------------------------------------------------------------------------------
+
+		NodeStringConstant();
+
+		void SerialiseWrapper(OutputArchive& arch) const override;
+		static void Serialise( const NodeStringConstant* pNode, OutputArchive& arch );
+		static NodeStringConstantPtr StaticUnserialise( InputArchive& arch );
+
+
+		//-----------------------------------------------------------------------------------------
+		// Node Interface
+		//-----------------------------------------------------------------------------------------
+
+        NodePtr Clone( NodeMapPtr pMap = 0 ) const override;
+
+        const NODE_TYPE* GetType() const override;
+		static const NODE_TYPE* GetStaticType();
+
+        int GetInputCount() const override;
+        Node* GetInputNode( int i ) const override;
+        void SetInputNode( int i, NodePtr pNode ) override;
+
+		//-----------------------------------------------------------------------------------------
+		// Own Interface
+		//-----------------------------------------------------------------------------------------
+
+		//! Get the value to be returned by the node.
+		const char* GetValue() const;
+
+		//! Set the value to be returned by the node.
+		void SetValue( const char* v );
+
+		//-----------------------------------------------------------------------------------------
+		// Interface pattern
+		//-----------------------------------------------------------------------------------------
+		class Private;
+		Private* GetPrivate() const;
+        Node::Private* GetBasePrivate() const override;
+
+	protected:
+
+		//! Forbidden. Manage with the Ptr<> template.
+		~NodeStringConstant();
+
+	private:
+
+		Private* m_pD;
+
+	};
+
+
+
+}
diff --git a/Engine/Plugins/Experimental/Mutable/Source/MutableTools/Public/NodeStringParameter.h b/Engine/Plugins/Experimental/Mutable/Source/MutableTools/Public/NodeStringParameter.h
new file mode 100644
--- /dev/null
+++ b/Engine/Plugins/Experimental/Mutable/Source/MutableTools/Public/NodeStringParameter.h
@@ -0,0 +1,100 @@
+// Copyright Epic Games, Inc. All Rights Reserved.
+
+#pragma once
+
+#include "NodeString.h"
+#include "Parameters.h"
+
+
+namespace mu
+{
+
+	// Forward definitions
+    class NodeImage;
+    using NodeImagePtr = Ptr<NodeImage>;
+    using NodeImagePtrConst = Ptr<const NodeImage>;
+
+    class NodeRange;
+    using NodeRangePtr = Ptr<NodeRange>;
+    using NodeRangePtrConst = Ptr<const NodeRange>;
+
+    class NodeStringParameter;
+    using NodeStringParameterPtr = Ptr<NodeStringParameter>;
+    using NodeStringParameterPtrConst = Ptr<const NodeStringParameter>;
+
+
+	//! Node that defines a string model parameter.
+	//! \ingroup model
+	class MUTABLETOOLS_API NodeStringParameter : public NodeString
+	{
+	public:
+
+		//-----------------------------------------------------------------------------------------
+		// Life cycle
+		//-----------------------------------------------------------------------------------------
+
+		NodeStringParameter();
+
+		void SerialiseWrapper(OutputArchive& arch) const override;
+		static void Serialise( const NodeStringParameter* pNode, OutputArchive& arch );
+		static NodeStringParameterPtr StaticUnserialise( InputArchive& arch );
+
+
+		//-----------------------------------------------------------------------------------------
+		// Node Interface
+		//-----------------------------------------------------------------------------------------
+
+        NodePtr Clone( NodeMapPtr pMap = 0 ) const override;
+
+        const NODE_TYPE* GetType() const override;
+		static const NODE_TYPE* GetStaticType();
+
+        int GetInputCount() const override;
+        Node* GetInputNode( int i ) const override;
+        void SetInputNode( int i, NodePtr pNode ) override;
+
+		//-----------------------------------------------------------------------------------------
+		// Own Interface
+		//-----------------------------------------------------------------------------------------
+
+		//! Get the name of the parameter. It will be exposed in the final compiled data.
+		const char* GetName() const;
+		void SetName( const char* );
+
+		//! Get the uid of the parameter. It will be exposed in the final compiled data.
+		const char* GetUid() const;
+		void SetUid( const char* );
+
+		//! Get the default value of the parameter.
+		const char* GetDefaultValue() const;
+		void SetDefaultValue( const char* v );
+
+		//! Get the additional information about the type
+		PARAMETER_DETAILED_TYPE GetDetailedType() const ;
+		void SetDetailedType( PARAMETER_DETAILED_TYPE type );
+
+        //! Set the number of ranges (dimensions) for this parameter.
+        //! By default a parameter has 0 ranges, meaning it only has one value.
+        void SetRangeCount( int i );
+        void SetRange( int i, NodeRangePtr pRange );
+
+		//-----------------------------------------------------------------------------------------
+		// Interface pattern
+		//-----------------------------------------------------------------------------------------
+		class Private;
+		Private* GetPrivate() const;
+        Node::Private* GetBasePrivate() const override;
+
+	protected:
+
+		//! Forbidden. Manage with the Ptr<> template.
+		~NodeStringParameter();
+
+	private:
+
+		Private* m_pD;
+
+	};
+
+
+}
diff --git a/Engine/Plugins/Experimental/Mutable/Source/MutableTools/Public/NodeSurface.h b/Engine/Plugins/Experimental/Mutable/Source/MutableTools/Public/NodeSurface.h
new file mode 100644
--- /dev/null
+++ b/Engine/Plugins/Experimental/Mutable/Source/MutableTools/Public/NodeSurface.h
@@ -0,0 +1,60 @@
+// Copyright Epic Games, Inc. All Rights Reserved.
+
+#pragma once
+
+#include "MutableTools/Public/Node.h"
+
+
+namespace mu
+{
+
+	// Forward definitions
+    class NodeSurface;
+    typedef Ptr<NodeSurface> NodeSurfacePtr;
+    typedef Ptr<const NodeSurface> NodeSurfacePtrConst;
+
+
+
+    //! This class is the parent of all nodes that output a Surface.
+	//! \ingroup model
+    class MUTABLETOOLS_API NodeSurface : public Node
+	{
+	public:
+
+		// Possible subclasses
+		enum class EType : uint8
+		{
+			New = 0,
+			Edit = 1,
+			Variation = 3,
+
+			None
+		};
+
+		//-----------------------------------------------------------------------------------------
+		// Life cycle
+		//-----------------------------------------------------------------------------------------
+
+        static void Serialise( const NodeSurface* pNode, OutputArchive& arch );
+        static NodeSurfacePtr StaticUnserialise( InputArchive& arch );
+
+		//-----------------------------------------------------------------------------------------
+		// Node Interface
+		//-----------------------------------------------------------------------------------------
+
+        const NODE_TYPE* GetType() const override;
+		static const NODE_TYPE* GetStaticType();
+
+
+	protected:
+
+		//! Forbidden. Manage with the Ptr<> template.
+        inline ~NodeSurface() {}
+
+		//!
+		EType Type = EType::None;
+
+	};
+
+
+}
diff --git a/Engine/Plugins/Experimental/Mutable/Source/MutableTools/Public/NodeSurfaceEdit.h b/Engine/Plugins/Experimental/Mutable/Source/MutableTools/Public/NodeSurfaceEdit.h
new file mode 100644
--- /dev/null
+++ b/Engine/Plugins/Experimental/Mutable/Source/MutableTools/Public/NodeSurfaceEdit.h
@@ -0,0 +1,163 @@
+// Copyright Epic Games, Inc. All Rights Reserved.
+
+#pragma once
+
+#include "MutableTools/Public/NodeSurface.h"
+
+
+namespace mu
+{
+
+	// Forward definitions
+    class NodeMesh;
+    typedef Ptr<NodeMesh> NodeMeshPtr;
+    typedef Ptr<const NodeMesh> NodeMeshPtrConst;
+
+    class NodeImage;
+	typedef Ptr<NodeImage> NodeImagePtr;
+	typedef Ptr<const NodeImage> NodeImagePtrConst;
+
+	class NodePatchMesh;
+	typedef Ptr<NodePatchMesh> NodePatchMeshPtr;
+	typedef Ptr<const NodePatchMesh> NodePatchMeshPtrConst;
+
+	class NodePatchImage;
+	typedef Ptr<NodePatchImage> NodePatchImagePtr;
+	typedef Ptr<const NodePatchImage> NodePatchImagePtrConst;
+
+    class NodeSurfaceEdit;
+    typedef Ptr<NodeSurfaceEdit> NodeSurfaceEditPtr;
+    typedef Ptr<const NodeSurfaceEdit> NodeSurfaceEditPtrConst;
+
+	class NodeScalar;
+	typedef Ptr<NodeScalar> NodeScalarPtr;
+	typedef Ptr<const NodeScalar> NodeScalarPtrConst;
+
+
+
+	//! This node modifies a node of the parent object of the object that this node belongs to.
+    //! It allows to extend, cut and morph the parent Surface's meshes.
+    //! It also allows to patch the parent Surface's textures.
+	//! \ingroup model
+    class MUTABLETOOLS_API NodeSurfaceEdit : public NodeSurface
+	{
+	public:
+
+		//-----------------------------------------------------------------------------------------
+		// Life cycle
+		//-----------------------------------------------------------------------------------------
+
+        NodeSurfaceEdit();
+
+		void SerialiseWrapper(OutputArchive& arch) const override;
+		static void Serialise( const NodeSurfaceEdit* pNode, OutputArchive& arch );
+        static NodeSurfaceEditPtr StaticUnserialise( InputArchive& arch );
+
+		//-----------------------------------------------------------------------------------------
+		// Node Interface
+		//-----------------------------------------------------------------------------------------
+
+        NodePtr Clone( NodeMapPtr pMap = 0 ) const override;
+
+        const NODE_TYPE* GetType() const override;
+		static const NODE_TYPE* GetStaticType();
+
+        virtual int GetInputCount() const override;
+        virtual Node* GetInputNode( int i ) const override;
+        void SetInputNode( int i, NodePtr pNode ) override;
+
+		//-----------------------------------------------------------------------------------------
+		// Own Interface
+		//-----------------------------------------------------------------------------------------
+
+        //! Set the parent Surface to modify. It should be a node from the same LOD of one of
+		//! the parent objects of this node's object.
+        void SetParent( NodeSurface* );
+        NodeSurface* GetParent() const;
+
+        //! \name Tags
+        //! \{
+
+        //! Add a tag to the surface:
+        //! - the surface will be affected by modifier nodes with the same tag
+        //! - the tag will be enabled when the surface is added to an object, and it can activate
+        //! variations for any surface.
+        void AddTag(const char* tagName);
+
+        //! Get the number of tags added to the Surface.
+        int GetTagCount() const;
+
+        //! Get a tag string from an index (0 to GetTagCount-1)
+        const char* GetTag( int ) const;
+
+        //! \}
+
+		//! \name Mesh modification
+		//! \{
+
+		//! Mesh a node to extend the parent mesh
+        void SetMesh( NodePatchMesh* );
+
+		//! Get the node to extend the parent mesh
+        NodePatchMesh* GetMesh() const;
+
+		//! Set a node to morph the parent mesh
+        void SetMorph( NodeMeshPtr );
+
+		//! Get the node to morph the parent mesh
+        NodeMeshPtr GetMorph() const;
+
+		//! Set the factor: weight used to select (and combine) the morphs to apply.
+		void SetFactor( NodeScalarPtr );
+
+		//! Get the factor: weight used to select (and combine) the morphs to apply.
+		NodeScalarPtr GetFactor() const;
+
+		//! \}
+
+		//! \name Image modification
+		//! \{
+
+		//! Get the number of images to modify. Should match the parent image count.
+		int GetImageCount() const;
+
+		//! Set the number of images to modify. Should match the parent image count.
+		void SetImageCount( int );
+
+		//! Get the image node that extends a parent image, adding new blocks
+		NodeImagePtr GetImage( int index ) const;
+
+		//! Set the image node that extends a parent image, adding new blocks
+		void SetImage( int index, NodeImagePtr );
+
+		//! Get the image node that patches a parent image, blending on existing blocks
+		NodePatchImagePtr GetPatch( int index ) const;
+
+		//! Set the image node that patches a parent image, blending on existing blocks
+		void SetPatch( int index, NodePatchImagePtr );
+
+		//!}
+
+
+		//-----------------------------------------------------------------------------------------
+		// Interface pattern
+		//-----------------------------------------------------------------------------------------
+		class Private;
+		Private* GetPrivate() const;
+        Node::Private* GetBasePrivate() const override;
+
+	protected:
+
+		//! Forbidden. Manage with the Ptr<> template.
+        ~NodeSurfaceEdit();
+
+	private:
+
+		Private* m_pD;
+
+	};
+
+
+
+}
+
diff --git a/Engine/Plugins/Experimental/Mutable/Source/MutableTools/Public/NodeSurfaceNew.h b/Engine/Plugins/Experimental/Mutable/Source/MutableTools/Public/NodeSurfaceNew.h
new file mode 100644
--- /dev/null
+++ b/Engine/Plugins/Experimental/Mutable/Source/MutableTools/Public/NodeSurfaceNew.h
@@ -0,0 +1,279 @@
+// Copyright Epic Games, Inc. All Rights Reserved.
+
+#pragma once
+
+#include "MutableTools/Public/NodeSurface.h"
+
+
+namespace mu
+{
+
+	// Forward definitions
+	class NodeMesh;
+	typedef Ptr<NodeMesh> NodeMeshPtr;
+	typedef Ptr<const NodeMesh> NodeMeshPtrConst;
+
+	class NodeImage;
+	typedef Ptr<NodeImage> NodeImagePtr;
+	typedef Ptr<const NodeImage> NodeImagePtrConst;
+
+    class NodeColour;
+    typedef Ptr<NodeColour> NodeColourPtr;
+    typedef Ptr<const NodeColour> NodeColourPtrConst;
+
+    class NodeScalar;
+    typedef Ptr<NodeScalar> NodeScalarPtr;
+    typedef Ptr<const NodeScalar> NodeScalarPtrConst;
+
+    class NodeString;
+    typedef Ptr<NodeString> NodeStringPtr;
+    typedef Ptr<const NodeString> NodeStringPtrConst;
+
+    class NodeSurface;
+    typedef Ptr<NodeSurface> NodeSurfacePtr;
+    typedef Ptr<const NodeSurface> NodeSurfacePtrConst;
+
+    class NodeSurfaceNew;
+    typedef Ptr<NodeSurfaceNew> NodeSurfaceNewPtr;
+    typedef Ptr<const NodeSurfaceNew> NodeSurfaceNewPtrConst;
+
+
+
+    //! This node makes a new Surface from several meshes and images.
+	//! \ingroup model
+    class MUTABLETOOLS_API NodeSurfaceNew : public NodeSurface
+	{
+	public:
+
+		//-----------------------------------------------------------------------------------------
+		// Life cycle
+		//-----------------------------------------------------------------------------------------
+
+        NodeSurfaceNew();
+
+		void SerialiseWrapper(OutputArchive& arch) const override;
+		static void Serialise( const NodeSurfaceNew* pNode, OutputArchive& arch );
+        static NodeSurfaceNewPtr StaticUnserialise( InputArchive& arch );
+
+		//-----------------------------------------------------------------------------------------
+		// Node Interface
+		//-----------------------------------------------------------------------------------------
+
+        NodePtr Clone( NodeMapPtr pMap = 0 ) const override;
+
+        const NODE_TYPE* GetType() const override;
+		static const NODE_TYPE* GetStaticType();
+
+        virtual int GetInputCount() const override;
+        virtual Node* GetInputNode( int i ) const override;
+        void SetInputNode( int i, NodePtr pNode ) override;
+
+		//-----------------------------------------------------------------------------------------
+		// Own Interface
+		//-----------------------------------------------------------------------------------------
+
+        //! Get the name of the Surface.
+		const char* GetName() const;
+
+        //! Set the name of the Surface.
+        void SetName( const char* );
+
+        //! Add an optional, opaque id that will be returned in the surfaces of the created
+        //! instances. Can be useful to identify surfaces on the application side.
+        //! See Instance::GetSurfaceCustomId
+        void SetCustomID( uint32_t id );
+
+        //! \name Tags
+        //! \{
+
+        //! Add a tag to the surface:
+        //! - the surface will be affected by modifier nodes with the same tag
+        //! - the tag will be enabled when the surface is added to an object, and it can activate
+        //! variations for any surface.
+		void AddTag(const char* tagName);
+
+        //! Get the number of tags added to the Surface.
+        int GetTagCount() const;
+
+        //! Get a tag string from an index (0 to GetTagCount-1)
+        const char* GetTag( int ) const;
+
+        //! \}
+
+		//! \name Meshes
+        //! \{
+
+        //! Get the number of meshes in the Surface.
+		int GetMeshCount() const;
+
+        //! Set the number of meshes in the Surface.
+		void SetMeshCount( int );
+
+        //! Get the node generating one of the meshes in the Surface.
+		//! \param index index of the mesh, from 0 to GetMeshCount()-1
+		NodeMeshPtr GetMesh( int index ) const;
+
+        //! Set the node generating one of the meshes in the Surface.
+		//! \param index index of the mesh, from 0 to GetMeshCount()-1
+		void SetMesh( int index, NodeMeshPtr );
+
+        //! Get the name of a mesh in the Surface.
+		//! \param index index of the mesh, from 0 to GetMeshCount()-1
+		const char* GetMeshName( int index ) const;
+
+        //! Set the name of a mesh in the Surface.
+		//! \param index index of the mesh, from 0 to GetMeshCount()-1
+		//! \param strName name of the mesh
+		void SetMeshName( int index, const char* strName );
+
+		//! \}
+
+
+		//! \name Images
+		//! \{
+
+		//! Get the number of images in the Surface.
+		int GetImageCount() const;
+
+		//! Set the number of images in the Surface.
+		void SetImageCount(int);
+
+		//! Get the node generating one of the images in the Surface.
+		//! \param index index of the image, from 0 to GetImageCount()-1
+		NodeImagePtr GetImage(int index) const;
+
+		//! Set the node generating one of the images in the Surface.
+		//! \param index index of the image, from 0 to GetImageCount()-1
+		void SetImage(int index, NodeImagePtr);
+
+		//! Get the name of a image in the Surface.
+		//! \param index index of the image, from 0 to GetImageCount()-1
+		const char* GetImageName(int index) const;
+
+        //! Set the name of a image in the Surface.
+        //! \param index index of the image, from 0 to GetImageCount()-1
+        //! \param strName name of the image
+        void SetImageName( int index, const char* strName );
+
+        //! Get the name of a image in the Surface.
+        //! \param index index of the image, from 0 to GetImageCount()-1
+        int GetImageLayoutIndex( int index ) const;
+
+        //! Set the name of a image in the Surface.
+        //! \param index index of the image, from 0 to GetImageCount()-1
+        //! \param layoutIndex index of the layout that will be used with this image.
+        void SetImageLayoutIndex( int imageIndex, int layoutIndex );
+
+		//! This can be used to ser additional information for error reporting.
+		//! It is not propagated to compiled objects in any way.
+		void SetImageAdditionalNames(int index, const char* strMaterialName, const char* strMaterialParameterName);
+
+        //! \}
+
+
+        //! \name Vectors
+        //! \{
+
+        //! Get the number of vectors in the Surface.
+        int GetVectorCount() const;
+
+        //! Set the number of vectors in the Surface.
+        void SetVectorCount(int);
+
+        //! Get the node generating one of the vectors in the Surface.
+        //! \param index index of the vector, from 0 to GetVectorCount()-1
+        NodeColourPtr GetVector(int index) const;
+
+        //! Set the node generating one of the vectors in the Surface.
+        //! \param index index of the vector, from 0 to GetVectorCount()-1
+        void SetVector(int index, NodeColourPtr);
+
+        //! Get the name of a vectors in the Surface.
+        //! \param index index of the vector, from 0 to GetVectorCount()-1
+        const char* GetVectorName(int index) const;
+
+        //! Set the name of a image in the Surface.
+        //! \param index index of the vector, from 0 to GetVectorCount()-1
+        //! \param strName name of the vector
+        void SetVectorName(int index, const char* strName);
+
+        //! \}
+
+
+        //! \name Scalars
+        //! \{
+
+        //! Get the number of Scalars in the Surface.
+        int GetScalarCount() const;
+
+        //! Set the number of Scalars in the Surface.
+        void SetScalarCount( int );
+
+        //! Get the node generating one of the Scalars in the Surface.
+        //! \param index index of the Scalar, from 0 to GetScalarCount()-1
+        NodeScalarPtr GetScalar( int index ) const;
+
+        //! Set the node generating one of the Scalars in the Surface.
+        //! \param index index of the Scalar, from 0 to GetScalarCount()-1
+        void SetScalar( int index, NodeScalarPtr );
+
+        //! Get the name of a Scalars in the Surface.
+        //! \param index index of the Scalar, from 0 to GetScalarCount()-1
+        const char* GetScalarName( int index ) const;
+
+        //! Set the name of a image in the Surface.
+        //! \param index index of the Scalar, from 0 to GetScalarCount()-1
+        //! \param strName name of the Scalar
+        void SetScalarName( int index, const char* strName );
+
+        //! \}
+
+
+        //! \name Strings
+        //! \{
+
+        //! Get the number of strings in the Surface.
+        int GetStringCount() const;
+
+        //! Set the number of strings in the Surface.
+        void SetStringCount( int );
+
+        //! Get the node generating one of the strings in the Surface.
+        //! \param index index of the strings, from 0 to GetStringCount()-1
+        NodeStringPtr GetString( int index ) const;
+
+        //! Set the node generating one of the Strings in the Surface.
+        //! \param index index of the String, from 0 to GetStringCount()-1
+        void SetString( int index, NodeStringPtr );
+
+        //! Get the name of a String in the Surface.
+        //! \param index index of the String, from 0 to GetStringCount()-1
+        const char* GetStringName( int index ) const;
+
+        //! Set the name of a image in the Surface.
+        //! \param index index of the String, from 0 to GetStringCount()-1
+        //! \param strName name of the String
+        void SetStringName( int index, const char* strName );
+
+        //! \}
+
+        //-----------------------------------------------------------------------------------------
+		// Interface pattern
+		//-----------------------------------------------------------------------------------------
+		class Private;
+		Private* GetPrivate() const;
+        Node::Private* GetBasePrivate() const override;
+
+	protected:
+
+		//! Forbidden. Manage with the Ptr<> template.
+        ~NodeSurfaceNew();
+
+	private:
+
+		Private* m_pD;
+
+	};
+
+
+}
diff --git a/Engine/Plugins/Experimental/Mutable/Source/MutableTools/Public/NodeSurfaceVariation.h b/Engine/Plugins/Experimental/Mutable/Source/MutableTools/Public/NodeSurfaceVariation.h
new file mode 100644
--- /dev/null
+++ b/Engine/Plugins/Experimental/Mutable/Source/MutableTools/Public/NodeSurfaceVariation.h
@@ -0,0 +1,108 @@
+// Copyright Epic Games, Inc. All Rights Reserved.
+
+#pragma once
+
+#include "MutableTools/Public/NodeSurface.h"
+
+
+namespace mu
+{
+
+	// Forward definitions
+    class NodeSurfaceVariation;
+    typedef Ptr<NodeSurfaceVariation> NodeSurfaceVariationPtr;
+    typedef Ptr<const NodeSurfaceVariation> NodeSurfaceVariationPtrConst;
+
+    class NodeModifier;
+
+
+	//! This node modifies a node of the parent object of the object that this node belongs to.
+    //! It allows to extend, cut and morph the parent Surface's meshes.
+    //! It also allows to patch the parent Surface's textures.
+	//! \ingroup model
+    class MUTABLETOOLS_API NodeSurfaceVariation : public NodeSurface
+	{
+	public:
+
+		//-----------------------------------------------------------------------------------------
+		// Life cycle
+		//-----------------------------------------------------------------------------------------
+
+        NodeSurfaceVariation();
+
+		void SerialiseWrapper(OutputArchive& arch) const override;
+		static void Serialise( const NodeSurfaceVariation* pNode, OutputArchive& arch );
+		static NodeSurfaceVariationPtr StaticUnserialise( InputArchive& arch );
+
+		//-----------------------------------------------------------------------------------------
+		// Node Interface
+		//-----------------------------------------------------------------------------------------
+
+        NodePtr Clone( NodeMapPtr pMap = 0 ) const override;
+
+        const NODE_TYPE* GetType() const override;
+		static const NODE_TYPE* GetStaticType();
+
+        virtual int GetInputCount() const override;
+        virtual Node* GetInputNode( int i ) const override;
+        void SetInputNode( int i, NodePtr pNode ) override;
+
+		//-----------------------------------------------------------------------------------------
+		// Own Interface
+		//-----------------------------------------------------------------------------------------
+
+		//! 
+        void AddDefaultSurface(NodeSurface* surface);
+        void AddDefaultModifier(NodeModifier* modifier);
+
+		//! Set the number of tags to consider in this variation
+		void SetVariationCount( int count );
+
+		//!
+		int GetVariationCount() const;
+
+        //!
+        enum class VariationType : uint8_t
+        {
+            //! The variation selection is controlled by tags defined in other surfaces.
+            //! Default value.
+            Tag = 0,
+
+            //! The variation selection is controlled by the state the object is in.
+            State
+        };
+
+        //!
+        void SetVariationType(VariationType);
+
+        //! Set the tag or state name that will enable a specific vartiation
+		void SetVariationTag(int index, const char* strTag);
+
+		//! 
+        void AddVariationSurface(int index, NodeSurface* surface);
+        void AddVariationModifier(int index, NodeModifier* modifier);
+
+		//!}
+
+
+		//-----------------------------------------------------------------------------------------
+		// Interface pattern
+		//-----------------------------------------------------------------------------------------
+		class Private;
+		Private* GetPrivate() const;
+                Node::Private* GetBasePrivate() const override;
+
+	protected:
+
+		//! Forbidden. Manage with the Ptr<> template.
+        ~NodeSurfaceVariation();
+
+	private:
+
+		Private* m_pD;
+
+	};
+
+
+
+}
diff --git a/Engine/Plugins/Experimental/Mutable/Source/MutableTools/Public/Streams.h b/Engine/Plugins/Experimental/Mutable/Source/MutableTools/Public/Streams.h
new file mode 100644
--- /dev/null
+++ b/Engine/Plugins/Experimental/Mutable/Source/MutableTools/Public/Streams.h
@@ -0,0 +1,175 @@
+// Copyright Epic Games, Inc. All Rights Reserved.
+
+#pragma once
+
+#include "MutableRuntime/Public/RefCounted.h"
+#include "MutableRuntime/Public/Serialisation.h"
+#include "MutableRuntime/Public/System.h"
+#include "MutableRuntime/Public/Model.h"
+
+#include "Containers/UnrealString.h"
+
+
+//! This tag is used to identify files containing serialised Model objects. The tag is not added
+//! or checked by the Model serialisation methods, but the involved tools should take care of it.
+#define MUTABLE_COMPILED_MODEL_FILETAG			"amc2"
+
+
+namespace mu
+{
+
+
+	//! Read a file into a serialisation stream.
+	//! \warning This classes are for reference and they are not specially good at error handling,
+	//! or working efficiently.
+	class MUTABLETOOLS_API InputFileStream : public InputStream
+	{
+	public:
+
+		//-----------------------------------------------------------------------------------------
+		// Life cycle
+		//-----------------------------------------------------------------------------------------
+
+		//! Create the stream by reading a file.
+		InputFileStream( const char* strFile );
+
+		~InputFileStream();
+
+
+		//-----------------------------------------------------------------------------------------
+		// InputStream interface
+		//-----------------------------------------------------------------------------------------
+        void Read( void* pData, uint64 size ) override;
+
+                
+		//-----------------------------------------------------------------------------------------
+		// Own interface
+		//-----------------------------------------------------------------------------------------
+
+        //! Return true if the file was found and it could be opened in the constructor.
+        bool IsOpen() const;
+
+        //! Return the number of bytes read from the stream. This should not be used as file 
+        //! position since other bytes might have been read outside this stream, or even Seek might 
+        //! have been called to repeat reads.
+		uint64 GetReadBytes() const;
+
+        //! Return the current position in bytes in the underlying file.
+		uint64 Tell() const;
+
+        //! Reposition the stream index
+        void Seek( uint64 position );
+        
+	private:
+
+		class Private;
+		Private* m_pD;
+	};
+
+
+	//! Write a serialisation stream to a file.
+	//! \warning This classes are for reference and they are not specially good at error handling,
+	//! or working efficiently.
+	class MUTABLETOOLS_API OutputFileStream : public OutputStream
+	{
+	public:
+
+		//-----------------------------------------------------------------------------------------
+		// Life cycle
+		//-----------------------------------------------------------------------------------------
+
+		//! Create a stream that will write to a file.
+		//! The file will be created overwriting any other file with the same name.
+		OutputFileStream( const char* strFile );
+
+		~OutputFileStream();
+
+
+		//-----------------------------------------------------------------------------------------
+		// OutputStream interface
+		//-----------------------------------------------------------------------------------------
+        void Write( const void* pData, uint64 size ) override;
+
+		//-----------------------------------------------------------------------------------------
+		// Own interface
+		//-----------------------------------------------------------------------------------------
+
+		//! Write the pending data to disk. After this call, the stream becomes invalid and cannot
+		//! receive any more serialisation.
+		void Flush();
+
+
+	private:
+
+		class Private;
+		Private* m_pD;
+
+	};
+
+
+    //! Implementation of the streaming interface that actually loads the file synchronously when
+    //! needed.
+    //! \ingroup tools
+    class MUTABLETOOLS_API FileModelStreamer : public ModelStreamer
+    {
+    public:
+
+        //-----------------------------------------------------------------------------------------
+        // Life cycle
+        //-----------------------------------------------------------------------------------------
+        FileModelStreamer();
+        ~FileModelStreamer() override;
+
+        //-----------------------------------------------------------------------------------------
+        // ModelStreamer interface
+        //-----------------------------------------------------------------------------------------
+		OPERATION_ID BeginReadBlock(const mu::Model*, uint64 key0, void* pBuffer, uint64 size ) override;
+        bool IsReadCompleted( OPERATION_ID ) override;
+        void EndRead( OPERATION_ID ) override;
+
+        void OpenWriteFile( const char* strModelName, uint64 key0 ) override;
+        void Write( const void* pBuffer, uint64 size ) override;
+        void CloseWriteFile() override;
+
+        //-----------------------------------------------------------------------------------------
+        // Own interface
+        //-----------------------------------------------------------------------------------------
+
+        //! Get the file name that this streamer would use for a specific model and data ID.
+        static void GetFileName(const FString& strModelName, uint64 key0, FString& outFileName );
+
+    private:
+
+        class Private;
+        Private* m_pD;
+    };
+
+
+    //! ProxyFactory that provides proxies for data stored in temporary files.
+    //! \ingroup tools
+    class MUTABLETOOLS_API ProxyFactoryFiles : public InputArchiveWithProxies::ProxyFactory
+    {
+    public:
+        Ptr<ResourceProxy<mu::Image>> NewImageProxy(InputArchive& arch) override;
+    };
+
+
+    //! ProxyFactory that provides proxies for data stored in a mutable_source file.
+    //! \ingroup tools
+    class MUTABLETOOLS_API ProxyFactoryMutableSourceFile : public InputArchiveWithProxies::ProxyFactory
+    {
+    public:
+        ProxyFactoryMutableSourceFile( const char* strFileName, InputFileStream* );
+        ~ProxyFactoryMutableSourceFile() override;
+
+        Ptr<ResourceProxy<mu::Image>> NewImageProxy(InputArchive& arch) override;
+
+    private:
+
+        class Private;
+        Private* m_pD;
+    };
+
+
+
+}
diff --git a/Engine/Plugins/Experimental/Mutable/Source/MutableTools/Public/Table.h b/Engine/Plugins/Experimental/Mutable/Source/MutableTools/Public/Table.h
new file mode 100644
--- /dev/null
+++ b/Engine/Plugins/Experimental/Mutable/Source/MutableTools/Public/Table.h
@@ -0,0 +1,105 @@
+// Copyright Epic Games, Inc. All Rights Reserved.
+
+#pragma once
+
+#include "RefCounted.h"
+#include "Ptr.h"
+#include "Types.h"
+#include "Serialisation.h"
+
+
+
+namespace mu
+{
+	// Forward declarations
+	class Table;
+	typedef Ptr<Table> TablePtr;
+	typedef Ptr<const Table> TablePtrConst;
+
+	class Mesh;
+	typedef Ptr<Mesh> MeshPtr;
+	typedef Ptr<const Mesh> MeshPtrConst;
+
+	class Image;
+	typedef Ptr<Image> ImagePtr;
+	typedef Ptr<const Image> ImagePtrConst;
+
+
+	//! Types of the values for the table cells.
+	//! \ingroup runtime
+	typedef enum
+	{
+		TCT_NONE,
+		TCT_SCALAR,
+		TCT_COLOUR,
+		TCT_IMAGE,
+		TCT_MESH,
+		TCT_STRING,
+		TCT_COUNT,
+
+		_TCT_FORCE32BITS = 0xFFFFFFFF
+	} TABLE_COLUMN_TYPE;
+
+
+	//! A table that contains many rows and defines attributes like meshes, images,
+	//! colours, etc. for every column. It is useful to define a big number of similarly structured
+	//! objects, by using the NodeDatabase in a model expression.
+	//! \ingroup model
+	class MUTABLETOOLS_API Table : public RefCounted
+	{
+	public:
+
+		//-----------------------------------------------------------------------------------------
+		// Life cycle.
+		//-----------------------------------------------------------------------------------------
+		Table();
+
+		static void Serialise( const Table* p, OutputArchive& arch );
+		static TablePtr StaticUnserialise( InputArchive& arch );
+
+		//-----------------------------------------------------------------------------------------
+		// Own interface
+		//-----------------------------------------------------------------------------------------
+
+		//!
+		void SetName( const char* strName );
+		const char* GetName() const;
+
+		//!
+		int AddColumn( const char* strName, TABLE_COLUMN_TYPE type );
+
+		//! Return the column index with the given name. -1 if not found.
+		int FindColumn( const char* strName ) const;
+
+		//!
+        void AddRow( uint32_t id );
+
+		//! Adds a "NONE" option
+		void SetNoneOption(bool bAddOption);
+
+		//!
+        void SetCell( int column, uint32_t rowId, float value );
+        void SetCell( int column, uint32_t rowId, float r, float g, float b );
+        void SetCell( int column, uint32_t rowId, Image* pImage );
+        void SetCell( int column, uint32_t rowId, Mesh* pMesh );
+        void SetCell( int column, uint32_t rowId, const char* strValue );
+
+
+		//-----------------------------------------------------------------------------------------
+		// Interface pattern
+		//-----------------------------------------------------------------------------------------
+		class Private;
+		Private* GetPrivate() const;
+
+	protected:
+
+		//! Forbidden. Manage with the Ptr<> template.
+		~Table();
+
+	private:
+
+		Private* m_pD;
+
+	};
+
+}
diff --git a/Engine/Plugins/Experimental/Mutable/Source/MutableTools/Public/TypeInfo.h b/Engine/Plugins/Experimental/Mutable/Source/MutableTools/Public/TypeInfo.h
new file mode 100644
--- /dev/null
+++ b/Engine/Plugins/Experimental/Mutable/Source/MutableTools/Public/TypeInfo.h
@@ -0,0 +1,36 @@
+// Copyright Epic Games, Inc. All Rights Reserved.
+
+#pragma once
+
+#include "Types.h"
+#include "Image.h"
+#include "Mesh.h"
+#include "Parameters.h"
+#include "MutableRuntime/Private/ParametersPrivate.h"
+
+// This header file provides additional information about data types defined in the runtime, that
+// can be useful only for the tools. For instance, it provides strings for some enumeration values.
+
+namespace mu
+{
+
+	class MUTABLETOOLS_API TypeInfo
+	{
+	public:
+
+		static const char* s_imageFormatName[size_t(EImageFormat::IF_COUNT)];
+
+		static const char* s_meshBufferSemanticName[ MBS_COUNT ];
+
+		static const char* s_meshBufferFormatName[ MBF_COUNT ];
+
+		static const char* s_projectorTypeName [ static_cast<uint32>(mu::PROJECTOR_TYPE::COUNT) ];
+
+		static const char* s_curveInterpolationModeName [ static_cast<uint8_t>(mu::CurveKeyFrame::InterpMode::Count) ];
+
+		static const char* s_curveTangentModeName [ static_cast<uint8_t> (mu::CurveKeyFrame::TangentMode::Count) ];
+
+		static const char* s_curveTangentWeightModeName [ static_cast<uint8_t> (mu::CurveKeyFrame::TangentWeightMode::Count) ];
+	};
+
+}
