[mutable] Restore changes lost during merge and fix compilation issues.

diff --git a/Engine/Plugins/Experimental/Mutable/Source/CustomizableObject/Classes/CustomizableObjectSystem.h b/Engine/Plugins/Experimental/Mutable/Source/CustomizableObject/Classes/CustomizableObjectSystem.h
--- a/Engine/Plugins/Experimental/Mutable/Source/CustomizableObject/Classes/CustomizableObjectSystem.h
+++ b/Engine/Plugins/Experimental/Mutable/Source/CustomizableObject/Classes/CustomizableObjectSystem.h
@@ -349,7 +349,7 @@ public:
 
 	// Pass a null ptr to reset to the default InstanceLODManagement
 	void SetInstanceLODManagement(UCustomizableInstanceLODManagementBase* NewInstanceLODManagement) 
-		{ CurrentInstanceLODManagement = NewInstanceLODManagement ? NewInstanceLODManagement : DefaultInstanceLODManagement.Get(); }
+		{ CurrentInstanceLODManagement = NewInstanceLODManagement ? NewInstanceLODManagement : ToRawPtr(DefaultInstanceLODManagement); }
 
 	// Find out the version of the plugin
 	UFUNCTION(BlueprintCallable, Category = Status)
diff --git a/Engine/Plugins/Experimental/Mutable/Source/CustomizableObject/Private/CustomizableObject.cpp b/Engine/Plugins/Experimental/Mutable/Source/CustomizableObject/Private/CustomizableObject.cpp
--- a/Engine/Plugins/Experimental/Mutable/Source/CustomizableObject/Private/CustomizableObject.cpp
+++ b/Engine/Plugins/Experimental/Mutable/Source/CustomizableObject/Private/CustomizableObject.cpp
@@ -732,7 +732,7 @@ bool UCustomizableObject::ConditionalAutoCompile()
 	else
 	{
 		const FAssetRegistryModule& AssetRegistryModule = FModuleManager::LoadModuleChecked<FAssetRegistryModule>("AssetRegistry");
-		const FAssetData AssetData = AssetRegistryModule.Get().GetAssetByObjectPath(*GetPathName());
+		const FAssetData AssetData = AssetRegistryModule.Get().GetAssetByObjectPath(FSoftObjectPath(GetPathName()));
 		System->RecompileCustomizableObjectAsync(AssetData, this);
 	}
 
diff --git a/Engine/Plugins/Experimental/Mutable/Source/CustomizableObject/Private/CustomizableObjectDGGUI.cpp b/Engine/Plugins/Experimental/Mutable/Source/CustomizableObject/Private/CustomizableObjectDGGUI.cpp
--- a/Engine/Plugins/Experimental/Mutable/Source/CustomizableObject/Private/CustomizableObjectDGGUI.cpp
+++ b/Engine/Plugins/Experimental/Mutable/Source/CustomizableObject/Private/CustomizableObjectDGGUI.cpp
@@ -70,7 +70,7 @@ void UDGGUI::CustomizableSkeletalMeshPreUpdate(UCustomizableSkeletalComponent* C
 		USkeleton* NextSkeleton = NextMesh->GetSkeleton();
 		if (Parent && NextSkeleton)
 		{
-			if (USkeletalMesh* ParentSkeletalMesh = Parent->SkeletalMesh)
+			if (USkeletalMesh* ParentSkeletalMesh = Cast<USkeletalMesh>(Parent->GetSkinnedAsset()))
 			{
 				if (USkeleton* ParentSkeleton = ParentSkeletalMesh->GetSkeleton())
 				{
diff --git a/Engine/Plugins/Experimental/Mutable/Source/CustomizableObject/Private/CustomizableSkeletalComponent.cpp b/Engine/Plugins/Experimental/Mutable/Source/CustomizableObject/Private/CustomizableSkeletalComponent.cpp
--- a/Engine/Plugins/Experimental/Mutable/Source/CustomizableObject/Private/CustomizableSkeletalComponent.cpp
+++ b/Engine/Plugins/Experimental/Mutable/Source/CustomizableObject/Private/CustomizableSkeletalComponent.cpp
@@ -48,7 +48,7 @@ void UCustomizableSkeletalComponent::SetSkeletalMesh(USkeletalMesh* SkeletalMesh
 		}
 #endif
 
-		if(Parent->SkeletalMesh == SkeletalMesh)
+		if(Parent->GetSkinnedAsset() == SkeletalMesh)
 		{
 			Parent->RecreateRenderState_Concurrent();
 		}
@@ -106,7 +106,7 @@ USkeletalMesh* UCustomizableSkeletalComponent::GetAttachedSkeletalMesh() const
 
 	if (Parent)
 	{
-		return Parent->SkeletalMesh;
+		return Parent->GetSkeletalMeshAsset();
 	}
 
 	return nullptr;
diff --git a/Engine/Plugins/Experimental/Mutable/Source/CustomizableObject/Public/UnrealPortabilityHelpers.h b/Engine/Plugins/Experimental/Mutable/Source/CustomizableObject/Public/UnrealPortabilityHelpers.h
--- a/Engine/Plugins/Experimental/Mutable/Source/CustomizableObject/Public/UnrealPortabilityHelpers.h
+++ b/Engine/Plugins/Experimental/Mutable/Source/CustomizableObject/Public/UnrealPortabilityHelpers.h
@@ -42,18 +42,18 @@ inline const TIndirectArray<FSkeletalMeshLODRenderData>& Helper_GetLODData(const
 	return Resource->LODRenderData;
 }
 
-inline TIndirectArray<FSkeletalMeshLODRenderData>& Helper_GetLODData(const USkeletalMesh* Mesh)
+inline TIndirectArray<FSkeletalMeshLODRenderData>& Helper_GetLODData(const USkinnedAsset* Mesh)
 {
 	return Mesh->GetResourceForRendering()->LODRenderData;
 }
 
-inline TArray<FSkelMeshRenderSection>& Helper_GetLODRenderSections(const USkeletalMesh* Mesh, int LODIndex)
+inline TArray<FSkelMeshRenderSection>& Helper_GetLODRenderSections(const USkinnedAsset* Mesh, int LODIndex)
 {
 	return  Mesh->GetResourceForRendering()->LODRenderData[LODIndex].RenderSections;
 }
 
 #if WITH_EDITORONLY_DATA
-inline FSkeletalMeshModel* Helper_GetImportedModel(const USkeletalMesh* Mesh)
+inline FSkeletalMeshModel* Helper_GetImportedModel(const USkinnedAsset* Mesh)
 {
 	return Mesh->GetImportedModel();
 }
@@ -74,27 +74,27 @@ inline const TArray<FSkelMeshRenderSection>* Helper_LODGetSections(const FSkelet
 	return &(Resource->RenderSections);
 }
 
-inline TArray<struct FSkeletalMeshLODInfo>& Helper_GetLODInfoArray(USkeletalMesh* Mesh)
+inline TArray<struct FSkeletalMeshLODInfo>& Helper_GetLODInfoArray(USkinnedAsset* Mesh)
 {
 	return Mesh->GetLODInfoArray();
 }
 
-inline const TArray<struct FSkeletalMeshLODInfo>& Helper_GetLODInfoArray(const USkeletalMesh* Mesh)
+inline const TArray<struct FSkeletalMeshLODInfo>& Helper_GetLODInfoArray(const USkinnedAsset* Mesh)
 {
 	return Mesh->GetLODInfoArray();
 }
 
-inline int Helper_GetLODNum(const USkeletalMesh* Mesh)
+inline int Helper_GetLODNum(const USkinnedAsset* Mesh)
 {
 	return Mesh->GetLODNum();
 }
 
-inline const TArray<FBoneReference>& Helper_GetLODInfoBonesToRemove(const USkeletalMesh* Mesh, int LODIndex)
+inline const TArray<FBoneReference>& Helper_GetLODInfoBonesToRemove(const USkinnedAsset* Mesh, int LODIndex)
 {
 	return Mesh->GetLODInfo(LODIndex)->BonesToRemove;
 }
 
-inline TArray<FBoneReference>& Helper_GetLODInfoBonesToRemove(USkeletalMesh* Mesh, int LODIndex)
+inline TArray<FBoneReference>& Helper_GetLODInfoBonesToRemove(USkinnedAsset* Mesh, int LODIndex)
 {
 	return Mesh->GetLODInfo(LODIndex)->BonesToRemove;
 }
diff --git a/Engine/Plugins/Experimental/Mutable/Source/CustomizableObjectEditor/Private/AssetTypeActions_CustomizableObject.cpp b/Engine/Plugins/Experimental/Mutable/Source/CustomizableObjectEditor/Private/AssetTypeActions_CustomizableObject.cpp
--- a/Engine/Plugins/Experimental/Mutable/Source/CustomizableObjectEditor/Private/AssetTypeActions_CustomizableObject.cpp
+++ b/Engine/Plugins/Experimental/Mutable/Source/CustomizableObjectEditor/Private/AssetTypeActions_CustomizableObject.cpp
@@ -214,7 +214,7 @@ void FAssetTypeActions_CustomizableObject::RecompileObjects(ERecompileCO Recompi
 	{
 		for (const TWeakObjectPtr<UCustomizableObject>& Object : InObjects)
 		{
-			FAssetData AssetData = AssetRegistryModule.Get().GetAssetByObjectPath(*Object->GetPathName());
+			FAssetData AssetData = AssetRegistryModule.Get().GetAssetByObjectPath(FSoftObjectPath(Object.Get()));
 			if (AssetData.IsValid())
 			{
 				ObjectsToRecompile.Add(AssetData);
diff --git a/Engine/Plugins/Experimental/Mutable/Source/CustomizableObjectEditor/Private/CustomizableObjectCompiler.cpp b/Engine/Plugins/Experimental/Mutable/Source/CustomizableObjectEditor/Private/CustomizableObjectCompiler.cpp
--- a/Engine/Plugins/Experimental/Mutable/Source/CustomizableObjectEditor/Private/CustomizableObjectCompiler.cpp
+++ b/Engine/Plugins/Experimental/Mutable/Source/CustomizableObjectEditor/Private/CustomizableObjectCompiler.cpp
@@ -252,7 +252,7 @@ void FCustomizableObjectCompiler::Compile(UCustomizableObject& Object, const FCo
 
 	for (FAssetData& Element : ArrayAssetData)
 	{
-		ArrayAssetToStream.Add(FSoftObjectPath(Element.ObjectPath.ToString()));
+		ArrayAssetToStream.Add(Element.GetSoftObjectPath());
 	}
 
 	MaxConvertToMutableTextureTime = ComputeAsyncLoadingTimeLimit() * 0.001f;
diff --git a/Engine/Plugins/Experimental/Mutable/Source/CustomizableObjectEditor/Private/CustomizableObjectEditorViewportClient.cpp b/Engine/Plugins/Experimental/Mutable/Source/CustomizableObjectEditor/Private/CustomizableObjectEditorViewportClient.cpp
--- a/Engine/Plugins/Experimental/Mutable/Source/CustomizableObjectEditor/Private/CustomizableObjectEditorViewportClient.cpp
+++ b/Engine/Plugins/Experimental/Mutable/Source/CustomizableObjectEditor/Private/CustomizableObjectEditorViewportClient.cpp
@@ -537,7 +537,7 @@ FCustomizableObjectEditorViewportClient::FCustomizableObjectEditorViewportClient
 void FCustomizableObjectEditorViewportClient::UpdateCameraSetup()
 {
 	static FRotator CustomOrbitRotation(-33.75, -135, 0);
-	if ( (SkeletalMeshComponents.Num() && SkeletalMeshComponents[0].IsValid() && SkeletalMeshComponents[0]->SkeletalMesh)
+	if ( (SkeletalMeshComponents.Num() && SkeletalMeshComponents[0].IsValid() && SkeletalMeshComponents[0]->GetSkinnedAsset())
 		||
 		(StaticMeshComponent.IsValid() && StaticMeshComponent->GetStaticMesh()) )
 	{
@@ -844,7 +844,7 @@ void FCustomizableObjectEditorViewportClient::Draw(FViewport* InViewport, FCanva
 	// Defensive check to avoid unreal crashing inside render if the mesh is degenereated
 	for (TWeakObjectPtr<USkeletalMeshComponent> SkeletalMeshComponent : SkeletalMeshComponents)
 	{
-		if (SkeletalMeshComponent.IsValid() && SkeletalMeshComponent->GetSkeletalMesh() && Helper_GetLODInfoArray(SkeletalMeshComponent->GetSkeletalMesh()).Num() == 0)
+		if (SkeletalMeshComponent.IsValid() && SkeletalMeshComponent->GetSkinnedAsset() && Helper_GetLODInfoArray(SkeletalMeshComponent->GetSkinnedAsset()).Num() == 0)
 		{
 			SkeletalMeshComponent->SetSkeletalMesh(nullptr);
 		}
@@ -1077,7 +1077,7 @@ void FCustomizableObjectEditorViewportClient::DrawUVs(FViewport* InViewport, FCa
 
 		for (TWeakObjectPtr<USkeletalMeshComponent> SkeletalMeshComponent : SkeletalMeshComponents)
 		{
-			if (!SkeletalMeshComponent.IsValid() || !SkeletalMeshComponent->GetSkeletalMesh() || CurrentComponentIndex != ComponentIndex)
+			if (!SkeletalMeshComponent.IsValid() || !SkeletalMeshComponent->GetSkinnedAsset() || CurrentComponentIndex != ComponentIndex)
 			{
 				CurrentComponentIndex++;
 				continue;
@@ -1085,7 +1085,7 @@ void FCustomizableObjectEditorViewportClient::DrawUVs(FViewport* InViewport, FCa
 
 			bool bFoundMaterial = false;
 
-			const FSkeletalMeshRenderData* MeshRes = SkeletalMeshComponent->GetSkeletalMesh()->GetResourceForRendering();
+			const FSkeletalMeshRenderData* MeshRes = SkeletalMeshComponent->GetSkinnedAsset()->GetResourceForRendering();
 			if (UVChannel < (int32)MeshRes->LODRenderData[LODLevel].GetNumTexCoords())
 			{
 				// Find material index from name
@@ -1548,9 +1548,9 @@ void FCustomizableObjectEditorViewportClient::ResetCamera()
 	float MaxSphereRadius = 0.0f;
 	for (const TWeakObjectPtr<UDebugSkelMeshComponent>& SkeletalMeshComponent : SkeletalMeshComponents)
 	{
-		if (SkeletalMeshComponent->GetSkeletalMesh())
+		if (SkeletalMeshComponent->GetSkinnedAsset())
 		{
-			MaxSphereRadius = FMath::Max(MaxSphereRadius, SkeletalMeshComponent->GetSkeletalMesh()->GetBounds().SphereRadius);
+			MaxSphereRadius = FMath::Max(MaxSphereRadius, SkeletalMeshComponent->GetSkinnedAsset()->GetBounds().SphereRadius);
 		}
 	}
 	
@@ -2031,8 +2031,8 @@ void FCustomizableObjectEditorViewportClient::SetAnimation(UAnimationAsset* Anim
 	for (TWeakObjectPtr<USkeletalMeshComponent> SkeletalMeshComponent : SkeletalMeshComponents)
 	{
 		if (SkeletalMeshComponent.IsValid() && Animation != nullptr
-			&& SkeletalMeshComponent->GetSkeletalMesh() != nullptr
-			&& SkeletalMeshComponent->GetSkeletalMesh()->GetSkeleton() == Animation->GetSkeleton()
+			&& SkeletalMeshComponent->GetSkinnedAsset() != nullptr
+			&& SkeletalMeshComponent->GetSkinnedAsset()->GetSkeleton() == Animation->GetSkeleton()
 			)
 		{
 			SkeletalMeshComponent->SetAnimationMode(AnimationType);
@@ -2131,7 +2131,7 @@ void FCustomizableObjectEditorViewportClient::SetFloorOffset(float NewValue)
 {
 	for (TWeakObjectPtr<USkeletalMeshComponent> SkeletalMeshComponent : SkeletalMeshComponents)
 	{
-		USkeletalMesh* Mesh = SkeletalMeshComponent.IsValid() ? SkeletalMeshComponent->GetSkeletalMesh() : nullptr;
+		USkeletalMesh* Mesh = SkeletalMeshComponent.IsValid() ? Cast<USkeletalMesh>(SkeletalMeshComponent->GetSkinnedAsset()) : nullptr;
 
 		if (Mesh)
 		{
@@ -2235,7 +2235,7 @@ void FCustomizableObjectEditorViewportClient::BakeInstance()
 	int32 NumComponents = SkeletalMeshComponents.Num();
 	for (int32 ComponentIndex = 0; ComponentIndex < NumComponents; ++ComponentIndex)
 	{
-		if (SkeletalMeshComponents[ComponentIndex].IsValid() && SkeletalMeshComponents[ComponentIndex]->GetSkeletalMesh())
+		if (SkeletalMeshComponents[ComponentIndex].IsValid() && SkeletalMeshComponents[ComponentIndex]->GetSkinnedAsset())
 		{
 			bHasSkeletalMesh = true;
 		}
@@ -2304,7 +2304,7 @@ void FCustomizableObjectEditorViewportClient::BakeInstance()
 			for (int32 ComponentIndex = 0; ComponentIndex < NumComponents; ++ComponentIndex)
 			{
 				USkeletalMesh* Mesh = SkeletalMeshComponents.Num() && SkeletalMeshComponents[ComponentIndex].IsValid() ?
-					SkeletalMeshComponents[ComponentIndex]->GetSkeletalMesh() : nullptr;
+					Cast<USkeletalMesh>(SkeletalMeshComponents[ComponentIndex]->GetSkinnedAsset()) : nullptr;
 				
 				if (!Mesh)
 				{
@@ -2814,9 +2814,9 @@ void FCustomizableObjectEditorViewportClient::ShowInstanceGeometryInformation(FC
 	// Show total number of triangles and vertices
 	for (TWeakObjectPtr<USkeletalMeshComponent> SkeletalMeshComponent : SkeletalMeshComponents)
 	{
-		if (SkeletalMeshComponent.IsValid() && SkeletalMeshComponent->GetSkeletalMesh())
+		if (SkeletalMeshComponent.IsValid() && SkeletalMeshComponent->GetSkinnedAsset())
 		{
-			const FSkeletalMeshRenderData* MeshRes = SkeletalMeshComponent->GetSkeletalMesh()->GetResourceForRendering();
+			const FSkeletalMeshRenderData* MeshRes = SkeletalMeshComponent->GetSkinnedAsset()->GetResourceForRendering();
 			int32 NumTriangles;
 			int32 NumVertices;
 			int32 NumLODLevel = MeshRes->LODRenderData.Num();
diff --git a/Engine/Plugins/Experimental/Mutable/Source/CustomizableObjectEditor/Private/CustomizableObjectInstanceFactory.cpp b/Engine/Plugins/Experimental/Mutable/Source/CustomizableObjectEditor/Private/CustomizableObjectInstanceFactory.cpp
--- a/Engine/Plugins/Experimental/Mutable/Source/CustomizableObjectEditor/Private/CustomizableObjectInstanceFactory.cpp
+++ b/Engine/Plugins/Experimental/Mutable/Source/CustomizableObjectEditor/Private/CustomizableObjectInstanceFactory.cpp
@@ -118,7 +118,7 @@ void UCustomizableObjectInstanceFactory::PostSpawnActor(UObject* Asset, AActor*
 			if (USkeletalMeshComponent* SkeletalMeshComp = NewCSMActor->GetSkeletalMeshComponentAt(ComponentIndex))
 			{
 				SkeletalMeshComp->UnregisterComponent();
-				SkeletalMeshComp->SkeletalMesh = SkeletalMesh;
+				SkeletalMeshComp->SetSkinnedAsset(SkeletalMesh);
 
 				if (ComponentIndex == 0 && NewCSMActor->GetWorld()->IsGameWorld())
 				{
diff --git a/Engine/Plugins/Experimental/Mutable/Source/CustomizableObjectEditor/Private/SCustomizableObjectEditorAdvancedPreviewSettings.cpp b/Engine/Plugins/Experimental/Mutable/Source/CustomizableObjectEditor/Private/SCustomizableObjectEditorAdvancedPreviewSettings.cpp
--- a/Engine/Plugins/Experimental/Mutable/Source/CustomizableObjectEditor/Private/SCustomizableObjectEditorAdvancedPreviewSettings.cpp
+++ b/Engine/Plugins/Experimental/Mutable/Source/CustomizableObjectEditor/Private/SCustomizableObjectEditorAdvancedPreviewSettings.cpp
@@ -115,7 +115,7 @@ FReply SLevelSelectWindow::OnButtonClick(EAppReturnType::Type ButtonID)
 			return FReply::Handled();
 		}
 
-		LevelAssetPath = LevelAssetData.ObjectPath.ToString();
+		LevelAssetPath = LevelAssetData.GetObjectPathString();
 
 		Result = LevelProfileManage::LoadProfileUObjectNames(LevelAssetPath, ArrayDirectionalLightName, ArrayPostProcessName, ArraySkyLightName);
 
diff --git a/Engine/Plugins/Experimental/Mutable/Source/CustomizableObjectEditor/Private/SCustomizableObjectEditorViewport.cpp b/Engine/Plugins/Experimental/Mutable/Source/CustomizableObjectEditor/Private/SCustomizableObjectEditorViewport.cpp
--- a/Engine/Plugins/Experimental/Mutable/Source/CustomizableObjectEditor/Private/SCustomizableObjectEditorViewport.cpp
+++ b/Engine/Plugins/Experimental/Mutable/Source/CustomizableObjectEditor/Private/SCustomizableObjectEditorViewport.cpp
@@ -881,12 +881,12 @@ void SCustomizableObjectEditorViewportTabBody::GenerateUVMaterialOptions()
 
 	for (UDebugSkelMeshComponent* PreviewSkeletalMeshComponent : PreviewSkeletalMeshComponents)
 	{
-		if (PreviewSkeletalMeshComponent != nullptr && PreviewSkeletalMeshComponent->SkeletalMesh != nullptr && PreviewSkeletalMeshComponent->SkeletalMesh->GetResourceForRendering() != nullptr)
+		if (PreviewSkeletalMeshComponent != nullptr && PreviewSkeletalMeshComponent->GetSkinnedAsset() != nullptr && PreviewSkeletalMeshComponent->GetSkinnedAsset()->GetResourceForRendering() != nullptr)
 		{
 			TMap<FString, int32> MaterialLODs;
 
 			// Add Suffix "__X" for materials with duplicated names and Suffing " LOD_X" for multiple LODs.
-			const FSkeletalMeshRenderData* MeshRes = PreviewSkeletalMeshComponent->SkeletalMesh->GetResourceForRendering();
+			const FSkeletalMeshRenderData* MeshRes = PreviewSkeletalMeshComponent->GetSkinnedAsset()->GetResourceForRendering();
 			const TArray<UMaterialInterface*> Materials = PreviewSkeletalMeshComponent->GetMaterials();
 			for (UMaterialInterface* m : Materials)
 			{
@@ -990,10 +990,10 @@ void SCustomizableObjectEditorViewportTabBody::GenerateUVChannelOptions(bool bRe
 
 	UDebugSkelMeshComponent* PreviewSkeletalMeshComponent = PreviewSkeletalMeshComponents[ComponentIndex];
 
-	if (PreviewSkeletalMeshComponent != nullptr && PreviewSkeletalMeshComponent->SkeletalMesh != nullptr
-		&& PreviewSkeletalMeshComponent->SkeletalMesh->GetResourceForRendering() != nullptr)
+	if (PreviewSkeletalMeshComponent != nullptr && PreviewSkeletalMeshComponent->GetSkinnedAsset() != nullptr
+		&& PreviewSkeletalMeshComponent->GetSkinnedAsset()->GetResourceForRendering() != nullptr)
 	{
-		const FSkeletalMeshRenderData* MeshRes = PreviewSkeletalMeshComponent->SkeletalMesh->GetResourceForRendering();
+		const FSkeletalMeshRenderData* MeshRes = PreviewSkeletalMeshComponent->GetSkinnedAsset()->GetResourceForRendering();
 		
 		int32 UVChannels = MeshRes->LODRenderData[LODIndex].GetNumTexCoords();
 		
@@ -1064,9 +1064,9 @@ int32 SCustomizableObjectEditorViewportTabBody::GetLODModelCount() const
 
 	for (UDebugSkelMeshComponent* PreviewComponent : GetSkeletalMeshComponents())
 	{
-		if (PreviewComponent && PreviewComponent->SkeletalMesh)
+		if (PreviewComponent && PreviewComponent->GetSkinnedAsset())
 		{
-			const TIndirectArray<FSkeletalMeshLODRenderData>& LODModels = Helper_GetLODData(PreviewComponent->SkeletalMesh);
+			const TIndirectArray<FSkeletalMeshLODRenderData>& LODModels = Helper_GetLODData(PreviewComponent->GetSkinnedAsset());
 			LODModelCount = FMath::Max(LODModelCount, LODModels.Num());
 		}
 	}
diff --git a/Engine/Plugins/Experimental/Mutable/Source/CustomizableObjectPopulationEditor/Private/CustomizableObjectPopulationClassEditor.cpp b/Engine/Plugins/Experimental/Mutable/Source/CustomizableObjectPopulationEditor/Private/CustomizableObjectPopulationClassEditor.cpp
--- a/Engine/Plugins/Experimental/Mutable/Source/CustomizableObjectPopulationEditor/Private/CustomizableObjectPopulationClassEditor.cpp
+++ b/Engine/Plugins/Experimental/Mutable/Source/CustomizableObjectPopulationEditor/Private/CustomizableObjectPopulationClassEditor.cpp
@@ -978,7 +978,7 @@ void FCustomizableObjectPopulationClassEditor::GeneratePopulationClassInstances(
 			{
 				//Search for asset with the same name
 				FString ObjectPath = FilePath + FString(".") + ObjectName;
-				AssetData = AssetRegistryModule.Get().GetAssetByObjectPath(*ObjectPath);
+				AssetData = AssetRegistryModule.Get().GetAssetByObjectPath(FSoftObjectPath(ObjectPath));
 
 				// Popup window
 				if (AssetData.IsValid() && AssetData.GetClass() == UCustomizableObjectInstance::StaticClass() && RetType != EAppReturnType::YesAll)
diff --git a/Engine/Plugins/Experimental/Mutable/Source/CustomizableObjectPopulationEditor/Private/CustomizableObjectPopulationEditor.cpp b/Engine/Plugins/Experimental/Mutable/Source/CustomizableObjectPopulationEditor/Private/CustomizableObjectPopulationEditor.cpp
--- a/Engine/Plugins/Experimental/Mutable/Source/CustomizableObjectPopulationEditor/Private/CustomizableObjectPopulationEditor.cpp
+++ b/Engine/Plugins/Experimental/Mutable/Source/CustomizableObjectPopulationEditor/Private/CustomizableObjectPopulationEditor.cpp
@@ -615,7 +615,7 @@ void FCustomizableObjectPopulationEditor::GeneratePopulationInstances()
 			{
 				//Search for asset with the same name
 				FString ObjectPath = FilePath + FString(".") + ObjectName;
-				AssetData = AssetRegistryModule.Get().GetAssetByObjectPath(*ObjectPath);
+				AssetData = AssetRegistryModule.Get().GetAssetByObjectPath(FSoftObjectPath(ObjectPath));
 
 				// Popup window
 				if (AssetData.IsValid() && AssetData.GetClass() == UCustomizableObjectInstance::StaticClass() && RetType != EAppReturnType::YesAll)
diff --git a/Engine/Plugins/Experimental/Mutable/Source/CustomizableObjectPopulationEditor/Private/CustomizableObjectPopulationEditorViewportClient.cpp b/Engine/Plugins/Experimental/Mutable/Source/CustomizableObjectPopulationEditor/Private/CustomizableObjectPopulationEditorViewportClient.cpp
--- a/Engine/Plugins/Experimental/Mutable/Source/CustomizableObjectPopulationEditor/Private/CustomizableObjectPopulationEditorViewportClient.cpp
+++ b/Engine/Plugins/Experimental/Mutable/Source/CustomizableObjectPopulationEditor/Private/CustomizableObjectPopulationEditorViewportClient.cpp
@@ -86,7 +86,7 @@ void FCustomizableObjectPopulationEditorViewportClient::ProcessClick(FSceneView&
 	// Checking hit
 	for (int32 i = 0; i < ColliderComponents.Num(); ++i)
 	{
-		if (SkeletalMeshComponents[i] && SkeletalMeshComponents[i]->SkeletalMesh)
+		if (SkeletalMeshComponents[i] && SkeletalMeshComponents[i]->GetSkinnedAsset())
 		{	
 			FVector BoxExtent = SkeletalMeshComponents[i]->Bounds.BoxExtent;
 			ColliderComponents[i]->SetVisibility(false);
diff --git a/Engine/Plugins/Experimental/Mutable/Source/MutableTools/Private/ASTOpImageCompose.cpp b/Engine/Plugins/Experimental/Mutable/Source/MutableTools/Private/ASTOpImageCompose.cpp
--- a/Engine/Plugins/Experimental/Mutable/Source/MutableTools/Private/ASTOpImageCompose.cpp
+++ b/Engine/Plugins/Experimental/Mutable/Source/MutableTools/Private/ASTOpImageCompose.cpp
@@ -63,7 +63,7 @@ uint64 ASTOpImageCompose::Hash() const
 //-------------------------------------------------------------------------------------------------
 mu::Ptr<ASTOp> ASTOpImageCompose::Clone(MapChildFunc& mapChild) const
 {
-    Ptr<ASTOpImageCompose> n = new ASTOpImageCompose();
+	mu::Ptr<ASTOpImageCompose> n = new ASTOpImageCompose();
     n->Layout = mapChild(Layout.child());
 	n->Base = mapChild(Base.child());
 	n->BlockImage = mapChild(BlockImage.child());
@@ -205,9 +205,9 @@ void ASTOpImageCompose::GetLayoutBlockSize(int* pBlockX, int* pBlockY)
 }
 
 
-Ptr<ASTOp> ASTOpImageCompose::OptimiseSemantic(const MODEL_OPTIMIZATION_OPTIONS& options) const
+mu::Ptr<ASTOp> ASTOpImageCompose::OptimiseSemantic(const MODEL_OPTIMIZATION_OPTIONS& options) const
 {
-	Ptr<ASTOp> at;
+	mu::Ptr<ASTOp> at;
 
 	auto baseAt = Base.child();
 	auto blockAt = BlockImage.child();
@@ -221,7 +221,7 @@ Ptr<ASTOp> ASTOpImageCompose::OptimiseSemantic(const MODEL_OPTIMIZATION_OPTIONS&
 		blockAt)
 	{
 		auto typedLayout = dynamic_cast<const ASTOpConstantResource*>(layoutAt.get());
-		Ptr<const mu::Layout> pLayout = static_cast<const mu::Layout*>(typedLayout->GetValue().get());
+		mu::Ptr<const mu::Layout> pLayout = static_cast<const mu::Layout*>(typedLayout->GetValue().get());
 
 		// Constant single-block full layout?
 		if (pLayout->GetBlockCount() == 1
@@ -245,7 +245,7 @@ Ptr<ASTOp> ASTOpImageCompose::OptimiseSemantic(const MODEL_OPTIMIZATION_OPTIONS&
 				&&
 				baseDesc.m_format != EImageFormat::IF_NONE)
 			{
-				Ptr<ASTOpImagePixelFormat> reformat = new ASTOpImagePixelFormat;
+				mu::Ptr<ASTOpImagePixelFormat> reformat = new ASTOpImagePixelFormat;
 				reformat->Format = baseDesc.m_format;
 				reformat->FormatIfAlpha = EImageFormat::IF_NONE;
 				reformat->Source = at;
@@ -258,7 +258,7 @@ Ptr<ASTOp> ASTOpImageCompose::OptimiseSemantic(const MODEL_OPTIMIZATION_OPTIONS&
 				&&
 				baseDesc.m_size[1] != 0)
 			{
-				Ptr<ASTOpFixed> resize = new ASTOpFixed;
+				mu::Ptr<ASTOpFixed> resize = new ASTOpFixed;
 				resize->op.type = OP_TYPE::IM_RESIZE;
 				resize->op.args.ImageResize.size[0] = baseDesc.m_size[0];
 				resize->op.args.ImageResize.size[1] = baseDesc.m_size[1];
diff --git a/Engine/Plugins/Experimental/Mutable/Source/MutableTools/Private/ASTOpImageMipmap.cpp b/Engine/Plugins/Experimental/Mutable/Source/MutableTools/Private/ASTOpImageMipmap.cpp
--- a/Engine/Plugins/Experimental/Mutable/Source/MutableTools/Private/ASTOpImageMipmap.cpp
+++ b/Engine/Plugins/Experimental/Mutable/Source/MutableTools/Private/ASTOpImageMipmap.cpp
@@ -59,7 +59,7 @@ uint64 ASTOpImageMipmap::Hash() const
 
 mu::Ptr<ASTOp> ASTOpImageMipmap::Clone(MapChildFunc& mapChild) const
 {
-    Ptr<ASTOpImageMipmap> n = new ASTOpImageMipmap();
+	mu::Ptr<ASTOpImageMipmap> n = new ASTOpImageMipmap();
     n->Source = mapChild(Source.child());
 	n->Levels = Levels;
 	n->BlockLevels = BlockLevels;
@@ -106,11 +106,11 @@ void ASTOpImageMipmap::Link( PROGRAM& program, const FLinkerOptions* )
 }
 
 
-Ptr<ASTOp> ASTOpImageMipmap::OptimiseSink(const MODEL_OPTIMIZATION_OPTIONS& options, OPTIMIZE_SINK_CONTEXT& context) const
+mu::Ptr<ASTOp> ASTOpImageMipmap::OptimiseSink(const MODEL_OPTIMIZATION_OPTIONS& options, OPTIMIZE_SINK_CONTEXT& context) const
 {
-	Ptr<ASTOp> at;
+	mu::Ptr<ASTOp> at;
 
-	Ptr<ASTOp> sourceAt = Source.child();
+	mu::Ptr<ASTOp> sourceAt = Source.child();
 	switch (sourceAt->GetOpType())
 	{
 
@@ -128,8 +128,8 @@ Ptr<ASTOp> ASTOpImageMipmap::OptimiseSink(const MODEL_OPTIMIZATION_OPTIONS& opti
 	{
 		// Swap unless the mipmap operation builds only the tail or is compressed.
 		// Otherwise, we could fall in a loop of swapping mipmaps and pixelformats.
-		Ptr<ASTOpImageMipmap> mop = mu::Clone<ASTOpImageMipmap>(this);
-		Ptr<ASTOpImagePixelFormat> fop = mu::Clone<ASTOpImagePixelFormat>(sourceAt);
+		mu::Ptr<ASTOpImageMipmap> mop = mu::Clone<ASTOpImageMipmap>(this);
+		mu::Ptr<ASTOpImagePixelFormat> fop = mu::Clone<ASTOpImagePixelFormat>(sourceAt);
 		bool isCompressedFormat = IsCompressedFormat(fop->Format);
 		if (isCompressedFormat && !mop->bOnlyTail)
 		{
@@ -227,7 +227,7 @@ bool ASTOpImageMipmap::IsImagePlainConstant(vec4<float>& colour) const
 
 mu::Ptr<ImageSizeExpression> ASTOpImageMipmap::GetImageSizeExpression() const
 {
-	Ptr<ImageSizeExpression> pRes;
+	mu::Ptr<ImageSizeExpression> pRes;
 
 	if (Source.child())
 	{
@@ -245,7 +245,7 @@ mu::Ptr<ImageSizeExpression> ASTOpImageMipmap::GetImageSizeExpression() const
 //---------------------------------------------------------------------------------------------
 //---------------------------------------------------------------------------------------------
 //---------------------------------------------------------------------------------------------
-Ptr<ASTOp> Sink_ImageMipmapAST::Apply(const ASTOp* root)
+mu::Ptr<ASTOp> Sink_ImageMipmapAST::Apply(const ASTOp* root)
 {
 	check(root->GetOpType() == OP_TYPE::IM_MIPMAP);
 
@@ -258,7 +258,7 @@ Ptr<ASTOp> Sink_ImageMipmapAST::Apply(const ASTOp* root)
 		return nullptr;
 	}
 
-	Ptr<ASTOp> newSource;
+	mu::Ptr<ASTOp> newSource;
 	m_initialSource = m_root->Source.child();
 
 	// Before sinking, see if it is worth splitting into miptail and mip.
@@ -271,13 +271,13 @@ Ptr<ASTOp> Sink_ImageMipmapAST::Apply(const ASTOp* root)
 		)
 	{
 		// the block mipmaps can be done before composition is done.
-		Ptr<ASTOpImageMipmap> newMip = mu::Clone<ASTOpImageMipmap>(m_root);
+		mu::Ptr<ASTOpImageMipmap> newMip = mu::Clone<ASTOpImageMipmap>(m_root);
 		newMip->Levels = m_root->BlockLevels;
 		newMip->BlockLevels = m_root->BlockLevels;
 		newMip->bOnlyTail = false;
 
 		// the smallest mipmaps after the composition is done.
-		Ptr<ASTOpImageMipmap> topMipOp = mu::Clone<ASTOpImageMipmap>(m_root);
+		mu::Ptr<ASTOpImageMipmap> topMipOp = mu::Clone<ASTOpImageMipmap>(m_root);
 		topMipOp->bOnlyTail = true;
 
 		// Proceed
@@ -306,7 +306,7 @@ Ptr<ASTOp> Sink_ImageMipmapAST::Apply(const ASTOp* root)
 //---------------------------------------------------------------------------------------------
 //---------------------------------------------------------------------------------------------
 //---------------------------------------------------------------------------------------------
-Ptr<ASTOp> Sink_ImageMipmapAST::Visit(Ptr<ASTOp> at, const ASTOpImageMipmap* currentMipmapOp)
+mu::Ptr<ASTOp> Sink_ImageMipmapAST::Visit(mu::Ptr<ASTOp> at, const ASTOpImageMipmap* currentMipmapOp)
 {
 	if (!at) return nullptr;
 
@@ -323,7 +323,7 @@ Ptr<ASTOp> Sink_ImageMipmapAST::Visit(Ptr<ASTOp> at, const ASTOpImageMipmap* cur
 		return cacheIt->second;
 	}
 
-	Ptr<ASTOp> newAt = at;
+	mu::Ptr<ASTOp> newAt = at;
 	switch (at->GetOpType())
 	{
 
@@ -406,7 +406,7 @@ Ptr<ASTOp> Sink_ImageMipmapAST::Visit(Ptr<ASTOp> at, const ASTOpImageMipmap* cur
 				}
 			}
 
-			Ptr<ASTOpImageMipmap> newMip = mu::Clone<ASTOpImageMipmap>(currentMipmapOp);
+			mu::Ptr<ASTOpImageMipmap> newMip = mu::Clone<ASTOpImageMipmap>(currentMipmapOp);
 			newMip->Levels = blockLevels;
 			newMip->BlockLevels = blockLevels;
 			newMip->bOnlyTail = false;
@@ -418,7 +418,7 @@ Ptr<ASTOp> Sink_ImageMipmapAST::Visit(Ptr<ASTOp> at, const ASTOpImageMipmap* cur
 
 			// We need to add a mipmap on top to finish the mipmapping
 			{
-				Ptr<ASTOpImageMipmap> topMipOp = mu::Clone<ASTOpImageMipmap>(currentMipmapOp);
+				mu::Ptr<ASTOpImageMipmap> topMipOp = mu::Clone<ASTOpImageMipmap>(currentMipmapOp);
 				topMipOp->Source = newOp;
 				topMipOp->bOnlyTail = true;
 				newAt = topMipOp;
@@ -435,7 +435,7 @@ Ptr<ASTOp> Sink_ImageMipmapAST::Visit(Ptr<ASTOp> at, const ASTOpImageMipmap* cur
 	// end on line, replace with mipmap
 	if (at == newAt && at != m_initialSource)
 	{
-		Ptr<ASTOpImageMipmap> newOp = mu::Clone<ASTOpImageMipmap>(currentMipmapOp);
+		mu::Ptr<ASTOpImageMipmap> newOp = mu::Clone<ASTOpImageMipmap>(currentMipmapOp);
 		check(newOp->GetOpType() == OP_TYPE::IM_MIPMAP);
 
 		newOp->Source = at;
diff --git a/Engine/Plugins/Experimental/Mutable/Source/MutableTools/Private/ASTOpImagePixelFormat.cpp b/Engine/Plugins/Experimental/Mutable/Source/MutableTools/Private/ASTOpImagePixelFormat.cpp
--- a/Engine/Plugins/Experimental/Mutable/Source/MutableTools/Private/ASTOpImagePixelFormat.cpp
+++ b/Engine/Plugins/Experimental/Mutable/Source/MutableTools/Private/ASTOpImagePixelFormat.cpp
@@ -50,7 +50,7 @@ uint64 ASTOpImagePixelFormat::Hash() const
 
 mu::Ptr<ASTOp> ASTOpImagePixelFormat::Clone(MapChildFunc& mapChild) const
 {
-    Ptr<ASTOpImagePixelFormat> n = new ASTOpImagePixelFormat();
+	mu::Ptr<ASTOpImagePixelFormat> n = new ASTOpImagePixelFormat();
     n->Source = mapChild(Source.child());
 	n->Format = Format;
 	n->FormatIfAlpha = FormatIfAlpha;
@@ -86,11 +86,11 @@ void ASTOpImagePixelFormat::Link( PROGRAM& program, const FLinkerOptions* )
 }
 
 
-Ptr<ASTOp> ASTOpImagePixelFormat::OptimiseSink(const MODEL_OPTIMIZATION_OPTIONS& options, OPTIMIZE_SINK_CONTEXT& context) const
+mu::Ptr<ASTOp> ASTOpImagePixelFormat::OptimiseSink(const MODEL_OPTIMIZATION_OPTIONS& options, OPTIMIZE_SINK_CONTEXT& context) const
 {
-	Ptr<ASTOp> at;
+	mu::Ptr<ASTOp> at;
 
-	Ptr<ASTOp> sourceAt = Source.child();
+	mu::Ptr<ASTOp> sourceAt = Source.child();
 
 	EImageFormat format = Format;
 	bool isCompressedFormat = IsCompressedFormat(Format);
@@ -104,7 +104,7 @@ Ptr<ASTOp> ASTOpImagePixelFormat::OptimiseSink(const MODEL_OPTIMIZATION_OPTIONS&
 	{
 		// Keep only the top pixel format
 		const ASTOpImagePixelFormat* typedSource = dynamic_cast<const ASTOpImagePixelFormat*>(sourceAt.get());
-		Ptr<ASTOpImagePixelFormat> formatOp = mu::Clone<ASTOpImagePixelFormat>(this);
+		mu::Ptr<ASTOpImagePixelFormat> formatOp = mu::Clone<ASTOpImagePixelFormat>(this);
 		formatOp->Source = typedSource->Source.child();
 		at = formatOp;
 		break;
@@ -115,9 +115,9 @@ Ptr<ASTOp> ASTOpImagePixelFormat::OptimiseSink(const MODEL_OPTIMIZATION_OPTIONS&
 		// This op doesn't support compressed formats
 		if (!isCompressedFormat)
 		{
-			Ptr<ASTOpFixed> newOp = mu::Clone<ASTOpFixed>(sourceAt);
+			mu::Ptr<ASTOpFixed> newOp = mu::Clone<ASTOpFixed>(sourceAt);
 
-			Ptr<ASTOpImagePixelFormat> fop = mu::Clone<ASTOpImagePixelFormat>(this);
+			mu::Ptr<ASTOpImagePixelFormat> fop = mu::Clone<ASTOpImagePixelFormat>(this);
 			fop->Source = newOp->children[newOp->op.args.ImageDisplace.source].child();
 			newOp->SetChild(newOp->op.args.ImageDisplace.source, fop);
 
@@ -135,7 +135,7 @@ Ptr<ASTOp> ASTOpImagePixelFormat::OptimiseSink(const MODEL_OPTIMIZATION_OPTIONS&
 		{
 			auto newOp = mu::Clone<ASTOpFixed>(sourceAt);
 
-			Ptr<ASTOpImagePixelFormat> fop = mu::Clone<ASTOpImagePixelFormat>(this);
+			mu::Ptr<ASTOpImagePixelFormat> fop = mu::Clone<ASTOpImagePixelFormat>(this);
 			fop->Source = newOp->children[newOp->op.args.ImageRasterMesh.image].child();
 			newOp->SetChild(newOp->op.args.ImageRasterMesh.image, fop);
 
@@ -245,7 +245,7 @@ bool ASTOpImagePixelFormat::IsImagePlainConstant(vec4<float>& colour) const
 
 mu::Ptr<ImageSizeExpression> ASTOpImagePixelFormat::GetImageSizeExpression() const
 {
-	Ptr<ImageSizeExpression> pRes;
+	mu::Ptr<ImageSizeExpression> pRes;
 
 	if (Source.child())
 	{
@@ -263,7 +263,7 @@ mu::Ptr<ImageSizeExpression> ASTOpImagePixelFormat::GetImageSizeExpression() con
 //---------------------------------------------------------------------------------------------
 //---------------------------------------------------------------------------------------------
 //---------------------------------------------------------------------------------------------
-Ptr<ASTOp> Sink_ImagePixelFormatAST::Apply(const ASTOp* root)
+mu::Ptr<ASTOp> Sink_ImagePixelFormatAST::Apply(const ASTOp* root)
 {
 	m_root = dynamic_cast<const ASTOpImagePixelFormat*>(root);
 	m_oldToNew.clear();
@@ -272,7 +272,7 @@ Ptr<ASTOp> Sink_ImagePixelFormatAST::Apply(const ASTOp* root)
 	check(root->GetOpType() == OP_TYPE::IM_PIXELFORMAT);
 
 	m_initialSource = m_root->Source.child();
-	Ptr<ASTOp> newSource = Visit(m_initialSource, m_root);
+	mu::Ptr<ASTOp> newSource = Visit(m_initialSource, m_root);
 
 	m_root = nullptr;
 
@@ -287,7 +287,7 @@ Ptr<ASTOp> Sink_ImagePixelFormatAST::Apply(const ASTOp* root)
 
 
 //---------------------------------------------------------------------------------------------
-Ptr<ASTOp> Sink_ImagePixelFormatAST::Visit(Ptr<ASTOp> at, const ASTOpImagePixelFormat* currentFormatOp)
+mu::Ptr<ASTOp> Sink_ImagePixelFormatAST::Visit(mu::Ptr<ASTOp> at, const ASTOpImagePixelFormat* currentFormatOp)
 {
 	if (!at) return nullptr;
 
@@ -308,7 +308,7 @@ Ptr<ASTOp> Sink_ImagePixelFormatAST::Visit(Ptr<ASTOp> at, const ASTOpImagePixelF
 		return cacheIt->second;
 	}
 
-	Ptr<ASTOp> newAt = at;
+	mu::Ptr<ASTOp> newAt = at;
 	switch (at->GetOpType())
 	{
 
@@ -501,7 +501,7 @@ Ptr<ASTOp> Sink_ImagePixelFormatAST::Visit(Ptr<ASTOp> at, const ASTOpImagePixelF
 	// end on tree branch, replace with format
 	if (at == newAt && at != m_initialSource)
 	{
-		Ptr<ASTOpImagePixelFormat> newOp = mu::Clone<ASTOpImagePixelFormat>(currentFormatOp);
+		mu::Ptr<ASTOpImagePixelFormat> newOp = mu::Clone<ASTOpImagePixelFormat>(currentFormatOp);
 		check(newOp->GetOpType() == OP_TYPE::IM_PIXELFORMAT);
 
 		newOp->Source = at;
diff --git a/Engine/Plugins/Experimental/Mutable/Source/MutableTools/Private/ASTOpMeshFormat.cpp b/Engine/Plugins/Experimental/Mutable/Source/MutableTools/Private/ASTOpMeshFormat.cpp
--- a/Engine/Plugins/Experimental/Mutable/Source/MutableTools/Private/ASTOpMeshFormat.cpp
+++ b/Engine/Plugins/Experimental/Mutable/Source/MutableTools/Private/ASTOpMeshFormat.cpp
@@ -51,7 +51,7 @@ uint64 ASTOpMeshFormat::Hash() const
 
 mu::Ptr<ASTOp> ASTOpMeshFormat::Clone(MapChildFunc& mapChild) const
 {
-    Ptr<ASTOpMeshFormat> n = new ASTOpMeshFormat();
+	mu::Ptr<ASTOpMeshFormat> n = new ASTOpMeshFormat();
     n->Source = mapChild(Source.child());
     n->Format = mapChild(Format.child());
 	n->Buffers = Buffers;
@@ -88,9 +88,9 @@ void ASTOpMeshFormat::Link( PROGRAM& program, const FLinkerOptions* )
 }
 
 
-Ptr<ASTOp> ASTOpMeshFormat::OptimiseSink(const MODEL_OPTIMIZATION_OPTIONS& options, OPTIMIZE_SINK_CONTEXT& context) const
+mu::Ptr<ASTOp> ASTOpMeshFormat::OptimiseSink(const MODEL_OPTIMIZATION_OPTIONS& options, OPTIMIZE_SINK_CONTEXT& context) const
 {
-	Ptr<ASTOp> at = context.MeshFormatSinker.Apply(this);
+	mu::Ptr<ASTOp> at = context.MeshFormatSinker.Apply(this);
 	return at;
 }
 
@@ -120,9 +120,9 @@ mu::Ptr<ASTOp> Sink_MeshFormatAST::Apply(const ASTOp* root)
 
 namespace
 {
-	Ptr<const Mesh> FindBaseMeshConstant(Ptr<ASTOp> at)
+	mu::Ptr<const Mesh> FindBaseMeshConstant(mu::Ptr<ASTOp> at)
 	{
-		Ptr<const Mesh> res;
+		mu::Ptr<const Mesh> res;
 
 		switch (at->GetOpType())
 		{
@@ -269,7 +269,7 @@ mu::Ptr<ASTOp> Sink_MeshFormatAST::Visit(const mu::Ptr<ASTOp>& at, const ASTOpMe
 		MeshPtrConst pTargetFormat = FindBaseMeshConstant(currentFormatOp->Format.child());
 		MeshPtrConst pTargetMorphFormat = MakeMorphTargetFormat(pTargetFormat);
 
-		Ptr<ASTOpConstantResource> motaop = new ASTOpConstantResource();
+		mu::Ptr<ASTOpConstantResource> motaop = new ASTOpConstantResource();
 		motaop->type = OP_TYPE::ME_CONSTANT;
 		motaop->SetValue(pTargetMorphFormat, false /* useDiskCache */);
 		auto targetMorphFormatAt = motaop;
@@ -278,7 +278,7 @@ mu::Ptr<ASTOp> Sink_MeshFormatAST::Visit(const mu::Ptr<ASTOp>& at, const ASTOpMe
 		{
 			if (newOp->children[newOp->op.args.MeshMorph2.targets[t]])
 			{
-				Ptr<ASTOpMeshFormat> newFormat = mu::Clone<ASTOpMeshFormat>(currentFormatOp);
+				mu::Ptr<ASTOpMeshFormat> newFormat = mu::Clone<ASTOpMeshFormat>(currentFormatOp);
 				newFormat->Buffers =
 					OP::MeshFormatArgs::BT_VERTEX
 					| OP::MeshFormatArgs::BT_IGNORE_MISSING;
@@ -312,7 +312,7 @@ mu::Ptr<ASTOp> Sink_MeshFormatAST::Visit(const mu::Ptr<ASTOp>& at, const ASTOpMe
 		MeshPtrConst pTargetFormat = FindBaseMeshConstant(currentFormatOp->Format.child());
 		MeshPtrConst pTargetMorphFormat = MakeMorphTargetFormat(pTargetFormat);
 
-		Ptr<ASTOpConstantResource> motaop = new ASTOpConstantResource();
+		mu::Ptr<ASTOpConstantResource> motaop = new ASTOpConstantResource();
 		motaop->type = OP_TYPE::ME_CONSTANT;
 		motaop->SetValue(pTargetMorphFormat, false /* useDiskCache */);
 		auto targetMorphFormatAt = motaop;
@@ -321,7 +321,7 @@ mu::Ptr<ASTOp> Sink_MeshFormatAST::Visit(const mu::Ptr<ASTOp>& at, const ASTOpMe
 		{
 			if (newOp->children[newOp->op.args.MeshInterpolate.targets[t]])
 			{
-				Ptr<ASTOpMeshFormat> newFormat = mu::Clone<ASTOpMeshFormat>(currentFormatOp);
+				mu::Ptr<ASTOpMeshFormat> newFormat = mu::Clone<ASTOpMeshFormat>(currentFormatOp);
 				newFormat->Buffers =
 					OP::MeshFormatArgs::BT_VERTEX
 					| OP::MeshFormatArgs::BT_IGNORE_MISSING;
@@ -386,7 +386,7 @@ mu::Ptr<ASTOp> Sink_MeshFormatAST::Visit(const mu::Ptr<ASTOp>& at, const ASTOpMe
 	default:
 		if (at != m_initialSource)
 		{
-			Ptr<ASTOpMeshFormat> newOp = mu::Clone<ASTOpMeshFormat>(currentFormatOp);
+			mu::Ptr<ASTOpMeshFormat> newOp = mu::Clone<ASTOpMeshFormat>(currentFormatOp);
 			newOp->Source = at;
 			newAt = newOp;
 		}
