[NO NON-PII COMMIT MESSAGE]

diff --git a/Engine/Plugins/Experimental/Mutable/Source/CustomizableObjectEditor/Private/CustomizableObjectInstanceEditor.cpp b/Engine/Plugins/Experimental/Mutable/Source/CustomizableObjectEditor/Private/CustomizableObjectInstanceEditor.cpp
--- a/Engine/Plugins/Experimental/Mutable/Source/CustomizableObjectEditor/Private/CustomizableObjectInstanceEditor.cpp
+++ b/Engine/Plugins/Experimental/Mutable/Source/CustomizableObjectEditor/Private/CustomizableObjectInstanceEditor.cpp
@@ -189,7 +189,7 @@ void FCustomizableObjectInstanceEditor::InitCustomizableObjectInstanceEditor( co
 	{
 		if (ObjectToEdit)
 		{
-			if (UCustomizableObject* CustomizableObject = CustomizableObjectInstance->GetCustomizableObject();
+			if (UCustomizableObject* CustomizableObject = ObjectToEdit->GetCustomizableObject();
 				CustomizableObject &&
 				!CustomizableObject->IsCompiled())
 			{
