[mutable] Fix incorrect usage of "using namespace".
Fix some localization key issues. (Horde 217717)

diff --git a/Engine/Plugins/Experimental/Mutable/Source/CustomizableObject/Private/CustomizableObjectSystem.cpp b/Engine/Plugins/Experimental/Mutable/Source/CustomizableObject/Private/CustomizableObjectSystem.cpp
--- a/Engine/Plugins/Experimental/Mutable/Source/CustomizableObject/Private/CustomizableObjectSystem.cpp
+++ b/Engine/Plugins/Experimental/Mutable/Source/CustomizableObject/Private/CustomizableObjectSystem.cpp
@@ -1995,38 +1995,45 @@ namespace impl
 		UCustomizableObjectInstance* CandidateInstance = Operation->CustomizableObjectInstance.Get();
 
 		bool bCancel = false;
+		UCustomizableObject* CustomizableObject = nullptr;
+		mu::Ptr<const mu::Parameters> Parameters;
 
 		// If the object is locked (for instance, compiling) we skip any instance update.
 		if (!CandidateInstance)
 		{
 			bCancel = true;
 		}
-
-		UCustomizableObject* CustomizableObject = CandidateInstance->GetCustomizableObject(); 
-		if (CustomizableObject->GetPrivate()->bLocked)
-		{
-			bCancel = true;
-		}
-		
-		// Only update resources if the instance is in range (it could have got far from the player since the task was queued)
-		if (
-			System->CurrentInstanceLODManagement->IsOnlyUpdateCloseCustomizableObjectsEnabled()
-			&& LastMinSquareDistFromComponentToPlayer > FMath::Square(System->CurrentInstanceLODManagement->GetOnlyUpdateCloseCustomizableObjectsDist())
-			&& LastMinSquareDistFromComponentToPlayer != FLT_MAX // This means it is the first frame so it has to be updated
-			)
+		else
 		{
-			bCancel = true;
-		}
+			CustomizableObject = CandidateInstance->GetCustomizableObject();
+			if (!CustomizableObject || CustomizableObject->GetPrivate()->bLocked)
+			{
+				bCancel = true;
+			}
 
-		mu::Ptr<const mu::Parameters> Parameters = Operation->GetParameters();
-		if (!Parameters)
-		{
-			bCancel = true;
+			// Only update resources if the instance is in range (it could have got far from the player since the task was queued)
+			if (
+				System->CurrentInstanceLODManagement->IsOnlyUpdateCloseCustomizableObjectsEnabled()
+				&& LastMinSquareDistFromComponentToPlayer > FMath::Square(System->CurrentInstanceLODManagement->GetOnlyUpdateCloseCustomizableObjectsDist())
+				&& LastMinSquareDistFromComponentToPlayer != FLT_MAX // This means it is the first frame so it has to be updated
+				)
+			{
+				bCancel = true;
+			}
+
+			Parameters = Operation->GetParameters();
+			if (!Parameters)
+			{
+				bCancel = true;
+			}
 		}
 
 		if (bCancel)
 		{
-			Operation->CustomizableObjectInstance->GetPrivate()->ClearCOInstanceFlags(Updating);
+			if (Operation->CustomizableObjectInstance.IsValid())
+			{
+				Operation->CustomizableObjectInstance->GetPrivate()->ClearCOInstanceFlags(Updating);
+			}
 			System->ClearCurrentMutableOperation();
 			return;
 		}
diff --git a/Engine/Plugins/Experimental/Mutable/Source/CustomizableObjectEditor/Private/CustomizableObjectCompiler.cpp b/Engine/Plugins/Experimental/Mutable/Source/CustomizableObjectEditor/Private/CustomizableObjectCompiler.cpp
--- a/Engine/Plugins/Experimental/Mutable/Source/CustomizableObjectEditor/Private/CustomizableObjectCompiler.cpp
+++ b/Engine/Plugins/Experimental/Mutable/Source/CustomizableObjectEditor/Private/CustomizableObjectCompiler.cpp
@@ -634,7 +634,7 @@ mu::NodeObjectPtr FCustomizableObjectCompiler::GenerateMutableRoot(
 		USkeletalMesh* RefSkeletalMesh = ActualRootObject->GetRefSkeletalMesh(ComponentIndex);
 		if (!RefSkeletalMesh)
 		{
-			CompilerLog(LOCTEXT("NoReferenceMesh", "Error! Missing reference mesh in the Object Properties Tab"), ActualRoot);
+			CompilerLog(LOCTEXT("NoReferenceMeshObjectTab", "Error! Missing reference mesh in the Object Properties Tab"), ActualRoot);
 			return nullptr;
 		}
 
@@ -1450,7 +1450,7 @@ void FCustomizableObjectCompiler::ForceFinishBeforeStartCompilation(UCustomizabl
 
 void FCustomizableObjectCompiler::AddCompileNotification(const FText& CompilationStep) const
 {
-	const FText Text = CurrentObject ? FText::FromString(FString::Printf(TEXT("Compiling %s"), *CurrentObject->GetName())) : LOCTEXT("CustomizableObjectCompileInProgress", "Compiling Customizable Object");
+	const FText Text = CurrentObject ? FText::FromString(FString::Printf(TEXT("Compiling %s"), *CurrentObject->GetName())) : LOCTEXT("CustomizableObjectCompileInProgressNotification", "Compiling Customizable Object");
 	
 	FCustomizableObjectEditorLogger::CreateLog(Text)
 	.SubText(CompilationStep)
diff --git a/Engine/Plugins/Experimental/Mutable/Source/CustomizableObjectEditor/Private/CustomizableObjectEditorActions.cpp b/Engine/Plugins/Experimental/Mutable/Source/CustomizableObjectEditor/Private/CustomizableObjectEditorActions.cpp
--- a/Engine/Plugins/Experimental/Mutable/Source/CustomizableObjectEditor/Private/CustomizableObjectEditorActions.cpp
+++ b/Engine/Plugins/Experimental/Mutable/Source/CustomizableObjectEditor/Private/CustomizableObjectEditorActions.cpp
@@ -67,7 +67,7 @@ void FCustomizableObjectEditorViewportCommands::RegisterCommands()
 	UI_COMMAND( SetShowBinormals, "Binormals", "Toggles display of vertex binormals (orthogonal vector to normal and tangent) in the Preview Pane.", EUserInterfaceActionType::ToggleButton, FInputChord() );
 	UI_COMMAND( SetShowPivot, "Show Pivot", "Display the pivot location of the static mesh.", EUserInterfaceActionType::ToggleButton, FInputChord() );
 
-	UI_COMMAND(BakeInstance, "Bake Instance", "Create baked unreal resources for the current preview instance.", EUserInterfaceActionType::Button, FInputChord());
+	UI_COMMAND(BakeInstance, "Bake Instance from Object", "Create baked unreal resources for the current preview instance.", EUserInterfaceActionType::Button, FInputChord());
 	UI_COMMAND(StateChangeTest, "State test", "Make a state change test, iterating through all possible states and runtime parameters for this instance", EUserInterfaceActionType::Button, FInputChord());
 	UI_COMMAND(StateChangeShowData, "Show or hide test results", "Show or hide test results", EUserInterfaceActionType::ToggleButton, FInputChord());
 	UI_COMMAND(StateChangeShowGeometryData, "Show instance geometry data", "Show instance geometry data", EUserInterfaceActionType::ToggleButton, FInputChord());
diff --git a/Engine/Plugins/Experimental/Mutable/Source/CustomizableObjectEditor/Private/CustomizableObjectEditorViewportClient.cpp b/Engine/Plugins/Experimental/Mutable/Source/CustomizableObjectEditor/Private/CustomizableObjectEditorViewportClient.cpp
--- a/Engine/Plugins/Experimental/Mutable/Source/CustomizableObjectEditor/Private/CustomizableObjectEditorViewportClient.cpp
+++ b/Engine/Plugins/Experimental/Mutable/Source/CustomizableObjectEditor/Private/CustomizableObjectEditorViewportClient.cpp
@@ -934,7 +934,7 @@ void FCustomizableObjectEditorViewportClient::Draw(FViewport* InViewport, FCanva
 		Canvas->DrawShadowedString(
 			6,
 			2,
-			*NSLOCTEXT("UnrealEd", "UVOverlay_F", "Warning! No reference mesh is set in the Object Properties tab.").ToString(),
+			*NSLOCTEXT("CustomizableObjectEditor", "NoReferenceMeshMutable", "Warning! No reference mesh is set in the Object Properties tab.").ToString(),
 			GEngine->GetSmallFont(),
 			FLinearColor::Red
 			);
@@ -970,7 +970,7 @@ void FCustomizableObjectEditorViewportClient::DrawUVs(FViewport* InViewport, FCa
 	InCanvas->DrawShadowedString( 
 		6,
 		InTextYPos,
-		*FText::Format( NSLOCTEXT("UnrealEd", "UVOverlay_F", "Showing UV channel {0} for LOD {1}"), FText::AsNumber(UVChannel), FText::AsNumber(LODLevel) ).ToString(),
+		*FText::Format( NSLOCTEXT("CustomizableObjectEditor", "UVOverlay_F", "Showing UV channel {0} for LOD {1}"), FText::AsNumber(UVChannel), FText::AsNumber(LODLevel) ).ToString(),
 		GEngine->GetSmallFont(),
 		FLinearColor::White
 		);
@@ -2222,7 +2222,7 @@ void FCustomizableObjectEditorViewportClient::BakeInstance()
 {
 	if (!AssetRegistryLoaded)
 	{
-		FNotificationInfo Info(NSLOCTEXT("CustomizableObject", "CustomizableObjectCompileTryLater", "Please wait until asset registry loads all assets"));
+		FNotificationInfo Info(NSLOCTEXT("CustomizableObjectEditor", "CustomizableObjectCompileTryLater", "Please wait until asset registry loads all assets"));
 		Info.bFireAndForget = true;
 		Info.bUseThrobber = true;
 		Info.FadeOutDuration = 1.0f;
@@ -2248,12 +2248,12 @@ void FCustomizableObjectEditorViewportClient::BakeInstance()
 
 	UCustomizableObjectInstance* Instance = CustomizableObjectEditorPtr.Pin()->GetPreviewInstance();
 	FString ObjectName = Instance->GetCustomizableObject()->GetName();
-	FText DefaultFileName = FText::Format(LOCTEXT("DefaultFileName", "{0}"), FText::AsCultureInvariant(ObjectName));
+	FText DefaultFileName = FText::Format(LOCTEXT("DefaultFileNameForBakeInstance", "{0}"), FText::AsCultureInvariant(ObjectName));
 	bool AddAllMaterialTextures = false;
 
 	TSharedRef<SMutableSelectFolderDlg> FolderDlg =
 		SNew(SMutableSelectFolderDlg)
-		.DefaultAssetPath(LOCTEXT("", ""))
+		.DefaultAssetPath(FText())
 		.DefaultFileName(DefaultFileName);
 
 	if (FolderDlg->ShowModal() != EAppReturnType::Cancel)
@@ -2836,7 +2836,7 @@ void FCustomizableObjectEditorViewportClient::ShowInstanceGeometryInformation(FC
 				InCanvas->DrawShadowedString(
 					6.0f,
 					YOffset,
-					*FText::Format(NSLOCTEXT("UnrealEd", "UVOverlay_F", "Component {3} LOD {0} has {1} vertices and {2} triangles"),
+					*FText::Format(NSLOCTEXT("CustomizableObjectEditor", "UVOverlay_F", "Component {3} LOD {0} has {1} vertices and {2} triangles"),
 						FText::AsNumber(i), FText::AsNumber(NumVertices), FText::AsNumber(NumTriangles), FText::AsNumber(ComponentIndex)).ToString(),
 					GEngine->GetSmallFont(),
 					FLinearColor::White
diff --git a/Engine/Plugins/Experimental/Mutable/Source/CustomizableObjectEditor/Private/CustomizableObjectInstanceEditorActions.cpp b/Engine/Plugins/Experimental/Mutable/Source/CustomizableObjectEditor/Private/CustomizableObjectInstanceEditorActions.cpp
--- a/Engine/Plugins/Experimental/Mutable/Source/CustomizableObjectEditor/Private/CustomizableObjectInstanceEditorActions.cpp
+++ b/Engine/Plugins/Experimental/Mutable/Source/CustomizableObjectEditor/Private/CustomizableObjectInstanceEditorActions.cpp
@@ -36,7 +36,7 @@ void FCustomizableObjectInstanceEditorCommands::RegisterCommands()
 	UI_COMMAND( SetShowBinormals, "Binormals", "Toggles display of vertex binormals (orthogonal vector to normal and tangent) in the Preview Pane.", EUserInterfaceActionType::ToggleButton, FInputChord() );
 	UI_COMMAND( SetShowPivot, "Show Pivot", "Display the pivot location of the static mesh.", EUserInterfaceActionType::ToggleButton, FInputChord() );
 
-	UI_COMMAND( BakeInstance, "Bake Instance", "Create baked unreal resources for the current instance.", EUserInterfaceActionType::Button, FInputChord());
+	UI_COMMAND( BakeInstance, "Bake Instance from Instance", "Create baked unreal resources for the current instance.", EUserInterfaceActionType::Button, FInputChord());
 	UI_COMMAND( StateChangeTest, "State test", "Make a state change test, iterating through all possible states and runtime parameters for this instance", EUserInterfaceActionType::Button, FInputChord());
 
 	UI_COMMAND(ShowParentCO, "Show Parent", "Show the Parent of the Customizable Object in the Content Browser", EUserInterfaceActionType::Button, FInputChord());
diff --git a/Engine/Plugins/Experimental/Mutable/Source/CustomizableObjectEditor/Private/CustomizableObjectInstanceFactory.cpp b/Engine/Plugins/Experimental/Mutable/Source/CustomizableObjectEditor/Private/CustomizableObjectInstanceFactory.cpp
--- a/Engine/Plugins/Experimental/Mutable/Source/CustomizableObjectEditor/Private/CustomizableObjectInstanceFactory.cpp
+++ b/Engine/Plugins/Experimental/Mutable/Source/CustomizableObjectEditor/Private/CustomizableObjectInstanceFactory.cpp
@@ -54,7 +54,7 @@ bool UCustomizableObjectInstanceFactory::CanCreateActorFrom(const FAssetData& As
         }
         else
         {
-            OutErrorMsg = LOCTEXT("NoCustomizableObjectInstance", "The CustomizableObjectInstance is null.");
+            OutErrorMsg = LOCTEXT("NoCustomizableObjectInstanceIsNull", "The CustomizableObjectInstance is null.");
         }
     }
 
diff --git a/Engine/Plugins/Experimental/Mutable/Source/CustomizableObjectEditor/Private/GenerateMutableSource/GenerateMutableSourceColor.cpp b/Engine/Plugins/Experimental/Mutable/Source/CustomizableObjectEditor/Private/GenerateMutableSource/GenerateMutableSourceColor.cpp
--- a/Engine/Plugins/Experimental/Mutable/Source/CustomizableObjectEditor/Private/GenerateMutableSource/GenerateMutableSourceColor.cpp
+++ b/Engine/Plugins/Experimental/Mutable/Source/CustomizableObjectEditor/Private/GenerateMutableSource/GenerateMutableSourceColor.cpp
@@ -336,3 +336,5 @@ mu::NodeColourPtr GenerateMutableSourceColor(const UEdGraphPin* Pin, FMutableGra
 
 	return Result;
 }
+
+#undef LOCTEXT_NAMESPACE
diff --git a/Engine/Plugins/Experimental/Mutable/Source/CustomizableObjectEditor/Private/GenerateMutableSource/GenerateMutableSourceFloat.cpp b/Engine/Plugins/Experimental/Mutable/Source/CustomizableObjectEditor/Private/GenerateMutableSource/GenerateMutableSourceFloat.cpp
--- a/Engine/Plugins/Experimental/Mutable/Source/CustomizableObjectEditor/Private/GenerateMutableSource/GenerateMutableSourceFloat.cpp
+++ b/Engine/Plugins/Experimental/Mutable/Source/CustomizableObjectEditor/Private/GenerateMutableSource/GenerateMutableSourceFloat.cpp
@@ -389,3 +389,6 @@ mu::NodeScalarPtr GenerateMutableSourceFloat(const UEdGraphPin* Pin, FMutableGra
 
 	return Result;
 }
+
+#undef LOCTEXT_NAMESPACE
+
diff --git a/Engine/Plugins/Experimental/Mutable/Source/CustomizableObjectEditor/Private/GenerateMutableSource/GenerateMutableSourceGroupProjector.cpp b/Engine/Plugins/Experimental/Mutable/Source/CustomizableObjectEditor/Private/GenerateMutableSource/GenerateMutableSourceGroupProjector.cpp
--- a/Engine/Plugins/Experimental/Mutable/Source/CustomizableObjectEditor/Private/GenerateMutableSource/GenerateMutableSourceGroupProjector.cpp
+++ b/Engine/Plugins/Experimental/Mutable/Source/CustomizableObjectEditor/Private/GenerateMutableSource/GenerateMutableSourceGroupProjector.cpp
@@ -616,3 +616,6 @@ bool GenerateMutableSourceGroupProjector(const UEdGraphPin* Pin, FMutableGraphGe
 
 	return true;
 }
+
+#undef LOCTEXT_NAMESPACE
+
diff --git a/Engine/Plugins/Experimental/Mutable/Source/CustomizableObjectEditor/Private/GenerateMutableSource/GenerateMutableSourceImage.cpp b/Engine/Plugins/Experimental/Mutable/Source/CustomizableObjectEditor/Private/GenerateMutableSource/GenerateMutableSourceImage.cpp
--- a/Engine/Plugins/Experimental/Mutable/Source/CustomizableObjectEditor/Private/GenerateMutableSource/GenerateMutableSourceImage.cpp
+++ b/Engine/Plugins/Experimental/Mutable/Source/CustomizableObjectEditor/Private/GenerateMutableSource/GenerateMutableSourceImage.cpp
@@ -818,3 +818,6 @@ mu::NodeImagePtr GenerateMutableSourceImage(const UEdGraphPin* Pin, FMutableGrap
 
 	return Result;
 }
+
+#undef LOCTEXT_NAMESPACE
+
diff --git a/Engine/Plugins/Experimental/Mutable/Source/CustomizableObjectEditor/Private/GenerateMutableSource/GenerateMutableSourceLayout.cpp b/Engine/Plugins/Experimental/Mutable/Source/CustomizableObjectEditor/Private/GenerateMutableSource/GenerateMutableSourceLayout.cpp
--- a/Engine/Plugins/Experimental/Mutable/Source/CustomizableObjectEditor/Private/GenerateMutableSource/GenerateMutableSourceLayout.cpp
+++ b/Engine/Plugins/Experimental/Mutable/Source/CustomizableObjectEditor/Private/GenerateMutableSource/GenerateMutableSourceLayout.cpp
@@ -120,3 +120,6 @@ mu::NodeLayoutPtr GenerateMutableSourceLayout(const UEdGraphPin * Pin, FMutableG
 
 	return Result;
 }
+
+#undef LOCTEXT_NAMESPACE
+
diff --git a/Engine/Plugins/Experimental/Mutable/Source/CustomizableObjectEditor/Private/GenerateMutableSource/GenerateMutableSourceMesh.cpp b/Engine/Plugins/Experimental/Mutable/Source/CustomizableObjectEditor/Private/GenerateMutableSource/GenerateMutableSourceMesh.cpp
--- a/Engine/Plugins/Experimental/Mutable/Source/CustomizableObjectEditor/Private/GenerateMutableSource/GenerateMutableSourceMesh.cpp
+++ b/Engine/Plugins/Experimental/Mutable/Source/CustomizableObjectEditor/Private/GenerateMutableSource/GenerateMutableSourceMesh.cpp
@@ -2891,3 +2891,6 @@ mu::NodeMeshPtr GenerateMutableSourceMesh(const UEdGraphPin * Pin,
 
 	return Result;
 }
+
+#undef LOCTEXT_NAMESPACE
+
diff --git a/Engine/Plugins/Experimental/Mutable/Source/CustomizableObjectEditor/Private/GenerateMutableSource/GenerateMutableSourceModifier.cpp b/Engine/Plugins/Experimental/Mutable/Source/CustomizableObjectEditor/Private/GenerateMutableSource/GenerateMutableSourceModifier.cpp
--- a/Engine/Plugins/Experimental/Mutable/Source/CustomizableObjectEditor/Private/GenerateMutableSource/GenerateMutableSourceModifier.cpp
+++ b/Engine/Plugins/Experimental/Mutable/Source/CustomizableObjectEditor/Private/GenerateMutableSource/GenerateMutableSourceModifier.cpp
@@ -204,4 +204,7 @@ mu::NodeModifierPtr GenerateMutableSourceModifier(const UEdGraphPin * Pin, FMuta
 	GenerationContext.GeneratedNodes.Add(Node);
 	
 	return Result;
-}
\ No newline at end of file
+}
+
+#undef LOCTEXT_NAMESPACE
+
diff --git a/Engine/Plugins/Experimental/Mutable/Source/CustomizableObjectEditor/Private/GenerateMutableSource/GenerateMutableSourceProjector.cpp b/Engine/Plugins/Experimental/Mutable/Source/CustomizableObjectEditor/Private/GenerateMutableSource/GenerateMutableSourceProjector.cpp
--- a/Engine/Plugins/Experimental/Mutable/Source/CustomizableObjectEditor/Private/GenerateMutableSource/GenerateMutableSourceProjector.cpp
+++ b/Engine/Plugins/Experimental/Mutable/Source/CustomizableObjectEditor/Private/GenerateMutableSource/GenerateMutableSourceProjector.cpp
@@ -170,3 +170,6 @@ mu::NodeProjectorPtr GenerateMutableSourceProjector(const UEdGraphPin* Pin, FMut
 
 	return Result;
 }
+
+#undef LOCTEXT_NAMESPACE
+
diff --git a/Engine/Plugins/Experimental/Mutable/Source/CustomizableObjectEditor/Private/GenerateMutableSource/GenerateMutableSourceSurface.cpp b/Engine/Plugins/Experimental/Mutable/Source/CustomizableObjectEditor/Private/GenerateMutableSource/GenerateMutableSourceSurface.cpp
--- a/Engine/Plugins/Experimental/Mutable/Source/CustomizableObjectEditor/Private/GenerateMutableSource/GenerateMutableSourceSurface.cpp
+++ b/Engine/Plugins/Experimental/Mutable/Source/CustomizableObjectEditor/Private/GenerateMutableSource/GenerateMutableSourceSurface.cpp
@@ -1458,3 +1458,6 @@ mu::NodeSurfacePtr GenerateMutableSourceSurface(const UEdGraphPin * Pin, FMutabl
 
 	return Result;
 }
+
+#undef LOCTEXT_NAMESPACE
+
diff --git a/Engine/Plugins/Experimental/Mutable/Source/CustomizableObjectEditor/Private/GenerateMutableSource/GenerateMutableSourceTable.cpp b/Engine/Plugins/Experimental/Mutable/Source/CustomizableObjectEditor/Private/GenerateMutableSource/GenerateMutableSourceTable.cpp
--- a/Engine/Plugins/Experimental/Mutable/Source/CustomizableObjectEditor/Private/GenerateMutableSource/GenerateMutableSourceTable.cpp
+++ b/Engine/Plugins/Experimental/Mutable/Source/CustomizableObjectEditor/Private/GenerateMutableSource/GenerateMutableSourceTable.cpp
@@ -683,3 +683,6 @@ mu::TablePtr GenerateMutableSourceTable(const FString& TableName, const UEdGraph
 
 	return MutableTable;
 }
+
+#undef LOCTEXT_NAMESPACE
+
diff --git a/Engine/Plugins/Experimental/Mutable/Source/CustomizableObjectEditor/Private/Nodes/CustomizableObjectNodeEditMaterialBaseDetails.cpp b/Engine/Plugins/Experimental/Mutable/Source/CustomizableObjectEditor/Private/Nodes/CustomizableObjectNodeEditMaterialBaseDetails.cpp
--- a/Engine/Plugins/Experimental/Mutable/Source/CustomizableObjectEditor/Private/Nodes/CustomizableObjectNodeEditMaterialBaseDetails.cpp
+++ b/Engine/Plugins/Experimental/Mutable/Source/CustomizableObjectEditor/Private/Nodes/CustomizableObjectNodeEditMaterialBaseDetails.cpp
@@ -51,7 +51,7 @@ void FCustomizableObjectNodeEditMaterialBaseDetails::CustomizeDetails(IDetailLay
 
 		TSharedRef<IPropertyHandle> LayoutProperty = DetailBuilder.GetProperty("ParentLayoutIndex");
 
-		BlocksCategory.AddCustomRow(LOCTEXT("FCustomizableObjectNodeEditMaterialBaseDetails", "Blocks"))
+		BlocksCategory.AddCustomRow(LOCTEXT("BlocksRow", "Blocks"))
 		[
 			SNew(SHorizontalBox)
 			+ SHorizontalBox::Slot()
@@ -83,7 +83,7 @@ void FCustomizableObjectNodeEditMaterialBaseDetails::CustomizeDetails(IDetailLay
 	}
 	else
 	{
-		BlocksCategory.AddCustomRow(LOCTEXT("FCustomizableObjectNodeEditMaterialBaseDetails", "Node"))
+		BlocksCategory.AddCustomRow(LOCTEXT("NodeRow", "Node"))
 		[
 			SNew(STextBlock)
 			.Text(LOCTEXT("Node not found", "Node not found"))
diff --git a/Engine/Plugins/Experimental/Mutable/Source/CustomizableObjectEditor/Private/Nodes/CustomizableObjectNodeExternalPin.cpp b/Engine/Plugins/Experimental/Mutable/Source/CustomizableObjectEditor/Private/Nodes/CustomizableObjectNodeExternalPin.cpp
--- a/Engine/Plugins/Experimental/Mutable/Source/CustomizableObjectEditor/Private/Nodes/CustomizableObjectNodeExternalPin.cpp
+++ b/Engine/Plugins/Experimental/Mutable/Source/CustomizableObjectEditor/Private/Nodes/CustomizableObjectNodeExternalPin.cpp
@@ -148,7 +148,7 @@ FText UCustomizableObjectNodeExternalPin::GetNodeTitle(ENodeTitleType::Type Titl
 	}
 	else
 	{
-		return FText::Format(LOCTEXT("External_Pin_Title", "{0}\nImport {1} Pin"), FText::FromString(ExternalObject->GetName()), PinTypeName);
+		return FText::Format(LOCTEXT("External_Pin_Title_WithName", "{0}\nImport {1} Pin"), FText::FromString(ExternalObject->GetName()), PinTypeName);
 	}
 }
 
diff --git a/Engine/Plugins/Experimental/Mutable/Source/CustomizableObjectEditor/Private/Nodes/CustomizableObjectNodeLayoutBlocksDetails.cpp b/Engine/Plugins/Experimental/Mutable/Source/CustomizableObjectEditor/Private/Nodes/CustomizableObjectNodeLayoutBlocksDetails.cpp
--- a/Engine/Plugins/Experimental/Mutable/Source/CustomizableObjectEditor/Private/Nodes/CustomizableObjectNodeLayoutBlocksDetails.cpp
+++ b/Engine/Plugins/Experimental/Mutable/Source/CustomizableObjectEditor/Private/Nodes/CustomizableObjectNodeLayoutBlocksDetails.cpp
@@ -44,7 +44,7 @@ void FCustomizableObjectNodeLayoutBlocksDetails::CustomizeDetails( IDetailLayout
 
 	if (Node)
 	{
-		BlocksCategory.AddCustomRow( LOCTEXT("FCustomizableObjectNodeLayoutBlocksDetails", "Blocks") )
+		BlocksCategory.AddCustomRow( LOCTEXT("FCustomizableObjectNodeLayoutBlocksDetails", "BlockInstructions") )
 		[
 			SNew( STextBlock )
 			.Text(  LOCTEXT("Use the node preview panel to set up the blocks.","Use the node preview panel to set up the blocks.") )
@@ -52,7 +52,7 @@ void FCustomizableObjectNodeLayoutBlocksDetails::CustomizeDetails( IDetailLayout
 	}
 	else
 	{
-		BlocksCategory.AddCustomRow( LOCTEXT("FCustomizableObjectNodeLayoutBlocksDetails", "Node") )
+		BlocksCategory.AddCustomRow( LOCTEXT("FCustomizableObjectNodeLayoutBlocksDetails", "NodeNotFound") )
 		[
 			SNew( STextBlock )
 			.Text( LOCTEXT( "Node not found", "Node not found" ) )
diff --git a/Engine/Plugins/Experimental/Mutable/Source/CustomizableObjectEditor/Private/SCustomizableInstanceProperties.cpp b/Engine/Plugins/Experimental/Mutable/Source/CustomizableObjectEditor/Private/SCustomizableInstanceProperties.cpp
--- a/Engine/Plugins/Experimental/Mutable/Source/CustomizableObjectEditor/Private/SCustomizableInstanceProperties.cpp
+++ b/Engine/Plugins/Experimental/Mutable/Source/CustomizableObjectEditor/Private/SCustomizableInstanceProperties.cpp
@@ -2332,7 +2332,7 @@ void SCreateProfileParameters::Construct(const FArguments& InArgs)
 		.AutoHeight()
 		[
 			SNew(STextBlock)
-			.Text(LOCTEXT("FileName", "Customizable Profile Name"))
+			.Text(LOCTEXT("CustomizableProfileName", "Customizable Profile Name"))
 		.Font(FSlateFontInfo(FPaths::EngineContentDir() / TEXT("Slate/Fonts/Roboto-Regular.ttf"), 14))
 		]
 
diff --git a/Engine/Plugins/Experimental/Mutable/Source/CustomizableObjectEditor/Private/SCustomizableObjectNodeSkeletalMeshRTMorphSelector.cpp b/Engine/Plugins/Experimental/Mutable/Source/CustomizableObjectEditor/Private/SCustomizableObjectNodeSkeletalMeshRTMorphSelector.cpp
--- a/Engine/Plugins/Experimental/Mutable/Source/CustomizableObjectEditor/Private/SCustomizableObjectNodeSkeletalMeshRTMorphSelector.cpp
+++ b/Engine/Plugins/Experimental/Mutable/Source/CustomizableObjectEditor/Private/SCustomizableObjectNodeSkeletalMeshRTMorphSelector.cpp
@@ -169,10 +169,10 @@ void SCustomizableObjectNodeSkeletalMeshRTMorphSelector::UpdateWidget()
             (
                 SNew(SHeaderRow)
                 + SHeaderRow::Column(Internal::ViewRows::NameColumn)
-                .DefaultLabel(LOCTEXT("MorphNameLabel", "Name"))
+                .DefaultLabel(LOCTEXT("MorphNameLabelSelector", "Name"))
 
                 + SHeaderRow::Column(Internal::ViewRows::SelectionColumn)
-                .DefaultLabel(LOCTEXT("SelectionLabel", "Selection"))
+                .DefaultLabel(LOCTEXT("SelectionLabelSelector", "Selection"))
             )
         ]
     ];
diff --git a/Engine/Plugins/Experimental/Mutable/Source/CustomizableObjectEditor/Private/SMutableCurveViewer.cpp b/Engine/Plugins/Experimental/Mutable/Source/CustomizableObjectEditor/Private/SMutableCurveViewer.cpp
--- a/Engine/Plugins/Experimental/Mutable/Source/CustomizableObjectEditor/Private/SMutableCurveViewer.cpp
+++ b/Engine/Plugins/Experimental/Mutable/Source/CustomizableObjectEditor/Private/SMutableCurveViewer.cpp
@@ -433,7 +433,7 @@ void SMutableCurveViewer::Construct(const FArguments& InArgs)
 				.DefaultLabel(LOCTEXT("TangentWeight", "Tangent Mode"))
 				
 				+ SHeaderRow::Column(MutableCurveKeyFramesListColumns::KeyFrameTangentWeightModeColumnID)
-				.DefaultLabel(LOCTEXT("TangentWeight", "Weight Mode"))
+				.DefaultLabel(LOCTEXT("TangentWeightMode", "Weight Mode"))
 
 			)
 		]
diff --git a/Engine/Plugins/Experimental/Mutable/Source/CustomizableObjectEditor/Private/SMutableGraphViewer.cpp b/Engine/Plugins/Experimental/Mutable/Source/CustomizableObjectEditor/Private/SMutableGraphViewer.cpp
--- a/Engine/Plugins/Experimental/Mutable/Source/CustomizableObjectEditor/Private/SMutableGraphViewer.cpp
+++ b/Engine/Plugins/Experimental/Mutable/Source/CustomizableObjectEditor/Private/SMutableGraphViewer.cpp
@@ -309,8 +309,8 @@ TSharedRef<SWidget> SMutableGraphViewer::GenerateCompileOptionsMenuContent()
 
 		// Disk as cache
 		MenuBuilder.AddMenuEntry(
-			LOCTEXT("MutableEnableTextureCompression", "Enable compiling using the disk as memory."),
-			LOCTEXT("MutableEnableTextureCompressionTooltip", "This is very slow but supports compiling huge objects. It requires a lot of free space in the OS disk."),
+			LOCTEXT("MutableDiskAsMemory", "Enable compiling using the disk as memory."),
+			LOCTEXT("MutableDiskAsMemoryTooltip", "This is very slow but supports compiling huge objects. It requires a lot of free space in the OS disk."),
 			FSlateIcon(),
 			FUIAction(
 				FExecuteAction::CreateLambda([this]() { CompileOptions.bUseDiskCompilation = !CompileOptions.bUseDiskCompilation; }),
diff --git a/Engine/Plugins/Experimental/Mutable/Source/CustomizableObjectEditor/Private/SMutableMeshViewer.cpp b/Engine/Plugins/Experimental/Mutable/Source/CustomizableObjectEditor/Private/SMutableMeshViewer.cpp
--- a/Engine/Plugins/Experimental/Mutable/Source/CustomizableObjectEditor/Private/SMutableMeshViewer.cpp
+++ b/Engine/Plugins/Experimental/Mutable/Source/CustomizableObjectEditor/Private/SMutableMeshViewer.cpp
@@ -531,7 +531,7 @@ TSharedRef<SWidget> SMutableMeshViewer::GenerateBuffersListView(
 	// Headers
 	const FText BufferIndexTitle = FText(LOCTEXT("BufferIndexTitle", "Buffer"));
 	const FText BufferChannelsTitle = FText::Format(
-		LOCTEXT("SemanticLabelTitle", "{0} Buffer Channels"), InBufferSetTypeName);
+		LOCTEXT("NumberOfBufferChannels", "{0} Buffer Channels"), InBufferSetTypeName);
 
 	// Tooltips
 	const FText BufferIndexTooltip = FText(LOCTEXT("BufferIndexTooltip",
diff --git a/Engine/Plugins/Experimental/Mutable/Source/CustomizableObjectEditor/Private/SMutableObjectViewer.cpp b/Engine/Plugins/Experimental/Mutable/Source/CustomizableObjectEditor/Private/SMutableObjectViewer.cpp
--- a/Engine/Plugins/Experimental/Mutable/Source/CustomizableObjectEditor/Private/SMutableObjectViewer.cpp
+++ b/Engine/Plugins/Experimental/Mutable/Source/CustomizableObjectEditor/Private/SMutableObjectViewer.cpp
@@ -164,7 +164,7 @@ void SMutableObjectViewer::Construct(const FArguments& InArgs, UCustomizableObje
 		FUIAction(FExecuteAction::CreateSP(this, &SMutableObjectViewer::CompileMutableCodePressed)),
 		NAME_None,
 		LOCTEXT("GenerateMutableCode", "Unreal to Mutable Code"),
-		LOCTEXT("GenerateMutableCodeTooltip", "Generate a mutable code from the customizable object source graph."),
+		LOCTEXT("GenerateMutableCodeTooltipFromGraph", "Generate a mutable code from the customizable object source graph."),
 		FSlateIcon(FCustomizableObjectEditorStyle::Get().GetStyleSetName(), "CustomizableObjectDebugger.CompileMutableCode", "CustomizableObjectDebugger.CompileMutableCode.Small"),
 		EUserInterfaceActionType::Button
 	);
diff --git a/Engine/Plugins/Experimental/Mutable/Source/CustomizableObjectPopulationEditor/Private/CustomizableObjectPopulationClassDetails.cpp b/Engine/Plugins/Experimental/Mutable/Source/CustomizableObjectPopulationEditor/Private/CustomizableObjectPopulationClassDetails.cpp
--- a/Engine/Plugins/Experimental/Mutable/Source/CustomizableObjectPopulationEditor/Private/CustomizableObjectPopulationClassDetails.cpp
+++ b/Engine/Plugins/Experimental/Mutable/Source/CustomizableObjectPopulationEditor/Private/CustomizableObjectPopulationClassDetails.cpp
@@ -1659,7 +1659,7 @@ void SRangeWidget::Construct(const FArguments& InArgs)
 				.Padding(15.0f, 0.0f, 0.0f, 0.0f)
 				[
 					SNew(SButton)
-					.ToolTipText(LOCTEXT("RemoveConstraint", "Remove Range"))
+					.ToolTipText(LOCTEXT("RemoveRange", "Remove Range"))
 					.OnClicked(this, &SRangeWidget::OnRemoveRangeButtonPressed)
 					[
 						SNew(SImage)
diff --git a/Engine/Plugins/Experimental/Mutable/Source/MutableRuntime/Private/CodeVisitor.cpp b/Engine/Plugins/Experimental/Mutable/Source/MutableRuntime/Private/CodeVisitor.cpp
--- a/Engine/Plugins/Experimental/Mutable/Source/MutableRuntime/Private/CodeVisitor.cpp
+++ b/Engine/Plugins/Experimental/Mutable/Source/MutableRuntime/Private/CodeVisitor.cpp
@@ -5,86 +5,89 @@
 
 #include "MutableTrace.h"
 
-using namespace mu;
+namespace mu
+{
 
 
-//-------------------------------------------------------------------------------------------------
-//-------------------------------------------------------------------------------------------------
-//-------------------------------------------------------------------------------------------------
-void SubtreeParametersVisitor::Run( OP::ADDRESS root, PROGRAM& program )
-{
-    // Cached?
-    auto it = m_resultCache.find(root);
-    if (it!=m_resultCache.end())
-    {
-        m_params = it->second;
-        return;
-    }
-
-    // Not cached
-    {
-        MUTABLE_CPUPROFILER_SCOPE(SubtreeParametersVisitor);
-
-        m_visited.resize( program.m_opAddress.size() );
-        if (program.m_opAddress.size())
-        {
-			FMemory::Memzero( &m_visited[0], m_visited.size() );
-        }
-
-        m_currentParams.resize( program.m_parameters.size() );
-        if (m_currentParams.size())
-        {
-			FMemory::Memzero( &m_currentParams[0], sizeof(int)*m_currentParams.size() );
-        }
-
-        m_pending.reserve( program.m_opAddress.size()/4 );
-        m_pending.push_back(root);
-
-        while ( m_pending.size() )
-        {
-            OP::ADDRESS at = m_pending.back();
-            m_pending.pop_back();
-
-            if (!m_visited[at])
-            {
-                m_visited[at] = true;
-
-                switch ( program.GetOpType(at) )
-                {
-                case OP_TYPE::NU_PARAMETER:
-                case OP_TYPE::SC_PARAMETER:
-                case OP_TYPE::BO_PARAMETER:
-                case OP_TYPE::CO_PARAMETER:
-                case OP_TYPE::PR_PARAMETER:
-                case OP_TYPE::IM_PARAMETER:
-                    m_currentParams[ program.GetOpArgs<OP::ParameterArgs>(at).variable ]++;
-                    break;
-
-                default:
-                    break;
-                }
-
-                ForEachReference( program, at, [&](OP::ADDRESS ref)
-                {
-                    if (ref)
-                    {
-                        m_pending.push_back( ref );
-                    }
-                });
-            }
-        }
-
-
-        // Build result
-        m_params.clear();
-        for (size_t i=0; i<m_currentParams.size(); ++i )
-        {
-            if ( m_currentParams[i] )
-            {
-                m_params.push_back( int(i) );
-            }
-        }
-
-        m_resultCache[root] = m_params;
-    }
-}
+	//-------------------------------------------------------------------------------------------------
+	//-------------------------------------------------------------------------------------------------
+	//-------------------------------------------------------------------------------------------------
+	void SubtreeParametersVisitor::Run(OP::ADDRESS root, PROGRAM& program)
+	{
+		// Cached?
+		auto it = m_resultCache.find(root);
+		if (it != m_resultCache.end())
+		{
+			m_params = it->second;
+			return;
+		}
+
+		// Not cached
+		{
+			MUTABLE_CPUPROFILER_SCOPE(SubtreeParametersVisitor);
+
+			m_visited.resize(program.m_opAddress.size());
+			if (program.m_opAddress.size())
+			{
+				FMemory::Memzero(&m_visited[0], m_visited.size());
+			}
+
+			m_currentParams.resize(program.m_parameters.size());
+			if (m_currentParams.size())
+			{
+				FMemory::Memzero(&m_currentParams[0], sizeof(int) * m_currentParams.size());
+			}
+
+			m_pending.reserve(program.m_opAddress.size() / 4);
+			m_pending.push_back(root);
+
+			while (m_pending.size())
+			{
+				OP::ADDRESS at = m_pending.back();
+				m_pending.pop_back();
+
+				if (!m_visited[at])
+				{
+					m_visited[at] = true;
+
+					switch (program.GetOpType(at))
+					{
+					case OP_TYPE::NU_PARAMETER:
+					case OP_TYPE::SC_PARAMETER:
+					case OP_TYPE::BO_PARAMETER:
+					case OP_TYPE::CO_PARAMETER:
+					case OP_TYPE::PR_PARAMETER:
+					case OP_TYPE::IM_PARAMETER:
+						m_currentParams[program.GetOpArgs<OP::ParameterArgs>(at).variable]++;
+						break;
+
+					default:
+						break;
+					}
+
+					ForEachReference(program, at, [&](OP::ADDRESS ref)
+						{
+							if (ref)
+							{
+								m_pending.push_back(ref);
+							}
+						});
+				}
+			}
+
+
+			// Build result
+			m_params.clear();
+			for (size_t i = 0; i < m_currentParams.size(); ++i)
+			{
+				if (m_currentParams[i])
+				{
+					m_params.push_back(int(i));
+				}
+			}
+
+			m_resultCache[root] = m_params;
+		}
+	}
+
+}
\ No newline at end of file
diff --git a/Engine/Plugins/Experimental/Mutable/Source/MutableRuntime/Private/Memory.cpp b/Engine/Plugins/Experimental/Mutable/Source/MutableRuntime/Private/Memory.cpp
--- a/Engine/Plugins/Experimental/Mutable/Source/MutableRuntime/Private/Memory.cpp
+++ b/Engine/Plugins/Experimental/Mutable/Source/MutableRuntime/Private/Memory.cpp
@@ -12,265 +12,266 @@
 #include <cstdlib>
 #include <cassert>
 
-using namespace mu;
+namespace mu
+{
 
-//! This would be nice, but apparently several compilers don't support it.
+	//! This would be nice, but apparently several compilers don't support it.
 #define MUTABLE_USE_STD_ALIGNED_ALLOC   1
 
 
 //! Static pointers to the user provided malloc and free methods.
 //! If they are null (default) the versions on the standard library will be used.
-static void* (*s_custom_malloc)(size_t, uint32_t) = nullptr;
-static void(*s_custom_free)(void*) = nullptr;
+	static void* (*s_custom_malloc)(size_t, uint32_t) = nullptr;
+	static void(*s_custom_free)(void*) = nullptr;
 
 
-static int s_initialized = 0;
-static int s_finalized = 0;
+	static int s_initialized = 0;
+	static int s_finalized = 0;
 
 
-//---------------------------------------------------------------------------------------------
+	//---------------------------------------------------------------------------------------------
 #if MUTABLE_USE_STD_ALIGNED_ALLOC
 
-inline void* mutable_aligned_alloc( size_t size, size_t align )
-{
+	inline void* mutable_aligned_alloc(size_t size, size_t align)
+	{
 #if defined(_WIN32) && _WIN32
-    return _aligned_malloc(size,align);
+		return _aligned_malloc(size, align);
 #elif defined(MUTABLE_PLATFORM_ANDROID)
-    return memalign( align, size );
+		return memalign(align, size);
 #else
-    return aligned_alloc( align, size );
+		return aligned_alloc(align, size);
 #endif
-}
+	}
 
-inline void mutable_aligned_free( void* p )
-{
+	inline void mutable_aligned_free(void* p)
+	{
 #if defined(_WIN32) && _WIN32
-    _aligned_free( p );
+		_aligned_free(p);
 #else
-    free( p );
+		free(p);
 #endif
-}
+	}
 
 #else
 
-typedef uint16_t offset_t;
+	typedef uint16_t offset_t;
 #define PTR_OFFSET_SZ sizeof( offset_t ) 
 #ifndef align_up
 #define align_up( num, align ) ( ( ( num ) + ( (align)-1 ) ) & ~( (align)-1 ) )
 #endif
 
-inline void* mutable_aligned_alloc( size_t size, size_t align )
-{
-    void* ptr = nullptr;
-
-    // We want it to be a power of two since align_up operates on powers of two
-    assert( ( align & ( align - 1 ) ) == 0 );
-
-    if ( align && size )
-    {
-        // We know we have to fit an offset value
-        // We also allocate extra bytes to ensure we
-        // can meet the alignment
-        size_t hdr_size = PTR_OFFSET_SZ + ( align - 1 );
-        void* p = malloc( size + hdr_size );
-
-        if ( p )
-        {
-            // Add the offset size to malloc's pointer (we will always store that)
-            // Then align the resulting value to the
-            // target alignment
-            ptr = (void*)align_up( ( (uintptr_t)p + PTR_OFFSET_SZ ), align );
-
-            // Calculate the offset and store it
-            // behind our aligned pointer
-            *( (offset_t*)ptr - 1 ) = (offset_t)( (uintptr_t)ptr - (uintptr_t)p );
-
-        } // else NULL, could not malloc
-    }     // else NULL, invalid arguments
-
-    return ptr;
-}
+	inline void* mutable_aligned_alloc(size_t size, size_t align)
+	{
+		void* ptr = nullptr;
+
+		// We want it to be a power of two since align_up operates on powers of two
+		assert((align & (align - 1)) == 0);
+
+		if (align && size)
+		{
+			// We know we have to fit an offset value
+			// We also allocate extra bytes to ensure we
+			// can meet the alignment
+			size_t hdr_size = PTR_OFFSET_SZ + (align - 1);
+			void* p = malloc(size + hdr_size);
+
+			if (p)
+			{
+				// Add the offset size to malloc's pointer (we will always store that)
+				// Then align the resulting value to the
+				// target alignment
+				ptr = (void*)align_up(((uintptr_t)p + PTR_OFFSET_SZ), align);
+
+				// Calculate the offset and store it
+				// behind our aligned pointer
+				*((offset_t*)ptr - 1) = (offset_t)((uintptr_t)ptr - (uintptr_t)p);
+
+			} // else NULL, could not malloc
+		}     // else NULL, invalid arguments
+
+		return ptr;
+	}
 
 
-void mutable_aligned_free( void* ptr )
-{
-    assert( ptr );
+	void mutable_aligned_free(void* ptr)
+	{
+		assert(ptr);
 
-    // Walk backwards from the passed-in pointer to get the pointer offset. We convert to an 
-    // offset_t pointer and rely on pointer math to get the data
-    offset_t offset = *( (offset_t*)ptr - 1 );
+		// Walk backwards from the passed-in pointer to get the pointer offset. We convert to an 
+		// offset_t pointer and rely on pointer math to get the data
+		offset_t offset = *((offset_t*)ptr - 1);
 
-    // Once we have the offset, we can get our original pointer and call free
-    void* p = (void*)( (uint8_t*)ptr - offset );
-    free( p );
-}
+		// Once we have the offset, we can get our original pointer and call free
+		void* p = (void*)((uint8_t*)ptr - offset);
+		free(p);
+	}
 
 #endif
 
 
-//---------------------------------------------------------------------------------------------
-//! Call the system malloc or the user-provided replacement.
-//---------------------------------------------------------------------------------------------
-inline void* lowerlevel_malloc( size_t bytes )
-{
-    void* pMem = 0;
-    if ( s_custom_malloc )
-    {
-        pMem = s_custom_malloc( bytes, 1 );
-    }
-    else
-    {
-        pMem = malloc( bytes );
-    }
-    return pMem;
-}
+	//---------------------------------------------------------------------------------------------
+	//! Call the system malloc or the user-provided replacement.
+	//---------------------------------------------------------------------------------------------
+	inline void* lowerlevel_malloc(size_t bytes)
+	{
+		void* pMem = 0;
+		if (s_custom_malloc)
+		{
+			pMem = s_custom_malloc(bytes, 1);
+		}
+		else
+		{
+			pMem = malloc(bytes);
+		}
+		return pMem;
+	}
 
-inline void* lowerlevel_malloc_aligned( size_t bytes, uint32_t alignment )
-{
-    void* pMem = 0;
-    if ( s_custom_malloc )
-    {
-        pMem = s_custom_malloc( bytes, alignment );
-    }
-    else
-    {
-        size_t padding = bytes % alignment;
-        bytes += padding ? alignment - padding : 0;
-        pMem = mutable_aligned_alloc( bytes, alignment );
-    }
-    return pMem;
-}
+	inline void* lowerlevel_malloc_aligned(size_t bytes, uint32_t alignment)
+	{
+		void* pMem = 0;
+		if (s_custom_malloc)
+		{
+			pMem = s_custom_malloc(bytes, alignment);
+		}
+		else
+		{
+			size_t padding = bytes % alignment;
+			bytes += padding ? alignment - padding : 0;
+			pMem = mutable_aligned_alloc(bytes, alignment);
+		}
+		return pMem;
+	}
 
-//-------------------------------------------------------------------------------------------------
-//! Call the system free or the user-provided replacement.
-//-------------------------------------------------------------------------------------------------
-inline void lowerlevel_free( void* ptr )
-{
-    if ( s_custom_free )
-    {
-        s_custom_free( ptr );
-    }
-    else
-    {
-        free( ptr );
-    }
-}
+	//-------------------------------------------------------------------------------------------------
+	//! Call the system free or the user-provided replacement.
+	//-------------------------------------------------------------------------------------------------
+	inline void lowerlevel_free(void* ptr)
+	{
+		if (s_custom_free)
+		{
+			s_custom_free(ptr);
+		}
+		else
+		{
+			free(ptr);
+		}
+	}
 
 
-//-------------------------------------------------------------------------------------------------
-inline void lowerlevel_free_aligned( void* ptr )
-{
-    if ( s_custom_free )
-    {
-        s_custom_free( ptr );
-    }
-    else
-    {
-        mutable_aligned_free( ptr );
-    }
-}
+	//-------------------------------------------------------------------------------------------------
+	inline void lowerlevel_free_aligned(void* ptr)
+	{
+		if (s_custom_free)
+		{
+			s_custom_free(ptr);
+		}
+		else
+		{
+			mutable_aligned_free(ptr);
+		}
+	}
 
 
-//-------------------------------------------------------------------------------------------------
-//! Memory management functions to be used inside the library. No other memory allocation is
-//! allowed.
-//-------------------------------------------------------------------------------------------------
-void* mu::mutable_malloc( size_t size )
-{
-	if (!size)
+	//-------------------------------------------------------------------------------------------------
+	//! Memory management functions to be used inside the library. No other memory allocation is
+	//! allowed.
+	//-------------------------------------------------------------------------------------------------
+	void* mutable_malloc(size_t size)
 	{
-		return nullptr;
-	}
+		if (!size)
+		{
+			return nullptr;
+		}
 
-    void* pMem = nullptr;
+		void* pMem = nullptr;
 
-	pMem = lowerlevel_malloc(size);
+		pMem = lowerlevel_malloc(size);
 
-    return pMem;
-}
+		return pMem;
+	}
 
 
-void* mu::mutable_malloc_aligned( size_t size, uint32_t alignment )
-{
-	if (!size)
+	void* mutable_malloc_aligned(size_t size, uint32_t alignment)
 	{
-		return nullptr;
-	}
+		if (!size)
+		{
+			return nullptr;
+		}
 
-    void* pMem = nullptr;
+		void* pMem = nullptr;
 
-    pMem = lowerlevel_malloc_aligned(size, alignment); 
+		pMem = lowerlevel_malloc_aligned(size, alignment);
 
-    return pMem;
-}
+		return pMem;
+	}
 
 
-void mu::mutable_free( void* ptr )
-{
-	if (!ptr)
+	void mutable_free(void* ptr)
 	{
-		return;
-	}
+		if (!ptr)
+		{
+			return;
+		}
 
-	lowerlevel_free(ptr);
-}
+		lowerlevel_free(ptr);
+	}
 
 
-void mu::mutable_free( void* ptr, size_t size )
-{
-	if (!ptr || !size)
+	void mutable_free(void* ptr, size_t size)
 	{
-		return;
-	}
+		if (!ptr || !size)
+		{
+			return;
+		}
 
-	lowerlevel_free(ptr);
-}
+		lowerlevel_free(ptr);
+	}
 
 
-void mu::mutable_free_aligned( void* ptr, size_t size )
-{
-	if (!size || !ptr)
+	void mutable_free_aligned(void* ptr, size_t size)
 	{
-		return;
-	}
+		if (!size || !ptr)
+		{
+			return;
+		}
 
-	lowerlevel_free_aligned(ptr);
-}
+		lowerlevel_free_aligned(ptr);
+	}
 
 
-//-------------------------------------------------------------------------------------------------
-void mu::Initialize
-    (
-        void* (*customMalloc)( size_t, uint32_t ),
-        void (*customFree)( void* )
-    )
-{
-	if (!s_initialized)
+	//-------------------------------------------------------------------------------------------------
+	void Initialize
+	(
+		void* (*customMalloc)(size_t, uint32_t),
+		void (*customFree)(void*)
+	)
 	{
-		s_initialized = 1;
-		s_finalized = 0;
+		if (!s_initialized)
+		{
+			s_initialized = 1;
+			s_finalized = 0;
 
-		s_custom_malloc = customMalloc;
-		s_custom_free = customFree;
+			s_custom_malloc = customMalloc;
+			s_custom_free = customFree;
 
-		miro::initialize();
+			miro::initialize();
+		}
 	}
-}
-
 
-//-------------------------------------------------------------------------------------------------
-void mu::Finalize()
-{
-    if (s_initialized && !s_finalized)
-    {
-        miro::finalize();
-
-        s_finalized = 1;
-        s_initialized = 0;
-        s_custom_malloc = nullptr;
-        s_custom_free = nullptr;
-    }
-}
 
+	//-------------------------------------------------------------------------------------------------
+	void Finalize()
+	{
+		if (s_initialized && !s_finalized)
+		{
+			miro::finalize();
+
+			s_finalized = 1;
+			s_initialized = 0;
+			s_custom_malloc = nullptr;
+			s_custom_free = nullptr;
+		}
+	}
 
 
+}
diff --git a/Engine/Plugins/Experimental/Mutable/Source/MutableRuntime/Private/Mesh.cpp b/Engine/Plugins/Experimental/Mutable/Source/MutableRuntime/Private/Mesh.cpp
--- a/Engine/Plugins/Experimental/Mutable/Source/MutableRuntime/Private/Mesh.cpp
+++ b/Engine/Plugins/Experimental/Mutable/Source/MutableRuntime/Private/Mesh.cpp
@@ -7,8 +7,8 @@
 #include "Platform.h"
 #include "MutableTrace.h"
 
-using namespace mu;
-
+namespace mu
+{
 
 //---------------------------------------------------------------------------------------------
 void Mesh::Serialise( const Mesh* p, OutputArchive& arch )
@@ -454,7 +454,7 @@ void Mesh::SetTag( int tagIndex, const char* strName )
 
 
 //---------------------------------------------------------------------------------------------
-const char* mu::Mesh::GetBonePoseName(int32 boneIndex) const
+const char* Mesh::GetBonePoseName(int32 boneIndex) const
 {
 	check(boneIndex >= 0 && boneIndex < GetBonePoseCount());
 	if (boneIndex >= 0 && boneIndex < GetBonePoseCount())
@@ -466,7 +466,7 @@ const char* mu::Mesh::GetBonePoseName(int32 boneIndex) const
 
 
 //---------------------------------------------------------------------------------------------
-int32 mu::Mesh::FindBonePose(const char* strName) const
+int32 Mesh::FindBonePose(const char* strName) const
 {
 	for (int32 i = 0; i < m_bonePoses.Num(); ++i)
 	{
@@ -1324,7 +1324,7 @@ static bool StaticMeshFormatIdentify_ProjectWrapping( const Mesh* pM )
 
 
 //---------------------------------------------------------------------------------------------
-STATIC_MESH_FORMAT_ID_FUNC mu::s_staticMeshFormatIdentify[] =
+STATIC_MESH_FORMAT_ID_FUNC s_staticMeshFormatIdentify[] =
 {
     StaticMeshFormatIdentify_None,
     StaticMeshFormatIdentify_Project,
@@ -1403,3 +1403,5 @@ void Mesh::Log( FString& out, int32 BufferElementLimit)
     out += "Faces:\n";
     LogBuffer( out, m_FaceBuffers, BufferElementLimit);
 }
+
+}
diff --git a/Engine/Plugins/Experimental/Mutable/Source/MutableRuntime/Private/MeshBufferSet.cpp b/Engine/Plugins/Experimental/Mutable/Source/MutableRuntime/Private/MeshBufferSet.cpp
--- a/Engine/Plugins/Experimental/Mutable/Source/MutableRuntime/Private/MeshBufferSet.cpp
+++ b/Engine/Plugins/Experimental/Mutable/Source/MutableRuntime/Private/MeshBufferSet.cpp
@@ -6,601 +6,603 @@
 #include "MutableTrace.h"
 
 
-using namespace mu;
-
+namespace mu
+{
 
-//---------------------------------------------------------------------------------------------
-static FMeshBufferFormatData s_meshBufferFormatData[ MBF_COUNT ] =
-    {
-        { 0, 0 },
-        { 2, 0 },
-        { 4, 0 },
+	//---------------------------------------------------------------------------------------------
+	static FMeshBufferFormatData s_meshBufferFormatData[MBF_COUNT] =
+	{
+		{ 0, 0 },
+		{ 2, 0 },
+		{ 4, 0 },
 
 		{ 1, 8 },
-        { 2, 16 },
-        { 4, 32 },
-        { 1, 7 },
-        { 2, 15 },
-        { 4, 31 },
+		{ 2, 16 },
+		{ 4, 32 },
+		{ 1, 7 },
+		{ 2, 15 },
+		{ 4, 31 },
 
 		{ 1, 0 },
-        { 2, 0 },
-        { 4, 0 },
-        { 1, 0 },
-        { 2, 0 },
-        { 4, 0 },
+		{ 2, 0 },
+		{ 4, 0 },
+		{ 1, 0 },
+		{ 2, 0 },
+		{ 4, 0 },
 
 		{ 1, 0 },
-        { 1, 0 },
-        { 1, 0 },
-        { 1, 0 },
+		{ 1, 0 },
+		{ 1, 0 },
+		{ 1, 0 },
 
 		{ 8, 0 }
-};
+	};
 
 
-//---------------------------------------------------------------------------------------------
-const FMeshBufferFormatData& mu::GetMeshFormatData( MESH_BUFFER_FORMAT format )
-{
-    check( format >= 0 );
-    check( format < MBF_COUNT );
-    return s_meshBufferFormatData[ format ];
-}
+	//---------------------------------------------------------------------------------------------
+	const FMeshBufferFormatData& GetMeshFormatData(MESH_BUFFER_FORMAT format)
+	{
+		check(format >= 0);
+		check(format < MBF_COUNT);
+		return s_meshBufferFormatData[format];
+	}
 
 
-//---------------------------------------------------------------------------------------------
-//---------------------------------------------------------------------------------------------
-//---------------------------------------------------------------------------------------------
-int32 FMeshBufferSet::GetElementCount() const
-{
-    return m_elementCount;
-}
+	//---------------------------------------------------------------------------------------------
+	//---------------------------------------------------------------------------------------------
+	//---------------------------------------------------------------------------------------------
+	int32 FMeshBufferSet::GetElementCount() const
+	{
+		return m_elementCount;
+	}
 
 
-//---------------------------------------------------------------------------------------------
-void FMeshBufferSet::SetElementCount( int32 count )
-{
-    check( count>=0 );
-	LLM_SCOPE_BYNAME(TEXT("MutableRuntime"));
+	//---------------------------------------------------------------------------------------------
+	void FMeshBufferSet::SetElementCount(int32 count)
+	{
+		check(count >= 0);
+		LLM_SCOPE_BYNAME(TEXT("MutableRuntime"));
 
-    m_elementCount = count;
+		m_elementCount = count;
 
-    for ( auto& buf: m_buffers )
-    {
-        buf.m_data.SetNumUninitialized( buf.m_elementSize * count, false );
-    }
-}
+		for (auto& buf : m_buffers)
+		{
+			buf.m_data.SetNumUninitialized(buf.m_elementSize * count, false);
+		}
+	}
 
 
-//---------------------------------------------------------------------------------------------
-int32 FMeshBufferSet::GetBufferCount() const
-{
-    return m_buffers.Num();
-}
+	//---------------------------------------------------------------------------------------------
+	int32 FMeshBufferSet::GetBufferCount() const
+	{
+		return m_buffers.Num();
+	}
 
 
-//---------------------------------------------------------------------------------------------
-void FMeshBufferSet::SetBufferCount( int32 count )
-{
-	LLM_SCOPE_BYNAME(TEXT("MutableRuntime"));
-    m_buffers.SetNum( count );
-}
+	//---------------------------------------------------------------------------------------------
+	void FMeshBufferSet::SetBufferCount(int32 count)
+	{
+		LLM_SCOPE_BYNAME(TEXT("MutableRuntime"));
+		m_buffers.SetNum(count);
+	}
 
 
-//---------------------------------------------------------------------------------------------
-int32 FMeshBufferSet::GetBufferChannelCount( int32 buffer ) const
-{
-    check(buffer >= 0 && buffer < m_buffers.Num());
-    if (buffer >= 0 && buffer < m_buffers.Num())
-    {
-        return m_buffers[buffer].m_channels.Num();
-    }
-    return 0;
-}
-
-
-//---------------------------------------------------------------------------------------------
-void FMeshBufferSet::GetChannel
-    (
-        int32 buffer,
-        int32 channel,
-        MESH_BUFFER_SEMANTIC* pSemantic,
-        int32* pSemanticIndex,
-        MESH_BUFFER_FORMAT* pFormat,
-        int32* pComponentCount,
-        int32* pOffset
-    ) const
-{
-    check( buffer>=0 && buffer<m_buffers.Num() );
-    check( channel>=0 && channel<m_buffers[buffer].m_channels.Num() );
-
-    const MESH_BUFFER_CHANNEL& chan = m_buffers[buffer].m_channels[channel];
-
-    if (pSemantic)
-    {
-        *pSemantic = chan.m_semantic;
-    }
-
-    if (pSemanticIndex)
-    {
-        *pSemanticIndex = chan.m_semanticIndex;
-    }
-
-    if (pFormat)
-    {
-        *pFormat = chan.m_format;
-    }
-
-    if (pComponentCount)
-    {
-        *pComponentCount = chan.m_componentCount;
-    }
-
-    if (pOffset)
-    {
-        *pOffset = chan.m_offset;
-    }
-}
-
-
-//---------------------------------------------------------------------------------------------
-void FMeshBufferSet::SetBuffer
-    (
-        int32 buffer,
-        int32 elementSize,
-        int32 channelCount,
-        const MESH_BUFFER_SEMANTIC* pSemantics,
-        const int32* pSemanticIndices,
-        const MESH_BUFFER_FORMAT* pFormats,
-        const int32* pComponentCount,
-        const int32* pOffsets
-    )
-{
-    check( buffer>=0 && buffer<m_buffers.Num() );
-	LLM_SCOPE_BYNAME(TEXT("MutableRuntime"));
-
-    MESH_BUFFER& buf = m_buffers[buffer];
-
-    unsigned minElemSize = 0;
-    buf.m_channels.SetNum( channelCount );
-    for ( int c=0; c<channelCount; ++c )
-    {
-        MESH_BUFFER_CHANNEL& chan = buf.m_channels[c];
-        chan.m_semantic = pSemantics ? pSemantics[c] : MBS_NONE;
-        chan.m_semanticIndex = pSemanticIndices ? ((uint8_t)pSemanticIndices[c]) : 0;
-        chan.m_format = pFormats ? pFormats[c] : MBF_NONE;
-        chan.m_componentCount = pComponentCount ? ((uint16_t)pComponentCount[c]) : 0;
-        chan.m_offset = pOffsets ? ((uint8_t)pOffsets[c]) : 0;
-
-        minElemSize = FMath::Max
-                (
-                    minElemSize,
-                    chan.m_offset+chan.m_componentCount*GetMeshFormatData(chan.m_format).m_size
-                );
-    }
-
-    // Set the user specified element size, or enlarge it if it was too small
-    buf.m_elementSize = FMath::Max( elementSize, int32(minElemSize) );
-
-    // Update the buffer data
-    buf.m_data.SetNumUninitialized( buf.m_elementSize * m_elementCount, false );
-}
-
-//---------------------------------------------------------------------------------------------
-void FMeshBufferSet::SetBufferChannel
-    (
-        int32 buffer,
-        int32 channelIndex,
-        MESH_BUFFER_SEMANTIC semantic,
-        int32 semanticIndex,
-        MESH_BUFFER_FORMAT format,
-        int32 componentCount,
-        int32 offset
-    )
-{
-    if (!( buffer>=0 && buffer<m_buffers.Num() ))
-    {
-		check(false);
-        return;
-    }
+	//---------------------------------------------------------------------------------------------
+	int32 FMeshBufferSet::GetBufferChannelCount(int32 buffer) const
+	{
+		check(buffer >= 0 && buffer < m_buffers.Num());
+		if (buffer >= 0 && buffer < m_buffers.Num())
+		{
+			return m_buffers[buffer].m_channels.Num();
+		}
+		return 0;
+	}
 
-	LLM_SCOPE_BYNAME(TEXT("MutableRuntime"));
 
-    MESH_BUFFER& buf = m_buffers[buffer];
+	//---------------------------------------------------------------------------------------------
+	void FMeshBufferSet::GetChannel
+	(
+		int32 buffer,
+		int32 channel,
+		MESH_BUFFER_SEMANTIC* pSemantic,
+		int32* pSemanticIndex,
+		MESH_BUFFER_FORMAT* pFormat,
+		int32* pComponentCount,
+		int32* pOffset
+	) const
+	{
+		check(buffer >= 0 && buffer < m_buffers.Num());
+		check(channel >= 0 && channel < m_buffers[buffer].m_channels.Num());
 
-    if (!( channelIndex>=0 && channelIndex<buf.m_channels.Num() ))
-    {
-		check(false);
-        return;
-    }
+		const MESH_BUFFER_CHANNEL& chan = m_buffers[buffer].m_channels[channel];
 
-    MESH_BUFFER_CHANNEL& chan = buf.m_channels[channelIndex];
-    chan.m_semantic = semantic;
-    chan.m_semanticIndex = uint8_t(semanticIndex);
-    chan.m_format = format;
-    chan.m_componentCount = uint16_t(componentCount);
-    chan.m_offset = uint8(offset);
-}
+		if (pSemantic)
+		{
+			*pSemantic = chan.m_semantic;
+		}
 
+		if (pSemanticIndex)
+		{
+			*pSemanticIndex = chan.m_semanticIndex;
+		}
 
-//---------------------------------------------------------------------------------------------
-uint8* FMeshBufferSet::GetBufferData( int32 buffer )
-{
-    check( buffer>=0 && buffer<m_buffers.Num() );
-	uint8_t* pResult = m_buffers[buffer].m_data.GetData();
-    return pResult;
-}
+		if (pFormat)
+		{
+			*pFormat = chan.m_format;
+		}
 
+		if (pComponentCount)
+		{
+			*pComponentCount = chan.m_componentCount;
+		}
 
-//---------------------------------------------------------------------------------------------
-const uint8* FMeshBufferSet::GetBufferData( int32 buffer ) const
-{
-    check( buffer>=0 && buffer<m_buffers.Num() );
-    const uint8_t* pResult = m_buffers[buffer].m_data.GetData();
-    return pResult;
-}
-
-
-//---------------------------------------------------------------------------------------------
-void FMeshBufferSet::FindChannel
-    (
-        MESH_BUFFER_SEMANTIC semantic, int32 semanticIndex,
-        int32* pBuffer, int32* pChannel
-    ) const
-{
-    check( pBuffer && pChannel );
-
-    *pBuffer = -1;
-    *pChannel = -1;
-
-    int index = 0;
-
-    for( size_t b=0; b<m_buffers.Num(); ++b )
-    {
-        for( size_t c=0; c<m_buffers[b].m_channels.Num(); ++c )
-        {
-            if ( m_buffers[b].m_channels[c].m_semantic == semantic )
-            {
-                if (index==semanticIndex)
-                {
-                    *pBuffer = int(b);
-                    *pChannel = int(c);
-                    return;
-                }
-                else
-                {
-                    ++index;
-                }
-            }
-        }
-    }
-}
-
-
-//---------------------------------------------------------------------------------------------
-int32 FMeshBufferSet::GetElementSize( int32 buffer ) const
-{
-    check( buffer>=0 && buffer<m_buffers.Num() );
+		if (pOffset)
+		{
+			*pOffset = chan.m_offset;
+		}
+	}
 
-    return m_buffers[buffer].m_elementSize;
-}
 
+	//---------------------------------------------------------------------------------------------
+	void FMeshBufferSet::SetBuffer
+	(
+		int32 buffer,
+		int32 elementSize,
+		int32 channelCount,
+		const MESH_BUFFER_SEMANTIC* pSemantics,
+		const int32* pSemanticIndices,
+		const MESH_BUFFER_FORMAT* pFormats,
+		const int32* pComponentCount,
+		const int32* pOffsets
+	)
+	{
+		check(buffer >= 0 && buffer < m_buffers.Num());
+		LLM_SCOPE_BYNAME(TEXT("MutableRuntime"));
 
-//---------------------------------------------------------------------------------------------
-int32 FMeshBufferSet::GetChannelOffset
-    (
-        int32 buffer,
-        int32 channel
-    ) const
-{
-    check( buffer>=0 && buffer<m_buffers.Num() );
-    check( channel>=0 && channel<m_buffers[buffer].m_channels.Num() );
+		MESH_BUFFER& buf = m_buffers[buffer];
+
+		unsigned minElemSize = 0;
+		buf.m_channels.SetNum(channelCount);
+		for (int c = 0; c < channelCount; ++c)
+		{
+			MESH_BUFFER_CHANNEL& chan = buf.m_channels[c];
+			chan.m_semantic = pSemantics ? pSemantics[c] : MBS_NONE;
+			chan.m_semanticIndex = pSemanticIndices ? ((uint8_t)pSemanticIndices[c]) : 0;
+			chan.m_format = pFormats ? pFormats[c] : MBF_NONE;
+			chan.m_componentCount = pComponentCount ? ((uint16_t)pComponentCount[c]) : 0;
+			chan.m_offset = pOffsets ? ((uint8_t)pOffsets[c]) : 0;
+
+			minElemSize = FMath::Max
+			(
+				minElemSize,
+				chan.m_offset + chan.m_componentCount * GetMeshFormatData(chan.m_format).m_size
+			);
+		}
 
-    return m_buffers[buffer].m_channels[channel].m_offset;
-}
+		// Set the user specified element size, or enlarge it if it was too small
+		buf.m_elementSize = FMath::Max(elementSize, int32(minElemSize));
 
+		// Update the buffer data
+		buf.m_data.SetNumUninitialized(buf.m_elementSize * m_elementCount, false);
+	}
 
-//---------------------------------------------------------------------------------------------
-void FMeshBufferSet::AddBuffer( const FMeshBufferSet& Other, int32 buffer )
-{
-    check( GetElementCount()==Other.GetElementCount() );
+	//---------------------------------------------------------------------------------------------
+	void FMeshBufferSet::SetBufferChannel
+	(
+		int32 buffer,
+		int32 channelIndex,
+		MESH_BUFFER_SEMANTIC semantic,
+		int32 semanticIndex,
+		MESH_BUFFER_FORMAT format,
+		int32 componentCount,
+		int32 offset
+	)
+	{
+		if (!(buffer >= 0 && buffer < m_buffers.Num()))
+		{
+			check(false);
+			return;
+		}
 
-	LLM_SCOPE_BYNAME(TEXT("MutableRuntime"));
+		LLM_SCOPE_BYNAME(TEXT("MutableRuntime"));
 
-    m_buffers.Add( Other.m_buffers[buffer] );
-}
+		MESH_BUFFER& buf = m_buffers[buffer];
 
+		if (!(channelIndex >= 0 && channelIndex < buf.m_channels.Num()))
+		{
+			check(false);
+			return;
+		}
 
-//---------------------------------------------------------------------------------------------
-void FMeshBufferSet::RemoveBuffer(int32 BufferIndex)
-{
-	if (BufferIndex >= 0 && BufferIndex < m_buffers.Num())
-	{
-		m_buffers.RemoveAt(BufferIndex);
+		MESH_BUFFER_CHANNEL& chan = buf.m_channels[channelIndex];
+		chan.m_semantic = semantic;
+		chan.m_semanticIndex = uint8_t(semanticIndex);
+		chan.m_format = format;
+		chan.m_componentCount = uint16_t(componentCount);
+		chan.m_offset = uint8(offset);
 	}
-	else
-	{
-		check(false);
-	}
-}
 
 
-//---------------------------------------------------------------------------------------------
-bool FMeshBufferSet::HasSameFormat( const FMeshBufferSet& Other ) const
-{
-    int32 bc = GetBufferCount();
-    if (Other.GetBufferCount() != bc )
-    {
-        return false;
-    }
-
-    for ( int32 b=0; b<bc; ++b )
-    {
-        if ( !HasSameFormat( b, Other ) )
-        {
-            return false;
-        }
-    }
-
-    return true;
-}
-
-
-//---------------------------------------------------------------------------------------------
-bool FMeshBufferSet::HasSameFormat( int32 buffer, const FMeshBufferSet& Other ) const
-{
-    if ( m_buffers[buffer].m_channels ==  Other.m_buffers[buffer].m_channels
-         &&
-         m_buffers[buffer].m_elementSize == Other.m_buffers[buffer].m_elementSize
-         )
-    {
-        return true;
-    }
-
-    return false;
-}
+	//---------------------------------------------------------------------------------------------
+	uint8* FMeshBufferSet::GetBufferData(int32 buffer)
+	{
+		check(buffer >= 0 && buffer < m_buffers.Num());
+		uint8_t* pResult = m_buffers[buffer].m_data.GetData();
+		return pResult;
+	}
 
 
-//---------------------------------------------------------------------------------------------
-size_t FMeshBufferSet::GetDataSize() const
-{
-    size_t res = 0;
+	//---------------------------------------------------------------------------------------------
+	const uint8* FMeshBufferSet::GetBufferData(int32 buffer) const
+	{
+		check(buffer >= 0 && buffer < m_buffers.Num());
+		const uint8_t* pResult = m_buffers[buffer].m_data.GetData();
+		return pResult;
+	}
 
-    for( int32 b=0; b<m_buffers.Num(); ++b )
-    {
-        res += m_buffers[b].m_elementSize * m_elementCount;
-    }
 
-    return res;
-}
+	//---------------------------------------------------------------------------------------------
+	void FMeshBufferSet::FindChannel
+	(
+		MESH_BUFFER_SEMANTIC semantic, int32 semanticIndex,
+		int32* pBuffer, int32* pChannel
+	) const
+	{
+		check(pBuffer && pChannel);
 
+		*pBuffer = -1;
+		*pChannel = -1;
 
-//-----------------------------------------------------------------------------------------
-void FMeshBufferSet::CopyElement(uint32_t fromIndex, uint32_t toIndex)
-{
-	check(fromIndex < m_elementCount);
-	check(toIndex < m_elementCount);
+		int index = 0;
 
-	if (fromIndex != toIndex)
-	{
-		for (auto& b : m_buffers)
+		for (size_t b = 0; b < m_buffers.Num(); ++b)
 		{
-			FMemory::Memcpy(&b.m_data[b.m_elementSize * toIndex],
-				&b.m_data[b.m_elementSize * fromIndex],
-				b.m_elementSize);
+			for (size_t c = 0; c < m_buffers[b].m_channels.Num(); ++c)
+			{
+				if (m_buffers[b].m_channels[c].m_semantic == semantic)
+				{
+					if (index == semanticIndex)
+					{
+						*pBuffer = int(b);
+						*pChannel = int(c);
+						return;
+					}
+					else
+					{
+						++index;
+					}
+				}
+			}
 		}
 	}
-}
 
 
-//-----------------------------------------------------------------------------------------
-bool FMeshBufferSet::IsSpecialBufferToIgnoreInSimilar(const MESH_BUFFER& b) const
-{
-	if (b.m_channels.Num() == 1
-		&&
-		b.m_channels[0].m_semantic == MBS_VERTEXINDEX)
+	//---------------------------------------------------------------------------------------------
+	int32 FMeshBufferSet::GetElementSize(int32 buffer) const
 	{
-		return true;
-	}
-	if (b.m_channels.Num() == 1
-		&&
-		b.m_channels[0].m_semantic == MBS_LAYOUTBLOCK)
-	{
-		return true;
+		check(buffer >= 0 && buffer < m_buffers.Num());
+
+		return m_buffers[buffer].m_elementSize;
 	}
-	return false;
-}
 
-//-----------------------------------------------------------------------------------------
-bool FMeshBufferSet::IsSimilarRobust(const FMeshBufferSet& Other) const
-{
 
-    MUTABLE_CPUPROFILER_SCOPE(FMeshBufferSet::IsSimilarRobust);
+	//---------------------------------------------------------------------------------------------
+	int32 FMeshBufferSet::GetChannelOffset
+	(
+		int32 buffer,
+		int32 channel
+	) const
+	{
+		check(buffer >= 0 && buffer < m_buffers.Num());
+		check(channel >= 0 && channel < m_buffers[buffer].m_channels.Num());
 
-	if ( m_elementCount != Other.m_elementCount )
-	{ 
-		return false;
+		return m_buffers[buffer].m_channels[channel].m_offset;
 	}
 
-	// IsSimilar is mush faster but can give false negatives if the buffer data description 
-	// omits parts of the data (e.g, memory layout paddings). It cannot have false positives. 
-	if ( IsSimilar(Other) )
+
+	//---------------------------------------------------------------------------------------------
+	void FMeshBufferSet::AddBuffer(const FMeshBufferSet& Other, int32 buffer)
 	{
-		return true;
+		check(GetElementCount() == Other.GetElementCount());
+
+		LLM_SCOPE_BYNAME(TEXT("MutableRuntime"));
+
+		m_buffers.Add(Other.m_buffers[buffer]);
 	}
 
-	const int32 ThisNumBuffers = m_buffers.Num();
-	const int32 OtherNumBuffers = Other.m_buffers.Num();
-	
-	int32 I = 0, J = 0;
-	while ( I < ThisNumBuffers && J < OtherNumBuffers )
+
+	//---------------------------------------------------------------------------------------------
+	void FMeshBufferSet::RemoveBuffer(int32 BufferIndex)
 	{
-		if ( IsSpecialBufferToIgnoreInSimilar(m_buffers[I]) )
+		if (BufferIndex >= 0 && BufferIndex < m_buffers.Num())
 		{
-			++I;
-			continue;
+			m_buffers.RemoveAt(BufferIndex);
 		}
-
-		if ( IsSpecialBufferToIgnoreInSimilar(Other.m_buffers[J]) )
+		else
 		{
-			++J;
-			continue;
+			check(false);
 		}
-		
-		const int32 ThisNumChannels = m_buffers[I].m_channels.Num();
-		const int32 OtherNumChannels = m_buffers[J].m_channels.Num();
+	}
 
-		if ( ThisNumChannels != OtherNumChannels )
+
+	//---------------------------------------------------------------------------------------------
+	bool FMeshBufferSet::HasSameFormat(const FMeshBufferSet& Other) const
+	{
+		int32 bc = GetBufferCount();
+		if (Other.GetBufferCount() != bc)
 		{
 			return false;
 		}
 
-		const MESH_BUFFER& ThisBuffer = m_buffers[I];
-		const MESH_BUFFER& OtherBuffer = Other.m_buffers[J];
-
-		for ( uint32 Elem = 0; Elem < m_elementCount; ++Elem )
+		for (int32 b = 0; b < bc; ++b)
 		{
-			for ( int32 C = 0; C < ThisNumChannels; ++C )
+			if (!HasSameFormat(b, Other))
 			{
-				const SIZE_T SizeA = GetMeshFormatData(ThisBuffer.m_channels[C].m_format).m_size * ThisBuffer.m_channels[C].m_componentCount;
-				const SIZE_T SizeB = GetMeshFormatData(OtherBuffer.m_channels[C].m_format).m_size * OtherBuffer.m_channels[C].m_componentCount;
+				return false;
+			}
+		}
 
-				if ( SizeA != SizeB )
-				{
-					return false;
-				}
+		return true;
+	}
 
-				const uint8* BuffA = ThisBuffer.m_data.GetData() + Elem * ThisBuffer.m_elementSize + ThisBuffer.m_channels[C].m_offset;
-				const uint8* BuffB = OtherBuffer.m_data.GetData() + Elem * OtherBuffer.m_elementSize + OtherBuffer.m_channels[C].m_offset;
 
-				if ( FMemory::Memcmp(BuffA, BuffB, SizeA) != 0 )
-				{
-					return false;
-				}
-			}
-		}		
+	//---------------------------------------------------------------------------------------------
+	bool FMeshBufferSet::HasSameFormat(int32 buffer, const FMeshBufferSet& Other) const
+	{
+		if (m_buffers[buffer].m_channels == Other.m_buffers[buffer].m_channels
+			&&
+			m_buffers[buffer].m_elementSize == Other.m_buffers[buffer].m_elementSize
+			)
+		{
+			return true;
+		}
 
-		++J;
-		++I;
+		return false;
 	}
 
-	if ( I + 1 < ThisNumBuffers || J + 1 < OtherNumBuffers )
+
+	//---------------------------------------------------------------------------------------------
+	size_t FMeshBufferSet::GetDataSize() const
 	{
-		return false;
+		size_t res = 0;
+
+		for (int32 b = 0; b < m_buffers.Num(); ++b)
+		{
+			res += m_buffers[b].m_elementSize * m_elementCount;
+		}
+
+		return res;
 	}
 
-	if ( I < ThisNumBuffers && !IsSpecialBufferToIgnoreInSimilar(m_buffers[I]) )
+
+	//-----------------------------------------------------------------------------------------
+	void FMeshBufferSet::CopyElement(uint32_t fromIndex, uint32_t toIndex)
 	{
-		return false;
+		check(fromIndex < m_elementCount);
+		check(toIndex < m_elementCount);
+
+		if (fromIndex != toIndex)
+		{
+			for (auto& b : m_buffers)
+			{
+				FMemory::Memcpy(&b.m_data[b.m_elementSize * toIndex],
+					&b.m_data[b.m_elementSize * fromIndex],
+					b.m_elementSize);
+			}
+		}
 	}
 
-	if ( J < OtherNumBuffers && !IsSpecialBufferToIgnoreInSimilar(Other.m_buffers[J]) )
+
+	//-----------------------------------------------------------------------------------------
+	bool FMeshBufferSet::IsSpecialBufferToIgnoreInSimilar(const MESH_BUFFER& b) const
 	{
+		if (b.m_channels.Num() == 1
+			&&
+			b.m_channels[0].m_semantic == MBS_VERTEXINDEX)
+		{
+			return true;
+		}
+		if (b.m_channels.Num() == 1
+			&&
+			b.m_channels[0].m_semantic == MBS_LAYOUTBLOCK)
+		{
+			return true;
+		}
 		return false;
 	}
-	
-	return true;
-}
 
-//-----------------------------------------------------------------------------------------
-bool FMeshBufferSet::IsSimilar(const FMeshBufferSet& o) const
-{
-    MUTABLE_CPUPROFILER_SCOPE(FMeshBufferSet::IsSimilar);
+	//-----------------------------------------------------------------------------------------
+	bool FMeshBufferSet::IsSimilarRobust(const FMeshBufferSet& Other) const
+	{
 
-	if (m_elementCount != o.m_elementCount) return false;
+		MUTABLE_CPUPROFILER_SCOPE(FMeshBufferSet::IsSimilarRobust);
 
-	// Compare all buffers except the vertex index channel, which should always be alone in
-	// the last buffer
-	int32 i = 0;
-	int32 oi = 0;
+		if (m_elementCount != Other.m_elementCount)
+		{
+			return false;
+		}
 
-	while (i < m_buffers.Num() && oi < o.m_buffers.Num())
-	{
-		// Is it a special buffer that we should ignore?
-		if (IsSpecialBufferToIgnoreInSimilar(m_buffers[i]))
+		// IsSimilar is mush faster but can give false negatives if the buffer data description 
+		// omits parts of the data (e.g, memory layout paddings). It cannot have false positives. 
+		if (IsSimilar(Other))
 		{
-			++i;
+			return true;
 		}
-		else
+
+		const int32 ThisNumBuffers = m_buffers.Num();
+		const int32 OtherNumBuffers = Other.m_buffers.Num();
+
+		int32 I = 0, J = 0;
+		while (I < ThisNumBuffers && J < OtherNumBuffers)
 		{
-			if (IsSpecialBufferToIgnoreInSimilar(o.m_buffers[oi]))
+			if (IsSpecialBufferToIgnoreInSimilar(m_buffers[I]))
 			{
-				++oi;
+				++I;
+				continue;
 			}
-			else
+
+			if (IsSpecialBufferToIgnoreInSimilar(Other.m_buffers[J]))
+			{
+				++J;
+				continue;
+			}
+
+			const int32 ThisNumChannels = m_buffers[I].m_channels.Num();
+			const int32 OtherNumChannels = m_buffers[J].m_channels.Num();
+
+			if (ThisNumChannels != OtherNumChannels)
+			{
+				return false;
+			}
+
+			const MESH_BUFFER& ThisBuffer = m_buffers[I];
+			const MESH_BUFFER& OtherBuffer = Other.m_buffers[J];
+
+			for (uint32 Elem = 0; Elem < m_elementCount; ++Elem)
 			{
-				if (!(m_buffers[i] == o.m_buffers[oi]))
+				for (int32 C = 0; C < ThisNumChannels; ++C)
 				{
-					return false;
+					const SIZE_T SizeA = GetMeshFormatData(ThisBuffer.m_channels[C].m_format).m_size * ThisBuffer.m_channels[C].m_componentCount;
+					const SIZE_T SizeB = GetMeshFormatData(OtherBuffer.m_channels[C].m_format).m_size * OtherBuffer.m_channels[C].m_componentCount;
+
+					if (SizeA != SizeB)
+					{
+						return false;
+					}
+
+					const uint8* BuffA = ThisBuffer.m_data.GetData() + Elem * ThisBuffer.m_elementSize + ThisBuffer.m_channels[C].m_offset;
+					const uint8* BuffB = OtherBuffer.m_data.GetData() + Elem * OtherBuffer.m_elementSize + OtherBuffer.m_channels[C].m_offset;
+
+					if (FMemory::Memcmp(BuffA, BuffB, SizeA) != 0)
+					{
+						return false;
+					}
 				}
-				++i;
-				++oi;
 			}
+
+			++J;
+			++I;
 		}
-	}
 
-	// Remaining buffers?
-	if (i + 1 < m_buffers.Num()
-		||
-		oi + 1 < o.m_buffers.Num())
-	{
-		return false;
-	}
+		if (I + 1 < ThisNumBuffers || J + 1 < OtherNumBuffers)
+		{
+			return false;
+		}
 
-	if (i < m_buffers.Num()
-		&&
-		!IsSpecialBufferToIgnoreInSimilar(m_buffers[i])
-		)
-	{
-		return false;
+		if (I < ThisNumBuffers && !IsSpecialBufferToIgnoreInSimilar(m_buffers[I]))
+		{
+			return false;
+		}
+
+		if (J < OtherNumBuffers && !IsSpecialBufferToIgnoreInSimilar(Other.m_buffers[J]))
+		{
+			return false;
+		}
+
+		return true;
 	}
 
-	if (oi < o.m_buffers.Num()
-		&&
-		!IsSpecialBufferToIgnoreInSimilar(o.m_buffers[oi])
-		)
+	//-----------------------------------------------------------------------------------------
+	bool FMeshBufferSet::IsSimilar(const FMeshBufferSet& o) const
 	{
-		return false;
-	}
+		MUTABLE_CPUPROFILER_SCOPE(FMeshBufferSet::IsSimilar);
 
-	return true;
-}
+		if (m_elementCount != o.m_elementCount) return false;
 
+		// Compare all buffers except the vertex index channel, which should always be alone in
+		// the last buffer
+		int32 i = 0;
+		int32 oi = 0;
 
-//-----------------------------------------------------------------------------------------
-void FMeshBufferSet::ResetBufferIndices()
-{
-	uint8_t currentIndices[MBS_COUNT];
-	memset(currentIndices, 0, sizeof(currentIndices));
+		while (i < m_buffers.Num() && oi < o.m_buffers.Num())
+		{
+			// Is it a special buffer that we should ignore?
+			if (IsSpecialBufferToIgnoreInSimilar(m_buffers[i]))
+			{
+				++i;
+			}
+			else
+			{
+				if (IsSpecialBufferToIgnoreInSimilar(o.m_buffers[oi]))
+				{
+					++oi;
+				}
+				else
+				{
+					if (!(m_buffers[i] == o.m_buffers[oi]))
+					{
+						return false;
+					}
+					++i;
+					++oi;
+				}
+			}
+		}
 
-	for (auto& b : m_buffers)
-	{
-		for (auto& c : b.m_channels)
+		// Remaining buffers?
+		if (i + 1 < m_buffers.Num()
+			||
+			oi + 1 < o.m_buffers.Num())
 		{
-			c.m_semanticIndex = currentIndices[c.m_semantic];
-			currentIndices[c.m_semantic]++;
+			return false;
+		}
+
+		if (i < m_buffers.Num()
+			&&
+			!IsSpecialBufferToIgnoreInSimilar(m_buffers[i])
+			)
+		{
+			return false;
+		}
+
+		if (oi < o.m_buffers.Num()
+			&&
+			!IsSpecialBufferToIgnoreInSimilar(o.m_buffers[oi])
+			)
+		{
+			return false;
 		}
+
+		return true;
 	}
-}
 
 
-//-----------------------------------------------------------------------------------------
-void FMeshBufferSet::UpdateOffsets(int32 b)
-{
-	uint8_t offset = 0;
-	for (auto& c : m_buffers[b].m_channels)
+	//-----------------------------------------------------------------------------------------
+	void FMeshBufferSet::ResetBufferIndices()
 	{
-		if (c.m_offset < offset)
+		uint8_t currentIndices[MBS_COUNT];
+		memset(currentIndices, 0, sizeof(currentIndices));
+
+		for (auto& b : m_buffers)
 		{
-			c.m_offset = offset;
+			for (auto& c : b.m_channels)
+			{
+				c.m_semanticIndex = currentIndices[c.m_semantic];
+				currentIndices[c.m_semantic]++;
+			}
 		}
-		else
+	}
+
+
+	//-----------------------------------------------------------------------------------------
+	void FMeshBufferSet::UpdateOffsets(int32 b)
+	{
+		uint8_t offset = 0;
+		for (auto& c : m_buffers[b].m_channels)
 		{
-			offset = c.m_offset;
+			if (c.m_offset < offset)
+			{
+				c.m_offset = offset;
+			}
+			else
+			{
+				offset = c.m_offset;
+			}
+			offset += uint8_t(c.m_componentCount * GetMeshFormatData(c.m_format).m_size);
 		}
-		offset += uint8_t(c.m_componentCount * GetMeshFormatData(c.m_format).m_size);
+
+		if (m_buffers[b].m_elementSize < offset)
+			m_buffers[b].m_elementSize = offset;
 	}
 
-	if (m_buffers[b].m_elementSize < offset)
-		m_buffers[b].m_elementSize = offset;
-}
+}
\ No newline at end of file
diff --git a/Engine/Plugins/Experimental/Mutable/Source/MutableRuntime/Private/OpImagePixelFormat.cpp b/Engine/Plugins/Experimental/Mutable/Source/MutableRuntime/Private/OpImagePixelFormat.cpp
--- a/Engine/Plugins/Experimental/Mutable/Source/MutableRuntime/Private/OpImagePixelFormat.cpp
+++ b/Engine/Plugins/Experimental/Mutable/Source/MutableRuntime/Private/OpImagePixelFormat.cpp
@@ -8,11 +8,13 @@
 
 #include "MutableTrace.h"
 
-using namespace mu;
+
+namespace mu
+{
 
 
 //-------------------------------------------------------------------------------------------------
-ImagePtr mu::ImagePixelFormat( int imageCompressionQuality, const Image* pBase,
+ImagePtr ImagePixelFormat( int imageCompressionQuality, const Image* pBase,
                                EImageFormat targetFormat, int onlyLOD )
 {
 	MUTABLE_CPUPROFILER_SCOPE(ImagePixelFormat);
@@ -52,7 +54,7 @@ ImagePtr mu::ImagePixelFormat( int imageCompressionQuality, const Image* pBase,
 
 
 //-------------------------------------------------------------------------------------------------
-bool mu::ImagePixelFormatInPlace( int imageCompressionQuality, Image* pResult, const Image* pBase,
+bool ImagePixelFormatInPlace( int imageCompressionQuality, Image* pResult, const Image* pBase,
                                   int onlyLOD )
 {
 	MUTABLE_CPUPROFILER_SCOPE(ImagePixelFormatInPlace);
@@ -1186,3 +1188,4 @@ bool mu::ImagePixelFormatInPlace( int imageCompressionQuality, Image* pResult, c
 }
 
 
+}
diff --git a/Engine/Plugins/Experimental/Mutable/Source/MutableRuntime/Private/OpImageProject.cpp b/Engine/Plugins/Experimental/Mutable/Source/MutableRuntime/Private/OpImageProject.cpp
--- a/Engine/Plugins/Experimental/Mutable/Source/MutableRuntime/Private/OpImageProject.cpp
+++ b/Engine/Plugins/Experimental/Mutable/Source/MutableRuntime/Private/OpImageProject.cpp
@@ -17,7 +17,8 @@
 #include <unordered_set>
 
 
-using namespace mu;
+namespace mu
+{
 
 
 namespace
@@ -988,7 +989,7 @@ void ImageRasterProjected_OptimisedWrapping( const Mesh* pMesh,
 
 
 //-------------------------------------------------------------------------------------------------
-void mu::ImageRasterProjectedPlanar( const Mesh* pMesh,
+void ImageRasterProjectedPlanar( const Mesh* pMesh,
                                      Image* pImage,
                                      const Image* pSource,
                                      const Image* pMask,
@@ -1075,7 +1076,7 @@ void mu::ImageRasterProjectedPlanar( const Mesh* pMesh,
 
 
 //-------------------------------------------------------------------------------------------------
-void mu::ImageRasterProjectedWrapping( const Mesh* pMesh,
+void ImageRasterProjectedWrapping( const Mesh* pMesh,
                                        Image* pImage,
                                        const Image* pSource,
                                        const Image* pMask,
@@ -1162,7 +1163,7 @@ void mu::ImageRasterProjectedWrapping( const Mesh* pMesh,
 
 
 //-------------------------------------------------------------------------------------------------
-void mu::ImageRasterProjectedCylindrical( const Mesh* pMesh,
+void ImageRasterProjectedCylindrical( const Mesh* pMesh,
                                           Image* pImage,
                                           const Image* pSource,
                                           const Image* pMask,
@@ -2638,7 +2639,7 @@ PRAGMA_ENABLE_OPTIMIZATION
 
 
 //-------------------------------------------------------------------------------------------------
-MeshPtr mu::MeshProject( const Mesh* pMesh,
+MeshPtr MeshProject( const Mesh* pMesh,
                          const PROJECTOR& projector )
 {
 	MUTABLE_CPUPROFILER_SCOPE(MeshProject);
@@ -2674,7 +2675,7 @@ MeshPtr mu::MeshProject( const Mesh* pMesh,
 
 
 //-------------------------------------------------------------------------------------------------
-MeshPtr mu::CreateMeshOptimisedForProjection( int layout )
+MeshPtr CreateMeshOptimisedForProjection( int layout )
 {
     MeshPtr pFormatMesh = new Mesh();
     pFormatMesh->GetVertexBuffers().SetBufferCount( 1 );
@@ -2708,7 +2709,7 @@ MeshPtr mu::CreateMeshOptimisedForProjection( int layout )
 
 
 //-------------------------------------------------------------------------------------------------
-MeshPtr mu::CreateMeshOptimisedForWrappingProjection( int layout )
+MeshPtr CreateMeshOptimisedForWrappingProjection( int layout )
 {
     MeshPtr pFormatMesh = new Mesh();
     pFormatMesh->GetVertexBuffers().SetBufferCount( 1 );
@@ -2740,3 +2741,5 @@ MeshPtr mu::CreateMeshOptimisedForWrappingProjection( int layout )
 
     return pFormatMesh;
 }
+
+}
\ No newline at end of file
diff --git a/Engine/Plugins/Experimental/Mutable/Source/MutableRuntime/Private/OpMeshFormat.cpp b/Engine/Plugins/Experimental/Mutable/Source/MutableRuntime/Private/OpMeshFormat.cpp
--- a/Engine/Plugins/Experimental/Mutable/Source/MutableRuntime/Private/OpMeshFormat.cpp
+++ b/Engine/Plugins/Experimental/Mutable/Source/MutableRuntime/Private/OpMeshFormat.cpp
@@ -9,755 +9,756 @@
 #include "MutableTrace.h"
 
 
-using namespace mu;
+namespace mu
+{
 
 
-//-------------------------------------------------------------------------------------------------
-// Eric Lengyel method's
-// http://www.terathon.com/code/tangent.html
-//-------------------------------------------------------------------------------------------------
-namespace
-{
+	//-------------------------------------------------------------------------------------------------
+	// Eric Lengyel method's
+	// http://www.terathon.com/code/tangent.html
+	//-------------------------------------------------------------------------------------------------
+	namespace
+	{
 
 #define MUTABLE_VERTEX_MERGE_TEX_RANGE          1024
 #define MUTABLE_TANGENT_GENERATION_EPSILON_1    0.000001f
 #define MUTABLE_TANGENT_GENERATION_EPSILON_2    0.001f
 #define MUTABLE_TANGENT_MIN_AXIS_DIFFERENCE     0.999f
 
-    struct TVertex
-    {
-
-        TVertex() = default;
-
-        TVertex( const vec3<float>& p, const vec3<float>&  n, const vec2<float>& t )
-        {
-            pos = p;
-            nor = n;
-            tex = t;
-        }
-
-        vec3<float> pos;
-        vec3<float> nor;
-        vec2<float> tex;
-
-        inline bool operator< ( const TVertex& other ) const
-        {
-            if( pos < other.pos )
-                return true;
-            if( other.pos < pos )
-                return false;
-
-            if( nor < other.nor )
-                return true;
-            if( other.nor < nor )
-                return false;
-
-            // Compare the texture coordinates with a particular precission.
-            vec2<int> uv0 = vec2<int>( int( tex[0] * MUTABLE_VERTEX_MERGE_TEX_RANGE ),
-                                       int( tex[1] * MUTABLE_VERTEX_MERGE_TEX_RANGE ) );
-            vec2<int> uv1 = vec2<int>( int( other.tex[0] * MUTABLE_VERTEX_MERGE_TEX_RANGE ),
-                                       int( other.tex[1] * MUTABLE_VERTEX_MERGE_TEX_RANGE ) );
-
-            if( uv0 < uv1 )
-                return true;
-            if( uv1 < uv0 )
-                return false;
-
-            return false;
-        }
-
-        inline bool operator== ( const TVertex& other ) const
-        {
-            return  ! ( (*this) < other )
-                    &&
-                    ! ( other < (*this) );
-        }
-    };
-
-
-    struct TFace
-    {
-        vec3<float> T;
-        vec3<float> B;
-        vec3<float> N;
-
-        TFace()
-            : T(0,0,0)
-            , B(0,0,0)
-            , N(0,0,0)
-        {}
-
-        TFace
-            (
-                const vec3<float>& v1,
-                const vec3<float>& v2,
-                const vec3<float>& v3,
-                const vec2<float>& w1,
-                const vec2<float>& w2,
-                const vec2<float>& w3
-            )
-        {
-            vec3<float> E1 = v2 - v1;
-            vec3<float> E2 = v3 - v1;
-
-            vec2<float> UV1 = w2 - w1;
-            vec2<float> UV2 = w3 - w1;
-
-            float  UVdet = UV1[0] * UV2[1] - UV2[0] * UV1[1];
-
-            N = normalise( cross( E1, E2 ) );
-
-            if( !(fabs(UVdet)<=MUTABLE_TANGENT_GENERATION_EPSILON_1 ) )
-            {
-                double r = 1.0 / UVdet;
-
-                T = vec3<float>
-                    (
-                        (UV2[1] * E1[0] - UV1[1] * E2[0]),
-                        (UV2[1] * E1[1] - UV1[1] * E2[1]),
-                        (UV2[1] * E1[2] - UV1[1] * E2[2])
-                    );
-
-                B = vec3<float>
-                    (
-                        (UV1[0] * E2[0] - UV2[0] * E1[0]),
-                        (UV1[0] * E2[1] - UV2[0] * E1[1]),
-                        (UV1[0] * E2[2] - UV2[0] * E1[2])
-                    );
-
-                T = T * float(r);
-                B = B * float(r);
-
-                T = normalise(T);
-                B = normalise(B);
-            }
-            else
-            {
-                T = vec3<float>(0,0,0);
-                B = vec3<float>(0,0,0);
-                N = vec3<float>(0,0,0);
-            }
-        }
-
-    };
-
-
-	void RebuildTangents
+		struct TVertex
+		{
+
+			TVertex() = default;
+
+			TVertex(const vec3<float>& p, const vec3<float>& n, const vec2<float>& t)
+			{
+				pos = p;
+				nor = n;
+				tex = t;
+			}
+
+			vec3<float> pos;
+			vec3<float> nor;
+			vec2<float> tex;
+
+			inline bool operator< (const TVertex& other) const
+			{
+				if (pos < other.pos)
+					return true;
+				if (other.pos < pos)
+					return false;
+
+				if (nor < other.nor)
+					return true;
+				if (other.nor < nor)
+					return false;
+
+				// Compare the texture coordinates with a particular precission.
+				vec2<int> uv0 = vec2<int>(int(tex[0] * MUTABLE_VERTEX_MERGE_TEX_RANGE),
+					int(tex[1] * MUTABLE_VERTEX_MERGE_TEX_RANGE));
+				vec2<int> uv1 = vec2<int>(int(other.tex[0] * MUTABLE_VERTEX_MERGE_TEX_RANGE),
+					int(other.tex[1] * MUTABLE_VERTEX_MERGE_TEX_RANGE));
+
+				if (uv0 < uv1)
+					return true;
+				if (uv1 < uv0)
+					return false;
+
+				return false;
+			}
+
+			inline bool operator== (const TVertex& other) const
+			{
+				return  !((*this) < other)
+					&&
+					!(other < (*this));
+			}
+		};
+
+
+		struct TFace
+		{
+			vec3<float> T;
+			vec3<float> B;
+			vec3<float> N;
+
+			TFace()
+				: T(0, 0, 0)
+				, B(0, 0, 0)
+				, N(0, 0, 0)
+			{}
+
+			TFace
+			(
+				const vec3<float>& v1,
+				const vec3<float>& v2,
+				const vec3<float>& v3,
+				const vec2<float>& w1,
+				const vec2<float>& w2,
+				const vec2<float>& w3
+			)
+			{
+				vec3<float> E1 = v2 - v1;
+				vec3<float> E2 = v3 - v1;
+
+				vec2<float> UV1 = w2 - w1;
+				vec2<float> UV2 = w3 - w1;
+
+				float  UVdet = UV1[0] * UV2[1] - UV2[0] * UV1[1];
+
+				N = normalise(cross(E1, E2));
+
+				if (!(fabs(UVdet) <= MUTABLE_TANGENT_GENERATION_EPSILON_1))
+				{
+					double r = 1.0 / UVdet;
+
+					T = vec3<float>
+						(
+							(UV2[1] * E1[0] - UV1[1] * E2[0]),
+							(UV2[1] * E1[1] - UV1[1] * E2[1]),
+							(UV2[1] * E1[2] - UV1[1] * E2[2])
+							);
+
+					B = vec3<float>
+						(
+							(UV1[0] * E2[0] - UV2[0] * E1[0]),
+							(UV1[0] * E2[1] - UV2[0] * E1[1]),
+							(UV1[0] * E2[2] - UV2[0] * E1[2])
+							);
+
+					T = T * float(r);
+					B = B * float(r);
+
+					T = normalise(T);
+					B = normalise(B);
+				}
+				else
+				{
+					T = vec3<float>(0, 0, 0);
+					B = vec3<float>(0, 0, 0);
+					N = vec3<float>(0, 0, 0);
+				}
+			}
+
+		};
+
+
+		void RebuildTangents
 		(
 			const FMeshBufferSet& IndexBuffers,
 			FMeshBufferSet& VertexBuffers
 		)
-	{
-		// Because of these, this method only works for a specific mesh format
-        UntypedMeshBufferIteratorConst vertex( VertexBuffers, MBS_POSITION, 0 );
-        UntypedMeshBufferIteratorConst normal( VertexBuffers, MBS_NORMAL, 0 );
-        UntypedMeshBufferIteratorConst texcoord( VertexBuffers, MBS_TEXCOORDS, 0 );
-
-		int vertexCount = VertexBuffers.GetElementCount();
-
-		// Group similar vertices
-		vector< TVertex > vertexData;
-		vertexData.reserve( vertexCount );
-        vector< uint32_t > lVertexGroups( vertexCount, 0 );
-
-        //Agrupamos los vertices por posicion, normal, texcoord
-        for( int32_t i = 0; i < vertexCount; ++i )
-        {
-            TVertex v( (vertex+i).GetAsVec4f().xyz(),
-                       (normal+i).GetAsVec4f().xyz(),
-                       (texcoord+i).GetAsVec4f().xy() );
-            vector<TVertex>::const_iterator it = std::find( vertexData.begin(),
-                                                            vertexData.end(),
-                                                            v);
-
-            size_t pos = 0;
-            if ( it != vertexData.end() )
-            {
-                pos = it - vertexData.begin();
-            }
-            else
-            {
-                pos = vertexData.size();
-                vertexData.push_back( v );
-            }
-
-            lVertexGroups[ i ] = (int)pos;
-        }
-
-        // Calculate the tangent space
-        MeshBufferIterator<MBF_FLOAT32, float, 3> tangent( VertexBuffers, MBS_TANGENT, 0 );
-        MeshBufferIterator<MBF_FLOAT32, float, 3> bitangent( VertexBuffers, MBS_BINORMAL, 0 );
-        UntypedMeshBufferIteratorConst indices( IndexBuffers, MBS_VERTEXINDEX, 0 );
-        int indexCount = IndexBuffers.GetElementCount();
-
-        vector<TFace> vertexFaces( indexCount );
-
-        // Worst case: a group for every vertex
-        vector< vector< TFace > > grp_faces( indexCount );
-
-        for ( long faceidx = 0, count = indexCount/3; faceidx < count; faceidx++ )
-        {
-            long i1 = 0;
-            long i2 = 0;
-            long i3 = 0;
-
-            if ( indices.GetFormat()==MBF_UINT16)
-            {
-                auto pIndices = reinterpret_cast<const uint16_t*>( indices.ptr() );
-                i1 = pIndices [faceidx * 3    ];
-                i2 = pIndices [faceidx * 3 + 1];
-                i3 = pIndices [faceidx * 3 + 2];
-            }
-            else if ( indices.GetFormat()==MBF_UINT32)
-            {
-                auto pIndices = reinterpret_cast<const uint32_t*>( indices.ptr() );
-                i1 = pIndices [faceidx * 3    ];
-                i2 = pIndices [faceidx * 3 + 1];
-                i3 = pIndices [faceidx * 3 + 2];
-            }
-            else
-            {
-                check( false );
-            }
-
-            TFace face;
-
-            const vec3<float> v1 = (vertex+(int)i1).GetAsVec4f().xyz();
-            const vec3<float> v2 = (vertex+(int)i2).GetAsVec4f().xyz();
-            const vec3<float> v3 = (vertex+(int)i3).GetAsVec4f().xyz();
-
-            // Compute face t & b
-            const vec2<float> w1 = (texcoord+(int)i1).GetAsVec4f().xy();
-            const vec2<float> w2 = (texcoord+(int)i2).GetAsVec4f().xy();
-            const vec2<float> w3 = (texcoord+(int)i3).GetAsVec4f().xy();
-            face = TFace( v1, v2, v3, w1, w2, w3 );
-
-            if( length( face.N ) != 0 )
-            {
-                vertexFaces[ i1 ] = vertexFaces[ i2 ] = vertexFaces[ i3 ] = face;
-
-                grp_faces[lVertexGroups[i1]].push_back( face );
-                grp_faces[lVertexGroups[i2]].push_back( face );
-                grp_faces[lVertexGroups[i3]].push_back( face );
-            }
-        }
-
-        for ( long vtxidx = 0; vtxidx < vertexCount; vtxidx++ )
-        {
-            const vec3<float> n = (normal+(int)vtxidx).GetAsVec4f().xyz();
-            const uint32_t group    = lVertexGroups[vtxidx];
-
-            vec3<float> lTangent  (0,0,0);
-            vec3<float> lBiTangent(0,0,0);
-
-            const vector< TFace >& faces   = grp_faces[group];
-
-            const TFace& vFace = vertexFaces[ vtxidx ];
-
-            // Ignore mirrors
-            for( unsigned i = 0; i < faces.size(); ++i )
-            {
-                const float DN = dot( vFace.N, faces[i].N );
-                const float DT = dot( vFace.T, faces[i].T );
-                const float DB = dot( vFace.B, faces[i].B );
-
-                if( DN > 0 && DT > 0 && DB > 0 )
-                {
-                    lTangent   += faces[i].T;
-                    lBiTangent += faces[i].B;
-                }
-            }
-
-            vec3<float> ortogonalizedTangent
-                    = normalise( lTangent
-                                 - ( n * dot( n, lTangent  ) ) );
-
-            vec3<float> ortogonalizedBiTangent
-                    = normalise( lBiTangent
-                                 - ( n * dot( n, lBiTangent) )
-                                 - ( ortogonalizedTangent * dot(ortogonalizedTangent,lBiTangent) ));
-
-            // this is where we can do a final check for zero length vectors
-            // and set them to something appropriate
-            float lenTan = length( ortogonalizedTangent );
-            float lenBin = length( ortogonalizedBiTangent );
-
-            if( (lenTan <= MUTABLE_TANGENT_GENERATION_EPSILON_2)
-                    || (lenBin <= MUTABLE_TANGENT_GENERATION_EPSILON_2)  ) //should be approx 1.0f
-            {
-                // the tangent space is ill defined at this vertex
-                // so we can generate a valid one based on the normal vector,
-                // which I'm assuming is valid!
-                if(lenTan > 0.5f)
-                {
-                    // the tangent is valid, so we can just use that to calculate the binormal
-                    ortogonalizedBiTangent = cross( n, ortogonalizedTangent );
-                }
-                else if(lenBin > 0.5)
-                {
-                    // the binormal is good and we can use it to calculate the tangent
-                    ortogonalizedTangent = cross( ortogonalizedBiTangent, n );
-                }
-                else
-                {
-                    // both vectors are invalid, so we should create something
-                    // that is at least valid if not correct
-                    // I'm checking two possible axis, because the normal could be one of them,
-                    // and we want to chose a different one to start making our valid basis.
-                    // I can find out which is further away from it by checking the dot product
-                    vec3<float> startAxis;
-
-                    if( dot(vec3<float>(1,0,0), n) <  dot(vec3<float>(0,1,0), n) )
-                    {
-                        //the xAxis is more different than the yAxis when compared to the normal
-                        startAxis = vec3<float>(1,0,0);
-                    }
-                    else
-                    {
-                        // the yAxis is more different than the xAxis when compared to the normal
-                        startAxis = vec3<float>(0,1,0);
-                    }
-
-                    ortogonalizedTangent   = cross( n, startAxis);
-                    ortogonalizedBiTangent = cross( n, ortogonalizedTangent);
-                }
-            }
-            else
-            {
-                // one final sanity check, make sure that they tangent and binormal are different
-                // enough
-                if ( dot(ortogonalizedTangent, ortogonalizedBiTangent )
-                     >
-                     MUTABLE_TANGENT_MIN_AXIS_DIFFERENCE )
-                {
-                    // then they are too similar lets make them more different
-                    ortogonalizedBiTangent = cross( n, ortogonalizedTangent );
-                }
-            }
-
-            tangent  [ vtxidx ] = ortogonalizedTangent;
-            bitangent[ vtxidx ] = ortogonalizedBiTangent;
-        }
-
-    }
+		{
+			// Because of these, this method only works for a specific mesh format
+			UntypedMeshBufferIteratorConst vertex(VertexBuffers, MBS_POSITION, 0);
+			UntypedMeshBufferIteratorConst normal(VertexBuffers, MBS_NORMAL, 0);
+			UntypedMeshBufferIteratorConst texcoord(VertexBuffers, MBS_TEXCOORDS, 0);
 
-}
+			int vertexCount = VertexBuffers.GetElementCount();
+
+			// Group similar vertices
+			vector< TVertex > vertexData;
+			vertexData.reserve(vertexCount);
+			vector< uint32_t > lVertexGroups(vertexCount, 0);
+
+			//Agrupamos los vertices por posicion, normal, texcoord
+			for (int32_t i = 0; i < vertexCount; ++i)
+			{
+				TVertex v((vertex + i).GetAsVec4f().xyz(),
+					(normal + i).GetAsVec4f().xyz(),
+					(texcoord + i).GetAsVec4f().xy());
+				vector<TVertex>::const_iterator it = std::find(vertexData.begin(),
+					vertexData.end(),
+					v);
+
+				size_t pos = 0;
+				if (it != vertexData.end())
+				{
+					pos = it - vertexData.begin();
+				}
+				else
+				{
+					pos = vertexData.size();
+					vertexData.push_back(v);
+				}
+
+				lVertexGroups[i] = (int)pos;
+			}
+
+			// Calculate the tangent space
+			MeshBufferIterator<MBF_FLOAT32, float, 3> tangent(VertexBuffers, MBS_TANGENT, 0);
+			MeshBufferIterator<MBF_FLOAT32, float, 3> bitangent(VertexBuffers, MBS_BINORMAL, 0);
+			UntypedMeshBufferIteratorConst indices(IndexBuffers, MBS_VERTEXINDEX, 0);
+			int indexCount = IndexBuffers.GetElementCount();
+
+			vector<TFace> vertexFaces(indexCount);
+
+			// Worst case: a group for every vertex
+			vector< vector< TFace > > grp_faces(indexCount);
+
+			for (long faceidx = 0, count = indexCount / 3; faceidx < count; faceidx++)
+			{
+				long i1 = 0;
+				long i2 = 0;
+				long i3 = 0;
+
+				if (indices.GetFormat() == MBF_UINT16)
+				{
+					auto pIndices = reinterpret_cast<const uint16_t*>(indices.ptr());
+					i1 = pIndices[faceidx * 3];
+					i2 = pIndices[faceidx * 3 + 1];
+					i3 = pIndices[faceidx * 3 + 2];
+				}
+				else if (indices.GetFormat() == MBF_UINT32)
+				{
+					auto pIndices = reinterpret_cast<const uint32_t*>(indices.ptr());
+					i1 = pIndices[faceidx * 3];
+					i2 = pIndices[faceidx * 3 + 1];
+					i3 = pIndices[faceidx * 3 + 2];
+				}
+				else
+				{
+					check(false);
+				}
+
+				TFace face;
+
+				const vec3<float> v1 = (vertex + (int)i1).GetAsVec4f().xyz();
+				const vec3<float> v2 = (vertex + (int)i2).GetAsVec4f().xyz();
+				const vec3<float> v3 = (vertex + (int)i3).GetAsVec4f().xyz();
+
+				// Compute face t & b
+				const vec2<float> w1 = (texcoord + (int)i1).GetAsVec4f().xy();
+				const vec2<float> w2 = (texcoord + (int)i2).GetAsVec4f().xy();
+				const vec2<float> w3 = (texcoord + (int)i3).GetAsVec4f().xy();
+				face = TFace(v1, v2, v3, w1, w2, w3);
+
+				if (length(face.N) != 0)
+				{
+					vertexFaces[i1] = vertexFaces[i2] = vertexFaces[i3] = face;
+
+					grp_faces[lVertexGroups[i1]].push_back(face);
+					grp_faces[lVertexGroups[i2]].push_back(face);
+					grp_faces[lVertexGroups[i3]].push_back(face);
+				}
+			}
+
+			for (long vtxidx = 0; vtxidx < vertexCount; vtxidx++)
+			{
+				const vec3<float> n = (normal + (int)vtxidx).GetAsVec4f().xyz();
+				const uint32_t group = lVertexGroups[vtxidx];
+
+				vec3<float> lTangent(0, 0, 0);
+				vec3<float> lBiTangent(0, 0, 0);
+
+				const vector< TFace >& faces = grp_faces[group];
 
+				const TFace& vFace = vertexFaces[vtxidx];
 
-//-------------------------------------------------------------------------------------------------
-void mu::MeshFormatBuffer
+				// Ignore mirrors
+				for (unsigned i = 0; i < faces.size(); ++i)
+				{
+					const float DN = dot(vFace.N, faces[i].N);
+					const float DT = dot(vFace.T, faces[i].T);
+					const float DB = dot(vFace.B, faces[i].B);
+
+					if (DN > 0 && DT > 0 && DB > 0)
+					{
+						lTangent += faces[i].T;
+						lBiTangent += faces[i].B;
+					}
+				}
+
+				vec3<float> ortogonalizedTangent
+					= normalise(lTangent
+						- (n * dot(n, lTangent)));
+
+				vec3<float> ortogonalizedBiTangent
+					= normalise(lBiTangent
+						- (n * dot(n, lBiTangent))
+						- (ortogonalizedTangent * dot(ortogonalizedTangent, lBiTangent)));
+
+				// this is where we can do a final check for zero length vectors
+				// and set them to something appropriate
+				float lenTan = length(ortogonalizedTangent);
+				float lenBin = length(ortogonalizedBiTangent);
+
+				if ((lenTan <= MUTABLE_TANGENT_GENERATION_EPSILON_2)
+					|| (lenBin <= MUTABLE_TANGENT_GENERATION_EPSILON_2)) //should be approx 1.0f
+				{
+					// the tangent space is ill defined at this vertex
+					// so we can generate a valid one based on the normal vector,
+					// which I'm assuming is valid!
+					if (lenTan > 0.5f)
+					{
+						// the tangent is valid, so we can just use that to calculate the binormal
+						ortogonalizedBiTangent = cross(n, ortogonalizedTangent);
+					}
+					else if (lenBin > 0.5)
+					{
+						// the binormal is good and we can use it to calculate the tangent
+						ortogonalizedTangent = cross(ortogonalizedBiTangent, n);
+					}
+					else
+					{
+						// both vectors are invalid, so we should create something
+						// that is at least valid if not correct
+						// I'm checking two possible axis, because the normal could be one of them,
+						// and we want to chose a different one to start making our valid basis.
+						// I can find out which is further away from it by checking the dot product
+						vec3<float> startAxis;
+
+						if (dot(vec3<float>(1, 0, 0), n) < dot(vec3<float>(0, 1, 0), n))
+						{
+							//the xAxis is more different than the yAxis when compared to the normal
+							startAxis = vec3<float>(1, 0, 0);
+						}
+						else
+						{
+							// the yAxis is more different than the xAxis when compared to the normal
+							startAxis = vec3<float>(0, 1, 0);
+						}
+
+						ortogonalizedTangent = cross(n, startAxis);
+						ortogonalizedBiTangent = cross(n, ortogonalizedTangent);
+					}
+				}
+				else
+				{
+					// one final sanity check, make sure that they tangent and binormal are different
+					// enough
+					if (dot(ortogonalizedTangent, ortogonalizedBiTangent)
+						>
+						MUTABLE_TANGENT_MIN_AXIS_DIFFERENCE)
+					{
+						// then they are too similar lets make them more different
+						ortogonalizedBiTangent = cross(n, ortogonalizedTangent);
+					}
+				}
+
+				tangent[vtxidx] = ortogonalizedTangent;
+				bitangent[vtxidx] = ortogonalizedBiTangent;
+			}
+
+		}
+
+	}
+
+
+	//-------------------------------------------------------------------------------------------------
+	void MeshFormatBuffer
 	(
 		const FMeshBufferSet& Source,
-        FMeshBufferSet& Result,
-        int bufferIndex
+		FMeshBufferSet& Result,
+		int bufferIndex
 	)
-{	
-	int vCount = Source.GetElementCount();
+	{
+		int vCount = Source.GetElementCount();
 
-    int b = bufferIndex;
-    {
-		// For every channel in this buffer
-		for ( int c=0; c<Result.GetBufferChannelCount(b); ++c )
+		int b = bufferIndex;
 		{
-			// Find this channel in the source mesh
-			MESH_BUFFER_SEMANTIC resultSemantic;
-			int resultSemanticIndex;
-			MESH_BUFFER_FORMAT resultFormat;
-			int resultComponents;
-			int resultOffset;
-			Result.GetChannel
-					(
-						b, c,
-						&resultSemantic, &resultSemanticIndex,
-						&resultFormat, &resultComponents,
-						&resultOffset
-					);
+			// For every channel in this buffer
+			for (int c = 0; c < Result.GetBufferChannelCount(b); ++c)
+			{
+				// Find this channel in the source mesh
+				MESH_BUFFER_SEMANTIC resultSemantic;
+				int resultSemanticIndex;
+				MESH_BUFFER_FORMAT resultFormat;
+				int resultComponents;
+				int resultOffset;
+				Result.GetChannel
+				(
+					b, c,
+					&resultSemantic, &resultSemanticIndex,
+					&resultFormat, &resultComponents,
+					&resultOffset
+				);
+
+				int sourceBuffer;
+				int sourceChannel;
+				Source.FindChannel
+				(resultSemantic, resultSemanticIndex, &sourceBuffer, &sourceChannel);
+
+				int resultElemSize = Result.GetElementSize(b);
+				int resultChannelSize = GetMeshFormatData(resultFormat).m_size * resultComponents;
+				uint8_t* pResultBuf = Result.GetBufferData(b);
+				pResultBuf += resultOffset;
+
+				if (sourceBuffer < 0)
+				{
+					// Not found: fill with zeros.
+
+					// Special case for derived channel data
+					bool generated = false;
+					if (resultSemantic == MBS_TANGENTSIGN)
+					{
+						// Look for the full tangent space
+						int tanXBuf, tanXChan, tanYBuf, tanYChan, tanZBuf, tanZChan;
+						Source.FindChannel(MBS_TANGENT, resultSemanticIndex, &tanXBuf, &tanXChan);
+						Source.FindChannel(MBS_BINORMAL, resultSemanticIndex, &tanYBuf, &tanYChan);
+						Source.FindChannel(MBS_NORMAL, resultSemanticIndex, &tanZBuf, &tanZChan);
 
-			int sourceBuffer;
-			int sourceChannel;
-			Source.FindChannel
-					( resultSemantic, resultSemanticIndex, &sourceBuffer, &sourceChannel );
+						if (tanXBuf >= 0 && tanYBuf >= 0 && tanZBuf >= 0)
+						{
+							generated = true;
+							UntypedMeshBufferIteratorConst xIt(Source, MBS_TANGENT, resultSemanticIndex);
+							UntypedMeshBufferIteratorConst yIt(Source, MBS_BINORMAL, resultSemanticIndex);
+							UntypedMeshBufferIteratorConst zIt(Source, MBS_NORMAL, resultSemanticIndex);
+							for (int v = 0; v < vCount; ++v)
+							{
+								mat3f mat;
+								mat[0] = xIt.GetAsVec4f().xyz();
+								mat[1] = yIt.GetAsVec4f().xyz();
+								mat[2] = zIt.GetAsVec4f().xyz();
+								float sign = mat.GetDeterminant() < 0 ? -1.0f : 1.0f;
+								ConvertData(0, pResultBuf, resultFormat, &sign, MBF_FLOAT32);
+
+								for (int i = 1; i < resultComponents; ++i)
+								{
+									// Add zeros
+									FMemory::Memzero
+									(
+										pResultBuf + GetMeshFormatData(resultFormat).m_size * i,
+										GetMeshFormatData(resultFormat).m_size
+									);
+								}
+								pResultBuf += resultElemSize;
+								xIt++;
+								yIt++;
+								zIt++;
+							}
+						}
+					}
 
-			int resultElemSize = Result.GetElementSize( b );
-			int resultChannelSize = GetMeshFormatData( resultFormat ).m_size * resultComponents;
-            uint8_t* pResultBuf = Result.GetBufferData( b );
-			pResultBuf += resultOffset;
+					// If we have to add colour channels, we will add them as white, to be neutral.
+					// \todo: normal channels also should have special values.
+					else if (resultSemantic == MBS_COLOUR)
+					{
+						generated = true;
 
-			if ( sourceBuffer<0 )
-			{
-				// Not found: fill with zeros.
-
-                // Special case for derived channel data
-                bool generated = false;
-                if ( resultSemantic == MBS_TANGENTSIGN )
-                {
-                    // Look for the full tangent space
-                    int tanXBuf, tanXChan, tanYBuf, tanYChan, tanZBuf, tanZChan;
-                    Source.FindChannel( MBS_TANGENT, resultSemanticIndex, &tanXBuf, &tanXChan );
-                    Source.FindChannel( MBS_BINORMAL, resultSemanticIndex, &tanYBuf, &tanYChan );
-                    Source.FindChannel( MBS_NORMAL, resultSemanticIndex, &tanZBuf, &tanZChan );
-
-                    if ( tanXBuf>=0 && tanYBuf>=0 && tanZBuf>=0 )
-                    {
-                        generated = true;
-                        UntypedMeshBufferIteratorConst xIt( Source, MBS_TANGENT, resultSemanticIndex );
-                        UntypedMeshBufferIteratorConst yIt( Source, MBS_BINORMAL, resultSemanticIndex );
-                        UntypedMeshBufferIteratorConst zIt( Source, MBS_NORMAL, resultSemanticIndex );
-                        for ( int v=0; v<vCount; ++v )
-                        {
-                            mat3f mat;
-                            mat[0] = xIt.GetAsVec4f().xyz();
-                            mat[1] = yIt.GetAsVec4f().xyz();
-                            mat[2] = zIt.GetAsVec4f().xyz();
-                            float sign = mat.GetDeterminant()<0?-1.0f:1.0f;
-                            ConvertData( 0, pResultBuf, resultFormat, &sign, MBF_FLOAT32 );
-
-                            for ( int i=1; i<resultComponents; ++i )
-                            {
-                                // Add zeros
-								FMemory::Memzero
-                                        (
-                                            pResultBuf+GetMeshFormatData( resultFormat ).m_size*i,
-                                            GetMeshFormatData( resultFormat ).m_size
-                                        );
-                            }
-                            pResultBuf += resultElemSize;
-                            xIt++;
-                            yIt++;
-                            zIt++;
-                        }
-                    }
-                }
-
-                // If we have to add colour channels, we will add them as white, to be neutral.
-                // \todo: normal channels also should have special values.
-                else if ( resultSemantic == MBS_COLOUR )
-                {
-                    generated = true;
-
-                    switch(resultFormat)
-                    {
-                    case MBF_FLOAT32:
-                    {
-						for (int v = 0; v < vCount; ++v)
+						switch (resultFormat)
 						{
-							auto pTypedResultBuf = (float*)pResultBuf;
-							for (int i = 0; i < resultComponents; ++i)
+						case MBF_FLOAT32:
+						{
+							for (int v = 0; v < vCount; ++v)
 							{
-								pTypedResultBuf[i] = 1.0f;
+								auto pTypedResultBuf = (float*)pResultBuf;
+								for (int i = 0; i < resultComponents; ++i)
+								{
+									pTypedResultBuf[i] = 1.0f;
+								}
+								pResultBuf += resultElemSize;
 							}
-							pResultBuf += resultElemSize;
+							break;
 						}
-                        break;
-                    }
 
-                    case MBF_NUINT8:
-                    {
-						for (int v = 0; v < vCount; ++v)
+						case MBF_NUINT8:
 						{
-							auto pTypedResultBuf = (uint8_t*)pResultBuf;
-							for (int i = 0; i < resultComponents; ++i)
+							for (int v = 0; v < vCount; ++v)
 							{
-								pTypedResultBuf[i] = 255;
+								auto pTypedResultBuf = (uint8_t*)pResultBuf;
+								for (int i = 0; i < resultComponents; ++i)
+								{
+									pTypedResultBuf[i] = 255;
+								}
+								pResultBuf += resultElemSize;
 							}
-							pResultBuf += resultElemSize;
+							break;
 						}
-                        break;
-                    }
 
-                    case MBF_NUINT16:
-                    {
-						for (int v = 0; v < vCount; ++v)
+						case MBF_NUINT16:
 						{
-							auto pTypedResultBuf = (uint16_t*)pResultBuf;
-							for (int i = 0; i < resultComponents; ++i)
+							for (int v = 0; v < vCount; ++v)
 							{
-								pTypedResultBuf[i] = 65535;
+								auto pTypedResultBuf = (uint16_t*)pResultBuf;
+								for (int i = 0; i < resultComponents; ++i)
+								{
+									pTypedResultBuf[i] = 65535;
+								}
+								pResultBuf += resultElemSize;
 							}
+							break;
+						}
+
+						default:
+							// Format not implemented
+							check(false);
+							break;
+						}
+					}
+
+					if (!generated)
+					{
+						// TODO: and maybe raise a warning?
+						for (int v = 0; v < vCount; ++v)
+						{
+							FMemory::Memzero(pResultBuf, resultChannelSize);
 							pResultBuf += resultElemSize;
 						}
-                        break;
-                    }
-
-                    default:
-                        // Format not implemented
-                        check( false );
-                        break;
-                    }
-                }
-
-                if (!generated)
-                {
-                    // TODO: and maybe raise a warning?
-                    for ( int v=0; v<vCount; ++v )
-                    {
-						FMemory::Memzero( pResultBuf, resultChannelSize );
-                        pResultBuf += resultElemSize;
-                    }
-                }
-			}
-			else
-			{
-				// Get the data about the source format
-				MESH_BUFFER_SEMANTIC sourceSemantic;
-				int sourceSemanticIndex;
-				MESH_BUFFER_FORMAT sourceFormat;
-				int sourceComponents;
-				int sourceOffset;
-				Source.GetChannel
-						(
-							sourceBuffer, sourceChannel,
-							&sourceSemantic, &sourceSemanticIndex,
-							&sourceFormat, &sourceComponents,
-							&sourceOffset
-						);
-				check( sourceSemantic==resultSemantic
-								&&
-								sourceSemanticIndex == resultSemanticIndex );
-
-				int sourceElemSize = Source.GetElementSize( sourceBuffer );
-                const uint8_t* pSourceBuf = Source.GetBufferData( sourceBuffer );
-				pSourceBuf += sourceOffset;
-
-				// Copy element by element
-				for ( int v=0; v<vCount; ++v )
-                {
-                    if ( resultFormat==sourceFormat && resultComponents==sourceComponents)
+					}
+				}
+				else
+				{
+					// Get the data about the source format
+					MESH_BUFFER_SEMANTIC sourceSemantic;
+					int sourceSemanticIndex;
+					MESH_BUFFER_FORMAT sourceFormat;
+					int sourceComponents;
+					int sourceOffset;
+					Source.GetChannel
+					(
+						sourceBuffer, sourceChannel,
+						&sourceSemantic, &sourceSemanticIndex,
+						&sourceFormat, &sourceComponents,
+						&sourceOffset
+					);
+					check(sourceSemantic == resultSemantic
+						&&
+						sourceSemanticIndex == resultSemanticIndex);
+
+					int sourceElemSize = Source.GetElementSize(sourceBuffer);
+					const uint8_t* pSourceBuf = Source.GetBufferData(sourceBuffer);
+					pSourceBuf += sourceOffset;
+
+					// Copy element by element
+					for (int v = 0; v < vCount; ++v)
 					{
-						memcpy( pResultBuf, pSourceBuf, resultChannelSize );
-                    }
-                    else if (resultFormat == MBF_PACKEDDIR8_W_TANGENTSIGN
-                             ||
-                             resultFormat == MBF_PACKEDDIRS8_W_TANGENTSIGN)
-                    {
-                        check(sourceComponents>=3);
-                        check(resultComponents==4);
-
-                        // convert the 3 first components
-                        //memcpy(pResultBuf, pSourceBuf, resultChannelSize);
-                        for ( int i=0; i<3; ++i )
-                        {
-                            if ( i<sourceComponents)
-                            {
-                                ConvertData
-                                        (
-                                            i,
-                                            pResultBuf, resultFormat,
-                                            pSourceBuf, sourceFormat
-                                        );
-                            }
-                        }
-
-
-                        // Add the tangent sign
-                        auto pData = (uint8_t*)pResultBuf;
-
-                        // Look for the full tangent space
-                        int tanXBuf, tanXChan, tanYBuf, tanYChan, tanZBuf, tanZChan;
-                        Source.FindChannel(MBS_TANGENT, resultSemanticIndex, &tanXBuf, &tanXChan);
-                        Source.FindChannel(MBS_BINORMAL, resultSemanticIndex, &tanYBuf, &tanYChan);
-                        Source.FindChannel(MBS_NORMAL, resultSemanticIndex, &tanZBuf, &tanZChan);
-
-                        if (tanXBuf >= 0 && tanYBuf >= 0 && tanZBuf >= 0)
-                        {
-                            UntypedMeshBufferIteratorConst xIt(Source, MBS_TANGENT, resultSemanticIndex);
-                            UntypedMeshBufferIteratorConst yIt(Source, MBS_BINORMAL, resultSemanticIndex);
-                            UntypedMeshBufferIteratorConst zIt(Source, MBS_NORMAL, resultSemanticIndex);
-
-                            mat3f mat;
-                            xIt += v;
-                            yIt += v;
-                            zIt += v;
-                            mat[0] = xIt.GetAsVec4f().xyz();
-                            mat[1] = yIt.GetAsVec4f().xyz();
-                            mat[2] = zIt.GetAsVec4f().xyz();
-
-                            uint8_t sign = 0;
-                            if (resultFormat == MBF_PACKEDDIR8_W_TANGENTSIGN)
-                            {
-                                sign = mat.GetDeterminant() < 0 ? 0 : 255;
-                            }
-                            else if (resultFormat == MBF_PACKEDDIRS8_W_TANGENTSIGN)
-                            {
-                                sign = mat.GetDeterminant() < 0 ? -128 : 127;
-                            }
-                            pData[3] = sign;
-                        }
-                    }
-                    else
-                    {
-						// Convert formats
-						for ( int i=0; i<resultComponents; ++i )
+						if (resultFormat == sourceFormat && resultComponents == sourceComponents)
+						{
+							memcpy(pResultBuf, pSourceBuf, resultChannelSize);
+						}
+						else if (resultFormat == MBF_PACKEDDIR8_W_TANGENTSIGN
+							||
+							resultFormat == MBF_PACKEDDIRS8_W_TANGENTSIGN)
 						{
-							if ( i<sourceComponents)
+							check(sourceComponents >= 3);
+							check(resultComponents == 4);
+
+							// convert the 3 first components
+							//memcpy(pResultBuf, pSourceBuf, resultChannelSize);
+							for (int i = 0; i < 3; ++i)
 							{
-								ConvertData
-										(
-											i,
-											pResultBuf, resultFormat,
-											pSourceBuf, sourceFormat
-										);
+								if (i < sourceComponents)
+								{
+									ConvertData
+									(
+										i,
+										pResultBuf, resultFormat,
+										pSourceBuf, sourceFormat
+									);
+								}
 							}
-							else
+
+
+							// Add the tangent sign
+							auto pData = (uint8_t*)pResultBuf;
+
+							// Look for the full tangent space
+							int tanXBuf, tanXChan, tanYBuf, tanYChan, tanZBuf, tanZChan;
+							Source.FindChannel(MBS_TANGENT, resultSemanticIndex, &tanXBuf, &tanXChan);
+							Source.FindChannel(MBS_BINORMAL, resultSemanticIndex, &tanYBuf, &tanYChan);
+							Source.FindChannel(MBS_NORMAL, resultSemanticIndex, &tanZBuf, &tanZChan);
+
+							if (tanXBuf >= 0 && tanYBuf >= 0 && tanZBuf >= 0)
 							{
-								// Add zeros. TODO: Warning?
-								FMemory::Memzero
-										(
-											pResultBuf+GetMeshFormatData( resultFormat ).m_size*i,
-											GetMeshFormatData( resultFormat ).m_size
-										);
+								UntypedMeshBufferIteratorConst xIt(Source, MBS_TANGENT, resultSemanticIndex);
+								UntypedMeshBufferIteratorConst yIt(Source, MBS_BINORMAL, resultSemanticIndex);
+								UntypedMeshBufferIteratorConst zIt(Source, MBS_NORMAL, resultSemanticIndex);
+
+								mat3f mat;
+								xIt += v;
+								yIt += v;
+								zIt += v;
+								mat[0] = xIt.GetAsVec4f().xyz();
+								mat[1] = yIt.GetAsVec4f().xyz();
+								mat[2] = zIt.GetAsVec4f().xyz();
+
+								uint8_t sign = 0;
+								if (resultFormat == MBF_PACKEDDIR8_W_TANGENTSIGN)
+								{
+									sign = mat.GetDeterminant() < 0 ? 0 : 255;
+								}
+								else if (resultFormat == MBF_PACKEDDIRS8_W_TANGENTSIGN)
+								{
+									sign = mat.GetDeterminant() < 0 ? -128 : 127;
+								}
+								pData[3] = sign;
 							}
 						}
+						else
+						{
+							// Convert formats
+							for (int i = 0; i < resultComponents; ++i)
+							{
+								if (i < sourceComponents)
+								{
+									ConvertData
+									(
+										i,
+										pResultBuf, resultFormat,
+										pSourceBuf, sourceFormat
+									);
+								}
+								else
+								{
+									// Add zeros. TODO: Warning?
+									FMemory::Memzero
+									(
+										pResultBuf + GetMeshFormatData(resultFormat).m_size * i,
+										GetMeshFormatData(resultFormat).m_size
+									);
+								}
+							}
 
 
-                        // Extra step to normalise some semantics in some formats
-                        // TODO: Make it optional, and add different normalisation types n, n^2
-                        // TODO: Optimise
-                        if (sourceSemantic==MBS_BONEWEIGHTS)
-                        {
-                            if (resultFormat==MBF_NUINT8)
-                            {
-                                auto pData = (uint8_t*)pResultBuf;
-                                uint8_t accum = 0;
-                                for ( int i=0; i<resultComponents; ++i )
-                                {
-                                    accum+= pData[i];
-                                }
-                                pData[0] += 255-accum;
-                            }
-                        }
-					}
+							// Extra step to normalise some semantics in some formats
+							// TODO: Make it optional, and add different normalisation types n, n^2
+							// TODO: Optimise
+							if (sourceSemantic == MBS_BONEWEIGHTS)
+							{
+								if (resultFormat == MBF_NUINT8)
+								{
+									auto pData = (uint8_t*)pResultBuf;
+									uint8_t accum = 0;
+									for (int i = 0; i < resultComponents; ++i)
+									{
+										accum += pData[i];
+									}
+									pData[0] += 255 - accum;
+								}
+							}
+						}
 
-					pResultBuf += resultElemSize;
-					pSourceBuf += sourceElemSize;
+						pResultBuf += resultElemSize;
+						pSourceBuf += sourceElemSize;
+					}
 				}
 			}
 		}
 	}
-}
 
 
-//-------------------------------------------------------------------------------------------------
-void FormatBufferSet
-    (
-        const FMeshBufferSet& Source,
-        FMeshBufferSet& Result,
-        bool keepSystemBuffers,
-        bool ignoreMissingChannels,
-        bool isVertexBuffer
-    )
-{
-    if ( ignoreMissingChannels )
-    {
-        // Remove from the result the channels that are not present in the source, and re-pack the
-        // offsets.
-        for ( int b=0; b<Result.GetBufferCount(); ++b )
-        {
-            vector<MESH_BUFFER_SEMANTIC> resultSemantics;
-            vector<int> resultSemanticIndexs;
-            vector<MESH_BUFFER_FORMAT> resultFormats;
-            vector<int> resultComponentss;
-            vector<int> resultOffsets;
-            int offset = 0;
-
-            // For every channel in this buffer
-            for ( int c=0; c<Result.GetBufferChannelCount(b); ++c )
-            {
-                MESH_BUFFER_SEMANTIC resultSemantic;
-                int resultSemanticIndex;
-                MESH_BUFFER_FORMAT resultFormat;
-                int resultComponents;
-
-                // Find this channel in the source mesh
-                Result.GetChannel
-                        (
-                            b, c,
-                            &resultSemantic, &resultSemanticIndex,
-                            &resultFormat, &resultComponents,
-                            nullptr
-                        );
-
-                int sourceBuffer;
-                int sourceChannel;
-                Source.FindChannel
-                        ( resultSemantic, resultSemanticIndex, &sourceBuffer, &sourceChannel );
-
-                if ( sourceBuffer>=0 )
-                {
-                    resultSemantics.push_back( resultSemantic );
-                    resultSemanticIndexs.push_back( resultSemanticIndex );
-                    resultFormats.push_back( resultFormat );
-                    resultComponentss.push_back( resultComponents );
-                    resultOffsets.push_back( offset );
-
-                    offset += GetMeshFormatData( resultFormat ).m_size * resultComponents;
-                }
-            }
-
-            if (resultSemantics.empty())
-            {
-                Result.SetBuffer( b, 0, 0, nullptr, nullptr, nullptr, nullptr, nullptr );
-            }
-            else
-            {
-                Result.SetBuffer( b, offset, (int)resultSemantics.size(),
-                                    &resultSemantics[0],
-                                    &resultSemanticIndexs[0],
-                                    &resultFormats[0],
-                                    &resultComponentss[0],
-                                    &resultOffsets[0] );
-            }
-        }
-    }
-
-
-    // For every vertex buffer in result
-    int vCount = Source.GetElementCount();
-    Result.SetElementCount( vCount );
-    for ( int b=0; b<Result.GetBufferCount(); ++b )
-    {
-        MeshFormatBuffer( Source, Result, b );
-    }
-
-
-    // Detect internal system buffers and clone them unmodified.
-    if ( keepSystemBuffers )
-    {
-        for ( int b=0; b<Source.GetBufferCount(); ++b )
-        {
-            // Detect system buffers and clone them unmodified.
-            if ( Source.GetBufferChannelCount(b)>0 )
-            {
-                MESH_BUFFER_SEMANTIC sourceSemantic;
-                int sourceSemanticIndex;
-                MESH_BUFFER_FORMAT sourceFormat;
-                int sourceComponents;
-                int sourceOffset;
-                Source.GetChannel
-                        (
-                            b, 0,
-                            &sourceSemantic, &sourceSemanticIndex,
-                            &sourceFormat, &sourceComponents,
-                            &sourceOffset
-                        );
-
-                if ( sourceSemantic==MBS_LAYOUTBLOCK
-                     || sourceSemantic==MBS_CHART
-                     || ( isVertexBuffer && sourceSemantic==MBS_VERTEXINDEX )
-                     )
-                {
-                    Result.AddBuffer( Source, b );
-                }
-            }
-        }
-    }
+	//-------------------------------------------------------------------------------------------------
+	void FormatBufferSet
+	(
+		const FMeshBufferSet& Source,
+		FMeshBufferSet& Result,
+		bool keepSystemBuffers,
+		bool ignoreMissingChannels,
+		bool isVertexBuffer
+	)
+	{
+		if (ignoreMissingChannels)
+		{
+			// Remove from the result the channels that are not present in the source, and re-pack the
+			// offsets.
+			for (int b = 0; b < Result.GetBufferCount(); ++b)
+			{
+				vector<MESH_BUFFER_SEMANTIC> resultSemantics;
+				vector<int> resultSemanticIndexs;
+				vector<MESH_BUFFER_FORMAT> resultFormats;
+				vector<int> resultComponentss;
+				vector<int> resultOffsets;
+				int offset = 0;
+
+				// For every channel in this buffer
+				for (int c = 0; c < Result.GetBufferChannelCount(b); ++c)
+				{
+					MESH_BUFFER_SEMANTIC resultSemantic;
+					int resultSemanticIndex;
+					MESH_BUFFER_FORMAT resultFormat;
+					int resultComponents;
+
+					// Find this channel in the source mesh
+					Result.GetChannel
+					(
+						b, c,
+						&resultSemantic, &resultSemanticIndex,
+						&resultFormat, &resultComponents,
+						nullptr
+					);
 
-}
+					int sourceBuffer;
+					int sourceChannel;
+					Source.FindChannel
+					(resultSemantic, resultSemanticIndex, &sourceBuffer, &sourceChannel);
+
+					if (sourceBuffer >= 0)
+					{
+						resultSemantics.push_back(resultSemantic);
+						resultSemanticIndexs.push_back(resultSemanticIndex);
+						resultFormats.push_back(resultFormat);
+						resultComponentss.push_back(resultComponents);
+						resultOffsets.push_back(offset);
+
+						offset += GetMeshFormatData(resultFormat).m_size * resultComponents;
+					}
+				}
+
+				if (resultSemantics.empty())
+				{
+					Result.SetBuffer(b, 0, 0, nullptr, nullptr, nullptr, nullptr, nullptr);
+				}
+				else
+				{
+					Result.SetBuffer(b, offset, (int)resultSemantics.size(),
+						&resultSemantics[0],
+						&resultSemanticIndexs[0],
+						&resultFormats[0],
+						&resultComponentss[0],
+						&resultOffsets[0]);
+				}
+			}
+		}
+
+
+		// For every vertex buffer in result
+		int vCount = Source.GetElementCount();
+		Result.SetElementCount(vCount);
+		for (int b = 0; b < Result.GetBufferCount(); ++b)
+		{
+			MeshFormatBuffer(Source, Result, b);
+		}
+
+
+		// Detect internal system buffers and clone them unmodified.
+		if (keepSystemBuffers)
+		{
+			for (int b = 0; b < Source.GetBufferCount(); ++b)
+			{
+				// Detect system buffers and clone them unmodified.
+				if (Source.GetBufferChannelCount(b) > 0)
+				{
+					MESH_BUFFER_SEMANTIC sourceSemantic;
+					int sourceSemanticIndex;
+					MESH_BUFFER_FORMAT sourceFormat;
+					int sourceComponents;
+					int sourceOffset;
+					Source.GetChannel
+					(
+						b, 0,
+						&sourceSemantic, &sourceSemanticIndex,
+						&sourceFormat, &sourceComponents,
+						&sourceOffset
+					);
+
+					if (sourceSemantic == MBS_LAYOUTBLOCK
+						|| sourceSemantic == MBS_CHART
+						|| (isVertexBuffer && sourceSemantic == MBS_VERTEXINDEX)
+						)
+					{
+						Result.AddBuffer(Source, b);
+					}
+				}
+			}
+		}
+
+	}
 
 
 
-//-------------------------------------------------------------------------------------------------
-MeshPtr mu::MeshFormat
+	//-------------------------------------------------------------------------------------------------
+	MeshPtr MeshFormat
 	(
 		const Mesh* pPureSource,
 		const Mesh* pFormat,
@@ -768,180 +769,181 @@ MeshPtr mu::MeshFormat
 		bool rebuildTangents,
 		bool ignoreMissingChannels
 	)
-{
-    MUTABLE_CPUPROFILER_SCOPE(MeshFormat);
-
-	if (!pPureSource) { return nullptr; }
-	if (!pFormat) { return pPureSource->Clone(); }
+	{
+		MUTABLE_CPUPROFILER_SCOPE(MeshFormat);
 
-    MeshPtrConst pSource = pPureSource;
+		if (!pPureSource) { return nullptr; }
+		if (!pFormat) { return pPureSource->Clone(); }
 
-    // Rebuild the tangent space if necessary
-    if ( rebuildTangents )
-    {
-		MeshPtr pClonedSource = pPureSource->Clone();
-		pSource = pClonedSource;
+		MeshPtrConst pSource = pPureSource;
 
-		// Make sure there are the tangents and binormal channels.
-        int tanBuf, tanChan, binBuf, binChan;
-        pSource->GetVertexBuffers().FindChannel( MBS_TANGENT, 0, &tanBuf, &tanChan );
-        pSource->GetVertexBuffers().FindChannel( MBS_BINORMAL, 0, &binBuf, &binChan );
+		// Rebuild the tangent space if necessary
+		if (rebuildTangents)
+		{
+			MeshPtr pClonedSource = pPureSource->Clone();
+			pSource = pClonedSource;
 
-		if (tanBuf<=0 ||binBuf<=0)
-        {
-			MESH_BUFFER buffer;
+			// Make sure there are the tangents and binormal channels.
+			int tanBuf, tanChan, binBuf, binChan;
+			pSource->GetVertexBuffers().FindChannel(MBS_TANGENT, 0, &tanBuf, &tanChan);
+			pSource->GetVertexBuffers().FindChannel(MBS_BINORMAL, 0, &binBuf, &binChan);
 
-			if (tanBuf<=0)
+			if (tanBuf <= 0 || binBuf <= 0)
 			{
-                MESH_BUFFER_CHANNEL channel;
-				channel.m_semantic = MBS_TANGENT;
-				channel.m_format = MBF_FLOAT32;
-				channel.m_componentCount = 3;
-				channel.m_semanticIndex = 0;
-				channel.m_offset = 0;
-				buffer.m_channels.Add( channel );
-				buffer.m_elementSize += 3*sizeof(float);
+				MESH_BUFFER buffer;
+
+				if (tanBuf <= 0)
+				{
+					MESH_BUFFER_CHANNEL channel;
+					channel.m_semantic = MBS_TANGENT;
+					channel.m_format = MBF_FLOAT32;
+					channel.m_componentCount = 3;
+					channel.m_semanticIndex = 0;
+					channel.m_offset = 0;
+					buffer.m_channels.Add(channel);
+					buffer.m_elementSize += 3 * sizeof(float);
+				}
+
+				if (binBuf <= 0)
+				{
+					MESH_BUFFER_CHANNEL channel;
+					channel.m_semantic = MBS_BINORMAL;
+					channel.m_format = MBF_FLOAT32;
+					channel.m_componentCount = 3;
+					channel.m_semanticIndex = 0;
+					channel.m_offset = (uint8)buffer.m_elementSize;
+					buffer.m_channels.Add(channel);
+					buffer.m_elementSize += 3 * sizeof(float);
+				}
+
+				buffer.m_data.SetNum(buffer.m_elementSize * pClonedSource->GetVertexCount());
+				pClonedSource->GetVertexBuffers().m_buffers.Add(buffer);
 			}
 
-			if (binBuf<=0)
+			RebuildTangents(pClonedSource->GetIndexBuffers(), pClonedSource->GetVertexBuffers());
+		}
+
+		MeshPtr pResult = pFormat->Clone();
+
+		// Make sure that the bone indices will fit in this format, or extend it.
+		if (formatVertices)
+		{
+			auto& buffs = pSource->GetVertexBuffers();
+
+			int semanticIndex = 0;
+			while (true)
 			{
-                MESH_BUFFER_CHANNEL channel;
-				channel.m_semantic = MBS_BINORMAL;
-				channel.m_format = MBF_FLOAT32;
-				channel.m_componentCount = 3;
-				channel.m_semanticIndex = 0;
-				channel.m_offset = (uint8)buffer.m_elementSize;
-				buffer.m_channels.Add( channel );
-				buffer.m_elementSize += 3*sizeof(float);
+				int buf = 0;
+				int chan = 0;
+				buffs.FindChannel(MBS_BONEINDICES, semanticIndex, &buf, &chan);
+				if (buf < 0)
+					break;
+
+				int resultBuf = 0;
+				int resultChan = 0;
+				auto& formBuffs = pResult->GetVertexBuffers();
+				formBuffs.FindChannel(MBS_BONEINDICES, semanticIndex, &resultBuf, &resultChan);
+				if (resultBuf >= 0)
+				{
+					UntypedMeshBufferIteratorConst it(buffs, MBS_BONEINDICES, semanticIndex);
+					int32_t maxBoneIndex = 0;
+					for (int v = 0; v < buffs.GetElementCount(); ++v)
+					{
+						auto va = it.GetAsVec8i();
+						for (int c = 0; c < it.GetComponents(); ++c)
+						{
+							maxBoneIndex = std::max(maxBoneIndex, va[c]);
+						}
+						++it;
+					}
+
+					auto& format = formBuffs.m_buffers[resultBuf].m_channels[resultChan].m_format;
+					if (maxBoneIndex > 0xffff && (format == MBF_UINT8 || format == MBF_UINT16))
+					{
+						format = MBF_UINT32;
+						formBuffs.UpdateOffsets(resultBuf);
+					}
+					else if (maxBoneIndex > 0x7fff && (format == MBF_INT8 || format == MBF_INT16))
+					{
+						format = MBF_UINT32;
+						formBuffs.UpdateOffsets(resultBuf);
+					}
+					else if (maxBoneIndex > 0xff && format == MBF_UINT8)
+					{
+						format = MBF_UINT16;
+						formBuffs.UpdateOffsets(resultBuf);
+					}
+					else if (maxBoneIndex > 0x7f && format == MBF_INT8)
+					{
+						format = MBF_INT16;
+						formBuffs.UpdateOffsets(resultBuf);
+					}
+				}
+
+				semanticIndex++;
 			}
+		}
 
-			buffer.m_data.SetNum( buffer.m_elementSize*pClonedSource->GetVertexCount() );
-			pClonedSource->GetVertexBuffers().m_buffers.Add( buffer );
-        }
-
-        RebuildTangents( pClonedSource->GetIndexBuffers(), pClonedSource->GetVertexBuffers() );
-    }
-
-    MeshPtr pResult = pFormat->Clone();
-
-    // Make sure that the bone indices will fit in this format, or extend it.
-    if ( formatVertices )
-    {
-        auto& buffs = pSource->GetVertexBuffers();
-
-        int semanticIndex = 0;
-        while( true )
-        {
-            int buf = 0;
-            int chan = 0;
-            buffs.FindChannel( MBS_BONEINDICES, semanticIndex, &buf, &chan );
-            if (buf<0)
-                break;
-
-            int resultBuf = 0;
-            int resultChan = 0;
-            auto& formBuffs = pResult->GetVertexBuffers();
-            formBuffs.FindChannel( MBS_BONEINDICES, semanticIndex, &resultBuf, &resultChan );
-            if ( resultBuf >=0 )
-            {
-                UntypedMeshBufferIteratorConst it( buffs, MBS_BONEINDICES, semanticIndex );
-                int32_t maxBoneIndex = 0;
-                for ( int v = 0; v < buffs.GetElementCount(); ++v )
-                {
-                    auto va = it.GetAsVec8i();
-                    for (int c=0; c<it.GetComponents(); ++c)
-                    {
-                        maxBoneIndex = std::max(maxBoneIndex,va[c]);
-                    }                    
-                    ++it;
-                }
-
-                auto& format = formBuffs.m_buffers[resultBuf].m_channels[resultChan].m_format;
-                if ( maxBoneIndex > 0xffff && (format == MBF_UINT8 || format == MBF_UINT16) )
-                {
-                    format = MBF_UINT32;
-                    formBuffs.UpdateOffsets( resultBuf );
-                }
-                else if ( maxBoneIndex > 0x7fff && (format == MBF_INT8 || format == MBF_INT16) )
-                {
-                    format = MBF_UINT32;
-                    formBuffs.UpdateOffsets( resultBuf );
-                }
-                else if ( maxBoneIndex > 0xff && format == MBF_UINT8 )
-                {
-                    format = MBF_UINT16;
-                    formBuffs.UpdateOffsets( resultBuf );
-                }
-                else if ( maxBoneIndex > 0x7f && format == MBF_INT8 )
-                {
-                    format = MBF_INT16;
-                    formBuffs.UpdateOffsets( resultBuf );
-                }
-            }
-
-            semanticIndex++;
-        }
-    }
-
-    // \todo Make sure that the vertex indices will fit in this format, or extend it.
-
-
-
-    if ( formatVertices )
-	{
+		// \todo Make sure that the vertex indices will fit in this format, or extend it.
 
-        FormatBufferSet( pSource->GetVertexBuffers(), pResult->GetVertexBuffers(),
-                         keepSystemBuffers, ignoreMissingChannels, true );
-    }
-	else
-	{
-		pResult->m_VertexBuffers = pSource->GetVertexBuffers();
-	}
 
-	if ( formatIndices )
-	{
-        FormatBufferSet( pSource->GetIndexBuffers(), pResult->GetIndexBuffers(), keepSystemBuffers,
-                         ignoreMissingChannels, false );
-    }
-	else
-	{
-		pResult->m_IndexBuffers = pSource->GetIndexBuffers();
-	}
 
-	if ( formatFaces )
-	{
-        FormatBufferSet( pSource->GetFaceBuffers(),
-                         pResult->GetFaceBuffers(),
-						 keepSystemBuffers,
-                         ignoreMissingChannels,
-                         false );
-	}
-	else
-	{
-		pResult->m_FaceBuffers = pSource->GetFaceBuffers();
-	}
+		if (formatVertices)
+		{
 
-	// Copy the rest of the data
-    pResult->SetSkeleton( pSource->GetSkeleton() );
-	pResult->SetPhysicsBody( pSource->GetPhysicsBody() );
+			FormatBufferSet(pSource->GetVertexBuffers(), pResult->GetVertexBuffers(),
+				keepSystemBuffers, ignoreMissingChannels, true);
+		}
+		else
+		{
+			pResult->m_VertexBuffers = pSource->GetVertexBuffers();
+		}
 
-	pResult->m_layouts.Empty();
-	for ( const auto& Layout: pSource->m_layouts )
-	{
-		pResult->m_layouts.Add(Layout->Clone());
-	}
+		if (formatIndices)
+		{
+			FormatBufferSet(pSource->GetIndexBuffers(), pResult->GetIndexBuffers(), keepSystemBuffers,
+				ignoreMissingChannels, false);
+		}
+		else
+		{
+			pResult->m_IndexBuffers = pSource->GetIndexBuffers();
+		}
 
-	pResult->ResetStaticFormatFlags();
-	pResult->EnsureSurfaceData();
+		if (formatFaces)
+		{
+			FormatBufferSet(pSource->GetFaceBuffers(),
+				pResult->GetFaceBuffers(),
+				keepSystemBuffers,
+				ignoreMissingChannels,
+				false);
+		}
+		else
+		{
+			pResult->m_FaceBuffers = pSource->GetFaceBuffers();
+		}
 
-	pResult->m_tags = pSource->m_tags;
+		// Copy the rest of the data
+		pResult->SetSkeleton(pSource->GetSkeleton());
+		pResult->SetPhysicsBody(pSource->GetPhysicsBody());
 
-	pResult->m_AdditionalBuffers = pSource->m_AdditionalBuffers;
+		pResult->m_layouts.Empty();
+		for (const auto& Layout : pSource->m_layouts)
+		{
+			pResult->m_layouts.Add(Layout->Clone());
+		}
 
-	pResult->m_bonePoses = pSource->m_bonePoses;
+		pResult->ResetStaticFormatFlags();
+		pResult->EnsureSurfaceData();
 
-	return pResult;
+		pResult->m_tags = pSource->m_tags;
 
-}
+		pResult->m_AdditionalBuffers = pSource->m_AdditionalBuffers;
+
+		pResult->m_bonePoses = pSource->m_bonePoses;
+
+		return pResult;
 
+	}
+
+}
 
diff --git a/Engine/Plugins/Experimental/Mutable/Source/MutableRuntime/Private/PhysicsBody.cpp b/Engine/Plugins/Experimental/Mutable/Source/MutableRuntime/Private/PhysicsBody.cpp
--- a/Engine/Plugins/Experimental/Mutable/Source/MutableRuntime/Private/PhysicsBody.cpp
+++ b/Engine/Plugins/Experimental/Mutable/Source/MutableRuntime/Private/PhysicsBody.cpp
@@ -2,462 +2,467 @@
 
 #include "PhysicsBody.h"
 
-using namespace mu;
 
-//-------------------------------------------------------------------------------------------------
-void PhysicsBody::Serialise( const PhysicsBody* p, OutputArchive& arch )
-{
-	arch << *p;
-}
-
-
-mu::Ptr<PhysicsBody> PhysicsBody::StaticUnserialise( InputArchive& arch )
-{
-	mu::Ptr<PhysicsBody> pResult = new PhysicsBody();
-	arch >> *pResult;
-	return pResult;
-}
-
-//-------------------------------------------------------------------------------------------------
-mu::Ptr<PhysicsBody> PhysicsBody::Clone() const
-{
-	mu::Ptr<PhysicsBody> pResult = new PhysicsBody();
-
-    pResult->Bodies = Bodies;
-	pResult->Bones = Bones; 
-	pResult->CustomIds = CustomIds; 
-	
-	return pResult;
-}
-
-//-------------------------------------------------------------------------------------------------
-void PhysicsBody::SetBodyCount( int32 Count )
-{
-	Bodies.SetNum(Count);
-	Bones.SetNum(Count);
-	CustomIds.Init(-1, Count);
-}
-
-//-------------------------------------------------------------------------------------------------
-int32 PhysicsBody::GetBodyCount( ) const
-{
-	return Bodies.Num();
-}
-
-//-------------------------------------------------------------------------------------------------
-const char* PhysicsBody::GetBodyBoneName( int32 B ) const
-{
-	check(B >= 0 && B < Bones.Num());
-
-	return Bones[B].c_str();
-}
-
-//-------------------------------------------------------------------------------------------------
-void PhysicsBody::SetBodyBoneName( int32 B, const char* BoneName )
-{
-	check(B >= 0 && B < Bones.Num());
-
-	Bones[B] = BoneName;
-}
-
-//-------------------------------------------------------------------------------------------------
-int32 PhysicsBody::GetBodyCustomId( int32 B ) const
-{
-	check(B >= 0 && B < CustomIds.Num());
-
-	return CustomIds[B];
-}
-
-//-------------------------------------------------------------------------------------------------
-void PhysicsBody::SetBodyCustomId( int32 B, int32 CustomId )
-{
-	check(B >= 0 && B < CustomIds.Num());
-
-	CustomIds[B] = CustomId;
-}
-
-//-------------------------------------------------------------------------------------------------
-int32 PhysicsBody::GetSphereCount( int32 B ) const
-{
-	return Bodies[B].Spheres.Num();
-}
-
-//-------------------------------------------------------------------------------------------------
-int32 PhysicsBody::GetBoxCount( int32 B ) const
-{
-	return Bodies[B].Boxes.Num();	
-}
-
-//-------------------------------------------------------------------------------------------------
-int32 PhysicsBody::GetConvexCount( int32 B ) const
-{
-	return Bodies[B].Convex.Num();
-}
-
-//-------------------------------------------------------------------------------------------------
-int32 PhysicsBody::GetSphylCount( int32 B ) const
-{
-	return Bodies[B].Sphyls.Num();
-}
-
-//-------------------------------------------------------------------------------------------------
-int32 PhysicsBody::GetTaperedCapsuleCount( int32 B ) const
-{
-	return Bodies[B].TaperedCapsules.Num();
-}
-
-//-------------------------------------------------------------------------------------------------
-void PhysicsBody::SetSphereCount( int32 B, int32 Count )
-{
-	Bodies[B].Spheres.SetNum(Count);
-}
-
-//-------------------------------------------------------------------------------------------------
-void PhysicsBody::SetBoxCount( int32 B, int32 Count )
-{
-	Bodies[B].Boxes.SetNum(Count);	
-}
-
-//-------------------------------------------------------------------------------------------------
-void PhysicsBody::SetConvexCount( int32 B, int32 Count)
-{
-    Bodies[B].Convex.SetNum(Count);	
-}
-
-//-------------------------------------------------------------------------------------------------
-void PhysicsBody::SetSphylCount( int32 B, int32 Count )
-{
-	Bodies[B].Sphyls.SetNum(Count);	
-}
-
-//-------------------------------------------------------------------------------------------------
-void PhysicsBody::SetTaperedCapsuleCount( int32 B, int32 Count )
-{
-	Bodies[B].TaperedCapsules.SetNum(Count);	
-}
-		
-//-------------------------------------------------------------------------------------------------
-void PhysicsBody::SetSphere( 
-		int32 B, int32 I, 
-		FVector3f Position, float Radius )
-{
-	check(B >= 0 && B < Bodies.Num());
-	check(I >= 0 && I < Bodies[B].Spheres.Num());
-	
-	Bodies[B].Spheres[I].Position = Position;
-	Bodies[B].Spheres[I].Radius = Radius;
-}
-
-//-------------------------------------------------------------------------------------------------
-void PhysicsBody::SetBox( 
-		int32 B, int32 I, 
-		FVector3f Position, FQuat4f Orientation, FVector3f Size )
-{
-	check(B >= 0 && B < Bodies.Num());
-	check(I >= 0 && I < Bodies[B].Boxes.Num());
-
-	Bodies[B].Boxes[I].Position = Position;
-	Bodies[B].Boxes[I].Orientation = Orientation;
-	Bodies[B].Boxes[I].Size = Size;
-}
-
-//-------------------------------------------------------------------------------------------------
-void PhysicsBody::SetConvex( 
-		int32 B, int32 I, 
-		const FVector3f* Vertices, int32 VerticesCount, 
-		const int32* Indices, int32 IndicesCount, 
-		const FTransform3f& Transform )
-{
-	check(B >= 0 && B < Bodies.Num());
-	check(I >= 0 && I < Bodies[B].Convex.Num());
-
-	Bodies[B].Convex[I].Vertices = TArray<FVector3f>( Vertices, VerticesCount );
-	Bodies[B].Convex[I].Indices = TArray<int32>( Indices, IndicesCount );
-	Bodies[B].Convex[I].Transform = Transform;
-	
-}
-
-//-------------------------------------------------------------------------------------------------
-void PhysicsBody::SetSphyl( 
-		int32 B, int32 I, 
-		FVector3f Position, FQuat4f Orientation, float Radius, float Length )
-{
-	check(B >= 0 && B < Bodies.Num());
-	check(I >= 0 && I < Bodies[B].Sphyls.Num());
-	
-	Bodies[B].Sphyls[I].Position = Position;
-	Bodies[B].Sphyls[I].Orientation = Orientation;
-	Bodies[B].Sphyls[I].Radius = Radius;
-	Bodies[B].Sphyls[I].Length = Length;
-}
-
-//-------------------------------------------------------------------------------------------------
-void PhysicsBody::SetSphereFlags( int32 B, int32 I, uint32 Flags )
-{
-	check(B >= 0 && B < Bodies.Num());
-	check(I >= 0 && I < Bodies[B].Spheres.Num());
-	
-	Bodies[B].Spheres[I].Flags = Flags;
-}
-
-//-------------------------------------------------------------------------------------------------
-void PhysicsBody::SetBoxFlags( int32 B, int32 I, uint32 Flags )
-{
-	check(B >= 0 && B < Bodies.Num());
-	check(I >= 0 && I < Bodies[B].Boxes.Num());
-
-	Bodies[B].Boxes[I].Flags = Flags;
-}
-
-//-------------------------------------------------------------------------------------------------
-void PhysicsBody::SetConvexFlags( int32 B, int32 I, uint32 Flags )
-{
-	check(B >= 0 && B < Bodies.Num());
-	check(I >= 0 && I < Bodies[B].Convex.Num());
-
-	Bodies[B].Convex[I].Flags = Flags;	
-}
-
-//-------------------------------------------------------------------------------------------------
-void PhysicsBody::SetSphylFlags( int32 B, int32 I, uint32 Flags )
-{
-	check(B >= 0 && B < Bodies.Num());
-	check(I >= 0 && I < Bodies[B].Sphyls.Num());
-	
-	Bodies[B].Sphyls[I].Flags = Flags;
-}
-
-//-------------------------------------------------------------------------------------------------
-void PhysicsBody::SetTaperedCapsuleFlags( int32 B, int32 I, uint32 Flags )
-{
-	check(B >= 0 && B < Bodies.Num());
-	check(I >= 0 && I < Bodies[B].TaperedCapsules.Num());
-	
-	Bodies[B].TaperedCapsules[I].Flags = Flags;
-}
-
-//-------------------------------------------------------------------------------------------------
-void PhysicsBody::SetSphereName( int32 B, int32 I, const char* Name )
-{
-	check(B >= 0 && B < Bodies.Num());
-	check(I >= 0 && I < Bodies[B].Spheres.Num());
-	
-	Bodies[B].Spheres[I].Name = Name;
-}
-
-//-------------------------------------------------------------------------------------------------
-void PhysicsBody::SetBoxName( int32 B, int32 I, const char* Name )
-{
-	check(B >= 0 && B < Bodies.Num());
-	check(I >= 0 && I < Bodies[B].Boxes.Num());
-
-	Bodies[B].Boxes[I].Name = Name;
-}
-
-//-------------------------------------------------------------------------------------------------
-void PhysicsBody::SetConvexName( int32 B, int32 I, const char* Name )
-{
-	check(B >= 0 && B < Bodies.Num());
-	check(I >= 0 && I < Bodies[B].Convex.Num());
-
-	Bodies[B].Convex[I].Name = Name;	
-}
-
-//-------------------------------------------------------------------------------------------------
-void PhysicsBody::SetSphylName( int32 B, int32 I, const char* Name )
-{
-	check(B >= 0 && B < Bodies.Num());
-	check(I >= 0 && I < Bodies[B].Sphyls.Num());
-	
-	Bodies[B].Sphyls[I].Name = Name;
-}
-
-//-------------------------------------------------------------------------------------------------
-void PhysicsBody::SetTaperedCapsuleName( int32 B, int32 I, const char* Name )
-{
-	check(B >= 0 && B < Bodies.Num());
-	check(I >= 0 && I < Bodies[B].TaperedCapsules.Num());
-	
-	Bodies[B].TaperedCapsules[I].Name = Name;
-}
-
-//-------------------------------------------------------------------------------------------------
-void PhysicsBody::SetTaperedCapsule( 
-		int32 B, int32 I, 
-		FVector3f Position, FQuat4f Orientation, 
-		float Radius0, float Radius1, float Length )
-{
-	check(B >= 0 && B < Bodies.Num());
-	check(I >= 0 && I < Bodies[B].TaperedCapsules.Num());
-	
-	Bodies[B].TaperedCapsules[I].Position = Position;
-	Bodies[B].TaperedCapsules[I].Orientation = Orientation;
-	Bodies[B].TaperedCapsules[I].Radius0 = Radius0;
-	Bodies[B].TaperedCapsules[I].Radius1 = Radius1;
-	Bodies[B].TaperedCapsules[I].Length = Length;
-}
-
-//-------------------------------------------------------------------------------------------------
-void PhysicsBody::GetSphere( 
-		int32 B, int32 I, 
+namespace mu
+{
+
+
+	//-------------------------------------------------------------------------------------------------
+	void PhysicsBody::Serialise(const PhysicsBody* p, OutputArchive& arch)
+	{
+		arch << *p;
+	}
+
+
+	mu::Ptr<PhysicsBody> PhysicsBody::StaticUnserialise(InputArchive& arch)
+	{
+		mu::Ptr<PhysicsBody> pResult = new PhysicsBody();
+		arch >> *pResult;
+		return pResult;
+	}
+
+	//-------------------------------------------------------------------------------------------------
+	mu::Ptr<PhysicsBody> PhysicsBody::Clone() const
+	{
+		mu::Ptr<PhysicsBody> pResult = new PhysicsBody();
+
+		pResult->Bodies = Bodies;
+		pResult->Bones = Bones;
+		pResult->CustomIds = CustomIds;
+
+		return pResult;
+	}
+
+	//-------------------------------------------------------------------------------------------------
+	void PhysicsBody::SetBodyCount(int32 Count)
+	{
+		Bodies.SetNum(Count);
+		Bones.SetNum(Count);
+		CustomIds.Init(-1, Count);
+	}
+
+	//-------------------------------------------------------------------------------------------------
+	int32 PhysicsBody::GetBodyCount() const
+	{
+		return Bodies.Num();
+	}
+
+	//-------------------------------------------------------------------------------------------------
+	const char* PhysicsBody::GetBodyBoneName(int32 B) const
+	{
+		check(B >= 0 && B < Bones.Num());
+
+		return Bones[B].c_str();
+	}
+
+	//-------------------------------------------------------------------------------------------------
+	void PhysicsBody::SetBodyBoneName(int32 B, const char* BoneName)
+	{
+		check(B >= 0 && B < Bones.Num());
+
+		Bones[B] = BoneName;
+	}
+
+	//-------------------------------------------------------------------------------------------------
+	int32 PhysicsBody::GetBodyCustomId(int32 B) const
+	{
+		check(B >= 0 && B < CustomIds.Num());
+
+		return CustomIds[B];
+	}
+
+	//-------------------------------------------------------------------------------------------------
+	void PhysicsBody::SetBodyCustomId(int32 B, int32 CustomId)
+	{
+		check(B >= 0 && B < CustomIds.Num());
+
+		CustomIds[B] = CustomId;
+	}
+
+	//-------------------------------------------------------------------------------------------------
+	int32 PhysicsBody::GetSphereCount(int32 B) const
+	{
+		return Bodies[B].Spheres.Num();
+	}
+
+	//-------------------------------------------------------------------------------------------------
+	int32 PhysicsBody::GetBoxCount(int32 B) const
+	{
+		return Bodies[B].Boxes.Num();
+	}
+
+	//-------------------------------------------------------------------------------------------------
+	int32 PhysicsBody::GetConvexCount(int32 B) const
+	{
+		return Bodies[B].Convex.Num();
+	}
+
+	//-------------------------------------------------------------------------------------------------
+	int32 PhysicsBody::GetSphylCount(int32 B) const
+	{
+		return Bodies[B].Sphyls.Num();
+	}
+
+	//-------------------------------------------------------------------------------------------------
+	int32 PhysicsBody::GetTaperedCapsuleCount(int32 B) const
+	{
+		return Bodies[B].TaperedCapsules.Num();
+	}
+
+	//-------------------------------------------------------------------------------------------------
+	void PhysicsBody::SetSphereCount(int32 B, int32 Count)
+	{
+		Bodies[B].Spheres.SetNum(Count);
+	}
+
+	//-------------------------------------------------------------------------------------------------
+	void PhysicsBody::SetBoxCount(int32 B, int32 Count)
+	{
+		Bodies[B].Boxes.SetNum(Count);
+	}
+
+	//-------------------------------------------------------------------------------------------------
+	void PhysicsBody::SetConvexCount(int32 B, int32 Count)
+	{
+		Bodies[B].Convex.SetNum(Count);
+	}
+
+	//-------------------------------------------------------------------------------------------------
+	void PhysicsBody::SetSphylCount(int32 B, int32 Count)
+	{
+		Bodies[B].Sphyls.SetNum(Count);
+	}
+
+	//-------------------------------------------------------------------------------------------------
+	void PhysicsBody::SetTaperedCapsuleCount(int32 B, int32 Count)
+	{
+		Bodies[B].TaperedCapsules.SetNum(Count);
+	}
+
+	//-------------------------------------------------------------------------------------------------
+	void PhysicsBody::SetSphere(
+		int32 B, int32 I,
+		FVector3f Position, float Radius)
+	{
+		check(B >= 0 && B < Bodies.Num());
+		check(I >= 0 && I < Bodies[B].Spheres.Num());
+
+		Bodies[B].Spheres[I].Position = Position;
+		Bodies[B].Spheres[I].Radius = Radius;
+	}
+
+	//-------------------------------------------------------------------------------------------------
+	void PhysicsBody::SetBox(
+		int32 B, int32 I,
+		FVector3f Position, FQuat4f Orientation, FVector3f Size)
+	{
+		check(B >= 0 && B < Bodies.Num());
+		check(I >= 0 && I < Bodies[B].Boxes.Num());
+
+		Bodies[B].Boxes[I].Position = Position;
+		Bodies[B].Boxes[I].Orientation = Orientation;
+		Bodies[B].Boxes[I].Size = Size;
+	}
+
+	//-------------------------------------------------------------------------------------------------
+	void PhysicsBody::SetConvex(
+		int32 B, int32 I,
+		const FVector3f* Vertices, int32 VerticesCount,
+		const int32* Indices, int32 IndicesCount,
+		const FTransform3f& Transform)
+	{
+		check(B >= 0 && B < Bodies.Num());
+		check(I >= 0 && I < Bodies[B].Convex.Num());
+
+		Bodies[B].Convex[I].Vertices = TArray<FVector3f>(Vertices, VerticesCount);
+		Bodies[B].Convex[I].Indices = TArray<int32>(Indices, IndicesCount);
+		Bodies[B].Convex[I].Transform = Transform;
+
+	}
+
+	//-------------------------------------------------------------------------------------------------
+	void PhysicsBody::SetSphyl(
+		int32 B, int32 I,
+		FVector3f Position, FQuat4f Orientation, float Radius, float Length)
+	{
+		check(B >= 0 && B < Bodies.Num());
+		check(I >= 0 && I < Bodies[B].Sphyls.Num());
+
+		Bodies[B].Sphyls[I].Position = Position;
+		Bodies[B].Sphyls[I].Orientation = Orientation;
+		Bodies[B].Sphyls[I].Radius = Radius;
+		Bodies[B].Sphyls[I].Length = Length;
+	}
+
+	//-------------------------------------------------------------------------------------------------
+	void PhysicsBody::SetSphereFlags(int32 B, int32 I, uint32 Flags)
+	{
+		check(B >= 0 && B < Bodies.Num());
+		check(I >= 0 && I < Bodies[B].Spheres.Num());
+
+		Bodies[B].Spheres[I].Flags = Flags;
+	}
+
+	//-------------------------------------------------------------------------------------------------
+	void PhysicsBody::SetBoxFlags(int32 B, int32 I, uint32 Flags)
+	{
+		check(B >= 0 && B < Bodies.Num());
+		check(I >= 0 && I < Bodies[B].Boxes.Num());
+
+		Bodies[B].Boxes[I].Flags = Flags;
+	}
+
+	//-------------------------------------------------------------------------------------------------
+	void PhysicsBody::SetConvexFlags(int32 B, int32 I, uint32 Flags)
+	{
+		check(B >= 0 && B < Bodies.Num());
+		check(I >= 0 && I < Bodies[B].Convex.Num());
+
+		Bodies[B].Convex[I].Flags = Flags;
+	}
+
+	//-------------------------------------------------------------------------------------------------
+	void PhysicsBody::SetSphylFlags(int32 B, int32 I, uint32 Flags)
+	{
+		check(B >= 0 && B < Bodies.Num());
+		check(I >= 0 && I < Bodies[B].Sphyls.Num());
+
+		Bodies[B].Sphyls[I].Flags = Flags;
+	}
+
+	//-------------------------------------------------------------------------------------------------
+	void PhysicsBody::SetTaperedCapsuleFlags(int32 B, int32 I, uint32 Flags)
+	{
+		check(B >= 0 && B < Bodies.Num());
+		check(I >= 0 && I < Bodies[B].TaperedCapsules.Num());
+
+		Bodies[B].TaperedCapsules[I].Flags = Flags;
+	}
+
+	//-------------------------------------------------------------------------------------------------
+	void PhysicsBody::SetSphereName(int32 B, int32 I, const char* Name)
+	{
+		check(B >= 0 && B < Bodies.Num());
+		check(I >= 0 && I < Bodies[B].Spheres.Num());
+
+		Bodies[B].Spheres[I].Name = Name;
+	}
+
+	//-------------------------------------------------------------------------------------------------
+	void PhysicsBody::SetBoxName(int32 B, int32 I, const char* Name)
+	{
+		check(B >= 0 && B < Bodies.Num());
+		check(I >= 0 && I < Bodies[B].Boxes.Num());
+
+		Bodies[B].Boxes[I].Name = Name;
+	}
+
+	//-------------------------------------------------------------------------------------------------
+	void PhysicsBody::SetConvexName(int32 B, int32 I, const char* Name)
+	{
+		check(B >= 0 && B < Bodies.Num());
+		check(I >= 0 && I < Bodies[B].Convex.Num());
+
+		Bodies[B].Convex[I].Name = Name;
+	}
+
+	//-------------------------------------------------------------------------------------------------
+	void PhysicsBody::SetSphylName(int32 B, int32 I, const char* Name)
+	{
+		check(B >= 0 && B < Bodies.Num());
+		check(I >= 0 && I < Bodies[B].Sphyls.Num());
+
+		Bodies[B].Sphyls[I].Name = Name;
+	}
+
+	//-------------------------------------------------------------------------------------------------
+	void PhysicsBody::SetTaperedCapsuleName(int32 B, int32 I, const char* Name)
+	{
+		check(B >= 0 && B < Bodies.Num());
+		check(I >= 0 && I < Bodies[B].TaperedCapsules.Num());
+
+		Bodies[B].TaperedCapsules[I].Name = Name;
+	}
+
+	//-------------------------------------------------------------------------------------------------
+	void PhysicsBody::SetTaperedCapsule(
+		int32 B, int32 I,
+		FVector3f Position, FQuat4f Orientation,
+		float Radius0, float Radius1, float Length)
+	{
+		check(B >= 0 && B < Bodies.Num());
+		check(I >= 0 && I < Bodies[B].TaperedCapsules.Num());
+
+		Bodies[B].TaperedCapsules[I].Position = Position;
+		Bodies[B].TaperedCapsules[I].Orientation = Orientation;
+		Bodies[B].TaperedCapsules[I].Radius0 = Radius0;
+		Bodies[B].TaperedCapsules[I].Radius1 = Radius1;
+		Bodies[B].TaperedCapsules[I].Length = Length;
+	}
+
+	//-------------------------------------------------------------------------------------------------
+	void PhysicsBody::GetSphere(
+		int32 B, int32 I,
 		FVector3f& OutPosition, float& OutRadius) const
-{
-	check(B >= 0 && B < Bodies.Num());
-	check(I >= 0 && I < Bodies[B].Spheres.Num());
+	{
+		check(B >= 0 && B < Bodies.Num());
+		check(I >= 0 && I < Bodies[B].Spheres.Num());
 
-	OutPosition = Bodies[B].Spheres[I].Position;
-	OutRadius = Bodies[B].Spheres[I].Radius;
-}
+		OutPosition = Bodies[B].Spheres[I].Position;
+		OutRadius = Bodies[B].Spheres[I].Radius;
+	}
 
-//-------------------------------------------------------------------------------------------------
-void PhysicsBody::GetBox( 
-		int32 B, int32 I, 
+	//-------------------------------------------------------------------------------------------------
+	void PhysicsBody::GetBox(
+		int32 B, int32 I,
 		FVector3f& OutPosition, FQuat4f& OutOrientation, FVector3f& OutSize) const
-{
-	check(B >= 0 && B < Bodies.Num());
-	check(I >= 0 && I < Bodies[B].Boxes.Num());
-
-	OutPosition = Bodies[B].Boxes[I].Position;
-	OutOrientation = Bodies[B].Boxes[I].Orientation;
-	OutSize = Bodies[B].Boxes[I].Size;
-}
-
-//-------------------------------------------------------------------------------------------------
-void PhysicsBody::GetConvex( 
-		int32 B, int32 I, 
-		FVector3f const*& OutVertices, int32& OutVerticesCount, 
-		int32 const*& OutIndices, int32& OutIndicesCount, 
+	{
+		check(B >= 0 && B < Bodies.Num());
+		check(I >= 0 && I < Bodies[B].Boxes.Num());
+
+		OutPosition = Bodies[B].Boxes[I].Position;
+		OutOrientation = Bodies[B].Boxes[I].Orientation;
+		OutSize = Bodies[B].Boxes[I].Size;
+	}
+
+	//-------------------------------------------------------------------------------------------------
+	void PhysicsBody::GetConvex(
+		int32 B, int32 I,
+		FVector3f const*& OutVertices, int32& OutVerticesCount,
+		int32 const*& OutIndices, int32& OutIndicesCount,
 		FTransform3f& OutTransform) const
-{
-	check(B >= 0 && B < Bodies.Num());
-	check(I >= 0 && I < Bodies[B].Convex.Num());
-	
-	OutVertices = Bodies[B].Convex[I].Vertices.GetData();
-	OutVerticesCount = Bodies[B].Convex[I].Vertices.Num();
-
-	OutIndices = Bodies[B].Convex[I].Indices.GetData();
-	OutIndicesCount = Bodies[B].Convex[I].Indices.Num();
-	
-	OutTransform = Bodies[B].Convex[I].Transform;
-}
+	{
+		check(B >= 0 && B < Bodies.Num());
+		check(I >= 0 && I < Bodies[B].Convex.Num());
+
+		OutVertices = Bodies[B].Convex[I].Vertices.GetData();
+		OutVerticesCount = Bodies[B].Convex[I].Vertices.Num();
+
+		OutIndices = Bodies[B].Convex[I].Indices.GetData();
+		OutIndicesCount = Bodies[B].Convex[I].Indices.Num();
+
+		OutTransform = Bodies[B].Convex[I].Transform;
+	}
+
+	//-------------------------------------------------------------------------------------------------
+	void PhysicsBody::GetSphyl(
+		int32 B, int32 I,
+		FVector3f& OutPosition, FQuat4f& OutOrientation, float& OutRadius, float& OutLength) const
+	{
+		check(B >= 0 && B < Bodies.Num());
+		check(I >= 0 && I < Bodies[B].Sphyls.Num());
+
+		OutPosition = Bodies[B].Sphyls[I].Position;
+		OutOrientation = Bodies[B].Sphyls[I].Orientation;
+		OutRadius = Bodies[B].Sphyls[I].Radius;
+		OutLength = Bodies[B].Sphyls[I].Length;
+	}
+
+	//-------------------------------------------------------------------------------------------------
+	void PhysicsBody::GetTaperedCapsule(
+		int32 B, int32 I,
+		FVector3f& OutPosition, FQuat4f& OutOrientation,
+		float& OutRadius0, float& OutRadius1, float& OutLength) const
+	{
+		check(B >= 0 && B < Bodies.Num());
+		check(I >= 0 && I < Bodies[B].TaperedCapsules.Num());
+
+		OutPosition = Bodies[B].TaperedCapsules[I].Position;
+		OutOrientation = Bodies[B].TaperedCapsules[I].Orientation;
+		OutRadius0 = Bodies[B].TaperedCapsules[I].Radius0;
+		OutRadius1 = Bodies[B].TaperedCapsules[I].Radius1;
+		OutLength = Bodies[B].TaperedCapsules[I].Length;
+	}
+
+
+	//-------------------------------------------------------------------------------------------------
+	uint32 PhysicsBody::GetSphereFlags(int32 B, int32 I) const
+	{
+		check(B >= 0 && B < Bodies.Num());
+		check(I >= 0 && I < Bodies[B].Spheres.Num());
+
+		return Bodies[B].Spheres[I].Flags;
+	}
+
+	//-------------------------------------------------------------------------------------------------
+	uint32 PhysicsBody::GetBoxFlags(int32 B, int32 I) const
+	{
+		check(B >= 0 && B < Bodies.Num());
+		check(I >= 0 && I < Bodies[B].Boxes.Num());
+
+		return Bodies[B].Boxes[I].Flags;
+	}
+
+	//-------------------------------------------------------------------------------------------------
+	uint32 PhysicsBody::GetConvexFlags(int32 B, int32 I) const
+	{
+		check(B >= 0 && B < Bodies.Num());
+		check(I >= 0 && I < Bodies[B].Convex.Num());
+
+		return Bodies[B].Convex[I].Flags;
+	}
+
+	//-------------------------------------------------------------------------------------------------
+	uint32 PhysicsBody::GetSphylFlags(int32 B, int32 I) const
+	{
+		check(B >= 0 && B < Bodies.Num());
+		check(I >= 0 && I < Bodies[B].Sphyls.Num());
+
+		return Bodies[B].Sphyls[I].Flags;
+	}
+
+	//-------------------------------------------------------------------------------------------------
+	uint32 PhysicsBody::GetTaperedCapsuleFlags(int32 B, int32 I) const
+	{
+		check(B >= 0 && B < Bodies.Num());
+		check(I >= 0 && I < Bodies[B].TaperedCapsules.Num());
+
+		return Bodies[B].TaperedCapsules[I].Flags;
+	}
+
+	//-------------------------------------------------------------------------------------------------
+	const char* PhysicsBody::GetSphereName(int32 B, int32 I) const
+	{
+		check(B >= 0 && B < Bodies.Num());
+		check(I >= 0 && I < Bodies[B].Spheres.Num());
+
+		return Bodies[B].Spheres[I].Name.c_str();
+	}
+
+	//-------------------------------------------------------------------------------------------------
+	const char* PhysicsBody::GetBoxName(int32 B, int32 I) const
+	{
+		check(B >= 0 && B < Bodies.Num());
+		check(I >= 0 && I < Bodies[B].Boxes.Num());
+
+		return Bodies[B].Boxes[I].Name.c_str();
+	}
+
+	//-------------------------------------------------------------------------------------------------
+	const char* PhysicsBody::GetConvexName(int32 B, int32 I) const
+	{
+		check(B >= 0 && B < Bodies.Num());
+		check(I >= 0 && I < Bodies[B].Convex.Num());
+
+		return Bodies[B].Convex[I].Name.c_str();
+	}
+
+	//-------------------------------------------------------------------------------------------------
+	const char* PhysicsBody::GetSphylName(int32 B, int32 I) const
+	{
+		check(B >= 0 && B < Bodies.Num());
+		check(I >= 0 && I < Bodies[B].Sphyls.Num());
+
+		return Bodies[B].Sphyls[I].Name.c_str();
+	}
+
+	//-------------------------------------------------------------------------------------------------
+	const char* PhysicsBody::GetTaperedCapsuleName(int32 B, int32 I) const
+	{
+		check(B >= 0 && B < Bodies.Num());
+		check(I >= 0 && I < Bodies[B].TaperedCapsules.Num());
+
+		return Bodies[B].TaperedCapsules[I].Name.c_str();
+	}
 
-//-------------------------------------------------------------------------------------------------
-void PhysicsBody::GetSphyl( 
-		int32 B, int32 I, 
-		FVector3f& OutPosition, FQuat4f& OutOrientation, float& OutRadius, float& OutLength ) const
-{
-	check(B >= 0 && B < Bodies.Num());
-	check(I >= 0 && I < Bodies[B].Sphyls.Num());
-
-	OutPosition = Bodies[B].Sphyls[I].Position;
-	OutOrientation = Bodies[B].Sphyls[I].Orientation;
-	OutRadius = Bodies[B].Sphyls[I].Radius;
-	OutLength = Bodies[B].Sphyls[I].Length;
-}
-
-//-------------------------------------------------------------------------------------------------
-void PhysicsBody::GetTaperedCapsule( 
-		int32 B, int32 I, 
-		FVector3f& OutPosition, FQuat4f& OutOrientation, 
-		float& OutRadius0, float& OutRadius1, float& OutLength ) const
-{
-	check(B >= 0 && B < Bodies.Num());
-	check(I >= 0 && I < Bodies[B].TaperedCapsules.Num());
-
-	OutPosition = Bodies[B].TaperedCapsules[I].Position;
-	OutOrientation = Bodies[B].TaperedCapsules[I].Orientation;
-	OutRadius0 = Bodies[B].TaperedCapsules[I].Radius0;
-	OutRadius1 = Bodies[B].TaperedCapsules[I].Radius1;
-	OutLength = Bodies[B].TaperedCapsules[I].Length;
-}
-
-
-//-------------------------------------------------------------------------------------------------
-uint32 PhysicsBody::GetSphereFlags( int32 B, int32 I ) const
-{
-	check(B >= 0 && B < Bodies.Num());
-	check(I >= 0 && I < Bodies[B].Spheres.Num());
-	
-	return Bodies[B].Spheres[I].Flags;
-}
-
-//-------------------------------------------------------------------------------------------------
-uint32 PhysicsBody::GetBoxFlags( int32 B, int32 I ) const
-{
-	check(B >= 0 && B < Bodies.Num());
-	check(I >= 0 && I < Bodies[B].Boxes.Num());
-
-	return Bodies[B].Boxes[I].Flags;
-}
-
-//-------------------------------------------------------------------------------------------------
-uint32 PhysicsBody::GetConvexFlags( int32 B, int32 I ) const
-{
-	check(B >= 0 && B < Bodies.Num());
-	check(I >= 0 && I < Bodies[B].Convex.Num());
-
-	return Bodies[B].Convex[I].Flags;	
-}
-
-//-------------------------------------------------------------------------------------------------
-uint32 PhysicsBody::GetSphylFlags( int32 B, int32 I ) const
-{
-	check(B >= 0 && B < Bodies.Num());
-	check(I >= 0 && I < Bodies[B].Sphyls.Num());
-	
-	return Bodies[B].Sphyls[I].Flags;
-}
-
-//-------------------------------------------------------------------------------------------------
-uint32 PhysicsBody::GetTaperedCapsuleFlags( int32 B, int32 I ) const
-{
-	check(B >= 0 && B < Bodies.Num());
-	check(I >= 0 && I < Bodies[B].TaperedCapsules.Num());
-	
-	return Bodies[B].TaperedCapsules[I].Flags;
-}
-
-//-------------------------------------------------------------------------------------------------
-const char* PhysicsBody::GetSphereName( int32 B, int32 I ) const
-{
-	check(B >= 0 && B < Bodies.Num());
-	check(I >= 0 && I < Bodies[B].Spheres.Num());
-	
-	return Bodies[B].Spheres[I].Name.c_str();
-}
-
-//-------------------------------------------------------------------------------------------------
-const char* PhysicsBody::GetBoxName( int32 B, int32 I ) const
-{
-	check(B >= 0 && B < Bodies.Num());
-	check(I >= 0 && I < Bodies[B].Boxes.Num());
-
-	return Bodies[B].Boxes[I].Name.c_str();
-}
-
-//-------------------------------------------------------------------------------------------------
-const char* PhysicsBody::GetConvexName( int32 B, int32 I ) const
-{
-	check(B >= 0 && B < Bodies.Num());
-	check(I >= 0 && I < Bodies[B].Convex.Num());
-
-	return Bodies[B].Convex[I].Name.c_str();
-}
-
-//-------------------------------------------------------------------------------------------------
-const char* PhysicsBody::GetSphylName( int32 B, int32 I ) const
-{
-	check(B >= 0 && B < Bodies.Num());
-	check(I >= 0 && I < Bodies[B].Sphyls.Num());
-	
-	return Bodies[B].Sphyls[I].Name.c_str();
-}
-
-//-------------------------------------------------------------------------------------------------
-const char* PhysicsBody::GetTaperedCapsuleName( int32 B, int32 I ) const
-{
-	check(B >= 0 && B < Bodies.Num());
-	check(I >= 0 && I < Bodies[B].TaperedCapsules.Num());
-	
-	return Bodies[B].TaperedCapsules[I].Name.c_str();
 }
diff --git a/Engine/Plugins/Experimental/Mutable/Source/MutableRuntime/Private/Skeleton.cpp b/Engine/Plugins/Experimental/Mutable/Source/MutableRuntime/Private/Skeleton.cpp
--- a/Engine/Plugins/Experimental/Mutable/Source/MutableRuntime/Private/Skeleton.cpp
+++ b/Engine/Plugins/Experimental/Mutable/Source/MutableRuntime/Private/Skeleton.cpp
@@ -7,127 +7,131 @@
 #include "MutableTrace.h"
 
 
-using namespace mu;
-
-//---------------------------------------------------------------------------------------------
-void Skeleton::Serialise( const Skeleton* p, OutputArchive& arch )
-{
-    arch << *p;
-}
-
-
-//---------------------------------------------------------------------------------------------
-mu::Ptr<Skeleton> Skeleton::StaticUnserialise( InputArchive& arch )
-{
-    SkeletonPtr pResult = new Skeleton();
-    arch >> *pResult;
-    return pResult;
-}
-
-
-//---------------------------------------------------------------------------------------------
-SkeletonPtr Skeleton::Clone() const
-{
-    SkeletonPtr pResult = new Skeleton();
-
-	pResult->m_bones = m_bones;
-	pResult->m_boneParents = m_boneParents;
-	pResult->m_boneIds = m_boneIds;
-
-    return pResult;
-}
-
-
-//---------------------------------------------------------------------------------------------
-int32 Skeleton::GetBoneCount() const
+namespace mu
 {
-    return m_bones.Num();
-}
 
 
-//---------------------------------------------------------------------------------------------
-void Skeleton::SetBoneCount( int32 b )
-{
-	m_bones.SetNum( b );
-	m_boneParents.Init(INDEX_NONE, b);
-	m_boneIds.Init(INDEX_NONE, b);
-}
-
-
-//---------------------------------------------------------------------------------------------
-const char* Skeleton::GetBoneName(int32 boneIndex ) const
-{
-    check( boneIndex>=0 && boneIndex<GetBoneCount() );
-    if (boneIndex >= 0 && boneIndex < GetBoneCount())
-    {
-        return m_bones[boneIndex].c_str();
-    }
-    return "";
-}
-
-
-//---------------------------------------------------------------------------------------------
-void Skeleton::SetBoneName(int32 boneIndex, const char* strName )
-{
-    check( boneIndex>=0 && boneIndex<GetBoneCount() );
-    if (boneIndex >= 0 && boneIndex < GetBoneCount())
-    {
-        m_bones[boneIndex] = strName;
-    }
-}
-
-
-//---------------------------------------------------------------------------------------------
-int Skeleton::GetBoneParent(int32 boneIndex ) const
-{
-    if ( boneIndex >= 0 && boneIndex < m_boneParents.Num() )
-    {
-        return m_boneParents[boneIndex];
-    }
-
-    return -1;
-}
-
-
-//-------------------------------------------------------------------------------------------------
-void Skeleton::SetBoneParent(int32 boneIndex, int32 parentBoneIndex )
-{
-    check( boneIndex >= 0 && boneIndex < GetBoneCount() );
-    check( parentBoneIndex >= -1 && parentBoneIndex < GetBoneCount() &&
-                    parentBoneIndex < 0xffff );
-    if ( boneIndex >= 0 && boneIndex < m_boneParents.Num() )
-    {
-        m_boneParents[boneIndex] = (int16_t)parentBoneIndex;
-    }
-}
-
-
-//-------------------------------------------------------------------------------------------------
-int32 Skeleton::GetBoneId(int32 boneIndex ) const
-{
-    if ( boneIndex >= 0 && boneIndex < m_boneIds.Num() )
-    {
-        return int(m_boneIds[boneIndex]);
-    }
-
-    return 0;
-}
-
-
-//-------------------------------------------------------------------------------------------------
-void Skeleton::SetBoneId(int32 boneIndex, int32 id )
-{
-    int boneCount = GetBoneCount();
-    check( boneIndex >= 0 && boneIndex < boneCount );
-    if ( boneIndex >= 0 && boneIndex < boneCount )
-    {
-        // Enlarge the id buffer if necessary, but enlarge it to the total number of bones to avoid
-        // repetitive allocation.
-        if ( boneIndex >= m_boneIds.Num() )
-        {
-            m_boneIds.SetNumZeroed( boneCount );
-        }
-
-        m_boneIds[boneIndex] = id;
-    }
-}
+	//---------------------------------------------------------------------------------------------
+	void Skeleton::Serialise(const Skeleton* p, OutputArchive& arch)
+	{
+		arch << *p;
+	}
+
+
+	//---------------------------------------------------------------------------------------------
+	mu::Ptr<Skeleton> Skeleton::StaticUnserialise(InputArchive& arch)
+	{
+		SkeletonPtr pResult = new Skeleton();
+		arch >> *pResult;
+		return pResult;
+	}
+
+
+	//---------------------------------------------------------------------------------------------
+	SkeletonPtr Skeleton::Clone() const
+	{
+		SkeletonPtr pResult = new Skeleton();
+
+		pResult->m_bones = m_bones;
+		pResult->m_boneParents = m_boneParents;
+		pResult->m_boneIds = m_boneIds;
+
+		return pResult;
+	}
+
+
+	//---------------------------------------------------------------------------------------------
+	int32 Skeleton::GetBoneCount() const
+	{
+		return m_bones.Num();
+	}
+
+
+	//---------------------------------------------------------------------------------------------
+	void Skeleton::SetBoneCount(int32 b)
+	{
+		m_bones.SetNum(b);
+		m_boneParents.Init(INDEX_NONE, b);
+		m_boneIds.Init(INDEX_NONE, b);
+	}
+
+
+	//---------------------------------------------------------------------------------------------
+	const char* Skeleton::GetBoneName(int32 boneIndex) const
+	{
+		check(boneIndex >= 0 && boneIndex < GetBoneCount());
+		if (boneIndex >= 0 && boneIndex < GetBoneCount())
+		{
+			return m_bones[boneIndex].c_str();
+		}
+		return "";
+	}
+
+
+	//---------------------------------------------------------------------------------------------
+	void Skeleton::SetBoneName(int32 boneIndex, const char* strName)
+	{
+		check(boneIndex >= 0 && boneIndex < GetBoneCount());
+		if (boneIndex >= 0 && boneIndex < GetBoneCount())
+		{
+			m_bones[boneIndex] = strName;
+		}
+	}
+
+
+	//---------------------------------------------------------------------------------------------
+	int Skeleton::GetBoneParent(int32 boneIndex) const
+	{
+		if (boneIndex >= 0 && boneIndex < m_boneParents.Num())
+		{
+			return m_boneParents[boneIndex];
+		}
+
+		return -1;
+	}
+
+
+	//-------------------------------------------------------------------------------------------------
+	void Skeleton::SetBoneParent(int32 boneIndex, int32 parentBoneIndex)
+	{
+		check(boneIndex >= 0 && boneIndex < GetBoneCount());
+		check(parentBoneIndex >= -1 && parentBoneIndex < GetBoneCount() &&
+			parentBoneIndex < 0xffff);
+		if (boneIndex >= 0 && boneIndex < m_boneParents.Num())
+		{
+			m_boneParents[boneIndex] = (int16_t)parentBoneIndex;
+		}
+	}
+
+
+	//-------------------------------------------------------------------------------------------------
+	int32 Skeleton::GetBoneId(int32 boneIndex) const
+	{
+		if (boneIndex >= 0 && boneIndex < m_boneIds.Num())
+		{
+			return int(m_boneIds[boneIndex]);
+		}
+
+		return 0;
+	}
+
+
+	//-------------------------------------------------------------------------------------------------
+	void Skeleton::SetBoneId(int32 boneIndex, int32 id)
+	{
+		int boneCount = GetBoneCount();
+		check(boneIndex >= 0 && boneIndex < boneCount);
+		if (boneIndex >= 0 && boneIndex < boneCount)
+		{
+			// Enlarge the id buffer if necessary, but enlarge it to the total number of bones to avoid
+			// repetitive allocation.
+			if (boneIndex >= m_boneIds.Num())
+			{
+				m_boneIds.SetNumZeroed(boneCount);
+			}
+
+			m_boneIds[boneIndex] = id;
+		}
+	}
+
+}
\ No newline at end of file
diff --git a/Engine/Plugins/Experimental/Mutable/Source/MutableTools/Private/AST.cpp b/Engine/Plugins/Experimental/Mutable/Source/MutableTools/Private/AST.cpp
--- a/Engine/Plugins/Experimental/Mutable/Source/MutableTools/Private/AST.cpp
+++ b/Engine/Plugins/Experimental/Mutable/Source/MutableTools/Private/AST.cpp
@@ -15,11 +15,10 @@
 std::atomic<uint32_t> mu::ASTOp::s_lastTraverseIndex(1);
 
 
-using namespace mu;
-
-
+namespace mu
+{
 
-void mu::DebugLogAST(const Ptr<ASTOp>& at, int indent, ASTOpList* done, const char* label )
+void DebugLogAST(const Ptr<ASTOp>& at, int indent, ASTOpList* done, const char* label )
 {
     if (!at) return;
 
@@ -1895,3 +1894,4 @@ mu::Ptr<ImageSizeExpression> ASTOpFixed::GetImageSizeExpression() const
     return pRes;
 }
 
+}
diff --git a/Engine/Plugins/Experimental/Mutable/Source/MutableTools/Private/ASTOpAddLOD.cpp b/Engine/Plugins/Experimental/Mutable/Source/MutableTools/Private/ASTOpAddLOD.cpp
--- a/Engine/Plugins/Experimental/Mutable/Source/MutableTools/Private/ASTOpAddLOD.cpp
+++ b/Engine/Plugins/Experimental/Mutable/Source/MutableTools/Private/ASTOpAddLOD.cpp
@@ -6,92 +6,94 @@
 #include "StreamsPrivate.h"
 
 
-using namespace mu;
-
-
-//-------------------------------------------------------------------------------------------------
-//-------------------------------------------------------------------------------------------------
-//-------------------------------------------------------------------------------------------------
-ASTOpAddLOD::ASTOpAddLOD()
+namespace mu
 {
-}
 
 
-ASTOpAddLOD::~ASTOpAddLOD()
-{
-    // Explicit call needed to avoid recursive destruction
-    ASTOp::RemoveChildren();
-}
-
+	//-------------------------------------------------------------------------------------------------
+	//-------------------------------------------------------------------------------------------------
+	//-------------------------------------------------------------------------------------------------
+	ASTOpAddLOD::ASTOpAddLOD()
+	{
+	}
+
+
+	ASTOpAddLOD::~ASTOpAddLOD()
+	{
+		// Explicit call needed to avoid recursive destruction
+		ASTOp::RemoveChildren();
+	}
+
+
+	bool ASTOpAddLOD::IsEqual(const ASTOp& otherUntyped) const
+	{
+		if (auto other = dynamic_cast<const ASTOpAddLOD*>(&otherUntyped))
+		{
+			return lods == other->lods;
+		}
+		return false;
+	}
+
+
+	uint64 ASTOpAddLOD::Hash() const
+	{
+		uint64 res = std::hash<uint64>()(uint64(OP_TYPE::IN_ADDLOD));
+		for (auto& c : lods)
+		{
+			hash_combine(res, c.child().get());
+		}
+		return res;
+	}
+
+
+	mu::Ptr<ASTOp> ASTOpAddLOD::Clone(MapChildFunc& mapChild) const
+	{
+		Ptr<ASTOpAddLOD> n = new ASTOpAddLOD();
+		for (auto& c : lods)
+		{
+			n->lods.emplace_back(n, mapChild(c.child()));
+		}
+		return n;
+	}
+
+
+	void ASTOpAddLOD::ForEachChild(const std::function<void(ASTChild&)>& f)
+	{
+		for (auto& l : lods)
+		{
+			f(l);
+		}
+	}
+
+
+	void ASTOpAddLOD::Link(PROGRAM& program, const FLinkerOptions*)
+	{
+		// Already linked?
+		if (!linkedAddress)
+		{
+			OP::InstanceAddLODArgs args;
+			memset(&args, 0, sizeof(args));
+
+			int i = 0;
+			for (auto& l : lods)
+			{
+				if (l)
+				{
+					args.lod[i] = l->linkedAddress;
+					++i;
+					if (i >= MUTABLE_OP_MAX_ADD_COUNT)
+					{
+						break;
+					}
+				}
+			}
+
+			linkedAddress = (OP::ADDRESS)program.m_opAddress.size();
+			program.m_opAddress.push_back((uint32_t)program.m_byteCode.size());
+			AppendCode(program.m_byteCode, OP_TYPE::IN_ADDLOD);
+			AppendCode(program.m_byteCode, args);
+		}
+
+	}
 
-bool ASTOpAddLOD::IsEqual(const ASTOp& otherUntyped) const
-{
-    if ( auto other = dynamic_cast<const ASTOpAddLOD*>(&otherUntyped) )
-    {
-        return lods==other->lods;
-    }
-    return false;
 }
-
-
-uint64 ASTOpAddLOD::Hash() const
-{
-	uint64 res = std::hash<uint64>()(uint64(OP_TYPE::IN_ADDLOD));
-    for(auto& c:lods)
-    {
-        hash_combine( res, c.child().get() );
-    }
-    return res;
-}
-
-
-mu::Ptr<ASTOp> ASTOpAddLOD::Clone(MapChildFunc& mapChild) const
-{
-    Ptr<ASTOpAddLOD> n = new ASTOpAddLOD();
-    for(auto& c:lods)
-    {
-        n->lods.emplace_back(n,mapChild(c.child()));
-    }
-    return n;
-}
-
-
-void ASTOpAddLOD::ForEachChild( const std::function<void(ASTChild&)>& f )
-{
-    for (auto& l:lods)
-    {
-        f(l);
-    }
-}
-
-
-void ASTOpAddLOD::Link( PROGRAM& program, const FLinkerOptions* )
-{
-    // Already linked?
-    if (!linkedAddress)
-    {
-        OP::InstanceAddLODArgs args;
-        memset( &args,0, sizeof(args) );
-
-        int i=0;
-        for (auto& l:lods)
-        {
-            if ( l )
-            {
-                args.lod[i] = l->linkedAddress;
-                ++i;
-                if (i>=MUTABLE_OP_MAX_ADD_COUNT)
-                {
-                    break;
-                }
-            }
-        }
-
-        linkedAddress = (OP::ADDRESS)program.m_opAddress.size();
-        program.m_opAddress.push_back((uint32_t)program.m_byteCode.size());
-        AppendCode(program.m_byteCode,OP_TYPE::IN_ADDLOD);
-        AppendCode(program.m_byteCode,args);
-    }
-
-}
-
diff --git a/Engine/Plugins/Experimental/Mutable/Source/MutableTools/Private/ASTOpConstantProjector.cpp b/Engine/Plugins/Experimental/Mutable/Source/MutableTools/Private/ASTOpConstantProjector.cpp
--- a/Engine/Plugins/Experimental/Mutable/Source/MutableTools/Private/ASTOpConstantProjector.cpp
+++ b/Engine/Plugins/Experimental/Mutable/Source/MutableTools/Private/ASTOpConstantProjector.cpp
@@ -6,56 +6,58 @@
 #include "StreamsPrivate.h"
 
 
-using namespace mu;
-
-
-//-------------------------------------------------------------------------------------------------
-//-------------------------------------------------------------------------------------------------
-//-------------------------------------------------------------------------------------------------
-void ASTOpConstantProjector::ForEachChild(const std::function<void(ASTChild&)>& )
-{
-}
-
-
-uint64 ASTOpConstantProjector::Hash() const
-{
-	uint64 res = std::hash<float>()(value.position[0]);
-    hash_combine( res, value.direction[0] );
-    return res;
-}
-
-
-bool ASTOpConstantProjector::IsEqual(const ASTOp& otherUntyped) const
+namespace mu
 {
-    if ( auto other = dynamic_cast<const ASTOpConstantProjector*>(&otherUntyped) )
-    {
-        return value==other->value;
-    }
-    return false;
-}
-
 
-mu::Ptr<ASTOp> ASTOpConstantProjector::Clone(MapChildFunc&) const
-{
-    Ptr<ASTOpConstantProjector> n = new ASTOpConstantProjector();
-    n->value = value;
-    return n;
-}
 
+	//-------------------------------------------------------------------------------------------------
+	//-------------------------------------------------------------------------------------------------
+	//-------------------------------------------------------------------------------------------------
+	void ASTOpConstantProjector::ForEachChild(const std::function<void(ASTChild&)>&)
+	{
+	}
+
+
+	uint64 ASTOpConstantProjector::Hash() const
+	{
+		uint64 res = std::hash<float>()(value.position[0]);
+		hash_combine(res, value.direction[0]);
+		return res;
+	}
+
+
+	bool ASTOpConstantProjector::IsEqual(const ASTOp& otherUntyped) const
+	{
+		if (auto other = dynamic_cast<const ASTOpConstantProjector*>(&otherUntyped))
+		{
+			return value == other->value;
+		}
+		return false;
+	}
+
+
+	mu::Ptr<ASTOp> ASTOpConstantProjector::Clone(MapChildFunc&) const
+	{
+		Ptr<ASTOpConstantProjector> n = new ASTOpConstantProjector();
+		n->value = value;
+		return n;
+	}
+
+
+	void ASTOpConstantProjector::Link(PROGRAM& program, const FLinkerOptions*)
+	{
+		if (!linkedAddress)
+		{
+			OP::ResourceConstantArgs args;
+			memset(&args, 0, sizeof(args));
+			args.value = program.AddConstant(value);
+
+			linkedAddress = (OP::ADDRESS)program.m_opAddress.size();
+			//program.m_code.push_back(op);
+			program.m_opAddress.push_back((uint32_t)program.m_byteCode.size());
+			AppendCode(program.m_byteCode, OP_TYPE::PR_CONSTANT);
+			AppendCode(program.m_byteCode, args);
+		}
+	}
 
-void ASTOpConstantProjector::Link( PROGRAM& program, const FLinkerOptions*)
-{
-    if (!linkedAddress)
-    {
-        OP::ResourceConstantArgs args;
-        memset( &args,0, sizeof(args) );
-        args.value = program.AddConstant(value);
-
-        linkedAddress = (OP::ADDRESS)program.m_opAddress.size();
-        //program.m_code.push_back(op);
-        program.m_opAddress.push_back((uint32_t)program.m_byteCode.size());
-        AppendCode(program.m_byteCode,OP_TYPE::PR_CONSTANT);
-        AppendCode(program.m_byteCode,args);
-    }
 }
-
diff --git a/Engine/Plugins/Experimental/Mutable/Source/MutableTools/Private/ASTOpMeshGeometryOperation.cpp b/Engine/Plugins/Experimental/Mutable/Source/MutableTools/Private/ASTOpMeshGeometryOperation.cpp
--- a/Engine/Plugins/Experimental/Mutable/Source/MutableTools/Private/ASTOpMeshGeometryOperation.cpp
+++ b/Engine/Plugins/Experimental/Mutable/Source/MutableTools/Private/ASTOpMeshGeometryOperation.cpp
@@ -5,85 +5,86 @@
 #include "CodeOptimiser.h"
 #include "StreamsPrivate.h"
 
-using namespace mu;
-
-
-ASTOpMeshGeometryOperation::ASTOpMeshGeometryOperation()
-    : meshA(this)
-    , meshB(this)
-	, scalarA(this)
-	, scalarB(this)
-{
-}
-
-
-ASTOpMeshGeometryOperation::~ASTOpMeshGeometryOperation()
-{
-    // Explicit call needed to avoid recursive destruction
-    ASTOp::RemoveChildren();
-}
-
-
-bool ASTOpMeshGeometryOperation::IsEqual(const ASTOp& otherUntyped) const
-{
-    if ( auto other = dynamic_cast<const ASTOpMeshGeometryOperation*>(&otherUntyped) )
-    {
-        return meshA ==other->meshA && meshB == other->meshB &&
-			scalarA == other->scalarA && scalarB == other->scalarB;
-    }
-    return false;
-}
-
-
-uint64 ASTOpMeshGeometryOperation::Hash() const
-{
-	uint64 res = std::hash<void*>()(meshA.child().get() );
-	hash_combine(res, meshB.child().get());
-	hash_combine(res, scalarA.child().get());
-	hash_combine(res, scalarB.child().get());
-	return res;
-}
-
-
-mu::Ptr<ASTOp> ASTOpMeshGeometryOperation::Clone(MapChildFunc& mapChild) const
-{
-    Ptr<ASTOpMeshGeometryOperation> n = new ASTOpMeshGeometryOperation();
-	n->meshA = mapChild(meshA.child());
-	n->meshB = mapChild(meshB.child());
-	n->scalarA = mapChild(scalarA.child());
-	n->scalarB = mapChild(scalarB.child());
-	return n;
-}
-
-
-void ASTOpMeshGeometryOperation::ForEachChild(const std::function<void(ASTChild&)>& f )
-{
-	f(meshA);
-	f(meshB);
-	f(scalarA);
-	f(scalarB);
-}
-
-
-void ASTOpMeshGeometryOperation::Link( PROGRAM& program, const FLinkerOptions* )
+namespace mu
 {
-    // Already linked?
-    if (!linkedAddress)
-    {
-        OP::MeshGeometryOperationArgs args;
-        memset( &args,0, sizeof(args) );
-
-		if (meshA) args.meshA = meshA->linkedAddress;
-		if (meshB) args.meshB = meshB->linkedAddress;
-		if (scalarA) args.scalarA = scalarA->linkedAddress;
-		if (scalarB) args.scalarB = scalarB->linkedAddress;
-
-        linkedAddress = (OP::ADDRESS)program.m_opAddress.size();
-        //program.m_code.push_back(op);
-        program.m_opAddress.push_back((uint32_t)program.m_byteCode.size());
-        AppendCode(program.m_byteCode,OP_TYPE::ME_GEOMETRYOPERATION);
-        AppendCode(program.m_byteCode,args);
-    }
-
-}
 
+	ASTOpMeshGeometryOperation::ASTOpMeshGeometryOperation()
+		: meshA(this)
+		, meshB(this)
+		, scalarA(this)
+		, scalarB(this)
+	{
+	}
+
+
+	ASTOpMeshGeometryOperation::~ASTOpMeshGeometryOperation()
+	{
+		// Explicit call needed to avoid recursive destruction
+		ASTOp::RemoveChildren();
+	}
+
+
+	bool ASTOpMeshGeometryOperation::IsEqual(const ASTOp& otherUntyped) const
+	{
+		if (auto other = dynamic_cast<const ASTOpMeshGeometryOperation*>(&otherUntyped))
+		{
+			return meshA == other->meshA && meshB == other->meshB &&
+				scalarA == other->scalarA && scalarB == other->scalarB;
+		}
+		return false;
+	}
+
+
+	uint64 ASTOpMeshGeometryOperation::Hash() const
+	{
+		uint64 res = std::hash<void*>()(meshA.child().get());
+		hash_combine(res, meshB.child().get());
+		hash_combine(res, scalarA.child().get());
+		hash_combine(res, scalarB.child().get());
+		return res;
+	}
+
+
+	mu::Ptr<ASTOp> ASTOpMeshGeometryOperation::Clone(MapChildFunc& mapChild) const
+	{
+		Ptr<ASTOpMeshGeometryOperation> n = new ASTOpMeshGeometryOperation();
+		n->meshA = mapChild(meshA.child());
+		n->meshB = mapChild(meshB.child());
+		n->scalarA = mapChild(scalarA.child());
+		n->scalarB = mapChild(scalarB.child());
+		return n;
+	}
+
+
+	void ASTOpMeshGeometryOperation::ForEachChild(const std::function<void(ASTChild&)>& f)
+	{
+		f(meshA);
+		f(meshB);
+		f(scalarA);
+		f(scalarB);
+	}
+
+
+	void ASTOpMeshGeometryOperation::Link(PROGRAM& program, const FLinkerOptions*)
+	{
+		// Already linked?
+		if (!linkedAddress)
+		{
+			OP::MeshGeometryOperationArgs args;
+			memset(&args, 0, sizeof(args));
+
+			if (meshA) args.meshA = meshA->linkedAddress;
+			if (meshB) args.meshB = meshB->linkedAddress;
+			if (scalarA) args.scalarA = scalarA->linkedAddress;
+			if (scalarB) args.scalarB = scalarB->linkedAddress;
+
+			linkedAddress = (OP::ADDRESS)program.m_opAddress.size();
+			//program.m_code.push_back(op);
+			program.m_opAddress.push_back((uint32_t)program.m_byteCode.size());
+			AppendCode(program.m_byteCode, OP_TYPE::ME_GEOMETRYOPERATION);
+			AppendCode(program.m_byteCode, args);
+		}
+
+	}
+
+}
\ No newline at end of file
diff --git a/Engine/Plugins/Experimental/Mutable/Source/MutableTools/Private/Visitor.h b/Engine/Plugins/Experimental/Mutable/Source/MutableTools/Private/Visitor.h
--- a/Engine/Plugins/Experimental/Mutable/Source/MutableTools/Private/Visitor.h
+++ b/Engine/Plugins/Experimental/Mutable/Source/MutableTools/Private/Visitor.h
@@ -1,11 +1,6 @@
 // Copyright Epic Games, Inc. All Rights Reserved.
 
-//! Partially based on the visitor description of the Loki library
-
-
-#ifndef _MUTABLE_VISITOR_H_
-#define _MUTABLE_VISITOR_H_
-
+#pragma once
 
 #include "MutableTools/Private/Platform.h"
 
@@ -135,6 +130,3 @@ namespace mu
 
 
 }
-
-
-#endif
