[mutable] Fix DGGUI Blueprint compilation issues.

- Updated references to the now private UCustomizableObjectInstance::CustomizableObject.
- Updated UCustomizableObjectInstance::GetProjectorValue parameters names.

diff --git a/Engine/Plugins/Experimental/Mutable/Source/CustomizableObject/Classes/CustomizableObjectInstance.h b/Engine/Plugins/Experimental/Mutable/Source/CustomizableObject/Classes/CustomizableObjectInstance.h
--- a/Engine/Plugins/Experimental/Mutable/Source/CustomizableObject/Classes/CustomizableObjectInstance.h
+++ b/Engine/Plugins/Experimental/Mutable/Source/CustomizableObject/Classes/CustomizableObjectInstance.h
@@ -153,9 +153,11 @@ public:
 	virtual FString GetDesc() override;
 	virtual bool IsEditorOnly() const override;
 
+	UFUNCTION(BlueprintCallable, Category = CustomizableObjectInstance)
 	void SetObject(UCustomizableObject* InObject);
 
 	//Get the current CustomizableObject 
+	UFUNCTION(BlueprintCallable, Category = CustomizableObjectInstance)
 	UCustomizableObject* GetCustomizableObject() const;
 
 	int32 GetState() const;
