[mutable] Fixed case where the generation of the cliping mesh was using wrong layout information stored in m_overrideLayoutsStack due to being connected to a table node (mesh var, mesh swt also).
Also, specified the LOD number in the unassigned UVs warning and fixed a minor bug of the autogeneration of blocks.

#rnx

diff --git a/Engine/Plugins/Experimental/Mutable/Source/MutableTools/Private/CodeGenerator_Mesh.cpp b/Engine/Plugins/Experimental/Mutable/Source/MutableTools/Private/CodeGenerator_Mesh.cpp
--- a/Engine/Plugins/Experimental/Mutable/Source/MutableTools/Private/CodeGenerator_Mesh.cpp
+++ b/Engine/Plugins/Experimental/Mutable/Source/MutableTools/Private/CodeGenerator_Mesh.cpp
@@ -196,8 +196,8 @@ namespace mu
             mutable_snprintf
                 (
                     buf, 256,
-                    "Source mesh has %d vertices not assigned to any layout block.",
-                    outside
+                    "Source mesh has %d vertices not assigned to any layout block in LOD %d",
+					outside, m_currentParents.back().m_lod
                 );
 			int blockCount = pLayout->GetBlockCount();
         
@@ -1024,9 +1024,16 @@ namespace mu
 		BOTTOM_UP_STATE temp = m_currentBottomUpState;
 		if (!m_activeTags.empty())
 		{
+			// Clearing layout stack to avoid unwanted information
+			vector< vector<Ptr<const Layout>> > m_overrideLayoutsStackCopy = m_overrideLayoutsStack;
+			m_overrideLayoutsStack.clear();
+
 			bool bModifiersForBeforeOperations = true;
 			result.meshOp = ApplyMeshModifiers(op, m_activeTags.back(),
 				bModifiersForBeforeOperations, node.m_errorContext);
+
+			// Retrieving stack information
+			m_overrideLayoutsStack = m_overrideLayoutsStackCopy;
 		}
 		m_currentBottomUpState = temp;
 
diff --git a/Engine/Plugins/Experimental/Mutable/Source/MutableTools/Private/NodeLayout.cpp b/Engine/Plugins/Experimental/Mutable/Source/MutableTools/Private/NodeLayout.cpp
--- a/Engine/Plugins/Experimental/Mutable/Source/MutableTools/Private/NodeLayout.cpp
+++ b/Engine/Plugins/Experimental/Mutable/Source/MutableTools/Private/NodeLayout.cpp
@@ -360,7 +360,7 @@ namespace mu
 				{
 					for (size_t j = 0; j < blocks.size(); ++j)
 					{
-						if (i != j && blocks[i].Contains(blocks[j].min))
+						if (i != j && blocks[i].IntersectsExclusive(blocks[j]))
 						{
 							blocks[i].Bound(blocks[j]);
 							blocks.erase(blocks.begin() + j);
