[Mutable] Fix COI parameters not showing up in the COI Editor

#rnx

diff --git a/Engine/Plugins/Experimental/Mutable/Source/CustomizableObject/Private/CustomizableObjectSystem.cpp b/Engine/Plugins/Experimental/Mutable/Source/CustomizableObject/Private/CustomizableObjectSystem.cpp
--- a/Engine/Plugins/Experimental/Mutable/Source/CustomizableObject/Private/CustomizableObjectSystem.cpp
+++ b/Engine/Plugins/Experimental/Mutable/Source/CustomizableObject/Private/CustomizableObjectSystem.cpp
@@ -2377,6 +2377,11 @@ FMutableOperation FMutableOperation::CreateInstanceUpdate(UCustomizableObjectIns
 	Op.bBuildParameterDecorations = InCustomizableObjectInstance->bBuildParameterDecorations;
 	Op.Parameters = InCustomizableObjectInstance->GetPrivate()->ReloadParametersFromObject(InCustomizableObjectInstance, false);
 
+	// TODO: Remove once MTBL-973 is done
+#if WITH_EDITOR
+	Op.bBuildParameterDecorations = true;
+#endif
+
 	InCustomizableObjectInstance->GetCustomizableObject()->ApplyStateForcedValuesToParameters(InCustomizableObjectInstance->GetState(), Op.Parameters.get());
 
 	if (!Op.Parameters)
