[mutable] Fix more cases of bad localization keys.
Remove code related to older engine versions.

diff --git a/Engine/Plugins/Experimental/Mutable/Source/CustomizableObject/Private/CustomizableObjectInstance.cpp b/Engine/Plugins/Experimental/Mutable/Source/CustomizableObject/Private/CustomizableObjectInstance.cpp
--- a/Engine/Plugins/Experimental/Mutable/Source/CustomizableObject/Private/CustomizableObjectInstance.cpp
+++ b/Engine/Plugins/Experimental/Mutable/Source/CustomizableObject/Private/CustomizableObjectInstance.cpp
@@ -4739,15 +4739,8 @@ void UpdateTextureRegionsMutable(UTexture2D* Texture, int32 MipIndex, const FUpd
 		RegionData->SrcPitch = SrcPitch;
 
 
-#if (ENGINE_MAJOR_VERSION == 4 && ENGINE_MINOR_VERSION > 22) || ENGINE_MAJOR_VERSION >= 5
 		ENQUEUE_RENDER_COMMAND(UpdateTextureRegionsMutable)(
 			[RegionData, BulkData](FRHICommandList& CmdList)
-#else
-		ENQUEUE_UNIQUE_RENDER_COMMAND_TWOPARAMETER(
-			UpdateTextureRegionsMutable,
-			FUpdateTextureRegionsData*, RegionData, RegionData,
-			FByteBulkData*, BulkData, BulkData,
-#endif
 			{
 				int32 CurrentFirstMip = RegionData->Texture2DResource->GetCurrentFirstMip();
 				uint8* SrcData = (uint8*)BulkData->LockReadOnly();
diff --git a/Engine/Plugins/Experimental/Mutable/Source/CustomizableObject/Private/UnrealConversionUtils.cpp b/Engine/Plugins/Experimental/Mutable/Source/CustomizableObject/Private/UnrealConversionUtils.cpp
--- a/Engine/Plugins/Experimental/Mutable/Source/CustomizableObject/Private/UnrealConversionUtils.cpp
+++ b/Engine/Plugins/Experimental/Mutable/Source/CustomizableObject/Private/UnrealConversionUtils.cpp
@@ -187,11 +187,7 @@ namespace UnrealConversionUtils
 		InMutableMesh->GetSurface(Surface, &FirstVertex, &VertexCount, &FirstIndex, &IndexCount);
 		FSkelMeshRenderSection& Section = Helper_GetLODRenderSections(OutSkeletalMesh, MeshLODIndex)[Surface];
 
-#if (ENGINE_MAJOR_VERSION == 4 && ENGINE_MINOR_VERSION >= 19) || ENGINE_MAJOR_VERSION >= 5
-		// \todo: investigate. Without this 4.19 rendering crashes. 
-		// It must be initialized even in disabled sections or else they will crash
 		Section.DuplicatedVerticesBuffer.Init(1, TMap<int, TArray<int32>>());
-#endif
 			
 		if (VertexCount == 0 || IndexCount == 0)
 		{
diff --git a/Engine/Plugins/Experimental/Mutable/Source/CustomizableObjectEditor/Classes/EdGraphSchema_CustomizableObject.h b/Engine/Plugins/Experimental/Mutable/Source/CustomizableObjectEditor/Classes/EdGraphSchema_CustomizableObject.h
--- a/Engine/Plugins/Experimental/Mutable/Source/CustomizableObjectEditor/Classes/EdGraphSchema_CustomizableObject.h
+++ b/Engine/Plugins/Experimental/Mutable/Source/CustomizableObjectEditor/Classes/EdGraphSchema_CustomizableObject.h
@@ -107,11 +107,7 @@ public:
 
 	FLinearColor GetPinTypeColor(const FName& PinType) const;
 
-#if (ENGINE_MAJOR_VERSION == 4 && ENGINE_MINOR_VERSION>=18) || ENGINE_MAJOR_VERSION >= 5
 	void GetContextMenuActions(class UToolMenu* Menu, class UGraphNodeContextMenuContext* Context) const override;
-#else
-	void GetContextMenuActions(const UEdGraph* CurrentGraph, const UEdGraphNode* InGraphNode, const UEdGraphPin* InGraphPin, class FMenuBuilder* MenuBuilder, bool bIsDebugging) const override;
-#endif
 
 	virtual void BreakPinLinks(UEdGraphPin& TargetPin, bool bSendsNodeNotification) const override;
 	virtual void BreakSinglePinLink(UEdGraphPin* SourcePin, UEdGraphPin* TargetPin) const override;
diff --git a/Engine/Plugins/Experimental/Mutable/Source/CustomizableObjectEditor/Private/CustomizableObjectDetails.cpp b/Engine/Plugins/Experimental/Mutable/Source/CustomizableObjectEditor/Private/CustomizableObjectDetails.cpp
--- a/Engine/Plugins/Experimental/Mutable/Source/CustomizableObjectEditor/Private/CustomizableObjectDetails.cpp
+++ b/Engine/Plugins/Experimental/Mutable/Source/CustomizableObjectEditor/Private/CustomizableObjectDetails.cpp
@@ -85,19 +85,11 @@ void FCustomizableObjectDetails::CustomizeDetails( IDetailLayoutBuilder& DetailB
 {
 	const UCustomizableObject* CustomObject = 0;
 
-#if (ENGINE_MAJOR_VERSION == 4 && ENGINE_MINOR_VERSION >= 18) || ENGINE_MAJOR_VERSION >= 5
 	const IDetailsView* DetailsView = DetailBuilder.GetDetailsView();
 	if (DetailsView->GetSelectedObjects().Num())
 	{
 		CustomObject = Cast<const UCustomizableObject>(DetailsView->GetSelectedObjects()[0].Get());
 	}
-#else
-	const IDetailsView& DetailsView = DetailBuilder.GetDetailsView();
-	if (DetailsView.GetSelectedObjects().Num())
-	{
-		CustomObject = Cast<const UCustomizableObject>(DetailsView.GetSelectedObjects()[0].Get());
-	}
-#endif
 
 	IDetailCategoryBuilder& StatesCategory = DetailBuilder.EditCategory( "States" );
 	//StatesCategory.CategoryIcon( "ActorClassIcon.CustomizableObject" );
@@ -176,7 +168,6 @@ void FCustomizableObjectDetails::CustomizeDetails( IDetailLayoutBuilder& DetailB
 
 	TSharedRef<IPropertyHandle> Property = DetailBuilder.GetProperty("ReferenceSkeletalMeshes");
 
-#if (ENGINE_MAJOR_VERSION == 4 && ENGINE_MINOR_VERSION >= 18) || ENGINE_MAJOR_VERSION >= 5
 	if (Property->IsValidHandle() && CustomObject)
 	{
 		if (CustomObject->bIsChildObject)
@@ -188,7 +179,6 @@ void FCustomizableObjectDetails::CustomizeDetails( IDetailLayoutBuilder& DetailB
 			Property->MarkResetToDefaultCustomized();
 		}
 	}
-#endif
 
 	StatesTree->RequestTreeRefresh();
 }
diff --git a/Engine/Plugins/Experimental/Mutable/Source/CustomizableObjectEditor/Private/CustomizableObjectEditor.cpp b/Engine/Plugins/Experimental/Mutable/Source/CustomizableObjectEditor/Private/CustomizableObjectEditor.cpp
--- a/Engine/Plugins/Experimental/Mutable/Source/CustomizableObjectEditor/Private/CustomizableObjectEditor.cpp
+++ b/Engine/Plugins/Experimental/Mutable/Source/CustomizableObjectEditor/Private/CustomizableObjectEditor.cpp
@@ -80,9 +80,7 @@
 #include "Widgets/Layout/SScrollBar.h"
 
 
-#if (ENGINE_MAJOR_VERSION == 4 && ENGINE_MINOR_VERSION >= 18) || ENGINE_MAJOR_VERSION >= 5
 #include "HAL/PlatformApplicationMisc.h"
-#endif
 
 
 #define LOCTEXT_NAMESPACE "CustomizableObjectEditor"
@@ -2007,11 +2005,7 @@ void FCustomizableObjectEditor::CopySelectedNodes()
 
 	FEdGraphUtilities::ExportNodesToText(SelectedNodes, /*out*/ ExportedText);
 
-#if (ENGINE_MAJOR_VERSION == 4 && ENGINE_MINOR_VERSION >= 18) || ENGINE_MAJOR_VERSION >= 5
 	FPlatformApplicationMisc::ClipboardCopy(*ExportedText);
-#else
-	FPlatformMisc::ClipboardCopy(*ExportedText);
-#endif
 
 	// Make sure Material remains the owner of the copied nodes
 	for (FGraphPanelSelectionSet::TConstIterator SelectedIter(SelectedNodes); SelectedIter; ++SelectedIter)
@@ -2062,11 +2056,7 @@ void FCustomizableObjectEditor::PasteNodesHere(const FVector2D& Location)
 
 	// Grab the text to paste from the clipboard.
 	FString TextToImport;
-#if (ENGINE_MAJOR_VERSION == 4 && ENGINE_MINOR_VERSION >= 18) || ENGINE_MAJOR_VERSION >= 5
 	FPlatformApplicationMisc::ClipboardPaste(TextToImport);
-#else
-	FPlatformMisc::ClipboardPaste(TextToImport);
-#endif
 
 	// Import the nodes
 	TSet<UEdGraphNode*> PastedNodes;
@@ -2167,11 +2157,7 @@ bool FCustomizableObjectEditor::CanPasteNodes() const
 {
 	FString ClipboardContent;
 
-#if (ENGINE_MAJOR_VERSION == 4 && ENGINE_MINOR_VERSION >= 18) || ENGINE_MAJOR_VERSION >= 5
 	FPlatformApplicationMisc::ClipboardPaste(ClipboardContent);
-#else
-	FPlatformMisc::ClipboardPaste(ClipboardContent);
-#endif
 
 	return FEdGraphUtilities::CanImportNodesFromText(CustomizableObject->Source, ClipboardContent);
 }
@@ -2873,9 +2859,7 @@ void FCustomizableObjectEditor::UpdateObjectProperties()
 {
 	if (CustomizableObjectDetailsView.IsValid())
 	{
-#if (ENGINE_MAJOR_VERSION == 4 && ENGINE_MINOR_VERSION >= 18) || ENGINE_MAJOR_VERSION >= 5
 		CustomizableObjectDetailsView->ForceRefresh();
-#endif
 	}
 }
 
diff --git a/Engine/Plugins/Experimental/Mutable/Source/CustomizableObjectEditor/Private/CustomizableObjectEditorViewportClient.cpp b/Engine/Plugins/Experimental/Mutable/Source/CustomizableObjectEditor/Private/CustomizableObjectEditorViewportClient.cpp
--- a/Engine/Plugins/Experimental/Mutable/Source/CustomizableObjectEditor/Private/CustomizableObjectEditorViewportClient.cpp
+++ b/Engine/Plugins/Experimental/Mutable/Source/CustomizableObjectEditor/Private/CustomizableObjectEditorViewportClient.cpp
@@ -2836,7 +2836,7 @@ void FCustomizableObjectEditorViewportClient::ShowInstanceGeometryInformation(FC
 				InCanvas->DrawShadowedString(
 					6.0f,
 					YOffset,
-					*FText::Format(NSLOCTEXT("CustomizableObjectEditor", "UVOverlay_F", "Component {3} LOD {0} has {1} vertices and {2} triangles"),
+					*FText::Format(NSLOCTEXT("CustomizableObjectEditor", "ComponentGeometryReport", "Component {3} LOD {0} has {1} vertices and {2} triangles"),
 						FText::AsNumber(i), FText::AsNumber(NumVertices), FText::AsNumber(NumTriangles), FText::AsNumber(ComponentIndex)).ToString(),
 					GEngine->GetSmallFont(),
 					FLinearColor::White
diff --git a/Engine/Plugins/Experimental/Mutable/Source/CustomizableObjectEditor/Private/CustomizableObjectInstanceEditorActions.cpp b/Engine/Plugins/Experimental/Mutable/Source/CustomizableObjectEditor/Private/CustomizableObjectInstanceEditorActions.cpp
--- a/Engine/Plugins/Experimental/Mutable/Source/CustomizableObjectEditor/Private/CustomizableObjectInstanceEditorActions.cpp
+++ b/Engine/Plugins/Experimental/Mutable/Source/CustomizableObjectEditor/Private/CustomizableObjectInstanceEditorActions.cpp
@@ -36,7 +36,7 @@ void FCustomizableObjectInstanceEditorCommands::RegisterCommands()
 	UI_COMMAND( SetShowBinormals, "Binormals", "Toggles display of vertex binormals (orthogonal vector to normal and tangent) in the Preview Pane.", EUserInterfaceActionType::ToggleButton, FInputChord() );
 	UI_COMMAND( SetShowPivot, "Show Pivot", "Display the pivot location of the static mesh.", EUserInterfaceActionType::ToggleButton, FInputChord() );
 
-	UI_COMMAND( BakeInstance, "Bake Instance from Instance", "Create baked unreal resources for the current instance.", EUserInterfaceActionType::Button, FInputChord());
+	UI_COMMAND( BakeInstanceFromInstance, "Bake Instance from Instance", "Create baked unreal resources for the current instance.", EUserInterfaceActionType::Button, FInputChord());
 	UI_COMMAND( StateChangeTest, "State test", "Make a state change test, iterating through all possible states and runtime parameters for this instance", EUserInterfaceActionType::Button, FInputChord());
 
 	UI_COMMAND(ShowParentCO, "Show Parent", "Show the Parent of the Customizable Object in the Content Browser", EUserInterfaceActionType::Button, FInputChord());
diff --git a/Engine/Plugins/Experimental/Mutable/Source/CustomizableObjectEditor/Private/CustomizableObjectNodeObjectGroupDetails.cpp b/Engine/Plugins/Experimental/Mutable/Source/CustomizableObjectEditor/Private/CustomizableObjectNodeObjectGroupDetails.cpp
--- a/Engine/Plugins/Experimental/Mutable/Source/CustomizableObjectEditor/Private/CustomizableObjectNodeObjectGroupDetails.cpp
+++ b/Engine/Plugins/Experimental/Mutable/Source/CustomizableObjectEditor/Private/CustomizableObjectNodeObjectGroupDetails.cpp
@@ -5,27 +5,8 @@
 #include "CustomizableObjectEditor.h"
 #include "CustomizableObjectEditorModule.h"
 #include "PropertyCustomizationHelpers.h"
-//#include "SPropertyEditorAsset.h"
-//#include "CustomizableObjectEditorUtilities.h"
-//
 #include "DetailLayoutBuilder.h"
 #include "DetailCategoryBuilder.h"
-//#include "LevelEditor.h"
-//#include "LevelEditorActions.h"
-//
-//#include "Widgets/Input/SButton.h"
-//#include "Widgets/Text/STextBlock.h"
-//
-//#include "Runtime/Launch/Resources/Version.h"
-//#include "Engine/SkeletalMesh.h"
-//#include "Widgets/Input/STextComboBox.h"
-//#include "Math/Vector4.h"
-//#include "Math/Vector.h"
-//
-//#include "Runtime/Launch/Resources/Version.h"
-//#if (ENGINE_MAJOR_VERSION == 4 && ENGINE_MINOR_VERSION > 22) || ENGINE_MAJOR_VERSION >= 5
-//#include "HAL/PlatformApplicationMisc.h"
-//#endif
 
 
 #define LOCTEXT_NAMESPACE "CustomizableObjectGroupDetails"
diff --git a/Engine/Plugins/Experimental/Mutable/Source/CustomizableObjectEditor/Private/CustomizableObjectWidget.cpp b/Engine/Plugins/Experimental/Mutable/Source/CustomizableObjectEditor/Private/CustomizableObjectWidget.cpp
--- a/Engine/Plugins/Experimental/Mutable/Source/CustomizableObjectEditor/Private/CustomizableObjectWidget.cpp
+++ b/Engine/Plugins/Experimental/Mutable/Source/CustomizableObjectEditor/Private/CustomizableObjectWidget.cpp
@@ -371,11 +371,7 @@ void FCustomizableObjectWidget::Render_Translate( const FSceneView* View, FPrimi
 		const FVector CameraYAxis = View->ViewMatrices.GetViewMatrix().GetColumn(1);
 		const FVector CameraZAxis = View->ViewMatrices.GetViewMatrix().GetColumn(2);
 
-#if (ENGINE_MAJOR_VERSION == 4 && ENGINE_MINOR_VERSION >= 17) || ENGINE_MAJOR_VERSION >= 5
 		DrawSphere( PDI, InLocation, FRotator::ZeroRotator, FVector( 4.0f * Scale ), 10, 5, Helper_GetMaterialProxy(XYZMaterial), SDPG_Foreground );
-#else
-		DrawSphere(PDI, InLocation, FVector(4.0f * Scale), 10, 5, Helper_GetMaterialProxy(XYZMaterial), SDPG_Foreground);
-#endif
 
 		PDI->SetHitProxy( NULL );
 	}
diff --git a/Engine/Plugins/Experimental/Mutable/Source/CustomizableObjectEditor/Private/EdGraphSchema_CustomizableObject.cpp b/Engine/Plugins/Experimental/Mutable/Source/CustomizableObjectEditor/Private/EdGraphSchema_CustomizableObject.cpp
--- a/Engine/Plugins/Experimental/Mutable/Source/CustomizableObjectEditor/Private/EdGraphSchema_CustomizableObject.cpp
+++ b/Engine/Plugins/Experimental/Mutable/Source/CustomizableObjectEditor/Private/EdGraphSchema_CustomizableObject.cpp
@@ -519,11 +519,7 @@ const FPinConnectionResponse UEdGraphSchema_CustomizableObject::CanCreateConnect
 	}
 
 	// See if we want to break existing connections (if its an input with an existing connection)
-#if (ENGINE_MAJOR_VERSION == 4 && ENGINE_MINOR_VERSION >= 17) || ENGINE_MAJOR_VERSION >= 5
 	const bool bBreakExistingDueToDataInput = (InputPin->LinkedTo.Num() > 0) && !InputPin->PinType.IsArray();
-#else
-	const bool bBreakExistingDueToDataInput = (InputPin->LinkedTo.Num() > 0) && !InputPin->PinType.bIsArray;
-#endif
 
 	if (bBreakExistingDueToDataOutput&&bBreakExistingDueToDataInput)
 	{
@@ -607,12 +603,7 @@ FLinearColor UEdGraphSchema_CustomizableObject::GetPinTypeColor(const FName& Typ
 
 FLinearColor UEdGraphSchema_CustomizableObject::GetPinTypeColor(const FEdGraphPinType& PinType) const
 {
-#if (ENGINE_MAJOR_VERSION == 4 && ENGINE_MINOR_VERSION <= 18)
-	const FString& TypeString = PinType.PinCategory;
-	const FName& TypeName = FName(*TypeString);
-#else
 	const FName TypeName = PinType.PinCategory;
-#endif
 
 	return GetPinTypeColor(TypeName);
 }
@@ -631,78 +622,6 @@ bool UEdGraphSchema_CustomizableObject::ShouldHidePinDefaultValue(UEdGraphPin* P
 }
 
 
-#if (ENGINE_MAJOR_VERSION == 4 && ENGINE_MINOR_VERSION < 18)
-void UEdGraphSchema_CustomizableObject::GetContextMenuActions(const UEdGraph* CurrentGraph, const UEdGraphNode* InGraphNode, const UEdGraphPin* InGraphPin, FMenuBuilder* MenuBuilder, bool bIsDebugging) const
-{
-	check(CurrentGraph);	
-	//FKismetNodeContextMenuBuilder Context(CurrentGraph, InGraphNode, InGraphPin, MenuBuilder, bIsDebugging);
-
-	if (InGraphPin != NULL)
-	{
-		MenuBuilder->BeginSection("PinActionsMenuHeader");
-		if (!bIsDebugging)
-		{
-			// Break pin links
-			if (InGraphPin->LinkedTo.Num() > 1)
-			{
-				MenuBuilder->AddMenuEntry( FGraphEditorCommands::Get().BreakPinLinks );
-			}
-
-			// add sub menu for break link to
-			if (InGraphPin->LinkedTo.Num() > 0)
-			{
-				MenuBuilder->AddSubMenu(
-					LOCTEXT( "BreakLinkTo", "Break Link To..." ),
-					LOCTEXT("BreakSpecificLinks", "Break a specific link..." ),
-					FNewMenuDelegate::CreateUObject( (UEdGraphSchema_CustomizableObject*const)this, &UEdGraphSchema_CustomizableObject::GetBreakLinkToSubMenuActions, const_cast<UEdGraphPin*>(InGraphPin)));
-			}
-
-		}
-		MenuBuilder->EndSection();
-
-	}
-	else if (InGraphNode != NULL)
-	{
-		if (!bIsDebugging)
-		{
-			// Node contextual actions
-			MenuBuilder->BeginSection("NodeActionsMenuHeader");
-			MenuBuilder->AddMenuEntry( FGenericCommands::Get().Delete );
-			MenuBuilder->AddMenuEntry( FGenericCommands::Get().Cut );
-			MenuBuilder->AddMenuEntry( FGenericCommands::Get().Copy );
-			MenuBuilder->AddMenuEntry( FGenericCommands::Get().Duplicate );
-			MenuBuilder->AddMenuEntry( FGraphEditorCommands::Get().ReconstructNodes );
-			MenuBuilder->AddMenuEntry( FGraphEditorCommands::Get().BreakNodeLinks );
-			// In the case of a UCustomizableObjectNodeObjectGroup, add the option to refresh all Customizable Object Material Node nodes of all the children of this node
-			const UCustomizableObjectNodeObjectGroup* TypedNode = Cast<UCustomizableObjectNodeObjectGroup>(InGraphNode);
-			if (TypedNode != nullptr)
-			{
-				MenuBuilder->AddMenuEntry(FCustomizableObjectEditorNodeContextCommands::Get().RefreshMaterialNodesInAllChildren);
-			}
-			MenuBuilder->EndSection();
-		}
-
-		struct SCommentUtility
-		{
-			static void CreateComment(const UEdGraphSchema_CustomizableObject* Schema, UEdGraph* Graph)
-			{
-				if (Schema && Graph)
-				{
-					Schema->AddComment(Graph, NULL, FVector2D::ZeroVector, true);
-				}
-			}
-		};
-		
-		MenuBuilder->BeginSection("SchemaActionComment", LOCTEXT("MultiCommentHeader", "Comment Group"));
-		MenuBuilder->AddMenuEntry(LOCTEXT("MultiCommentDesc", "Create Comment from Selection"),
-			LOCTEXT("CommentToolTip", "Create a resizable comment box around selection."),
-			FSlateIcon(),
-			FUIAction(FExecuteAction::CreateStatic(SCommentUtility::CreateComment, this, const_cast<UEdGraph*>(CurrentGraph))));
-		MenuBuilder->EndSection();
-	}
-}
-
-#else
 void UEdGraphSchema_CustomizableObject::GetContextMenuActions(class UToolMenu* Menu, class UGraphNodeContextMenuContext* Context) const
 {
 	if (Context && !Context->Pin && Context->Node)
@@ -741,14 +660,10 @@ void UEdGraphSchema_CustomizableObject::GetContextMenuActions(class UToolMenu* M
 		section.AddMenuEntry("MultiCommentDesc", LOCTEXT("MultiCommentDesc", "Create Comment from Selection"),
 			LOCTEXT("CommentToolTip", "Create a resizable comment box around selection."),
 			FSlateIcon(),
-#if ENGINE_MAJOR_VERSION >= 5
 			FUIAction(FExecuteAction::CreateStatic(SCommentUtility::CreateComment, this, const_cast<UEdGraph*>(ToRawPtr(Context->Graph)))));
-#else
-			FUIAction(FExecuteAction::CreateStatic(SCommentUtility::CreateComment, this, const_cast<UEdGraph*>(Context->Graph))));
-#endif
 	}
 }
-#endif
+
 
 
 void UEdGraphSchema_CustomizableObject::BreakPinLinks(UEdGraphPin& TargetPin, bool bSendsNodeNotification) const
diff --git a/Engine/Plugins/Experimental/Mutable/Source/CustomizableObjectEditor/Private/GenerateMutableSource/GenerateMutableSourceMesh.cpp b/Engine/Plugins/Experimental/Mutable/Source/CustomizableObjectEditor/Private/GenerateMutableSource/GenerateMutableSourceMesh.cpp
--- a/Engine/Plugins/Experimental/Mutable/Source/CustomizableObjectEditor/Private/GenerateMutableSource/GenerateMutableSourceMesh.cpp
+++ b/Engine/Plugins/Experimental/Mutable/Source/CustomizableObjectEditor/Private/GenerateMutableSource/GenerateMutableSourceMesh.cpp
@@ -2534,7 +2534,7 @@ mu::NodeMeshPtr GenerateMutableSourceMesh(const UEdGraphPin * Pin,
 			}
 			else
 			{
-				GenerationContext.Compiler->CompilerLog(LOCTEXT("MeshFailed", "Mesh generation failed."), Node);
+				GenerationContext.Compiler->CompilerLog(LOCTEXT("MeshGenerationFailed", "Mesh generation failed."), Node);
 			}
 		}
 		else
@@ -2553,7 +2553,7 @@ mu::NodeMeshPtr GenerateMutableSourceMesh(const UEdGraphPin * Pin,
 			}
 			else
 			{
-				GenerationContext.Compiler->CompilerLog(LOCTEXT("MeshFailed", "Mesh generation failed."), Node);
+				GenerationContext.Compiler->CompilerLog(LOCTEXT("MeshGenerationFailed", "Mesh generation failed."), Node);
 			}
 		}
 
@@ -2566,7 +2566,7 @@ mu::NodeMeshPtr GenerateMutableSourceMesh(const UEdGraphPin * Pin,
 			}
 			else
 			{
-				GenerationContext.Compiler->CompilerLog(LOCTEXT("MeshFailed", "Scalar generation failed."), Node);
+				GenerationContext.Compiler->CompilerLog(LOCTEXT("ScalarGenerationFailed", "Scalar generation failed."), Node);
 			}
 		}
 
@@ -2579,7 +2579,7 @@ mu::NodeMeshPtr GenerateMutableSourceMesh(const UEdGraphPin * Pin,
 			}
 			else
 			{
-				GenerationContext.Compiler->CompilerLog(LOCTEXT("MeshFailed", "Scalar generation failed."), Node);
+				GenerationContext.Compiler->CompilerLog(LOCTEXT("ScalarGenerationFailed", "Scalar generation failed."), Node);
 			}
 		}
 	}
diff --git a/Engine/Plugins/Experimental/Mutable/Source/CustomizableObjectEditor/Private/GenerateMutableSource/GenerateMutableSourceSurface.cpp b/Engine/Plugins/Experimental/Mutable/Source/CustomizableObjectEditor/Private/GenerateMutableSource/GenerateMutableSourceSurface.cpp
--- a/Engine/Plugins/Experimental/Mutable/Source/CustomizableObjectEditor/Private/GenerateMutableSource/GenerateMutableSourceSurface.cpp
+++ b/Engine/Plugins/Experimental/Mutable/Source/CustomizableObjectEditor/Private/GenerateMutableSource/GenerateMutableSourceSurface.cpp
@@ -946,7 +946,7 @@ mu::NodeSurfacePtr GenerateMutableSourceSurface(const UEdGraphPin * Pin, FMutabl
 				CopyMaterialParentMaterialNode = ParentMaterialCopyNodeCast->GetMaterialNode();
 				if (!CopyMaterialParentMaterialNode)
 				{
-					GenerationContext.Compiler->CompilerLog(LOCTEXT("ParentMissing", "Base Material not set (or not found)."), ParentMaterialCopyNodeCast);
+					GenerationContext.Compiler->CompilerLog(LOCTEXT("BaseMissing", "Base Material not set (or not found)."), ParentMaterialCopyNodeCast);
 					return;
 				}
 			}
@@ -1157,7 +1157,7 @@ mu::NodeSurfacePtr GenerateMutableSourceSurface(const UEdGraphPin * Pin, FMutabl
 
 				if (!Layouts.IsValidIndex(TypedNodeRemBlocks->ParentLayoutIndex))
 				{
-					GenerationContext.Compiler->CompilerLog(LOCTEXT("Invalid layout", "Node refers to an invalid texture layout. Add a layout to the selected parent material, then select some block(s)."), Node);
+					GenerationContext.Compiler->CompilerLog(LOCTEXT("RemoveMeshNode Invalid layout", "Node refers to an invalid texture layout. Add a layout to the selected parent material, then select some block(s)."), Node);
 					UE_LOG(LogMutable, Warning, TEXT("In object [%s] UCustomizableObjectNodeRemoveMeshBlocks refers to an invalid texture layout index %d. Parent node has %d layouts."),
 						*GenerationContext.Object->GetName(), TypedNodeRemBlocks->ParentLayoutIndex, Layouts.Num());
 				}
@@ -1176,7 +1176,7 @@ mu::NodeSurfacePtr GenerateMutableSourceSurface(const UEdGraphPin * Pin, FMutabl
 							}
 							else
 							{
-								GenerationContext.Compiler->CompilerLog(LOCTEXT("InvalidLayoutBlock", "Node refers to an invalid layout block."), Node);
+								GenerationContext.Compiler->CompilerLog(LOCTEXT("RemoveMeshNode InvalidLayoutBlock", "Node refers to an invalid layout block."), Node);
 								UE_LOG(LogMutable, Warning, TEXT("In object [%s] UCustomizableObjectNodeRemoveMeshBlocks refers to an invalid layout block id %s. Parent node has %d blocks."),
 									*GenerationContext.Object->GetName(), *Id.ToString(), Layout->Blocks.Num());
 							}
@@ -1258,7 +1258,7 @@ mu::NodeSurfacePtr GenerateMutableSourceSurface(const UEdGraphPin * Pin, FMutabl
 
 					if (!Layouts.IsValidIndex(TypedNodeEdit->ParentLayoutIndex))
 					{
-						GenerationContext.Compiler->CompilerLog(LOCTEXT("Invalid layout", "Node refers to an invalid texture layout."), Node);
+						GenerationContext.Compiler->CompilerLog(LOCTEXT("EditNode Invalid layout", "Node refers to an invalid texture layout."), Node);
 						UE_LOG(LogMutable, Warning, TEXT("In object [%s] UCustomizableObjectNodeEditMaterial refers to an invalid texture layout index %d. Parent node has %d layouts."),
 							*GenerationContext.Object->GetName(), TypedNodeEdit->ParentLayoutIndex, Layouts.Num());
 					}
@@ -1279,7 +1279,7 @@ mu::NodeSurfacePtr GenerateMutableSourceSurface(const UEdGraphPin * Pin, FMutabl
 									}
 									else
 									{
-										GenerationContext.Compiler->CompilerLog(LOCTEXT("InvalidLayoutBlock", "Node refers to an invalid layout block."), Node);
+										GenerationContext.Compiler->CompilerLog(LOCTEXT("EditNode InvalidLayoutBlock", "Node refers to an invalid layout block."), Node);
 										UE_LOG(LogMutable, Warning, TEXT("In object [%s] UCustomizableObjectNodeEditMaterial refers to an invalid layout block id %s. Parent node has %d blocks."),
 											*GenerationContext.Object->GetName(), *Id.ToString(), Layout->Blocks.Num());
 									}
diff --git a/Engine/Plugins/Experimental/Mutable/Source/CustomizableObjectEditor/Private/Nodes/CustomizableObjectNodeAnimationPose.cpp b/Engine/Plugins/Experimental/Mutable/Source/CustomizableObjectEditor/Private/Nodes/CustomizableObjectNodeAnimationPose.cpp
--- a/Engine/Plugins/Experimental/Mutable/Source/CustomizableObjectEditor/Private/Nodes/CustomizableObjectNodeAnimationPose.cpp
+++ b/Engine/Plugins/Experimental/Mutable/Source/CustomizableObjectEditor/Private/Nodes/CustomizableObjectNodeAnimationPose.cpp
@@ -118,7 +118,6 @@ void UCustomizableObjectNodeAnimationPose::StaticRetrievePoseInformation(UPoseAs
 	ExtractionContext.bExtractRootMotion = false;
 	ExtractionContext.CurrentTime = 0.0f;
 
-#if (ENGINE_MAJOR_VERSION == 4 && ENGINE_MINOR_VERSION >= 24) || ENGINE_MAJOR_VERSION >= 5
 	OutCurve.UIDToArrayIndexLUT = AnimCurves.UIDToArrayIndexLUT;
 
 	const TArray<FSmartName>& PoseNames = PoseAsset->GetPoseNames();
@@ -127,12 +126,6 @@ void UCustomizableObjectNodeAnimationPose::StaticRetrievePoseInformation(UPoseAs
 	PoseAsset->GetAnimationPose(SecondOutAnimData, ExtractionContext);
 	OutPose = OutAnimData.GetPose();
 	OutCurve = OutAnimData.GetCurve();
-#else
-	OutCurve.UIDList = AnimCurves.UIDList;
-
-	ExtractionContext.PoseCurves.Add(1.0f);
-	PoseAsset->GetAnimationPose(OutPose, OutCurve, ExtractionContext);
-#endif
 
 	const TArray<FTransform, FAnimStackAllocator>& ArrayPoseBoneTransform = OutPose.GetBones();
 	const TArray<FBoneIndexType>& ArrayPoseBoneIndices = OutPose.GetBoneContainer().GetBoneIndicesArray();
diff --git a/Engine/Plugins/Experimental/Mutable/Source/CustomizableObjectEditor/Private/Nodes/CustomizableObjectNodeLayoutBlocksDetails.cpp b/Engine/Plugins/Experimental/Mutable/Source/CustomizableObjectEditor/Private/Nodes/CustomizableObjectNodeLayoutBlocksDetails.cpp
--- a/Engine/Plugins/Experimental/Mutable/Source/CustomizableObjectEditor/Private/Nodes/CustomizableObjectNodeLayoutBlocksDetails.cpp
+++ b/Engine/Plugins/Experimental/Mutable/Source/CustomizableObjectEditor/Private/Nodes/CustomizableObjectNodeLayoutBlocksDetails.cpp
@@ -26,19 +26,11 @@ TSharedRef<IDetailCustomization> FCustomizableObjectNodeLayoutBlocksDetails::Mak
 void FCustomizableObjectNodeLayoutBlocksDetails::CustomizeDetails( IDetailLayoutBuilder& DetailBuilder )
 {
 	Node = 0;
-#if (ENGINE_MAJOR_VERSION == 4 && ENGINE_MINOR_VERSION >= 18) || ENGINE_MAJOR_VERSION >= 5
 	const IDetailsView* DetailsView = DetailBuilder.GetDetailsView();
 	if (DetailsView->GetSelectedObjects().Num())
 	{
 		Node = Cast<UCustomizableObjectNodeLayoutBlocks>(DetailsView->GetSelectedObjects()[0].Get());
 	}
-#else
-	const IDetailsView& DetailsView = DetailBuilder.GetDetailsView();
-	if ( DetailsView.GetSelectedObjects().Num() )
-	{
-		Node = Cast<UCustomizableObjectNodeLayoutBlocks>( DetailsView.GetSelectedObjects()[0].Get() );
-	}
-#endif
 
 	IDetailCategoryBuilder& BlocksCategory = DetailBuilder.EditCategory( "Blocks" );
 
diff --git a/Engine/Plugins/Experimental/Mutable/Source/CustomizableObjectEditor/Private/Nodes/CustomizableObjectNodeMeshClipMorphDetails.cpp b/Engine/Plugins/Experimental/Mutable/Source/CustomizableObjectEditor/Private/Nodes/CustomizableObjectNodeMeshClipMorphDetails.cpp
--- a/Engine/Plugins/Experimental/Mutable/Source/CustomizableObjectEditor/Private/Nodes/CustomizableObjectNodeMeshClipMorphDetails.cpp
+++ b/Engine/Plugins/Experimental/Mutable/Source/CustomizableObjectEditor/Private/Nodes/CustomizableObjectNodeMeshClipMorphDetails.cpp
@@ -38,19 +38,11 @@ void FCustomizableObjectNodeMeshClipMorphDetails::CustomizeDetails(IDetailLayout
 	Node = nullptr;
 	DetailBuilderPtr = &DetailBuilder;
 
-#if (ENGINE_MAJOR_VERSION == 4 && ENGINE_MINOR_VERSION >= 18) || ENGINE_MAJOR_VERSION >= 5
 	const IDetailsView* DetailsView = DetailBuilder.GetDetailsView();
 	if (DetailsView->GetSelectedObjects().Num())
 	{
 		Node = Cast<UCustomizableObjectNodeMeshClipMorph>(DetailsView->GetSelectedObjects()[0].Get());
 	}
-#else
-	const IDetailsView& DetailsView = DetailBuilder.GetDetailsView();
-	if (DetailsView.GetSelectedObjects().Num())
-	{
-		Node = Cast<UCustomizableObjectNodeMeshClipMorph>(DetailsView.GetSelectedObjects()[0].Get());
-	}
-#endif
 
 	IDetailCategoryBuilder& BlocksCategory = DetailBuilder.EditCategory("MeshToClipAndMorph");
 	DetailBuilder.HideProperty("BoneName");
diff --git a/Engine/Plugins/Experimental/Mutable/Source/CustomizableObjectEditor/Private/Nodes/CustomizableObjectNodeMeshClipWithMeshDetails.cpp b/Engine/Plugins/Experimental/Mutable/Source/CustomizableObjectEditor/Private/Nodes/CustomizableObjectNodeMeshClipWithMeshDetails.cpp
--- a/Engine/Plugins/Experimental/Mutable/Source/CustomizableObjectEditor/Private/Nodes/CustomizableObjectNodeMeshClipWithMeshDetails.cpp
+++ b/Engine/Plugins/Experimental/Mutable/Source/CustomizableObjectEditor/Private/Nodes/CustomizableObjectNodeMeshClipWithMeshDetails.cpp
@@ -38,19 +38,11 @@ void FCustomizableObjectNodeMeshClipWithMeshDetails::CustomizeDetails(IDetailLay
 	Node = nullptr;
 	DetailBuilderPtr = &DetailBuilder;
 
-#if (ENGINE_MAJOR_VERSION == 4 && ENGINE_MINOR_VERSION >= 18) || ENGINE_MAJOR_VERSION >= 5
 	const IDetailsView* DetailsView = DetailBuilder.GetDetailsView();
 	if (DetailsView->GetSelectedObjects().Num())
 	{
 		Node = Cast<UCustomizableObjectNodeMeshClipWithMesh>(DetailsView->GetSelectedObjects()[0].Get());
 	}
-#else
-	const IDetailsView& DetailsView = DetailBuilder.GetDetailsView();
-	if (DetailsView.GetSelectedObjects().Num())
-	{
-		Node = Cast<UCustomizableObjectNodeMeshClipWithMesh>(DetailsView.GetSelectedObjects()[0].Get());
-	}
-#endif
 
 	IDetailCategoryBuilder& CustomizableObjectToClipCategory = DetailBuilder.EditCategory("CustomizableObjectToClip");
 
diff --git a/Engine/Plugins/Experimental/Mutable/Source/CustomizableObjectEditor/Private/Nodes/CustomizableObjectNodeMorphMaterialDetails.cpp b/Engine/Plugins/Experimental/Mutable/Source/CustomizableObjectEditor/Private/Nodes/CustomizableObjectNodeMorphMaterialDetails.cpp
--- a/Engine/Plugins/Experimental/Mutable/Source/CustomizableObjectEditor/Private/Nodes/CustomizableObjectNodeMorphMaterialDetails.cpp
+++ b/Engine/Plugins/Experimental/Mutable/Source/CustomizableObjectEditor/Private/Nodes/CustomizableObjectNodeMorphMaterialDetails.cpp
@@ -39,19 +39,11 @@ TSharedRef<IDetailCustomization> FCustomizableObjectNodeMorphMaterialDetails::Ma
 void FCustomizableObjectNodeMorphMaterialDetails::CustomizeDetails( IDetailLayoutBuilder& DetailBuilder )
 {
 
-#if (ENGINE_MAJOR_VERSION == 4 && ENGINE_MINOR_VERSION >= 18) || ENGINE_MAJOR_VERSION >= 5
 	const IDetailsView* DetailsView = DetailBuilder.GetDetailsView();
 	if ( DetailsView->GetSelectedObjects().Num() )
 	{
 		Node = Cast<UCustomizableObjectNodeMorphMaterial>( DetailsView->GetSelectedObjects()[0].Get() );
 	}
-#else
-	const IDetailsView& DetailsView = DetailBuilder.GetDetailsView();
-	if (DetailsView.GetSelectedObjects().Num())
-	{
-		Node = Cast<UCustomizableObjectNodeMorphMaterial>(DetailsView.GetSelectedObjects()[0].Get());
-	}
-#endif
 
 	IDetailCategoryBuilder& BlocksCategory = DetailBuilder.EditCategory( "Customizable Object" );
 	//BlocksCategory.CategoryIcon( "ActorClassIcon.CustomizableObject" );
diff --git a/Engine/Plugins/Experimental/Mutable/Source/CustomizableObjectEditor/Private/Nodes/CustomizableObjectNodeObjectDetails.cpp b/Engine/Plugins/Experimental/Mutable/Source/CustomizableObjectEditor/Private/Nodes/CustomizableObjectNodeObjectDetails.cpp
--- a/Engine/Plugins/Experimental/Mutable/Source/CustomizableObjectEditor/Private/Nodes/CustomizableObjectNodeObjectDetails.cpp
+++ b/Engine/Plugins/Experimental/Mutable/Source/CustomizableObjectEditor/Private/Nodes/CustomizableObjectNodeObjectDetails.cpp
@@ -38,19 +38,11 @@ void FCustomizableObjectNodeObjectDetails::CustomizeDetails( IDetailLayoutBuilde
 	Node = nullptr;
 	DetailBuilderPtr = &DetailBuilder;
 
-#if (ENGINE_MAJOR_VERSION == 4 && ENGINE_MINOR_VERSION >= 18) || ENGINE_MAJOR_VERSION >= 5
 	const IDetailsView* DetailsView = DetailBuilder.GetDetailsView();
 	if ( DetailsView->GetSelectedObjects().Num() )
 	{
 		Node = Cast<UCustomizableObjectNodeObject>( DetailsView->GetSelectedObjects()[0].Get() );
 	}
-#else
-	const IDetailsView& DetailsView = DetailBuilder.GetDetailsView();
-	if (DetailsView.GetSelectedObjects().Num())
-	{
-		Node = Cast<UCustomizableObjectNodeObject>(DetailsView.GetSelectedObjects()[0].Get());
-	}
-#endif
 
 	IDetailCategoryBuilder& StatesCategory = DetailBuilder.EditCategory("States");
 	IDetailCategoryBuilder& RuntimeParameters = DetailBuilder.EditCategory("States Runtime Parameters");
diff --git a/Engine/Plugins/Experimental/Mutable/Source/CustomizableObjectEditor/Private/Nodes/CustomizableObjectNodeProjectorParameterDetails.cpp b/Engine/Plugins/Experimental/Mutable/Source/CustomizableObjectEditor/Private/Nodes/CustomizableObjectNodeProjectorParameterDetails.cpp
--- a/Engine/Plugins/Experimental/Mutable/Source/CustomizableObjectEditor/Private/Nodes/CustomizableObjectNodeProjectorParameterDetails.cpp
+++ b/Engine/Plugins/Experimental/Mutable/Source/CustomizableObjectEditor/Private/Nodes/CustomizableObjectNodeProjectorParameterDetails.cpp
@@ -20,9 +20,7 @@
 #include "Math/Vector.h"
 
 #include "Runtime/Launch/Resources/Version.h"
-#if (ENGINE_MAJOR_VERSION == 4 && ENGINE_MINOR_VERSION > 22) || ENGINE_MAJOR_VERSION >= 5
 #include "HAL/PlatformApplicationMisc.h"
-#endif
 
 
 #define LOCTEXT_NAMESPACE "CustomizableObjectDetails"
@@ -40,7 +38,6 @@ void FCustomizableObjectNodeProjectorParameterDetails::CustomizeDetails( IDetail
 	NodeParameter = nullptr;
 	DetailBuilderPtr = &DetailBuilder;
 
-#if (ENGINE_MAJOR_VERSION == 4 && ENGINE_MINOR_VERSION >= 18) || ENGINE_MAJOR_VERSION >= 5
 	const IDetailsView* DetailsView = DetailBuilder.GetDetailsView();
 	if ( DetailsView->GetSelectedObjects().Num() )
 	{
@@ -53,20 +50,6 @@ void FCustomizableObjectNodeProjectorParameterDetails::CustomizeDetails( IDetail
 			NodeParameter = Cast<UCustomizableObjectNodeProjectorParameter>(DetailsView->GetSelectedObjects()[0].Get());
 		}
 	}
-#else
-	const IDetailsView& DetailsView = DetailBuilder.GetDetailsView();
-	if (DetailsView.GetSelectedObjects().Num())
-	{
-		if (DetailsView.GetSelectedObjects()[0].Get()->IsA(UCustomizableObjectNodeProjectorConstant::StaticClass()))
-		{
-			NodeConstant = Cast<UCustomizableObjectNodeProjectorConstant>(DetailsView.GetSelectedObjects()[0].Get());
-		}
-		else if (DetailsView.GetSelectedObjects()[0].Get()->IsA(UCustomizableObjectNodeProjectorParameter::StaticClass()))
-		{
-			NodeParameter = Cast<UCustomizableObjectNodeProjectorParameter>(DetailsView.GetSelectedObjects()[0].Get());
-		}
-	}
-#endif
 
 	IDetailCategoryBuilder& BlocksCategory = DetailBuilder.EditCategory( "Clipboard" );
 	//BlocksCategory.CategoryIcon( "ActorClassIcon.CustomizableObject" );
@@ -210,11 +193,7 @@ FReply FCustomizableObjectNodeProjectorParameterDetails::OnProjectorCopyPressed(
         Struct->ExportText(ExportedText, &NodeConstant->Value, nullptr, nullptr, (PPF_ExportsNotFullyQualified | PPF_Copy | PPF_Delimited | PPF_IncludeTransient), nullptr);
 	}
 
-#if (ENGINE_MAJOR_VERSION == 4 && ENGINE_MINOR_VERSION > 22) || ENGINE_MAJOR_VERSION >= 5
 	FPlatformApplicationMisc::ClipboardCopy(*ExportedText);
-#else
-	FPlatformMisc::ClipboardCopy(*ExportedText);
-#endif
 
 	return FReply::Handled();
 }
@@ -224,11 +203,7 @@ FReply FCustomizableObjectNodeProjectorParameterDetails::OnProjectorPastePressed
 {
 	FString ClipText;
 
-#if (ENGINE_MAJOR_VERSION == 4 && ENGINE_MINOR_VERSION > 22) || ENGINE_MAJOR_VERSION >= 5
 	FPlatformApplicationMisc::ClipboardPaste(ClipText);
-#else
-	FPlatformMisc::ClipboardPaste(ClipText);
-#endif
 
 	TSharedPtr<ICustomizableObjectEditor> Editor = nullptr;
 
diff --git a/Engine/Plugins/Experimental/Mutable/Source/CustomizableObjectEditor/Private/SCustomizableObjectEditorViewport.cpp b/Engine/Plugins/Experimental/Mutable/Source/CustomizableObjectEditor/Private/SCustomizableObjectEditorViewport.cpp
--- a/Engine/Plugins/Experimental/Mutable/Source/CustomizableObjectEditor/Private/SCustomizableObjectEditorViewport.cpp
+++ b/Engine/Plugins/Experimental/Mutable/Source/CustomizableObjectEditor/Private/SCustomizableObjectEditorViewport.cpp
@@ -1139,11 +1139,7 @@ void SCustomizableObjectEditorViewportTabBody::OnSetLODModel(int32 LODSelectionT
 	{
 		if (PreviewComponent)
 		{
-#if (ENGINE_MAJOR_VERSION == 4 && ENGINE_MINOR_VERSION >= 24) || ENGINE_MAJOR_VERSION >= 5
 			PreviewComponent->SetForcedLOD(LODSelection);
-#else
-			PreviewComponent->ForcedLodModel = LODSelection;
-#endif
 			//PopulateUVChoices();
 		}
 	}
diff --git a/Engine/Plugins/Experimental/Mutable/Source/CustomizableObjectEditor/Private/SCustomizableObjectEditorViewportToolBar.cpp b/Engine/Plugins/Experimental/Mutable/Source/CustomizableObjectEditor/Private/SCustomizableObjectEditorViewportToolBar.cpp
--- a/Engine/Plugins/Experimental/Mutable/Source/CustomizableObjectEditor/Private/SCustomizableObjectEditorViewportToolBar.cpp
+++ b/Engine/Plugins/Experimental/Mutable/Source/CustomizableObjectEditor/Private/SCustomizableObjectEditorViewportToolBar.cpp
@@ -423,10 +423,6 @@ void SCustomizableObjectEditorViewportToolBar::Construct(const FArguments& InArg
 			[
 				SNew(SBorder)
 				.BorderImage(FAppStyle::GetBrush("NoBorder"))
-#if ENGINE_MAJOR_VERSION < 5
-				// Color and opacity is changed based on whether or not the mouse cursor is hovering over the toolbar area
-				.ColorAndOpacity(this, &SViewportToolBar::OnGetColorAndOpacity)
-#endif
 				.ForegroundColor(FAppStyle::GetSlateColor(DefaultForegroundName))
 				[
 					SNew(SVerticalBox)
@@ -1259,16 +1255,11 @@ void SCustomizableObjectEditorViewportToolBar::OnFOVValueChanged(float NewValue)
 {
 	TSharedPtr<FCustomizableObjectEditorViewportClient> ViewportClient = Viewport.Pin()->GetViewportClient();
 
-#if (ENGINE_MAJOR_VERSION == 4 && ENGINE_MINOR_VERSION >= 18) || ENGINE_MAJOR_VERSION >= 5
 	ViewportClient->FOVAngle = NewValue;
 	// \todo: this editor name should be somewhere else.
 	FString EditorName("CustomizableObjectEditor");
 	int ViewportIndex=0;
 	ViewportClient->ConfigOption->SetViewFOV(FName(*EditorName),NewValue,ViewportIndex);
-#else
-	ViewportClient->FOVAngle = NewValue;
-	ViewportClient->ConfigOption->SetViewFOV(NewValue);
-#endif
 
 	ViewportClient->ViewFOV = NewValue;
 	ViewportClient->Invalidate();
diff --git a/Engine/Plugins/Experimental/Mutable/Source/CustomizableObjectEditor/Private/SCustomizableObjectHighresScreenshot.cpp b/Engine/Plugins/Experimental/Mutable/Source/CustomizableObjectEditor/Private/SCustomizableObjectHighresScreenshot.cpp
--- a/Engine/Plugins/Experimental/Mutable/Source/CustomizableObjectEditor/Private/SCustomizableObjectHighresScreenshot.cpp
+++ b/Engine/Plugins/Experimental/Mutable/Source/CustomizableObjectEditor/Private/SCustomizableObjectHighresScreenshot.cpp
@@ -195,7 +195,7 @@ void SCustomizableObjectHighresScreenshot::Construct( const FArguments& InArgs )
 					.AutoWidth()
 					[
 						SNew( SButton )
-						.ToolTipText( NSLOCTEXT("HighResScreenshot", "ScreenshotCaptureTooltop", "Take a screenshot") )
+						.ToolTipText( NSLOCTEXT("CustomizableObjectEditor", "ScreenshotCaptureTooltop", "Take a screenshot") )
 						.OnClicked( this, &SCustomizableObjectHighresScreenshot::OnCaptureClicked )
 						[
 							SNew(SImage)
@@ -261,13 +261,8 @@ void SCustomizableObjectHighresScreenshot::OnShowFloorChanged(ECheckBoxState New
 	{
 		if (NewValue == ECheckBoxState::Unchecked)
 		{
-#if (ENGINE_MAJOR_VERSION == 4 && ENGINE_MINOR_VERSION >= 18) || ENGINE_MAJOR_VERSION >= 5
 			ExposureLogOffsetCached = LevelViewportClient.Pin()->ExposureSettings.FixedEV100;
 			LevelViewportClient.Pin()->ExposureSettings.FixedEV100 = 0;
-#else
-			ExposureLogOffsetCached = LevelViewportClient.Pin()->ExposureSettings.LogOffset;
-			LevelViewportClient.Pin()->ExposureSettings.LogOffset = 0;
-#endif
 
 			ExposureLogFixedCached = LevelViewportClient.Pin()->ExposureSettings.bFixed;
 			LevelViewportClient.Pin()->ExposureSettings.bFixed = true;
@@ -281,21 +276,12 @@ void SCustomizableObjectHighresScreenshot::OnShowFloorChanged(ECheckBoxState New
 		}
 		else
 		{
-#if (ENGINE_MAJOR_VERSION == 4 && ENGINE_MINOR_VERSION >= 18) || ENGINE_MAJOR_VERSION >= 5
 			if ((LevelViewportClient.Pin()->ExposureSettings.FixedEV100 == 0) &&
 				(LevelViewportClient.Pin()->ExposureSettings.bFixed == true))
 			{
 				LevelViewportClient.Pin()->ExposureSettings.FixedEV100 = ExposureLogOffsetCached;
 				LevelViewportClient.Pin()->ExposureSettings.bFixed = ExposureLogFixedCached;
 			}
-#else
-			if ((LevelViewportClient.Pin()->ExposureSettings.LogOffset == 0) &&
-				(LevelViewportClient.Pin()->ExposureSettings.bFixed == true))
-			{
-				LevelViewportClient.Pin()->ExposureSettings.LogOffset = ExposureLogOffsetCached;
-				LevelViewportClient.Pin()->ExposureSettings.bFixed = ExposureLogFixedCached;
-			}
-#endif
 
 			LevelViewportClient.Pin()->SetFloorVisibility(FloorVisibilityCached);
 			LevelViewportClient.Pin()->SetShowGrid();
@@ -311,7 +297,6 @@ void SCustomizableObjectHighresScreenshot::RestoreViewportValues()
 	if (LevelViewportClient.IsValid())
 	{
 		// Restore exposure values if modified with the "Show / hide background" option
-#if (ENGINE_MAJOR_VERSION == 4 && ENGINE_MINOR_VERSION >= 18) || ENGINE_MAJOR_VERSION >= 5
 		if ((LevelViewportClient.Pin()->ExposureSettings.FixedEV100 == 0) &&
 			(LevelViewportClient.Pin()->ExposureSettings.bFixed == true) &&
 			(ExposureLogOffsetCached != -1))
@@ -320,17 +305,6 @@ void SCustomizableObjectHighresScreenshot::RestoreViewportValues()
 			LevelViewportClient.Pin()->ExposureSettings.bFixed = ExposureLogFixedCached;
 		}
 
-#else
-		if ((LevelViewportClient.Pin()->ExposureSettings.LogOffset == 0) &&
-			(LevelViewportClient.Pin()->ExposureSettings.bFixed == true) &&
-			(ExposureLogOffsetCached != -1))
-		{
-			LevelViewportClient.Pin()->ExposureSettings.LogOffset = ExposureLogOffsetCached;
-			LevelViewportClient.Pin()->ExposureSettings.bFixed = ExposureLogFixedCached;
-		}
-
-#endif
-
 		// Restore floor and environment mesh visibility
 		if (!LevelViewportClient.Pin()->GetEnvironmentMeshVisibility())
 		{
@@ -393,7 +367,7 @@ TWeakPtr<class SWindow> SCustomizableObjectHighresScreenshot::OpenDialog(
 		// No dialog window currently open - need to create one
 		TSharedRef<SCustomizableObjectHighresScreenshot> Dialog = SNew(SCustomizableObjectHighresScreenshot);
 		TSharedRef<SWindow> Window = SNew(SWindow)
-			.Title( NSLOCTEXT("HighResScreenshot", "HighResolutionScreenshot", "High Resolution Screenshot") )
+			.Title( NSLOCTEXT("CustomizableObjectEditor", "HighResolutionScreenshot", "High Resolution Screenshot") )
 			.ClientSize(FVector2D(484,231))
 			.SupportsMinimize(false)
 			.SupportsMaximize(false)
diff --git a/Engine/Plugins/Experimental/Mutable/Source/CustomizableObjectEditor/Private/SCustomizableObjectNodeLayoutBlocksEditor.cpp b/Engine/Plugins/Experimental/Mutable/Source/CustomizableObjectEditor/Private/SCustomizableObjectNodeLayoutBlocksEditor.cpp
--- a/Engine/Plugins/Experimental/Mutable/Source/CustomizableObjectEditor/Private/SCustomizableObjectNodeLayoutBlocksEditor.cpp
+++ b/Engine/Plugins/Experimental/Mutable/Source/CustomizableObjectEditor/Private/SCustomizableObjectNodeLayoutBlocksEditor.cpp
@@ -41,7 +41,7 @@ public:
 	FLayoutEditorCommands() : TCommands<FLayoutEditorCommands>
 	(
 		"LayoutEditorCommands", // Context name for fast lookup
-		NSLOCTEXT( "Contexts", "LayoutEditorCommands", "Layout Editor" ), // Localized context name for displaying
+		NSLOCTEXT( "CustomizableObjectEditor", "LayoutEditorCommands", "Layout Editor" ), // Localized context name for displaying
 		NAME_None, // Parent
 		FCustomizableObjectEditorStyle::GetStyleSetName()
 	)
diff --git a/Engine/Plugins/Experimental/Mutable/Source/CustomizableObjectEditor/Private/SCustomizableObjectNodeLayoutBlocksSelector.cpp b/Engine/Plugins/Experimental/Mutable/Source/CustomizableObjectEditor/Private/SCustomizableObjectNodeLayoutBlocksSelector.cpp
--- a/Engine/Plugins/Experimental/Mutable/Source/CustomizableObjectEditor/Private/SCustomizableObjectNodeLayoutBlocksSelector.cpp
+++ b/Engine/Plugins/Experimental/Mutable/Source/CustomizableObjectEditor/Private/SCustomizableObjectNodeLayoutBlocksSelector.cpp
@@ -33,7 +33,7 @@ public:
 	FLayoutBlockSelectorCommands() : TCommands<FLayoutBlockSelectorCommands>
 	(
 		"LayoutBlockSelector", // Context name for fast lookup
-		NSLOCTEXT( "Contexts", "LayoutBlockSelector", "Layout Block Selector" ), // Localized context name for displaying
+		NSLOCTEXT( "CustomizableObjectEditor", "LayoutBlockSelector", "Layout Block Selector" ), // Localized context name for displaying
 		NAME_None, // Parent
 		FCustomizableObjectEditorStyle::GetStyleSetName()
 	)
diff --git a/Engine/Plugins/Experimental/Mutable/Source/CustomizableObjectEditor/Private/SMutableCodeViewer.cpp b/Engine/Plugins/Experimental/Mutable/Source/CustomizableObjectEditor/Private/SMutableCodeViewer.cpp
--- a/Engine/Plugins/Experimental/Mutable/Source/CustomizableObjectEditor/Private/SMutableCodeViewer.cpp
+++ b/Engine/Plugins/Experimental/Mutable/Source/CustomizableObjectEditor/Private/SMutableCodeViewer.cpp
@@ -1453,8 +1453,8 @@ TSharedPtr<SWidget> SMutableCodeViewer::OnTreeContextMenuOpening()
 		if (TreeView->GetSelectedItems().Num() == 1)
 		{
 			MenuBuilder.AddMenuEntry(
-		LOCTEXT("Set_as_search_operation_type","Set as search Operation"),
-		LOCTEXT("Set_as_search_operation_type_tip", "Sets the type of this operation as the type to be looking for when searching for operations on the tree view"),
+				LOCTEXT("Set_as_search_operation_type","Set as search Operation"),
+				LOCTEXT("Set_as_search_operation_type_Tooltip", "Sets the type of this operation as the type to be looking for when searching for operations on the tree view"),
 				FSlateIcon(),
 				FUIAction(FExecuteAction::CreateSP(this, &SMutableCodeViewer::OnSelectedOperationTypeFromTree)
 				)	
@@ -1464,8 +1464,8 @@ TSharedPtr<SWidget> SMutableCodeViewer::OnTreeContextMenuOpening()
 		MenuBuilder.AddMenuSeparator();
 		
 		MenuBuilder.AddMenuEntry(
-			LOCTEXT("Expand_Selected", "Expand Selected Operation"),
-			LOCTEXT("Expand_Selected", "Expands only the selected Operation and leaves the other as they are."),
+			LOCTEXT("Code_Expand_Selected", "Expand Selected Operation"),
+			LOCTEXT("Code_Expand_Selected_Tooltip", "Expands only the selected Operation and leaves the other as they are."),
 			FSlateIcon(),
 			FUIAction(FExecuteAction::CreateSP(this, &SMutableCodeViewer::TreeExpandSelected)
 			)
@@ -1474,8 +1474,8 @@ TSharedPtr<SWidget> SMutableCodeViewer::OnTreeContextMenuOpening()
 
 
 	MenuBuilder.AddMenuEntry(
-		LOCTEXT("Expand_Instance", "Expand Instance-Level Operations"),
-		LOCTEXT("Expand_Instance", "Expands all the operations in the tree that are instance operations (not images, meshes, booleans, etc.)."),
+		LOCTEXT("Code_Expand_Instance", "Expand Instance-Level Operations"),
+		LOCTEXT("Code_Expand_Instance_Tooltip", "Expands all the operations in the tree that are instance operations (not images, meshes, booleans, etc.)."),
 		FSlateIcon(),
 		FUIAction(FExecuteAction::CreateSP(this, &SMutableCodeViewer::TreeExpandInstance)
 			//, FCanExecuteAction::CreateSP(this, &SMutableCodeViewer::HasAnyItemInPalette)
@@ -1483,8 +1483,8 @@ TSharedPtr<SWidget> SMutableCodeViewer::OnTreeContextMenuOpening()
 	);
 
 	MenuBuilder.AddMenuEntry(
-		LOCTEXT("Expand_Unique", "Expand All Unique Operations"),
-		LOCTEXT("Expand_Unique", "Expands all the operations in the tree that have not been expanded yet."),
+		LOCTEXT("Code_Expand_Unique", "Expand All Unique Operations"),
+		LOCTEXT("Code_Expand_Unique_Tooltip", "Expands all the operations in the tree that have not been expanded yet."),
 		FSlateIcon(),
 		FUIAction(FExecuteAction::CreateSP(this, &SMutableCodeViewer::TreeExpandUnique)
 		)
diff --git a/Engine/Plugins/Experimental/Mutable/Source/CustomizableObjectEditor/Private/SMutableColorPreviewBox.cpp b/Engine/Plugins/Experimental/Mutable/Source/CustomizableObjectEditor/Private/SMutableColorPreviewBox.cpp
--- a/Engine/Plugins/Experimental/Mutable/Source/CustomizableObjectEditor/Private/SMutableColorPreviewBox.cpp
+++ b/Engine/Plugins/Experimental/Mutable/Source/CustomizableObjectEditor/Private/SMutableColorPreviewBox.cpp
@@ -35,4 +35,5 @@ FSlateColor SMutableColorPreviewBox::GetColor() const
 {
 	return this->Color;
 }
-#undef LOCTEXT_NAMESPACES
\ No newline at end of file
+
+#undef LOCTEXT_NAMESPACE
diff --git a/Engine/Plugins/Experimental/Mutable/Source/CustomizableObjectEditor/Private/SMutableGraphViewer.cpp b/Engine/Plugins/Experimental/Mutable/Source/CustomizableObjectEditor/Private/SMutableGraphViewer.cpp
--- a/Engine/Plugins/Experimental/Mutable/Source/CustomizableObjectEditor/Private/SMutableGraphViewer.cpp
+++ b/Engine/Plugins/Experimental/Mutable/Source/CustomizableObjectEditor/Private/SMutableGraphViewer.cpp
@@ -279,10 +279,10 @@ TSharedRef<SWidget> SMutableGraphViewer::GenerateCompileOptionsMenuContent()
 
 		// Optimisation level
 		CompileOptimizationStrings.Empty();
-		CompileOptimizationStrings.Add(MakeShareable(new FString(NSLOCTEXT("UnrealEd", "OptimizationNone", "None").ToString())));
-		CompileOptimizationStrings.Add(MakeShareable(new FString(NSLOCTEXT("UnrealEd", "OptimizationMin", "Minimal").ToString())));
-		CompileOptimizationStrings.Add(MakeShareable(new FString(NSLOCTEXT("UnrealEd", "OptimizationMed", "Medium").ToString())));
-		CompileOptimizationStrings.Add(MakeShareable(new FString(NSLOCTEXT("UnrealEd", "OptimizationMax", "Maximum").ToString())));
+		CompileOptimizationStrings.Add(MakeShareable(new FString(NSLOCTEXT("CustomizableObjectEditor", "OptimizationNone", "None").ToString())));
+		CompileOptimizationStrings.Add(MakeShareable(new FString(NSLOCTEXT("CustomizableObjectEditor", "OptimizationMin", "Minimal").ToString())));
+		CompileOptimizationStrings.Add(MakeShareable(new FString(NSLOCTEXT("CustomizableObjectEditor", "OptimizationMed", "Medium").ToString())));
+		CompileOptimizationStrings.Add(MakeShareable(new FString(NSLOCTEXT("CustomizableObjectEditor", "OptimizationMax", "Maximum").ToString())));
 
 		CompileOptions.OptimizationLevel = FMath::Min(CompileOptions.OptimizationLevel, CompileOptimizationStrings.Num() - 1);
 
@@ -399,8 +399,8 @@ TSharedPtr<SWidget> SMutableGraphViewer::OnTreeContextMenuOpening()
 	FMenuBuilder MenuBuilder(true, nullptr);
 
 	MenuBuilder.AddMenuEntry(
-		LOCTEXT("Expand_Instance", "Expand Instance-Level Operations"),
-		LOCTEXT("Expand_Instance", "Expands all the operations in the tree that are instance operations (not images, meshes, booleans, etc.)."),
+		LOCTEXT("Graph_Expand_Instance", "Expand Instance-Level Operations"),
+		LOCTEXT("Graph_Expand_Instance_Tooltip", "Expands all the operations in the tree that are instance operations (not images, meshes, booleans, etc.)."),
 		FSlateIcon(),
 		FUIAction(FExecuteAction::CreateSP(this, &SMutableGraphViewer::TreeExpandUnique)
 			//, FCanExecuteAction::CreateSP(this, &SMutableCodeViewer::HasAnyItemInPalette)
diff --git a/Engine/Plugins/Experimental/Mutable/Source/CustomizableObjectEditor/Private/SMutableObjectViewer.cpp b/Engine/Plugins/Experimental/Mutable/Source/CustomizableObjectEditor/Private/SMutableObjectViewer.cpp
--- a/Engine/Plugins/Experimental/Mutable/Source/CustomizableObjectEditor/Private/SMutableObjectViewer.cpp
+++ b/Engine/Plugins/Experimental/Mutable/Source/CustomizableObjectEditor/Private/SMutableObjectViewer.cpp
@@ -314,8 +314,8 @@ TSharedRef<SWidget> SMutableObjectViewer::GenerateCompileOptionsMenuContent()
 
 		// Texture compresssion
 		MenuBuilder.AddMenuEntry(
-			LOCTEXT("MutableEnableTextureCompression", "Enable texture compression."),
-			LOCTEXT("MutableEnableTextureCompressionTooltip", "Enable or disable the compression of textures in the mutable code compilation."),
+			LOCTEXT("Generate_MutableEnableTextureCompression", "Enable texture compression."),
+			LOCTEXT("Generate_MutableEnableTextureCompressionTooltip", "Enable or disable the compression of textures in the mutable code compilation."),
 			FSlateIcon(),
 			FUIAction(
 				FExecuteAction::CreateLambda([this]() { CompileOptions.bTextureCompression = !CompileOptions.bTextureCompression; }),
@@ -330,10 +330,10 @@ TSharedRef<SWidget> SMutableObjectViewer::GenerateCompileOptionsMenuContent()
 
 		// Optimisation level
 		CompileOptimizationStrings.Empty();
-		CompileOptimizationStrings.Add(MakeShareable(new FString(NSLOCTEXT("UnrealEd", "OptimizationNone", "None").ToString())));
-		CompileOptimizationStrings.Add(MakeShareable(new FString(NSLOCTEXT("UnrealEd", "OptimizationMin", "Minimal").ToString())));
-		CompileOptimizationStrings.Add(MakeShareable(new FString(NSLOCTEXT("UnrealEd", "OptimizationMed", "Medium").ToString())));
-		CompileOptimizationStrings.Add(MakeShareable(new FString(NSLOCTEXT("UnrealEd", "OptimizationMax", "Maximum").ToString())));
+		CompileOptimizationStrings.Add(MakeShareable(new FString(LOCTEXT("OptimizationNone", "None").ToString())));
+		CompileOptimizationStrings.Add(MakeShareable(new FString(LOCTEXT("OptimizationMin", "Minimal").ToString())));
+		CompileOptimizationStrings.Add(MakeShareable(new FString(LOCTEXT("OptimizationMed", "Medium").ToString())));
+		CompileOptimizationStrings.Add(MakeShareable(new FString(LOCTEXT("OptimizationMax", "Maximum").ToString())));
 
 		CompileOptions.OptimizationLevel = FMath::Min(CompileOptions.OptimizationLevel, CompileOptimizationStrings.Num() - 1);
 
@@ -347,8 +347,8 @@ TSharedRef<SWidget> SMutableObjectViewer::GenerateCompileOptionsMenuContent()
 
 		// Parallel compilation
 		MenuBuilder.AddMenuEntry(
-			LOCTEXT("MutableEnableParallelCompilation", "Enable compiling in multiple threads."),
-			LOCTEXT("MutableEnableParallelCompilationTooltip", "This is faster but use more memory."),
+			LOCTEXT("Generate_MutableEnableParallelCompilation", "Enable compiling in multiple threads."),
+			LOCTEXT("Generate_MutableEnableParallelCompilationTooltip", "This is faster but use more memory."),
 			FSlateIcon(),
 			FUIAction(
 				FExecuteAction::CreateLambda([this]() { CompileOptions.bUseParallelCompilation = !CompileOptions.bUseParallelCompilation; }),
@@ -360,8 +360,8 @@ TSharedRef<SWidget> SMutableObjectViewer::GenerateCompileOptionsMenuContent()
 
 		// Disk as cache
 		MenuBuilder.AddMenuEntry(
-			LOCTEXT("MutableEnableTextureCompression", "Enable compiling using the disk as memory."),
-			LOCTEXT("MutableEnableTextureCompressionTooltip", "This is very slow but supports compiling huge objects. It requires a lot of free space in the OS disk."),
+			LOCTEXT("Generate_MutableUseDisk", "Enable compiling using the disk as memory."),
+			LOCTEXT("Generate_MutableUseDiskTooltip", "This is very slow but supports compiling huge objects. It requires a lot of free space in the OS disk."),
 			FSlateIcon(),
 			FUIAction(
 				FExecuteAction::CreateLambda([this]() { CompileOptions.bUseDiskCompilation = !CompileOptions.bUseDiskCompilation; }),
diff --git a/Engine/Plugins/Experimental/Mutable/Source/CustomizableObjectEditor/Public/CustomizableObjectEditorViewportLODCommands.h b/Engine/Plugins/Experimental/Mutable/Source/CustomizableObjectEditor/Public/CustomizableObjectEditorViewportLODCommands.h
--- a/Engine/Plugins/Experimental/Mutable/Source/CustomizableObjectEditor/Public/CustomizableObjectEditorViewportLODCommands.h
+++ b/Engine/Plugins/Experimental/Mutable/Source/CustomizableObjectEditor/Public/CustomizableObjectEditorViewportLODCommands.h
@@ -16,7 +16,7 @@ public:
 		: TCommands<FCustomizableObjectEditorViewportLODCommands>
 		(
 			TEXT("COEditorViewportLODCommands"), // Context name for fast lookup
-			NSLOCTEXT("Contexts", "COEditorViewportLODCommands", "Customizable Object Editor Commands"), // Localized context name for displaying
+			NSLOCTEXT("CustomizableObjectEditor", "COEditorViewportLODCommands", "Customizable Object Editor Commands"), // Localized context name for displaying
 			NAME_None, // Parent context name. 
 			FCustomizableObjectEditorStyle::GetStyleSetName() // Icon Style Set
 			)
diff --git a/Engine/Plugins/Experimental/Mutable/Source/CustomizableObjectEditor/Public/CustomizableObjectEditorViewportMenuCommands.h b/Engine/Plugins/Experimental/Mutable/Source/CustomizableObjectEditor/Public/CustomizableObjectEditorViewportMenuCommands.h
--- a/Engine/Plugins/Experimental/Mutable/Source/CustomizableObjectEditor/Public/CustomizableObjectEditorViewportMenuCommands.h
+++ b/Engine/Plugins/Experimental/Mutable/Source/CustomizableObjectEditor/Public/CustomizableObjectEditorViewportMenuCommands.h
@@ -16,7 +16,7 @@ public:
 		: TCommands<FCustomizableObjectEditorViewportMenuCommands>
 		(
 			TEXT("CustomizableObjectEditorViewportMenu"), // Context name for fast lookup
-			NSLOCTEXT("Contexts", "CustomizableObjectEditorViewportMenu", "CustomizableObject Viewport Menu"), // Localized context name for displaying
+			NSLOCTEXT("CustomizableObjectEditor", "CustomizableObjectEditorViewportMenu", "CustomizableObject Viewport Menu"), // Localized context name for displaying
 			NAME_None, // Parent context name.  
 			FCustomizableObjectEditorStyle::GetStyleSetName() // Icon Style Set
 			)
diff --git a/Engine/Plugins/Experimental/Mutable/Source/CustomizableObjectEditor/Public/CustomizableObjectInstanceEditorActions.h b/Engine/Plugins/Experimental/Mutable/Source/CustomizableObjectEditor/Public/CustomizableObjectInstanceEditorActions.h
--- a/Engine/Plugins/Experimental/Mutable/Source/CustomizableObjectEditor/Public/CustomizableObjectInstanceEditorActions.h
+++ b/Engine/Plugins/Experimental/Mutable/Source/CustomizableObjectEditor/Public/CustomizableObjectInstanceEditorActions.h
@@ -30,7 +30,7 @@ public:
 	TSharedPtr< FUICommandInfo > SetShowBinormals;
 	TSharedPtr< FUICommandInfo > SetShowPivot;
 
-	TSharedPtr< FUICommandInfo > BakeInstance;
+	TSharedPtr< FUICommandInfo > BakeInstanceFromInstance;
 	TSharedPtr< FUICommandInfo > StateChangeTest;
 
 	//Toolbar Commands
diff --git a/Engine/Plugins/Experimental/Mutable/Source/CustomizableObjectPopulationEditor/CustomizableObjectPopulationEditor.Build.cs b/Engine/Plugins/Experimental/Mutable/Source/CustomizableObjectPopulationEditor/CustomizableObjectPopulationEditor.Build.cs
--- a/Engine/Plugins/Experimental/Mutable/Source/CustomizableObjectPopulationEditor/CustomizableObjectPopulationEditor.Build.cs
+++ b/Engine/Plugins/Experimental/Mutable/Source/CustomizableObjectPopulationEditor/CustomizableObjectPopulationEditor.Build.cs
@@ -5,45 +5,6 @@ using System.IO;
 
 public class CustomizableObjectPopulationEditor : ModuleRules
 {
-	private int ReadEngineMajorVersion()
-	{
-		string line;
-
-		System.IO.StreamReader file =
-		   new System.IO.StreamReader("Runtime/Launch/Resources/Version.h");
-
-		while ((line = file.ReadLine()) != null)
-		{
-			if (line.StartsWith("#define ENGINE_MAJOR_VERSION"))
-			{
-				file.Close();
-				return int.Parse(line.Split()[2]);
-			}
-		}
-
-		file.Close();
-		return 0;
-	}
-
-	private int ReadEngineMinorVersion()
-	{
-		string line;
-
-		System.IO.StreamReader file =
-		   new System.IO.StreamReader("Runtime/Launch/Resources/Version.h");
-
-		while ((line = file.ReadLine()) != null)
-		{
-			if (line.StartsWith("#define ENGINE_MINOR_VERSION"))
-			{
-				file.Close();
-				return int.Parse(line.Split()[2]);
-			}
-		}
-
-		file.Close();
-		return 0;
-	}
 
 	public CustomizableObjectPopulationEditor(ReadOnlyTargetRules TargetRules) : base(TargetRules)
 	{
@@ -97,14 +58,6 @@ public class CustomizableObjectPopulationEditor : ModuleRules
 			}
 		);
 
-		if (ReadEngineMajorVersion() >= 4 && ReadEngineMinorVersion() >= 18)
-		{
-			PrivateDependencyModuleNames.AddRange(new string[]
-			{
-				"ApplicationCore",
-				"ToolMenus",
-			});
-		}
 
 		DynamicallyLoadedModuleNames.AddRange(
 			new string[] {
diff --git a/Engine/Plugins/Experimental/Mutable/Source/CustomizableObjectPopulationEditor/Private/AssetTypeActions_CustomizableObjectPopulation.cpp b/Engine/Plugins/Experimental/Mutable/Source/CustomizableObjectPopulationEditor/Private/AssetTypeActions_CustomizableObjectPopulation.cpp
--- a/Engine/Plugins/Experimental/Mutable/Source/CustomizableObjectPopulationEditor/Private/AssetTypeActions_CustomizableObjectPopulation.cpp
+++ b/Engine/Plugins/Experimental/Mutable/Source/CustomizableObjectPopulationEditor/Private/AssetTypeActions_CustomizableObjectPopulation.cpp
@@ -18,7 +18,7 @@
 
 FText FAssetTypeActions_CustomizableObjectPopulation::GetName() const
 {
-	return NSLOCTEXT("AssetTypeActions", "AssetTypeActions_CustomizableObjectPopulation", "Customizable Population");
+	return LOCTEXT("AssetTypeActions_CustomizableObjectPopulation", "Customizable Population");
 }
 
 
diff --git a/Engine/Plugins/Experimental/Mutable/Source/CustomizableObjectPopulationEditor/Private/CustomizableObjectPopulationClassEditor.cpp b/Engine/Plugins/Experimental/Mutable/Source/CustomizableObjectPopulationEditor/Private/CustomizableObjectPopulationClassEditor.cpp
--- a/Engine/Plugins/Experimental/Mutable/Source/CustomizableObjectPopulationEditor/Private/CustomizableObjectPopulationClassEditor.cpp
+++ b/Engine/Plugins/Experimental/Mutable/Source/CustomizableObjectPopulationEditor/Private/CustomizableObjectPopulationClassEditor.cpp
@@ -946,7 +946,7 @@ void FCustomizableObjectPopulationClassEditor::GeneratePopulationClassInstances(
 
 	TSharedRef<SSelectPopulationClassFolderDlg> FolderDlg =
 		SNew(SSelectPopulationClassFolderDlg)
-		.DefaultAssetPath(LOCTEXT("", ""))
+		.DefaultAssetPath(FText())
 		.DefaultFileName(DefaultFileName);
 
 	// Popup window management
diff --git a/Engine/Plugins/Experimental/Mutable/Source/CustomizableObjectPopulationEditor/Private/CustomizableObjectPopulationEditor.cpp b/Engine/Plugins/Experimental/Mutable/Source/CustomizableObjectPopulationEditor/Private/CustomizableObjectPopulationEditor.cpp
--- a/Engine/Plugins/Experimental/Mutable/Source/CustomizableObjectPopulationEditor/Private/CustomizableObjectPopulationEditor.cpp
+++ b/Engine/Plugins/Experimental/Mutable/Source/CustomizableObjectPopulationEditor/Private/CustomizableObjectPopulationEditor.cpp
@@ -582,7 +582,7 @@ void FCustomizableObjectPopulationEditor::GeneratePopulationInstances()
 
 	TSharedRef<SSelectPopulationFolderDlg> FolderDlg =
 		SNew(SSelectPopulationFolderDlg)
-		.DefaultAssetPath(LOCTEXT("", ""))
+		.DefaultAssetPath(FText())
 		.DefaultFileName(DefaultFileName);
 
 	if (FolderDlg->ShowModal() != EAppReturnType::Cancel)
