[mutable] Even more localization key fixes.

diff --git a/Engine/Plugins/Experimental/Mutable/Source/CustomizableObjectEditor/Private/Nodes/CustomizableObjectNodeLayoutBlocksDetails.cpp b/Engine/Plugins/Experimental/Mutable/Source/CustomizableObjectEditor/Private/Nodes/CustomizableObjectNodeLayoutBlocksDetails.cpp
--- a/Engine/Plugins/Experimental/Mutable/Source/CustomizableObjectEditor/Private/Nodes/CustomizableObjectNodeLayoutBlocksDetails.cpp
+++ b/Engine/Plugins/Experimental/Mutable/Source/CustomizableObjectEditor/Private/Nodes/CustomizableObjectNodeLayoutBlocksDetails.cpp
@@ -36,7 +36,7 @@ void FCustomizableObjectNodeLayoutBlocksDetails::CustomizeDetails( IDetailLayout
 
 	if (Node)
 	{
-		BlocksCategory.AddCustomRow( LOCTEXT("FCustomizableObjectNodeLayoutBlocksDetails", "BlockInstructions") )
+		BlocksCategory.AddCustomRow( LOCTEXT("BlocksDetails_BlockInstructions", "BlockInstructions") )
 		[
 			SNew( STextBlock )
 			.Text(  LOCTEXT("Use the node preview panel to set up the blocks.","Use the node preview panel to set up the blocks.") )
@@ -44,7 +44,7 @@ void FCustomizableObjectNodeLayoutBlocksDetails::CustomizeDetails( IDetailLayout
 	}
 	else
 	{
-		BlocksCategory.AddCustomRow( LOCTEXT("FCustomizableObjectNodeLayoutBlocksDetails", "NodeNotFound") )
+		BlocksCategory.AddCustomRow( LOCTEXT("BlocksDetails_NodeNotFound", "NodeNotFound") )
 		[
 			SNew( STextBlock )
 			.Text( LOCTEXT( "Node not found", "Node not found" ) )
diff --git a/Engine/Plugins/Experimental/Mutable/Source/CustomizableObjectEditor/Private/Nodes/CustomizableObjectNodeMeshClipMorphDetails.cpp b/Engine/Plugins/Experimental/Mutable/Source/CustomizableObjectEditor/Private/Nodes/CustomizableObjectNodeMeshClipMorphDetails.cpp
--- a/Engine/Plugins/Experimental/Mutable/Source/CustomizableObjectEditor/Private/Nodes/CustomizableObjectNodeMeshClipMorphDetails.cpp
+++ b/Engine/Plugins/Experimental/Mutable/Source/CustomizableObjectEditor/Private/Nodes/CustomizableObjectNodeMeshClipMorphDetails.cpp
@@ -106,7 +106,7 @@ void FCustomizableObjectNodeMeshClipMorphDetails::CustomizeDetails(IDetailLayout
 			// Add them to the parent combo box
 			TSharedRef<IPropertyHandle> BoneProperty = DetailBuilder.GetProperty("BoneName");
 
-			BlocksCategory.AddCustomRow(LOCTEXT("FCustomizableObjectNodeMeshClipMorphDetails", "Bone Name"))
+			BlocksCategory.AddCustomRow(LOCTEXT("ClipMorphDetails_BoneName", "Bone Name"))
 			[
 				SNew(SProperty, BoneProperty)
 				.ShouldDisplayName(false)
@@ -121,7 +121,7 @@ void FCustomizableObjectNodeMeshClipMorphDetails::CustomizeDetails(IDetailLayout
 						+SHorizontalBox::Slot().HAlign(HAlign_Left).VAlign(VAlign_Center)
 						[
 							SNew(STextBlock)
-							.Text( LOCTEXT("FCustomizableObjectNodeRemoveMeshDetails", "Bone Name"))
+							.Text( LOCTEXT("ClipMorphDetails_BoneName", "Bone Name"))
 							.Font(IDetailLayoutBuilder::GetDetailFont())
 						]
 						+SHorizontalBox::Slot().HAlign(HAlign_Fill)
@@ -136,7 +136,7 @@ void FCustomizableObjectNodeMeshClipMorphDetails::CustomizeDetails(IDetailLayout
 				]
 			];
 
-			MeshClipParametersCategory.AddCustomRow(LOCTEXT("FCustomizableObjectNodeMeshClipMorphDetails", "Invert plane normal"))
+			MeshClipParametersCategory.AddCustomRow(LOCTEXT("ClipMorphDetails_PlaneNormal", "Invert plane normal"))
 			[
 				SNew(SProperty, BoneProperty)
 				.ShouldDisplayName(false)
@@ -151,14 +151,14 @@ void FCustomizableObjectNodeMeshClipMorphDetails::CustomizeDetails(IDetailLayout
 						+SHorizontalBox::Slot().HAlign(HAlign_Left).VAlign(VAlign_Center)
 						[
 							SNew(STextBlock)
-							.Text( LOCTEXT("FCustomizableObjectNodeMeshClipMorphDetails", "Invert plane normal"))
+							.Text( LOCTEXT("ClipMorphDetails_PlaneNormal", "Invert plane normal"))
 						]
 						+SHorizontalBox::Slot().HAlign(HAlign_Left)
 						[
 							SNew(SCheckBox)
 							.OnCheckStateChanged(this, &FCustomizableObjectNodeMeshClipMorphDetails::OnInvertNormalCheckboxChanged)
 							.IsChecked(this, &FCustomizableObjectNodeMeshClipMorphDetails::GetInvertNormalCheckBoxState)
-							.ToolTipText(LOCTEXT("InvertNormal", "Invert normal direction of the clip plane"))
+							.ToolTipText(LOCTEXT("ClipMorphDetails_InvertNormal_Tooltip", "Invert normal direction of the clip plane"))
 						]
 					]
 				]
@@ -170,7 +170,7 @@ void FCustomizableObjectNodeMeshClipMorphDetails::CustomizeDetails(IDetailLayout
 		BlocksCategory.AddCustomRow(LOCTEXT("Node", "Node"))
 		[
 			SNew(STextBlock)
-			.Text(LOCTEXT("Node not found", "Node not found"))
+			.Text(LOCTEXT("ClipMorphDetails_InvertNormal_NodeNotFound", "Node not found"))
 		];
 	}
 }
diff --git a/Engine/Plugins/Experimental/Mutable/Source/CustomizableObjectEditor/Private/Nodes/CustomizableObjectNodeMeshMorphDetails.cpp b/Engine/Plugins/Experimental/Mutable/Source/CustomizableObjectEditor/Private/Nodes/CustomizableObjectNodeMeshMorphDetails.cpp
--- a/Engine/Plugins/Experimental/Mutable/Source/CustomizableObjectEditor/Private/Nodes/CustomizableObjectNodeMeshMorphDetails.cpp
+++ b/Engine/Plugins/Experimental/Mutable/Source/CustomizableObjectEditor/Private/Nodes/CustomizableObjectNodeMeshMorphDetails.cpp
@@ -86,7 +86,7 @@ void FCustomizableObjectNodeMeshMorphDetails::CustomizeDetails( IDetailLayoutBui
 
 			TSharedRef<IPropertyHandle> MorphTargetNameProperty = DetailBuilder.GetProperty("MorphTargetName");
 
-			BlocksCategory.AddCustomRow( LOCTEXT("FCustomizableObjectNodeMeshMorphDetails", "Target") )
+			BlocksCategory.AddCustomRow( LOCTEXT("MeshMorphDetails_Target", "Target") )
 			[
 				SNew(SProperty, MorphTargetNameProperty)
 				.ShouldDisplayName( false)
@@ -108,10 +108,10 @@ void FCustomizableObjectNodeMeshMorphDetails::CustomizeDetails( IDetailLayoutBui
 	
 	if (!SourceMeshFound)
 	{
-		BlocksCategory.AddCustomRow( LOCTEXT("FCustomizableObjectNodeMeshMorphDetails", "Node") )
+		BlocksCategory.AddCustomRow( LOCTEXT("MeshMorphDetails_Node", "Node") )
 		[
 			SNew( STextBlock )
-			.Text( LOCTEXT( "No source mesh found.", "No source mesh found." ) )			
+			.Text( LOCTEXT( "MeshMorphDetails_NoSource", "No source mesh found." ) )			
 		];
 	}
 
diff --git a/Engine/Plugins/Experimental/Mutable/Source/CustomizableObjectEditor/Private/Nodes/CustomizableObjectNodeMorphMaterialDetails.cpp b/Engine/Plugins/Experimental/Mutable/Source/CustomizableObjectEditor/Private/Nodes/CustomizableObjectNodeMorphMaterialDetails.cpp
--- a/Engine/Plugins/Experimental/Mutable/Source/CustomizableObjectEditor/Private/Nodes/CustomizableObjectNodeMorphMaterialDetails.cpp
+++ b/Engine/Plugins/Experimental/Mutable/Source/CustomizableObjectEditor/Private/Nodes/CustomizableObjectNodeMorphMaterialDetails.cpp
@@ -80,7 +80,7 @@ void FCustomizableObjectNodeMorphMaterialDetails::CustomizeDetails( IDetailLayou
         MorphTargetComboOptions.Sort(CompareNames);
 
 		TSharedRef<IPropertyHandle> MorphTargetNameProperty = DetailBuilder.GetProperty("MorphTargetName");
-		BlocksCategory.AddCustomRow( LOCTEXT("FCustomizableObjectNodeMorphMaterialDetails", "Target") )
+		BlocksCategory.AddCustomRow( LOCTEXT("MorphMaterialDetails_Target", "Target") )
 		[
 			SNew(SProperty, MorphTargetNameProperty)
 			.ShouldDisplayName( false)
@@ -97,7 +97,7 @@ void FCustomizableObjectNodeMorphMaterialDetails::CustomizeDetails( IDetailLayou
 						.VAlign(VAlign_Center)
 					[
 						SNew( STextBlock )
-							.Text( LOCTEXT("Morph Target","Morph Target") )
+							.Text( LOCTEXT("MorphMaterialDetails_MorphTarget","Morph Target") )
 					]
 
 					+ SHorizontalBox::Slot()
@@ -116,10 +116,10 @@ void FCustomizableObjectNodeMorphMaterialDetails::CustomizeDetails( IDetailLayou
 	}
 	else
 	{
-		BlocksCategory.AddCustomRow( LOCTEXT("FCustomizableObjectNodeMorphMaterialDetails", "Node") )
+		BlocksCategory.AddCustomRow( LOCTEXT("MorphMaterialDetails_Node", "Node") )
 		[
 			SNew( STextBlock )
-			.Text( LOCTEXT( "Node not found", "Node not found" ) )
+			.Text( LOCTEXT( "MorphMaterialDetails_NodeNotFound", "Node not found" ) )
 		];
 	}
 
diff --git a/Engine/Plugins/Experimental/Mutable/Source/CustomizableObjectEditor/Private/SCustomizableObjectHighresScreenshot.cpp b/Engine/Plugins/Experimental/Mutable/Source/CustomizableObjectEditor/Private/SCustomizableObjectHighresScreenshot.cpp
--- a/Engine/Plugins/Experimental/Mutable/Source/CustomizableObjectEditor/Private/SCustomizableObjectHighresScreenshot.cpp
+++ b/Engine/Plugins/Experimental/Mutable/Source/CustomizableObjectEditor/Private/SCustomizableObjectHighresScreenshot.cpp
@@ -56,37 +56,37 @@ void SCustomizableObjectHighresScreenshot::Construct( const FArguments& InArgs )
 							.VAlign(VAlign_Center)
 							[
 								SNew( STextBlock )
-								.Text( NSLOCTEXT("HighResScreenshot", "ScreenshotSizeMultiplier", "Screenshot Size Multiplier") )
+								.Text( NSLOCTEXT("CustomizableObjectEditor", "ScreenshotSizeMultiplier", "Screenshot Size Multiplier") )
 							]
 							+SVerticalBox::Slot()
 							.VAlign(VAlign_Center)
 							[
 								SNew( STextBlock )
-								.Text( NSLOCTEXT("HighResScreenshot", "IncludeBufferVisTargets", "Include Buffer Visualization Targets") )
+								.Text( NSLOCTEXT("CustomizableObjectEditor", "IncludeBufferVisTargets", "Include Buffer Visualization Targets") )
 							]
 							+ SVerticalBox::Slot()
 							.VAlign(VAlign_Center)
 							[
 								SAssignNew(HDRLabel, STextBlock)
-								.Text(NSLOCTEXT("HighResScreenshot", "CaptureHDR", "Write HDR format visualization targets"))
+								.Text(NSLOCTEXT("CustomizableObjectEditor", "CaptureHDR", "Write HDR format visualization targets"))
 							]
 							+SVerticalBox::Slot()
 							.VAlign(VAlign_Center)
 							[
 								SAssignNew(Force128BitRenderingLabel, STextBlock)
-								.Text(NSLOCTEXT("HighResScreenshot", "Force128BitPipeline", "Force 128-bit buffers for rendering pipeline"))
+								.Text(NSLOCTEXT("CustomizableObjectEditor", "Force128BitPipeline", "Force 128-bit buffers for rendering pipeline"))
 							]
 							+SVerticalBox::Slot()
 							.VAlign(VAlign_Center)
 							[
 								SNew( STextBlock )
-								.Text( NSLOCTEXT("HighResScreenshot", "UseCustomDepth", "Use custom depth as mask") )
+								.Text( NSLOCTEXT("CustomizableObjectEditor", "UseCustomDepth", "Use custom depth as mask") )
 							]
 							+SVerticalBox::Slot()
 							.VAlign(VAlign_Center)
 							[
 								SNew( STextBlock )
-								.Text( NSLOCTEXT("HighResScreenshot", "ShowHideBackground", "Show / hide background") )
+								.Text( NSLOCTEXT("CustomizableObjectEditor", "ShowHideBackground", "Show / hide background") )
 							]
 						]
 						+SSplitter::Slot()
