[mutable] Fix wrong localization key.

diff --git a/Engine/Plugins/Experimental/Mutable/Source/CustomizableObjectEditor/Private/SCustomizableObjectHighresScreenshot.cpp b/Engine/Plugins/Experimental/Mutable/Source/CustomizableObjectEditor/Private/SCustomizableObjectHighresScreenshot.cpp
--- a/Engine/Plugins/Experimental/Mutable/Source/CustomizableObjectEditor/Private/SCustomizableObjectHighresScreenshot.cpp
+++ b/Engine/Plugins/Experimental/Mutable/Source/CustomizableObjectEditor/Private/SCustomizableObjectHighresScreenshot.cpp
@@ -170,7 +170,7 @@ void SCustomizableObjectHighresScreenshot::Construct( const FArguments& InArgs )
 					.BorderImage(FAppStyle::GetBrush("ToolPanel.GroupBorder"))
 					[
 						SNew( STextBlock )
-						.Text( NSLOCTEXT("HighResScreenshot", "CaptureWarningText", "Due to the high system requirements of a high resolution screenshot, very large multipliers might cause the graphics driver to become unresponsive and possibly crash. In these circumstances, please try using a lower multiplier") )
+						.Text( NSLOCTEXT("CustomizableObjectEditor", "CaptureWarningText", "Due to the high system requirements of a high resolution screenshot, very large multipliers might cause the graphics driver to become unresponsive and possibly crash. In these circumstances, please try using a lower multiplier") )
 						.AutoWrapText(true)
 					]
 				]
