[Mutable] Remove CIS Warning: Clipping mesh not closed.

#rnx

diff --git a/Engine/Plugins/Experimental/Mutable/Source/CustomizableObjectEditor/Private/GenerateMutableSource/GenerateMutableSourceModifier.cpp b/Engine/Plugins/Experimental/Mutable/Source/CustomizableObjectEditor/Private/GenerateMutableSource/GenerateMutableSourceModifier.cpp
--- a/Engine/Plugins/Experimental/Mutable/Source/CustomizableObjectEditor/Private/GenerateMutableSource/GenerateMutableSourceModifier.cpp
+++ b/Engine/Plugins/Experimental/Mutable/Source/CustomizableObjectEditor/Private/GenerateMutableSource/GenerateMutableSourceModifier.cpp
@@ -145,7 +145,7 @@ mu::NodeModifierPtr GenerateMutableSourceModifier(const UEdGraphPin * Pin, FMuta
 
 				if (!bClosed)
 				{
-					GenerationContext.Compiler->CompilerLog(LOCTEXT("Clipping mesh", "Clipping mesh not closed (i.e., it does not enclose a volume)."), MeshData.Node, EMessageSeverity::Warning);
+					//GenerationContext.Compiler->CompilerLog(LOCTEXT("Clipping mesh", "Clipping mesh not closed (i.e., it does not enclose a volume)."), MeshData.Node, EMessageSeverity::Warning);
 				}
 			}
 
