[mutable] Fix crash with certain combination of data during obejct optimization.

diff --git a/Engine/Plugins/Experimental/Mutable/Source/MutableTools/Private/CodeOptimiser_Logic.cpp b/Engine/Plugins/Experimental/Mutable/Source/MutableTools/Private/CodeOptimiser_Logic.cpp
--- a/Engine/Plugins/Experimental/Mutable/Source/MutableTools/Private/CodeOptimiser_Logic.cpp
+++ b/Engine/Plugins/Experimental/Mutable/Source/MutableTools/Private/CodeOptimiser_Logic.cpp
@@ -1018,9 +1018,9 @@ namespace mu
 
                 if (n->GetOpType()==OP_TYPE::IN_CONDITIONAL)
                 {
-                    auto topConditional = dynamic_cast<ASTOpConditional*>(n.get());
-                    auto yes = topConditional->yes.child();
-                    if (yes->GetOpType()==OP_TYPE::IN_ADDSURFACE)
+					ASTOpConditional* topConditional = dynamic_cast<ASTOpConditional*>(n.get());
+					Ptr<ASTOp> yes = topConditional->yes.child();
+                    if (yes && yes->GetOpType()==OP_TYPE::IN_ADDSURFACE)
                     {
                         auto addSurface = dynamic_cast<ASTOpInstanceAdd*>(yes.get());
 
@@ -1028,10 +1028,10 @@ namespace mu
                         while (!ended)
                         {
                             ended = true;
-                            auto base = addSurface->instance.child();
+							Ptr<ASTOp> base = addSurface->instance.child();
                             if ( base && base->GetOpType()==OP_TYPE::IN_CONDITIONAL )
                             {
-                                auto bottomConditional = dynamic_cast<ASTOpConditional*>(base.get());
+								ASTOpConditional* bottomConditional = dynamic_cast<ASTOpConditional*>(base.get());
 
                                 // Are the two conditions exclusive?
                                 bool conditionaAreExclusive = false;
@@ -1040,7 +1040,7 @@ namespace mu
                                 facts.push_back(topConditional->condition.child());
 
                                 // Check if the child condition has a value with the current facts
-                                const auto& pChildCond = bottomConditional->condition.child();
+                                Ptr<ASTOp> pChildCond = bottomConditional->condition.child();
                                 ASTOp::BOOL_EVAL_RESULT result;
                                 {
                                     //MUTABLE_CPUPROFILER_SCOPE(EvaluateBool);
@@ -1062,7 +1062,7 @@ namespace mu
                                     {
                                         // Other parents may not impose the same condition that allows
                                         // the optimisation.
-                                        auto newAddSurface = mu::Clone<ASTOpInstanceAdd>(addSurface);
+                                        Ptr<ASTOpInstanceAdd> newAddSurface = mu::Clone<ASTOpInstanceAdd>(addSurface);
                                         newAddSurface->instance = bottomConditional->no.child();
                                         topConditional->yes = newAddSurface;
                                     }
@@ -1077,9 +1077,9 @@ namespace mu
 
                 else if (n->GetOpType()==OP_TYPE::ME_CONDITIONAL)
                 {
-                    auto topConditional = dynamic_cast<ASTOpConditional*>(n.get());
-                    auto yes = topConditional->yes.child();
-                    if ( yes->GetOpType()==OP_TYPE::ME_MERGE )
+					ASTOpConditional* topConditional = dynamic_cast<ASTOpConditional*>(n.get());
+                    Ptr<ASTOp> yes = topConditional->yes.child();
+                    if ( yes && yes->GetOpType()==OP_TYPE::ME_MERGE )
                     {
                         Ptr<ASTOpFixed> merge = dynamic_cast<ASTOpFixed*>(yes.get());
 
@@ -1087,10 +1087,10 @@ namespace mu
                         while (!ended)
                         {
                             ended = true;
-                            auto base = merge->children[merge->op.args.MeshMerge.base].child();
+							Ptr<ASTOp> base = merge->children[merge->op.args.MeshMerge.base].child();
                             if ( base && base->GetOpType()==OP_TYPE::ME_CONDITIONAL )
                             {
-                                auto bottomConditional = dynamic_cast<ASTOpConditional*>(base.get());
+								ASTOpConditional* bottomConditional = dynamic_cast<ASTOpConditional*>(base.get());
 
                                 // Are the two conditions exclusive?
                                 bool conditionaAreExclusive = false;
@@ -1099,7 +1099,7 @@ namespace mu
                                 facts.push_back(topConditional->condition.child());
 
                                 // Check if the child condition has a value with the current facts
-                                auto pChildCond = bottomConditional->condition.child();
+								Ptr<ASTOp> pChildCond = bottomConditional->condition.child();
                                 ASTOp::BOOL_EVAL_RESULT result = pChildCond->EvaluateBool( facts );
                                 if ( result==ASTOp::BET_FALSE )
                                 {
@@ -1117,7 +1117,7 @@ namespace mu
                                     {
                                         // Other parents may not impose the same condition that allows
                                         // the optimisation.
-                                        auto newMerge = mu::Clone<ASTOpFixed>(merge);
+                                        Ptr<ASTOpFixed> newMerge = mu::Clone<ASTOpFixed>(merge);
                                         newMerge->SetChild(newMerge->op.args.MeshMerge.base, bottomConditional->no.child());
                                         topConditional->yes = newMerge;
                                         merge = newMerge;
@@ -1278,7 +1278,7 @@ namespace mu
             {
                 bool recurse = true;
 
-                auto topSwitch = dynamic_cast<ASTOpSwitch*>(n.get());
+				ASTOpSwitch* topSwitch = dynamic_cast<ASTOpSwitch*>(n.get());
                 if (topSwitch)
                 {
                     bool first = true;
