Moved the Mutable plugin files to live under a directory named after the module. This allows us to get rid of long include paths that contain public and private folder names in the path and fixes head

diff --git a/Engine/Plugins/Experimental/Mutable/Source/CustomizableObjectEditor/Private/MuCOE/SMutableGraphViewer.cpp b/Engine/Plugins/Experimental/Mutable/Source/CustomizableObjectEditor/Private/MuCOE/SMutableGraphViewer.cpp
new file mode 100644
--- /dev/null
+++ b/Engine/Plugins/Experimental/Mutable/Source/CustomizableObjectEditor/Private/MuCOE/SMutableGraphViewer.cpp
@@ -0,0 +1,520 @@
+// Copyright Epic Games, Inc. All Rights Reserved.
+
+#include "SMutableGraphViewer.h"
+
+#include "MuCOE/CustomizableObjectEditorStyle.h"
+#include "MuCOE/CustomizableObjectCompiler.h"
+#include "SMutableCodeViewer.h"
+#include "MuT/Streams.h"
+
+#include "Misc/App.h"
+#include "Framework/MultiBox/MultiBoxBuilder.h"
+#include "Widgets/Input/SButton.h"
+#include "Widgets/Images/SImage.h"
+#include "Widgets/Views/STreeView.h"
+#include "Widgets/Input/STextComboBox.h"
+#include "Widgets/Docking/SDockTab.h"
+#include "EditorDirectories.h"
+#include "DesktopPlatformModule.h"
+#include "HAL/PlatformApplicationMisc.h"
+
+#define LOCTEXT_NAMESPACE "SMutableDebugger"
+
+
+// \todo: multi-column tree
+namespace MutableGraphTreeViewColumns
+{
+	static const FName Name("Name");
+};
+
+
+class SMutableGraphTreeRow : public STableRow<TSharedPtr<FMutableGraphTreeElement>>
+{
+public:
+
+	void Construct(const FArguments& Args, const TSharedRef<STableViewBase>& InOwnerTableView, const TSharedPtr<FMutableGraphTreeElement>& InRowItem)
+	{
+		RowItem = InRowItem;
+
+		const char* TypeName = RowItem->MutableNode->GetType()->m_strName;
+
+		FText MainLabel = FText::FromString(ANSI_TO_TCHAR(TypeName));
+		if (RowItem->DuplicatedOf)
+		{
+			MainLabel = FText::FromString( FString::Printf(TEXT("%s (Duplicated)"), ANSI_TO_TCHAR(TypeName)));
+		}
+
+		// TODO
+		const FSlateBrush* IconBrush = nullptr; 
+
+		this->ChildSlot
+		[
+			SNew(SHorizontalBox)
+			+ SHorizontalBox::Slot()
+			.VAlign(VAlign_Center)
+			.AutoWidth()
+			[
+				SNew(SExpanderArrow, SharedThis(this))
+			]
+
+			+ SHorizontalBox::Slot()
+			.VAlign(VAlign_Center)
+			.Padding(FMargin(5.f, 0.f, 5.f, 0.f))
+			.AutoWidth()
+			[
+				SNew(SOverlay)
+				+ SOverlay::Slot()
+				[
+					SNew(SImage)
+					.Image(IconBrush ? IconBrush : FCoreStyle::Get().GetDefaultBrush())
+					.ColorAndOpacity(IconBrush ? FLinearColor::White : FLinearColor::Transparent)
+				]
+			]
+
+			+ SHorizontalBox::Slot()
+			[
+				SNew(STextBlock)
+				.Text(MainLabel)
+			]
+		];
+
+		STableRow< TSharedPtr<FMutableGraphTreeElement> >::ConstructInternal(
+			STableRow::FArguments()
+			//.Style(FAppStyle::Get(), "DetailsView.TreeView.TableRow")
+			.ShowSelection(true)
+			, InOwnerTableView
+		);
+
+	}
+
+
+private:
+
+	TSharedPtr<FMutableGraphTreeElement> RowItem;
+};
+
+
+void SMutableGraphViewer::AddReferencedObjects(FReferenceCollector& Collector)
+{
+	// Add UObjects here if we own any at some point
+	//Collector.AddReferencedObject(CustomizableObject);
+}
+
+
+FString SMutableGraphViewer::GetReferencerName() const
+{
+	return TEXT("SMutableGraphViewer");
+}
+
+
+void SMutableGraphViewer::Construct(const FArguments& InArgs, const mu::NodePtr& InRootNode, 
+	const FCompilationOptions& InCompileOptions,
+	TWeakPtr<FTabManager> InParentTabManager, const FName& InParentNewTabId)
+{
+	DataTag = InArgs._DataTag;
+	RootNode = InRootNode;
+	CompileOptions = InCompileOptions;	
+	ParentTabManager = InParentTabManager;
+	ParentNewTabId = InParentNewTabId;
+
+	FToolBarBuilder ToolbarBuilder(TSharedPtr<const FUICommandList>(), FMultiBoxCustomization::None, TSharedPtr<FExtender>(), true);
+	ToolbarBuilder.SetLabelVisibility(EVisibility::Visible);
+	ToolbarBuilder.SetStyle(&FAppStyle::Get(), "SlimToolBar");
+
+	// Export
+	ToolbarBuilder.AddToolBarButton(
+		FUIAction(
+			FExecuteAction::CreateLambda([InRootNode]()
+				{
+					TArray<FString> SaveFilenames;
+					IDesktopPlatform* DesktopPlatform = FDesktopPlatformModule::Get();
+					bool bSave = false;
+					if (!DesktopPlatform) return;
+
+					FString LastExportPath = FEditorDirectories::Get().GetLastDirectory(ELastDirectory::GENERIC_EXPORT);
+					FString FileTypes = TEXT("Mutable source data files|*.mutable_source|All files|*.*");
+					bSave = DesktopPlatform->SaveFileDialog(
+						FSlateApplication::Get().FindBestParentWindowHandleForDialogs(nullptr),
+						TEXT("Export Mutable object"),
+						*LastExportPath,
+						TEXT("exported.mutable_source"),
+						*FileTypes,
+						EFileDialogFlags::None,
+						SaveFilenames
+					);
+
+					if (!bSave) return;
+
+					// Dump source model to a file.
+					FString SaveFileName = FString(SaveFilenames[0]);
+					mu::OutputFileStream stream(TCHAR_TO_ANSI(*SaveFileName));
+					stream.Write(MUTABLE_SOURCE_MODEL_FILETAG, 4);
+					mu::OutputArchive arch(&stream);
+					mu::Node::Serialise(InRootNode.get(), arch);
+					stream.Flush();
+
+					FEditorDirectories::Get().SetLastDirectory(ELastDirectory::GENERIC_EXPORT, SaveFileName);
+				})
+			),
+			NAME_None,
+			LOCTEXT("ExportMutableGraph", "Export"),
+			LOCTEXT("ExportMutableGraphTooltip", "Export a debug mutable graph file."),
+			FSlateIcon(),
+			EUserInterfaceActionType::Button
+		);
+		
+	ToolbarBuilder.EndSection();
+
+
+	ToolbarBuilder.BeginSection("Compilation");
+
+	ToolbarBuilder.AddToolBarButton(
+		FUIAction(FExecuteAction::CreateSP(this, &SMutableGraphViewer::CompileMutableCodePressed)),
+		NAME_None,
+		LOCTEXT("GenerateMutableCode", "Unreal to Mutable Code"),
+		LOCTEXT("GenerateMutableCodeTooltip", "Generate a mutable code from the mutable graph."),
+		FSlateIcon(FCustomizableObjectEditorStyle::Get().GetStyleSetName(), "CustomizableObjectDebugger.CompileMutableCode", "CustomizableObjectDebugger.CompileMutableCode.Small"),
+		EUserInterfaceActionType::Button
+	);
+
+	ToolbarBuilder.AddComboButton(
+		FUIAction(),
+		FOnGetContent::CreateSP(this, &SMutableGraphViewer::GenerateCompileOptionsMenuContent),
+		LOCTEXT("Compile_Options_Label", "Compile Options"),
+		LOCTEXT("Compile_Options_Tooltip", "Change Compile Options"),
+		TAttribute<FSlateIcon>(),
+		true);
+
+	ToolbarBuilder.EndSection();
+
+	ToolbarBuilder.AddWidget(SNew(STextBlock).Text(MakeAttributeLambda([this]() { return FText::FromString(DataTag); })));
+
+	ChildSlot
+	[
+		SNew(SVerticalBox)
+		+ SVerticalBox::Slot()
+		.AutoHeight()
+		.VAlign(VAlign_Center)
+		[
+			ToolbarBuilder.MakeWidget()
+		]
+		+ SVerticalBox::Slot()
+		.VAlign(VAlign_Fill)
+		[
+			SNew(SSplitter)
+			.Orientation(EOrientation::Orient_Horizontal)
+			+ SSplitter::Slot()
+			.Value(0.25f)
+			[
+				SNew(SBorder)
+				.BorderImage(FAppStyle::GetBrush("ToolPanel.GroupBorder"))
+				.Padding(FMargin(4.0f, 4.0f))
+				[
+					SAssignNew(TreeView, STreeView<TSharedPtr<FMutableGraphTreeElement>>)
+					.TreeItemsSource(&RootNodes)
+					.OnGenerateRow(this,&SMutableGraphViewer::GenerateRowForNodeTree)
+					.OnGetChildren(this, &SMutableGraphViewer::GetChildrenForInfo)
+					.OnSetExpansionRecursive(this, &SMutableGraphViewer::TreeExpandRecursive)
+					.OnContextMenuOpening(this, &SMutableGraphViewer::OnTreeContextMenuOpening)
+					.SelectionMode(ESelectionMode::Single)
+					.HeaderRow
+					(
+						SNew(SHeaderRow)
+						+ SHeaderRow::Column(MutableGraphTreeViewColumns::Name)
+						.FillWidth(25.f)
+						.DefaultLabel(LOCTEXT("Node Name", "Node Name"))
+					)
+				]
+			]
+			+ SSplitter::Slot()
+			.Value(0.75f)
+			[
+				SNew(SBorder)
+				.BorderImage(FAppStyle::GetBrush("ToolPanel.GroupBorder"))
+				.Padding(FMargin(4.0f, 4.0f))
+				//[
+				//	SubjectsTreeView->AsShared()
+				//]
+			]
+		]
+	];
+	
+	RebuildTree();
+}
+
+
+void SMutableGraphViewer::CompileMutableCodePressed()
+{
+	// Do the compilation to Mutable Code synchronously.
+	TSharedPtr<FCustomizableObjectCompileRunnable> CompileTask = MakeShareable(new FCustomizableObjectCompileRunnable(RootNode, false));
+	CompileTask->Options = CompileOptions;
+	CompileTask->Init();
+	CompileTask->Run();
+
+	FString NewDataTag = FString::Printf(TEXT("%s graph, opt %d "), *DataTag, CompileOptions.OptimizationLevel);
+
+	TSharedPtr<SDockTab> NewMutableCodeTab = SNew(SDockTab)
+		.Label(LOCTEXT("MutableCode", "Mutable Code"))
+		[
+			SNew(SMutableCodeViewer, CompileTask->Model)
+			.DataTag(NewDataTag)
+		];
+
+	TSharedPtr<FTabManager> TabManager = ParentTabManager.Pin();
+	check(TabManager);
+	TabManager->InsertNewDocumentTab(ParentNewTabId, FTabManager::ESearchPreference::PreferLiveTab, NewMutableCodeTab.ToSharedRef());
+}
+
+
+TSharedRef<SWidget> SMutableGraphViewer::GenerateCompileOptionsMenuContent()
+{
+	const bool bShouldCloseWindowAfterMenuSelection = false;
+	FMenuBuilder MenuBuilder(bShouldCloseWindowAfterMenuSelection, nullptr);
+
+	// settings
+	MenuBuilder.BeginSection("Optimization", LOCTEXT("MutableCompileOptimizationHeading", "Optimization"));
+	{
+		// Compilation options
+		//-----------------------------------
+
+		// Optimisation level
+		CompileOptimizationStrings.Empty();
+		CompileOptimizationStrings.Add(MakeShareable(new FString(NSLOCTEXT("CustomizableObjectEditor", "OptimizationNone", "None").ToString())));
+		CompileOptimizationStrings.Add(MakeShareable(new FString(NSLOCTEXT("CustomizableObjectEditor", "OptimizationMin", "Minimal").ToString())));
+		CompileOptimizationStrings.Add(MakeShareable(new FString(NSLOCTEXT("CustomizableObjectEditor", "OptimizationMed", "Medium").ToString())));
+		CompileOptimizationStrings.Add(MakeShareable(new FString(NSLOCTEXT("CustomizableObjectEditor", "OptimizationMax", "Maximum").ToString())));
+
+		CompileOptions.OptimizationLevel = FMath::Min(CompileOptions.OptimizationLevel, CompileOptimizationStrings.Num() - 1);
+
+		CompileOptimizationCombo =
+			SNew(STextComboBox)
+			.OptionsSource(&CompileOptimizationStrings)
+			.InitiallySelectedItem(CompileOptimizationStrings[CompileOptions.OptimizationLevel])
+			.OnSelectionChanged(this, &SMutableGraphViewer::OnChangeCompileOptimizationLevel)
+			;
+		MenuBuilder.AddWidget(CompileOptimizationCombo.ToSharedRef(), LOCTEXT("MutableCompileOptimizationLevel", "Optimization Level"));
+
+		// Parallel compilation
+		MenuBuilder.AddMenuEntry(
+			LOCTEXT("MutableEnableParallelCompilation", "Enable compiling in multiple threads."),
+			LOCTEXT("MutableEnableParallelCompilationTooltip", "This is faster but use more memory."),
+			FSlateIcon(),
+			FUIAction(
+				FExecuteAction::CreateLambda([this]() { CompileOptions.bUseParallelCompilation = !CompileOptions.bUseParallelCompilation; }),
+				FCanExecuteAction(),
+				FIsActionChecked::CreateLambda([this]() { return CompileOptions.bUseParallelCompilation; })),
+			NAME_None,
+			EUserInterfaceActionType::ToggleButton
+		);
+
+		// Disk as cache
+		MenuBuilder.AddMenuEntry(
+			LOCTEXT("MutableDiskAsMemory", "Enable compiling using the disk as memory."),
+			LOCTEXT("MutableDiskAsMemoryTooltip", "This is very slow but supports compiling huge objects. It requires a lot of free space in the OS disk."),
+			FSlateIcon(),
+			FUIAction(
+				FExecuteAction::CreateLambda([this]() { CompileOptions.bUseDiskCompilation = !CompileOptions.bUseDiskCompilation; }),
+				FCanExecuteAction(),
+				FIsActionChecked::CreateLambda([this]() { return CompileOptions.bUseDiskCompilation; })),
+			NAME_None,
+			EUserInterfaceActionType::ToggleButton
+		);
+	}
+	MenuBuilder.EndSection();
+
+	return MenuBuilder.MakeWidget();
+}
+
+
+void SMutableGraphViewer::OnChangeCompileOptimizationLevel(TSharedPtr<FString> NewSelection, ESelectInfo::Type SelectInfo)
+{
+	CompileOptions.OptimizationLevel = CompileOptimizationStrings.Find(NewSelection);
+}
+
+
+void SMutableGraphViewer::RebuildTree()
+{
+	RootNodes.Reset();
+	ItemCache.Reset();
+	MainItemPerNode.Reset();
+
+	RootNodes.Add(MakeShareable(new FMutableGraphTreeElement(RootNode)));
+	TreeView->RequestTreeRefresh();
+	TreeExpandUnique();
+}
+
+
+TSharedRef<ITableRow> SMutableGraphViewer::GenerateRowForNodeTree(TSharedPtr<FMutableGraphTreeElement> InTreeNode, const TSharedRef<STableViewBase>& InOwnerTable)
+{
+	TSharedRef<SMutableGraphTreeRow> Row = SNew(SMutableGraphTreeRow, InOwnerTable, InTreeNode);
+	return Row;
+}
+
+void SMutableGraphViewer::GetChildrenForInfo(TSharedPtr<FMutableGraphTreeElement> InInfo, TArray<TSharedPtr<FMutableGraphTreeElement>>& OutChildren)
+{
+	if (!InInfo->MutableNode)
+	{
+		return;
+	}
+
+	// If this is a duplicated of another row, don't provide its children.
+	if (InInfo->DuplicatedOf)
+	{
+		return;
+	}
+
+	// Generic node case
+	int32 InputCount = InInfo->MutableNode->GetInputCount();
+	for (int32 InputIndex=0; InputIndex<InputCount; ++InputIndex)
+	{
+		mu::NodePtr ChildNode = InInfo->MutableNode->GetInputNode( InputIndex );
+		if (ChildNode)
+		{
+			FItemCacheKey Key = { InInfo->MutableNode.get(),ChildNode.get(), uint32(InputIndex) };
+			TSharedPtr<FMutableGraphTreeElement>* CachedItem = ItemCache.Find(Key);
+
+			if (CachedItem)
+			{
+				OutChildren.Add(*CachedItem);
+			}
+			else
+			{
+				TSharedPtr<FMutableGraphTreeElement>* MainItemPtr = MainItemPerNode.Find(ChildNode.get());
+				TSharedPtr<FMutableGraphTreeElement> Item = MakeShareable(new FMutableGraphTreeElement(ChildNode, MainItemPtr));
+				OutChildren.Add(Item);
+				ItemCache.Add(Key, Item);
+
+				if (!MainItemPtr)
+				{
+					MainItemPerNode.Add(ChildNode.get(),Item);
+				}
+			}
+		}
+	}
+}
+
+
+TSharedPtr<SWidget> SMutableGraphViewer::OnTreeContextMenuOpening()
+{
+	FMenuBuilder MenuBuilder(true, nullptr);
+
+	MenuBuilder.AddMenuEntry(
+		LOCTEXT("Graph_Expand_Instance", "Expand Instance-Level Operations"),
+		LOCTEXT("Graph_Expand_Instance_Tooltip", "Expands all the operations in the tree that are instance operations (not images, meshes, booleans, etc.)."),
+		FSlateIcon(),
+		FUIAction(FExecuteAction::CreateSP(this, &SMutableGraphViewer::TreeExpandUnique)
+			//, FCanExecuteAction::CreateSP(this, &SMutableCodeViewer::HasAnyItemInPalette)
+		)
+	);
+
+	return MenuBuilder.MakeWidget();
+}
+
+
+void SMutableGraphViewer::TreeExpandRecursive(TSharedPtr<FMutableGraphTreeElement> InInfo, bool bExpand)
+{
+	if (bExpand)
+	{
+		TreeExpandUnique();
+	}
+}
+
+
+void SMutableGraphViewer::TreeExpandUnique()
+{
+	TArray<TSharedPtr<FMutableGraphTreeElement>> Pending = RootNodes;
+
+	TSet<TSharedPtr<FMutableGraphTreeElement>> Processed;
+
+	TArray<TSharedPtr<FMutableGraphTreeElement>> Children;
+
+	while (!Pending.IsEmpty())
+	{
+		TSharedPtr<FMutableGraphTreeElement> Item = Pending.Pop();
+		TreeView->SetItemExpansion(Item, true);
+
+		Children.SetNum(0);
+		GetChildrenForInfo(Item, Children);
+		Pending.Append(Children);
+	}
+}
+
+
+
+FReply SMutableGraphViewer::OnDragOver(const FGeometry& MyGeometry, const FDragDropEvent& DragDropEvent)
+{
+	if (TSharedPtr<FExternalDragOperation> DragDropOp = DragDropEvent.GetOperationAs<FExternalDragOperation>())
+	{
+		if (DragDropOp->HasFiles())
+		{
+			// For now, only allow a single file.
+			const TArray<FString>& Files = DragDropOp->GetFiles();
+			if (Files.Num() == 1)
+			{
+				const FString DraggedFileExtension = FPaths::GetExtension(Files[0], true);
+				if (DraggedFileExtension == TEXT(".mutable_source"))
+				{
+					// Dump source model to a file.
+					mu::InputFileStream stream(TCHAR_TO_ANSI(*Files[0]));
+
+					char MutableSourceTag[4] = {};
+					stream.Read(MutableSourceTag, 4);
+
+					if (!FMemory::Memcmp(MutableSourceTag, MUTABLE_SOURCE_MODEL_FILETAG, 4))
+					{
+						return FReply::Handled();
+					}
+
+					return FReply::Unhandled();
+				}
+			}
+		}
+	}
+
+	return FReply::Unhandled();
+}
+
+
+FReply SMutableGraphViewer::OnDrop(const FGeometry& MyGeometry, const FDragDropEvent& DragDropEvent)
+{
+	SCompoundWidget::OnDrop(MyGeometry, DragDropEvent);
+
+	if (TSharedPtr<FExternalDragOperation> DragDropOp = DragDropEvent.GetOperationAs<FExternalDragOperation>())
+	{
+		if (DragDropOp->HasFiles())
+		{
+			// For now, only allow a single file.
+			const TArray<FString>& Files = DragDropOp->GetFiles();
+			if (Files.Num() == 1)
+			{
+				const FString DraggedFileExtension = FPaths::GetExtension(Files[0], true);
+				if (DraggedFileExtension == TEXT(".mutable_source"))
+				{
+					// Dump source model to a file.
+					mu::InputFileStream stream(TCHAR_TO_ANSI(*Files[0]));
+
+					char MutableSourceTag[4] = {};
+					stream.Read(MutableSourceTag, 4);
+
+					if (!FMemory::Memcmp(MutableSourceTag, MUTABLE_SOURCE_MODEL_FILETAG, 4))
+					{
+						mu::InputArchive arch(&stream);
+						RootNode = mu::Node::StaticUnserialise( arch );
+						DataTag = FString("dropped-file ") + FPaths::GetCleanFilename(Files[0]);
+						RebuildTree();
+
+						return FReply::Handled();
+					}
+
+					return FReply::Unhandled();
+				}
+			}
+		}
+
+		return FReply::Unhandled();
+	}
+
+	return FReply::Unhandled();
+}
+
+#undef LOCTEXT_NAMESPACE 
diff --git a/Engine/Plugins/Experimental/Mutable/Source/CustomizableObjectEditor/Private/MuCOE/SMutableGraphViewer.h b/Engine/Plugins/Experimental/Mutable/Source/CustomizableObjectEditor/Private/MuCOE/SMutableGraphViewer.h
new file mode 100644
--- /dev/null
+++ b/Engine/Plugins/Experimental/Mutable/Source/CustomizableObjectEditor/Private/MuCOE/SMutableGraphViewer.h
@@ -0,0 +1,135 @@
+// Copyright Epic Games, Inc. All Rights Reserved.
+
+#pragma once
+
+#include "Widgets/SCompoundWidget.h"
+#include "MuCO/CustomizableObject.h"
+#include "MuT/Node.h"
+#include "Widgets/Views/STreeView.h"
+#include "Widgets/Input/STextComboBox.h"
+#include "MuCOE/CustomizableObjectCompiler.h"
+#include "Framework/Docking/TabManager.h"
+
+
+// Forward declarations
+class FMutableGraphTreeElement;
+
+
+/** This widget shows the internal Mutable Graph for debugging purposes. 
+ * This is not the Unreal source graph in the UCustomizableObject, but an intermediate step of the compilation process.
+ */
+class SMutableGraphViewer :
+	public SCompoundWidget,
+	public FGCObject
+{
+public:
+
+	// SWidget interface
+	SLATE_BEGIN_ARGS(SMutableGraphViewer) {}
+
+	/** User-visible tag to indentify the source of the data shwon. */
+	SLATE_ARGUMENT(FString, DataTag)
+
+	SLATE_END_ARGS()
+
+	void Construct(const FArguments& InArgs, const mu::NodePtr& InRootNode, const FCompilationOptions& InCompileOptions,
+		TWeakPtr<FTabManager> InParentTabManager, const FName& InParentNewTabId);
+
+	// SWidget interface
+	FReply OnDragOver(const FGeometry& MyGeometry, const FDragDropEvent& DragDropEvent) override;
+	FReply OnDrop(const FGeometry& MyGeometry, const FDragDropEvent& DragDropEvent) override;
+
+	// FGCObject interface
+	virtual void AddReferencedObjects(FReferenceCollector& Collector) override;
+	virtual FString GetReferencerName() const override;
+
+private:
+
+	/** Tag for the user to identify the data being shown. */
+	FString DataTag;
+
+	/** The Mutable Graph to show, represented by its root. */
+	mu::NodePtr RootNode;
+
+	/** Compilation options to use in the debugger operations. */
+	FCompilationOptions CompileOptions;
+
+	/** Object compiler. */
+	FCustomizableObjectCompiler Compiler;
+
+	/** UI references used to create new tabs. */
+	TWeakPtr<FTabManager> ParentTabManager;
+	FName ParentNewTabId;
+
+	/** Root nodes of the tree widget. */
+	TArray<TSharedPtr<FMutableGraphTreeElement>> RootNodes;
+
+	/** Tree showing the graph. */
+	TSharedPtr<STreeView<TSharedPtr<FMutableGraphTreeElement>>> TreeView;
+
+	/** Cache of tree elements matching the graph nodes that have been generated so far. 
+	* We store both the parent and the node in the key, because a single node may appear multiple times if it has different parents.
+	*/
+	struct FItemCacheKey
+	{
+		const mu::Node* Parent = nullptr;
+		const mu::Node* Child = nullptr;
+		uint32 ChildIndexInParent = 0;
+
+		friend FORCEINLINE bool operator == (const FItemCacheKey& A, const FItemCacheKey& B)
+		{
+			return A.Parent == B.Parent && A.Child == B.Child && A.ChildIndexInParent == B.ChildIndexInParent;
+		}
+
+		friend FORCEINLINE uint32 GetTypeHash(const FItemCacheKey& Key)
+		{
+			return HashCombine( GetTypeHash(Key.Parent), HashCombine(GetTypeHash(Key.Child), Key.ChildIndexInParent));
+		}
+	};
+	TMap< FItemCacheKey, TSharedPtr<FMutableGraphTreeElement>> ItemCache;
+
+	/** Main tree item for each node. A graph node can be representes with multiple tree ndoes if it is reachable from different paths. */
+	TMap< const mu::Node*, TSharedPtr<FMutableGraphTreeElement>> MainItemPerNode;
+
+	/** */
+	void RebuildTree();
+
+	/** UI callbacks */
+	void CompileMutableCodePressed();
+	TSharedRef<SWidget> GenerateCompileOptionsMenuContent();
+	TSharedPtr<STextComboBox> CompileOptimizationCombo;
+	TArray< TSharedPtr<FString> > CompileOptimizationStrings;
+	void OnChangeCompileOptimizationLevel(TSharedPtr<FString> NewSelection, ESelectInfo::Type SelectInfo);
+
+	/** Callbacks from the tree widget. */
+ 	TSharedRef<ITableRow> GenerateRowForNodeTree(TSharedPtr<FMutableGraphTreeElement> InTreeNode, const TSharedRef<STableViewBase>& InOwnerTable);
+	void GetChildrenForInfo(TSharedPtr<FMutableGraphTreeElement> InInfo, TArray< TSharedPtr<FMutableGraphTreeElement> >& OutChildren);
+	TSharedPtr<SWidget> OnTreeContextMenuOpening();
+	void TreeExpandRecursive(TSharedPtr<FMutableGraphTreeElement> Info, bool bExpand);
+	void TreeExpandUnique();
+
+};
+
+
+/** An row of the code tree in the SMutableGraphViewer. */
+class FMutableGraphTreeElement : public TSharedFromThis<FMutableGraphTreeElement>
+{
+public:
+	FMutableGraphTreeElement(const mu::NodePtr& InNode, TSharedPtr<FMutableGraphTreeElement>* InDuplicatedOf=nullptr )
+	{
+		MutableNode = InNode;
+		if (InDuplicatedOf)
+		{
+			DuplicatedOf = *InDuplicatedOf;
+		}
+	}
+
+public:
+
+	/** Mutable Graph Node represented in this tree row. */
+	mu::NodePtr MutableNode;
+
+	/** If this tree element is a duplicated of another node, this is the node. */
+	TSharedPtr<FMutableGraphTreeElement> DuplicatedOf;
+
+};
diff --git a/Engine/Plugins/Experimental/Mutable/Source/CustomizableObjectEditor/Private/MuCOE/SMutableMeshViewer.h b/Engine/Plugins/Experimental/Mutable/Source/CustomizableObjectEditor/Private/MuCOE/SMutableMeshViewer.h
new file mode 100644
--- /dev/null
+++ b/Engine/Plugins/Experimental/Mutable/Source/CustomizableObjectEditor/Private/MuCOE/SMutableMeshViewer.h
@@ -0,0 +1,209 @@
+// Copyright Epic Games, Inc. All Rights Reserved.
+
+#pragma once
+
+#include "Containers/Array.h"
+#include "MuR/Mesh.h"
+#include "Widgets/DeclarativeSyntaxSupport.h"
+#include "Widgets/SCompoundWidget.h"
+#include "Widgets/Views/SListView.h"
+#include "Widgets/Views/STableRow.h"
+#include "Widgets/Layout/SBorder.h"
+#include "AssetThumbnail.h"
+#include "UObject/GCObject.h"
+
+#include "Templates/SharedPointer.h"
+
+#include "SMutableMeshViewport.h"
+#include "Widgets/Input/SSearchBox.h"
+#include "Widgets/Layout/SScrollBox.h"
+
+#include "PropertyCustomizationHelpers.h"
+#include "SMutableSkeletonViewer.h"
+
+
+/** Container designed to hold the buffer channel data of a mutable mesh buffer to be later used by the UI */
+struct FBufferChannelElement
+{
+	FText SemanticIndex = FText(INVTEXT(""));
+	FText BufferSemantic = FText(INVTEXT(""));
+	FText BufferFormat = FText(INVTEXT(""));
+	FText BufferComponentCount = FText(INVTEXT(""));
+};
+
+
+/** Element representing a mutable buffer. It contains an array with all the elements representing the channels the mutable buffer is made of.*/
+struct FBufferElement
+{
+	/** The index of the buffer on the origin mutable buffer set */
+	FText BufferIndex = FText(INVTEXT(""));
+
+	/** An array of BufferChannels that represent the relative mutable channels on the mutable buffer*/
+	TSharedPtr<TArray<TSharedPtr<FBufferChannelElement>>> BufferChannels;
+};
+
+/** Widget designed to show the statistical data from a Mutable Mesh*/
+class SMutableMeshViewer : public SCompoundWidget, public FGCObject
+{
+public:
+	SLATE_BEGIN_ARGS(SMutableMeshViewer) {}
+		SLATE_ARGUMENT_DEFAULT(mu::MeshPtrConst, Mesh){nullptr};
+	SLATE_END_ARGS()
+
+	/** Builds the widget */
+	void Construct(const FArguments& InArgs);
+
+	// SWidget interface
+	// void Tick(const FGeometry& AllottedGeometry, const double InCurrentTime, const float InDeltaTime) override;
+
+	// FGCObject interface
+	virtual void AddReferencedObjects(FReferenceCollector& Collector) override;
+
+	virtual FString GetReferencerName() const override
+	{
+		return TEXT("SMutableMeshViewer");
+	}
+
+	/** Set the Mutable Mesh to be used for this widget */
+	void SetMesh(const mu::MeshPtrConst& InMesh);
+
+	
+	/** Generate a new SListView for all the ChannelElements provided.
+	 * @param InBufferChannelElements - List with all the elements that will be ued to fill InHostListView
+	 * @return A SListView containing all the channel elements as UI elements. 
+	 */
+	TSharedRef<SWidget> GenerateBufferChannelsListView(
+		const TSharedPtr<TArray<TSharedPtr<FBufferChannelElement>>>& InBufferChannelElements);
+
+private:
+	
+	/** Slate whose task is to display the skeleton found on this mesh as a slate tree view */
+	TSharedPtr<SMutableSkeletonViewer> MutableSkeletonViewer;
+
+	/** Thumbnail objects used to render a view of the selected skeletal mesh for the mesh viewport*/
+	TSharedPtr<FAssetThumbnailPool> AssetThumbnailPool;
+	TSharedPtr<FAssetThumbnail> AssetThumbnail;
+
+	/** Data backend for the widget. It represents the mesh that is being "displayed" */
+	mu::MeshPtrConst MutableMesh = nullptr;
+
+	/** Is true, mesh has changed and we need to update. */
+	bool bIsPendingUpdate = false;
+
+	/** Splitter used to separate the two sides of the slate (tables and viewport) */
+	TSharedPtr<SSplitter> SpaceSplitter;
+
+	/** Slate object containing all the buffer tables alongside with the bone tree */
+	TSharedPtr<SScrollBox> DataSpaceSlate;
+
+	/** Viewport object to preview the current mesh inside an actual Unreal scene */
+	TSharedPtr<SMutableMeshViewport> MeshViewport;
+	
+	/*
+	 * Slate views for the main types of mesh buffers (vertex , index and face)
+	 * Each buffer element also contains the channels it uses
+	 */
+	
+	TSharedPtr<SListView<TSharedPtr<FBufferElement>>> VertexBuffersSlateView;
+	TSharedPtr<SListView<TSharedPtr<FBufferElement>>> IndexBuffersSlateView;
+	TSharedPtr<SListView<TSharedPtr<FBufferElement>>> FaceBuffersSlateView;
+
+private:
+	
+	/** The border where the object responsible for providing the user of a way to select objects is */
+	TSharedPtr<SBorder> ReferenceMeshSelectionSpace;
+
+	/** Generates all slate objects related with the Mesh Viewport Slate */
+	TSharedRef<SWidget> GenerateViewportSlates();
+
+	/** Generates the tables showing the buffer data on the mesh alongside with the bone tree found on the mutable mesh */
+	TSharedRef<SWidget> GenerateDataTableSlates();
+	
+	/** Width and height for the skeletal mesh thumbnail */
+	FVector2f ThumbnailSize = FVector2f();
+	FVector2f WidgetSize = FVector2f();
+
+	// Single property that only draws the combo box widget of the skeletal mesh
+	TSharedPtr<class ISinglePropertyView> SkeletalMeshSelector;
+	
+	/** Method called each time the mesh selected changes so the UI gets updated reliably */
+	void OnMeshChanged();
+
+private:
+	
+	/* Generic UI callbacks used by the Widget */
+	FText GetVertexCount() const;
+	FText GetFaceCount() const;
+	FText GetBoneCount() const;
+
+	/*
+	 * Elements used to feed the buffers list (index and buffer channels as an internal list)
+	 */
+	
+	TArray<TSharedPtr<FBufferElement>> VertexBuffers;
+	TArray<TSharedPtr<FBufferElement>> IndexBuffers;
+	TArray<TSharedPtr<FBufferElement>> FaceBuffers;
+	
+	/** 
+	* Fills the provided TArray with the buffer definitions generated from the mutable buffers found on the provided mutable buffer set
+	* @param InMuMeshBufferSetPtr - Mutable buffer set pointer pointing to the buffer set witch buffer data is wanted to be saved on the provided TArray
+	* @param OutBuffersDataArray - TArray witch will contain the parsed data from the provided mutable buffer set converted onto FBufferElements
+	* @param InHostListView - The slate list using the previously referenced array as data backend. 
+	*/
+	void FillTargetBufferSetDataArray(const mu::FMeshBufferSet& InMuMeshBufferSetPtr,
+	                                  TArray<TSharedPtr<FBufferElement>>& OutBuffersDataArray,
+	                                  TSharedPtr<SListView<TSharedPtr<FBufferElement>>>& InHostListView);
+
+	
+	/**
+	 * Provided an array of buffer elements and the type of buffer it generates a new SListView for said buffer elements
+	 * @param InBufferElements - The FBufferElement objects to be used as backend for the SListView
+	 * @param InBufferSetTypeName - The name of the buffer set you are generating. Some examples are "Vertex", "index" or "Face"
+	 * @param OutHostListView - THe SListView generated.
+	 * @return - A shared reference to the SListView generated. Required by this slate operation.
+	 */
+	TSharedRef<SWidget> GenerateBuffersListView(TSharedPtr<SListView<TSharedPtr<FBufferElement>>>& OutHostListView,
+	                                          const TArray<TSharedPtr<FBufferElement>>& InBufferElements,
+	                                          const FText& InBufferSetTypeName);
+	
+	/** Callback method invoked each time a new row of the SListView containing the buffer elements needs to be built
+	 * @param InBuffer - The buffer element being used to build the new row.
+	 * @param OwnerTable - The table that is going to get the new row added into it.
+	 * @return The new row object as the base interface used for all the table rows
+	 */
+	TSharedRef<ITableRow> OnGenerateBufferRow(TSharedPtr<FBufferElement, ESPMode::ThreadSafe> InBuffer,
+											  const TSharedRef<STableViewBase, ESPMode::ThreadSafe>& OwnerTable);
+	
+	/** 
+	* Callback Method responsible of generating each row of the buffer channel lists based on the channel definition provided.
+	* @param InBufferChannel - The mesh Buffer channel definition that is being used as blueprint for the actual UI row
+	* @param OwnerTable - The parent UI element that is going to be expanded with the new row object
+	* @return The new row object as the base interface used for all the table rows
+	*/
+	TSharedRef<ITableRow> OnGenerateBufferChannelRow(TSharedPtr<FBufferChannelElement> InBufferChannel,
+	                                                 const TSharedRef<STableViewBase>& OwnerTable);
+	
+	
+	/*
+	 * Reference skeleton mesh selection
+	 */
+
+	/** Currently selected reference mesh */
+	const USkeletalMesh* SelectedReferenceMesh = nullptr;
+
+	/** Callback invoked when the selected asset on the reference skeletal mesh selection space has changed
+	 * @param AssetData - The data of the selected asset on the selection space
+	 */
+	void OnSelectedReferenceSkeletalMeshChanged(const FAssetData& AssetData);
+
+	/** Filters the assets found in order to know if the asset should be exposed on the UI for usage
+	 * @param AssetData - The data of the asset being checked for validation
+	 * @return - True means that the asset WILL NOT appear on the list of selectable resources
+	 */
+	bool OnShouldFilterAsset(const FAssetData& AssetData) const;
+
+	/**
+	 * Provides the SObjectPropertyEntryBox a way to know what to display in the selected object space
+	 */
+	FString OnObjectPath() const;
+};
diff --git a/Engine/Plugins/Experimental/Mutable/Source/CustomizableObjectEditor/Private/MuCOE/SMutableMeshViewport.h b/Engine/Plugins/Experimental/Mutable/Source/CustomizableObjectEditor/Private/MuCOE/SMutableMeshViewport.h
new file mode 100644
--- /dev/null
+++ b/Engine/Plugins/Experimental/Mutable/Source/CustomizableObjectEditor/Private/MuCOE/SMutableMeshViewport.h
@@ -0,0 +1,96 @@
+// Copyright Epic Games, Inc. All Rights Reserved.
+
+#pragma once
+
+#include "MuR/Mesh.h"
+#include "FMutableViewportClient.h"
+#include "Widgets/SCompoundWidget.h"
+#include "AdvancedPreviewScene.h"
+#include "Components/SkeletalMeshComponent.h"
+#include "SEditorViewport.h"
+#include "UnrealWidget.h"
+#include "UObject/GCObject.h"
+
+#include "Engine/SkeletalMesh.h"
+
+
+/**
+* Object with the objective of showing a preview of a selected mutable mesh. It is designed to have few dependencies
+* as possible
+*/
+class SMutableMeshViewport : public SEditorViewport , public FGCObject
+{
+	SLATE_BEGIN_ARGS(SMutableMeshViewport) {}
+    	SLATE_ARGUMENT_DEFAULT(mu::MeshPtrConst, Mesh) { nullptr };
+    SLATE_END_ARGS()
+    
+public:
+    /** Builds the widget */
+    void Construct(const FArguments& InArgs);
+	
+    /** Set the Mutable Mesh to be used for this widget
+     * @param InMesh - The mutable mesh to be displayed
+     */
+    void SetMesh(const mu::MeshPtrConst& InMesh);
+
+	/** Set the reference mesh used by the mutable to skeletal mesh conversion process
+	 * @param InReferenceMesh - The Skeletal mesh to be used to aid on the conversion of the MutableMesh to a SkeletalMesh
+	 */
+	void SetReferenceMesh(const USkeletalMesh* InReferenceMesh);
+
+	// FGCObject interface
+	virtual void AddReferencedObjects(FReferenceCollector& Collector) override;
+	virtual FString GetReferencerName() const override
+	{
+		return TEXT("SMutableMeshViewport");
+	}
+	
+    // SWidget interface
+	virtual void Tick(const FGeometry& AllottedGeometry, const double InCurrentTime, const float InDeltaTime) override;
+
+protected:
+
+	// SEditorViewport interface (called by SEditorViewport::Construct)
+	virtual TSharedRef<FEditorViewportClient> MakeEditorViewportClient() override;
+
+private:
+
+	/** The mutable mesh being displayed */
+	mu::MeshPtrConst MutableMesh = nullptr;
+
+	/*
+	 *	UI  
+	 */
+	
+	/** Skeleton provided by the user */
+	const USkeletalMesh* ReferenceSkeletalMesh = nullptr;
+	
+	/** The preview scene that we are viewing */
+	TSharedPtr<FAdvancedPreviewScene> PreviewScene = nullptr;
+
+	/** Editor viewport client */
+	TSharedPtr<FMutableMeshViewportClient> ViewportClient = nullptr;
+
+	/** Restores the default state for the viewport */
+	void ClearViewport() const;
+
+	/** Adds the required content to the viewport*/
+	void SendMeshToViewport();
+
+	/** Clears or displays a mesh depending on the presence of a mutable mesh on MutableMesh**/
+	void RefreshViewportContents();
+
+	/*
+	 * Utility
+	 */
+
+	/** Mutable mesh converted into a Unreal mesh object */
+	USkeletalMeshComponent* SkeletalMeshComponent = nullptr;
+	
+	/** Generates a new USkeletalMesh from the MutableMesh set on this object
+	 * @return - If the operation was able to be processed or not
+	 */
+	bool GenerateUnrealMesh();
+};
+
+
diff --git a/Engine/Plugins/Experimental/Mutable/Source/CustomizableObjectEditor/Private/MuCOE/SMutableObjectViewer.cpp b/Engine/Plugins/Experimental/Mutable/Source/CustomizableObjectEditor/Private/MuCOE/SMutableObjectViewer.cpp
new file mode 100644
--- /dev/null
+++ b/Engine/Plugins/Experimental/Mutable/Source/CustomizableObjectEditor/Private/MuCOE/SMutableObjectViewer.cpp
@@ -0,0 +1,449 @@
+// Copyright Epic Games, Inc. All Rights Reserved.
+
+#include "SMutableObjectViewer.h"
+
+#include "MuCOE/CustomizableObjectEditorStyle.h"
+#include "MuCOE/CustomizableObjectCompiler.h"
+#include "SMutableGraphViewer.h"
+#include "SMutableCodeViewer.h"
+#include "MuT/Streams.h"
+
+#include "Misc/App.h"
+#include "Framework/MultiBox/MultiBoxBuilder.h"
+#include "Widgets/Input/SButton.h"
+#include "Widgets/Images/SImage.h"
+#include "Widgets/Input/STextComboBox.h"
+#include "Widgets/Docking/SDockTab.h"
+#include "EditorDirectories.h"
+#include "DesktopPlatformModule.h"
+#include "HAL/PlatformApplicationMisc.h"
+#include "Interfaces/ITargetPlatform.h"
+#include "Interfaces/ITargetPlatformManagerModule.h"
+
+#define LOCTEXT_NAMESPACE "SMutableDebugger"
+
+// \todo: multi-column tree
+namespace MutableObjectTreeViewColumns
+{
+	static const FName Name("Property");
+	static const FName Value("Value");
+};
+
+
+class SMutableObjectTreeRow : public STableRow<TSharedPtr<FMutableObjectTreeElement>>
+{
+public:
+
+	void Construct(const FArguments& Args, const TSharedRef<STableViewBase>& InOwnerTableView, const TSharedPtr<FMutableObjectTreeElement>& InRowItem, const UCustomizableObject* InObject )
+	{
+		RowItem = InRowItem;
+
+		FText MainLabel;;
+		switch (RowItem->Type)
+		{
+		case FMutableObjectTreeElement::EType::SectionCaption:
+		{
+			switch (RowItem->Section)
+			{
+			case FMutableObjectTreeElement::ESection::General:
+				MainLabel = FText::FromString(TEXT("General"));
+				break;
+			case FMutableObjectTreeElement::ESection::ChildObjects:
+				MainLabel = FText::FromString(TEXT("Children"));
+				break;
+			default:
+				MainLabel = FText::FromString(TEXT("Unknown"));
+				break;
+			}
+			break;
+		}
+
+		case FMutableObjectTreeElement::EType::Name:
+			MainLabel = FText::FromString(InObject->GetName() );
+
+		default:
+			MainLabel = FText::FromString(TEXT("Unknown"));
+			break;
+		}
+
+
+		// TODO
+		const FSlateBrush* IconBrush = nullptr;
+
+		this->ChildSlot
+		[
+			SNew(SHorizontalBox)
+			+ SHorizontalBox::Slot()
+			.VAlign(VAlign_Center)
+			.AutoWidth()
+			[
+				SNew(SExpanderArrow, SharedThis(this))
+			]
+
+			+ SHorizontalBox::Slot()
+				.VAlign(VAlign_Center)
+				.Padding(FMargin(5.f, 0.f, 5.f, 0.f))
+				.AutoWidth()
+				[
+					SNew(SOverlay)
+					+ SOverlay::Slot()
+					[
+						SNew(SImage)
+						.Image(IconBrush ? IconBrush : FCoreStyle::Get().GetDefaultBrush())
+						.ColorAndOpacity(IconBrush ? FLinearColor::White : FLinearColor::Transparent)
+					]
+				]
+
+			+ SHorizontalBox::Slot()
+				[
+					SNew(STextBlock)
+					.Text(MainLabel)
+				]
+		];
+
+		STableRow< TSharedPtr<FMutableObjectTreeElement> >::ConstructInternal(
+			STableRow::FArguments()
+			//.Style(FAppStyle::Get(), "DetailsView.TreeView.TableRow")
+			.ShowSelection(true)
+			, InOwnerTableView
+		);
+
+	}
+
+
+private:
+
+	TSharedPtr<FMutableObjectTreeElement> RowItem;
+};
+
+
+void SMutableObjectViewer::AddReferencedObjects(FReferenceCollector& Collector)
+{
+	Collector.AddReferencedObject(CustomizableObject);
+}
+
+
+FString SMutableObjectViewer::GetReferencerName() const
+{
+	return TEXT("SMutableObjectViewer");
+}
+
+
+void SMutableObjectViewer::Construct(const FArguments& InArgs, UCustomizableObject* InObject, 
+	TWeakPtr<FTabManager> InParentTabManager, const FName& InParentNewTabId )
+{
+	CustomizableObject = InObject;
+	ParentTabManager = InParentTabManager;
+	ParentNewTabId = InParentNewTabId;
+
+	// Initialize the debugger compile options
+	CompileOptions.bTextureCompression = true;
+	CompileOptions.OptimizationLevel = 3;
+	{
+		ITargetPlatformManagerModule* TPM = GetTargetPlatformManager();
+		const TArray<ITargetPlatform*>& Platforms = TPM->GetActiveTargetPlatforms();
+		CompileOptions.TargetPlatform = Platforms.IsEmpty() ? nullptr : Platforms[0];
+	}
+
+	FToolBarBuilder ToolbarBuilder(TSharedPtr<const FUICommandList>(), FMultiBoxCustomization::None, TSharedPtr<FExtender>(), true);
+	ToolbarBuilder.SetLabelVisibility(EVisibility::Visible);
+	ToolbarBuilder.SetStyle(&FAppStyle::Get(), "SlimToolBar");
+
+	ToolbarBuilder.BeginSection("Compilation");
+
+	ToolbarBuilder.AddToolBarButton(
+		FUIAction(FExecuteAction::CreateSP(this, &SMutableObjectViewer::GenerateMutableGraphPressed)),
+		NAME_None,
+		LOCTEXT("GenerateMutableGraph", "Unreal to Mutable Graph"),
+		LOCTEXT("GenerateMutableGraphTooltip", "Generate a mutable graph from the customizable object source graph."),
+		FSlateIcon(FCustomizableObjectEditorStyle::Get().GetStyleSetName(), "CustomizableObjectDebugger.GenerateMutableGraph", "CustomizableObjectDebugger.GenerateMutableGraph.Small"),
+		EUserInterfaceActionType::Button
+	);
+
+	ToolbarBuilder.AddToolBarButton(
+		FUIAction(FExecuteAction::CreateSP(this, &SMutableObjectViewer::CompileMutableCodePressed)),
+		NAME_None,
+		LOCTEXT("GenerateMutableCode", "Unreal to Mutable Code"),
+		LOCTEXT("GenerateMutableCodeTooltipFromGraph", "Generate a mutable code from the customizable object source graph."),
+		FSlateIcon(FCustomizableObjectEditorStyle::Get().GetStyleSetName(), "CustomizableObjectDebugger.CompileMutableCode", "CustomizableObjectDebugger.CompileMutableCode.Small"),
+		EUserInterfaceActionType::Button
+	);
+
+	ToolbarBuilder.AddComboButton(
+		FUIAction(),
+		FOnGetContent::CreateSP(this, &SMutableObjectViewer::GenerateCompileOptionsMenuContent),
+		LOCTEXT("Compile_Options_Label", "Compile Options"),
+		LOCTEXT("Compile_Options_Tooltip", "Change Compile Options"),
+		TAttribute<FSlateIcon>(),
+		true);
+
+	ToolbarBuilder.EndSection();
+
+	// Initialize the root tree nodes
+	RootTreeNodes.Add(MakeShareable(new FMutableObjectTreeElement(FMutableObjectTreeElement::EType::SectionCaption, FMutableObjectTreeElement::ESection::General, InObject)));
+	RootTreeNodes.Add(MakeShareable(new FMutableObjectTreeElement(FMutableObjectTreeElement::EType::SectionCaption, FMutableObjectTreeElement::ESection::ChildObjects, InObject)));
+
+	ChildSlot
+	[
+		SNew(SVerticalBox)
+		+ SVerticalBox::Slot()
+		.AutoHeight()
+		.VAlign(VAlign_Center)
+		[
+			ToolbarBuilder.MakeWidget()
+		]
+		+ SVerticalBox::Slot()
+		.VAlign(VAlign_Fill)
+		[
+			SNew(SBorder)
+			.BorderImage(FAppStyle::GetBrush("ToolPanel.GroupBorder"))
+			.Padding(FMargin(4.0f, 4.0f))
+			[
+				SAssignNew(TreeView, STreeView<TSharedPtr<FMutableObjectTreeElement>>)
+				.TreeItemsSource(&RootTreeNodes)
+				.OnGenerateRow(this,&SMutableObjectViewer::GenerateRowForNodeTree)
+				.OnGetChildren(this, &SMutableObjectViewer::GetChildrenForInfo)
+				.SelectionMode(ESelectionMode::None)					
+			]
+		]
+	];	
+}
+
+
+void SMutableObjectViewer::GenerateMutableGraphPressed()
+{
+	// Convert from Unreal graph to Mutable graph.
+	mu::NodePtr RootNode = Compiler.Export(CustomizableObject, CompileOptions);
+	if (!RootNode)
+	{
+		// TODO: Show errors
+		return;
+	}
+
+	FString DataTag = FString::Printf(TEXT("%s for %s"), *CustomizableObject->GetName(), *CompileOptions.TargetPlatform->PlatformName());
+
+	TSharedPtr<SDockTab> NewMutableGraphTab = SNew(SDockTab)
+		.Label(LOCTEXT("MutableGraph", "Mutable Graph"))
+		[
+			SNew(SMutableGraphViewer, RootNode, CompileOptions, ParentTabManager, ParentNewTabId)
+			.DataTag(DataTag)
+		];
+
+	TSharedPtr<FTabManager> TabManager = ParentTabManager.Pin();
+	check(TabManager);
+	TabManager->InsertNewDocumentTab(ParentNewTabId, FTabManager::ESearchPreference::PreferLiveTab, NewMutableGraphTab.ToSharedRef());
+}
+
+
+void SMutableObjectViewer::CompileMutableCodePressed()
+{
+	// Convert from Unreal graph to Mutable graph.
+	mu::NodePtr RootNode = Compiler.Export(CustomizableObject, CompileOptions);
+	if (!RootNode)
+	{
+		// TODO: Show errors
+		return;
+	}
+
+	// Do the compilation to Mutable Code synchronously.
+	TSharedPtr<FCustomizableObjectCompileRunnable> CompileTask = MakeShareable(new FCustomizableObjectCompileRunnable(RootNode, false));
+	CompileTask->Options = CompileOptions;
+	CompileTask->Init();
+	CompileTask->Run();
+
+	FString DataTag = FString::Printf( TEXT("%s for %s opt %d "), *CustomizableObject->GetName(), *CompileOptions.TargetPlatform->PlatformName(), CompileOptions.OptimizationLevel );
+
+	TSharedPtr<SDockTab> NewMutableCodeTab = SNew(SDockTab)
+		.Label(LOCTEXT("MutableCode", "Mutable Code"))
+		[
+			SNew(SMutableCodeViewer, CompileTask->Model)
+			.DataTag(DataTag)
+		];
+
+	TSharedPtr<FTabManager> TabManager = ParentTabManager.Pin();
+	check(TabManager);
+	TabManager->InsertNewDocumentTab(ParentNewTabId, FTabManager::ESearchPreference::PreferLiveTab, NewMutableCodeTab.ToSharedRef());
+}
+
+
+TSharedRef<SWidget> SMutableObjectViewer::GenerateCompileOptionsMenuContent()
+{
+	const bool bShouldCloseWindowAfterMenuSelection = false;
+	FMenuBuilder MenuBuilder(bShouldCloseWindowAfterMenuSelection, nullptr);
+
+	// settings
+	MenuBuilder.BeginSection("Optimization", LOCTEXT("MutableCompileOptimizationHeading", "Optimization"));
+	{
+		// Unreal Graph to Mutable Graph options
+		//-----------------------------------
+
+		// Platform
+		DebugPlatformStrings.Empty();
+
+		TSharedPtr<FString> SelectedPlatform;
+		{
+			ITargetPlatformManagerModule* TPM = GetTargetPlatformManager();
+			check(TPM);
+
+			ITargetPlatform* CurrentPlatform = NULL;
+			const TArray<ITargetPlatform*>& Platforms = TPM->GetTargetPlatforms();
+			for (const ITargetPlatform* Platform : Platforms)
+			{
+				TSharedPtr<FString> ThisPlatform = MakeShareable(new FString(Platform->PlatformName()));
+				DebugPlatformStrings.Add(ThisPlatform);
+				if (Platform == CompileOptions.TargetPlatform)
+				{
+					SelectedPlatform = ThisPlatform;
+				}
+			}
+		}
+
+		if (!SelectedPlatform.IsValid() && DebugPlatformStrings.Num())
+		{
+			SelectedPlatform = DebugPlatformStrings[0];
+		}
+
+		DebugPlatformCombo =
+			SNew(STextComboBox)
+			.OptionsSource(&DebugPlatformStrings)
+			.InitiallySelectedItem(SelectedPlatform)
+			.OnSelectionChanged(this, &SMutableObjectViewer::OnChangeDebugPlatform)
+			;
+
+		MenuBuilder.AddWidget(DebugPlatformCombo.ToSharedRef(), LOCTEXT("MutableDebugPlatform", "Target Platform"));
+
+		// Texture compresssion
+		MenuBuilder.AddMenuEntry(
+			LOCTEXT("Generate_MutableEnableTextureCompression", "Enable texture compression."),
+			LOCTEXT("Generate_MutableEnableTextureCompressionTooltip", "Enable or disable the compression of textures in the mutable code compilation."),
+			FSlateIcon(),
+			FUIAction(
+				FExecuteAction::CreateLambda([this]() { CompileOptions.bTextureCompression = !CompileOptions.bTextureCompression; }),
+				FCanExecuteAction(),
+				FIsActionChecked::CreateLambda([this]() { return CompileOptions.bTextureCompression; })),
+			NAME_None,
+			EUserInterfaceActionType::ToggleButton
+		);
+
+		// Compilation options
+		//-----------------------------------
+
+		// Optimisation level
+		CompileOptimizationStrings.Empty();
+		CompileOptimizationStrings.Add(MakeShareable(new FString(LOCTEXT("OptimizationNone", "None").ToString())));
+		CompileOptimizationStrings.Add(MakeShareable(new FString(LOCTEXT("OptimizationMin", "Minimal").ToString())));
+		CompileOptimizationStrings.Add(MakeShareable(new FString(LOCTEXT("OptimizationMed", "Medium").ToString())));
+		CompileOptimizationStrings.Add(MakeShareable(new FString(LOCTEXT("OptimizationMax", "Maximum").ToString())));
+
+		CompileOptions.OptimizationLevel = FMath::Min(CompileOptions.OptimizationLevel, CompileOptimizationStrings.Num() - 1);
+
+		CompileOptimizationCombo =
+			SNew(STextComboBox)
+			.OptionsSource(&CompileOptimizationStrings)
+			.InitiallySelectedItem(CompileOptimizationStrings[CompileOptions.OptimizationLevel])
+			.OnSelectionChanged(this, &SMutableObjectViewer::OnChangeCompileOptimizationLevel)
+			;
+		MenuBuilder.AddWidget(CompileOptimizationCombo.ToSharedRef(), LOCTEXT("MutableCompileOptimizationLevel", "Optimization Level"));
+
+		// Parallel compilation
+		MenuBuilder.AddMenuEntry(
+			LOCTEXT("Generate_MutableEnableParallelCompilation", "Enable compiling in multiple threads."),
+			LOCTEXT("Generate_MutableEnableParallelCompilationTooltip", "This is faster but use more memory."),
+			FSlateIcon(),
+			FUIAction(
+				FExecuteAction::CreateLambda([this]() { CompileOptions.bUseParallelCompilation = !CompileOptions.bUseParallelCompilation; }),
+				FCanExecuteAction(),
+				FIsActionChecked::CreateLambda([this]() { return CompileOptions.bUseParallelCompilation; })),
+			NAME_None,
+			EUserInterfaceActionType::ToggleButton
+		);
+
+		// Disk as cache
+		MenuBuilder.AddMenuEntry(
+			LOCTEXT("Generate_MutableUseDisk", "Enable compiling using the disk as memory."),
+			LOCTEXT("Generate_MutableUseDiskTooltip", "This is very slow but supports compiling huge objects. It requires a lot of free space in the OS disk."),
+			FSlateIcon(),
+			FUIAction(
+				FExecuteAction::CreateLambda([this]() { CompileOptions.bUseDiskCompilation = !CompileOptions.bUseDiskCompilation; }),
+				FCanExecuteAction(),
+				FIsActionChecked::CreateLambda([this]() { return CompileOptions.bUseDiskCompilation; })),
+			NAME_None,
+			EUserInterfaceActionType::ToggleButton
+		);
+	}
+	MenuBuilder.EndSection();
+
+	return MenuBuilder.MakeWidget();
+}
+
+
+void SMutableObjectViewer::OnChangeCompileOptimizationLevel(TSharedPtr<FString> NewSelection, ESelectInfo::Type SelectInfo)
+{
+	CompileOptions.OptimizationLevel = CompileOptimizationStrings.Find(NewSelection);
+}
+
+
+void SMutableObjectViewer::OnChangeDebugPlatform(TSharedPtr<FString> NewSelection, ESelectInfo::Type SelectInfo)
+{
+	if (!CustomizableObject) return;
+
+	CompileOptions.TargetPlatform = nullptr;
+
+	ITargetPlatformManagerModule* TPM = GetTargetPlatformManager();
+	check(TPM);
+	const TArray<ITargetPlatform*>& Platforms = TPM->GetTargetPlatforms();
+	check(Platforms.Num());
+
+	CompileOptions.TargetPlatform = Platforms[0];
+
+	for (int32 Index = 1; Index < Platforms.Num(); Index++)
+	{
+		if (Platforms[Index]->PlatformName() == *NewSelection)
+		{
+			CompileOptions.TargetPlatform = Platforms[Index];
+			break;
+		}
+	}
+}
+
+
+TSharedRef<ITableRow> SMutableObjectViewer::GenerateRowForNodeTree(TSharedPtr<FMutableObjectTreeElement> InTreeNode, const TSharedRef<STableViewBase>& InOwnerTable)
+{
+	TSharedRef<SMutableObjectTreeRow> Row = SNew(SMutableObjectTreeRow, InOwnerTable, InTreeNode, CustomizableObject);
+	return Row;
+}
+
+
+void SMutableObjectViewer::GetChildrenForInfo(TSharedPtr<FMutableObjectTreeElement> InInfo, TArray<TSharedPtr<FMutableObjectTreeElement>>& OutChildren)
+{
+	switch (InInfo->Type)
+	{
+	case FMutableObjectTreeElement::EType::SectionCaption:
+		switch (InInfo->Section)
+		{
+
+		case FMutableObjectTreeElement::ESection::General:
+		{
+			OutChildren.Add(MakeShareable(new FMutableObjectTreeElement(FMutableObjectTreeElement::EType::Name, InInfo->Section, CustomizableObject )));
+			break;
+		}
+
+		case FMutableObjectTreeElement::ESection::ChildObjects:
+		{
+			// TODO: based on what is done in FCustomizableObjectCompiler::ProcessChildObjectsRecursively
+			//OutChildren.Add(MakeShareable(new FMutableObjectTreeElement(FMutableObjectTreeElement::EType::Name, InInfo->Section)));
+			break;
+		}
+
+		default:
+			break;
+		}
+		break;
+
+	default:
+		break;
+	}
+}
+
+
+#undef LOCTEXT_NAMESPACE 
diff --git a/Engine/Plugins/Experimental/Mutable/Source/CustomizableObjectEditor/Private/MuCOE/SMutableObjectViewer.h b/Engine/Plugins/Experimental/Mutable/Source/CustomizableObjectEditor/Private/MuCOE/SMutableObjectViewer.h
new file mode 100644
--- /dev/null
+++ b/Engine/Plugins/Experimental/Mutable/Source/CustomizableObjectEditor/Private/MuCOE/SMutableObjectViewer.h
@@ -0,0 +1,122 @@
+// Copyright Epic Games, Inc. All Rights Reserved.
+
+#pragma once
+
+#include "Widgets/SCompoundWidget.h"
+#include "Framework/Docking/TabManager.h"
+#include "Widgets/Views/STreeView.h"
+#include "MuCO/CustomizableObject.h"
+#include "MuCOE/CustomizableObjectCompiler.h"
+
+
+// Forward declarations
+class FMutableObjectTreeElement;
+
+
+/** This widget shows the internals of a CustomizableObject for debugging purposes. 
+ * This is not the Unreal source graph in the UCustomizableObject, but an intermediate step of the compilation process.
+ */
+class SMutableObjectViewer :
+	public SCompoundWidget,
+	public FGCObject
+{
+public:
+
+	// SWidget interface
+	SLATE_BEGIN_ARGS(SMutableObjectViewer) {}
+	SLATE_END_ARGS()
+
+	void Construct(const FArguments& InArgs, UCustomizableObject*,
+		TWeakPtr<FTabManager> InParentTabManager, const FName& InParentNewTabId );
+
+	// FGCObject interface
+	virtual void AddReferencedObjects(FReferenceCollector& Collector) override;
+	virtual FString GetReferencerName() const override;
+
+private:
+
+	/** The Mutable Graph to show, represented by its root. */
+	UCustomizableObject* CustomizableObject;
+
+	/** Compilation options to use in the debugger operations. */
+	FCompilationOptions CompileOptions;
+	TArray< TSharedPtr<FString> > DebugPlatformStrings;
+	TSharedPtr<STextComboBox> DebugPlatformCombo;
+
+	/** Object compiler. */
+	FCustomizableObjectCompiler Compiler;
+
+	/** UI references used to create new tabs. */
+	TWeakPtr<FTabManager> ParentTabManager;
+	FName ParentNewTabId;
+
+	/** Tree showing the object properties. */
+	TSharedPtr<STreeView<TSharedPtr<FMutableObjectTreeElement>>> TreeView;
+
+	/** Root nodes of the tree widget. */
+	TArray<TSharedPtr<FMutableObjectTreeElement>> RootTreeNodes;
+
+	/** UI callbacks */
+	void GenerateMutableGraphPressed();
+	void CompileMutableCodePressed();
+	TSharedRef<SWidget> GenerateCompileOptionsMenuContent();
+	TSharedPtr<STextComboBox> CompileOptimizationCombo;
+	TArray< TSharedPtr<FString> > CompileOptimizationStrings;
+	void OnChangeCompileOptimizationLevel(TSharedPtr<FString> NewSelection, ESelectInfo::Type SelectInfo);
+	void OnChangeDebugPlatform(TSharedPtr<FString> NewSelection, ESelectInfo::Type SelectInfo);
+
+	/** Tree callbacks. */
+	TSharedRef<ITableRow> GenerateRowForNodeTree(TSharedPtr<FMutableObjectTreeElement> InTreeNode, const TSharedRef<STableViewBase>& InOwnerTable);
+	void GetChildrenForInfo(TSharedPtr<FMutableObjectTreeElement> InInfo, TArray< TSharedPtr<FMutableObjectTreeElement> >& OutChildren);
+
+};
+
+
+/** An row of the code tree in . */
+class FMutableObjectTreeElement : public TSharedFromThis<FMutableObjectTreeElement>
+{
+public:
+
+	enum class EType
+	{
+		None,
+		SectionCaption,
+		Name,
+		ChildObject
+	};
+
+	enum class ESection
+	{
+		None,
+		General,
+		ChildObjects
+	};
+
+	/** Constructor for section rows */
+	FMutableObjectTreeElement(ESection InSection, UCustomizableObject* InObject)
+	{
+		Type = EType::SectionCaption;
+		Section = InSection;
+		CustomizableObject = InObject;
+	}
+
+	/** Constructor for generic property rows. */
+	FMutableObjectTreeElement(EType InType, ESection InSection, UCustomizableObject* InObject)
+	{
+		Type = InType;
+		Section = InSection;
+		CustomizableObject = InObject;
+	}
+
+public:
+
+	/** Row type. */
+	EType Type = EType::None;
+
+	/** If the row is a section row. */
+	ESection Section = ESection::None;
+
+	/** */
+	UCustomizableObject* CustomizableObject = nullptr;
+
+};
diff --git a/Engine/Plugins/Experimental/Mutable/Source/CustomizableObjectEditor/Private/MuCOE/SMutableParametersWidget.cpp b/Engine/Plugins/Experimental/Mutable/Source/CustomizableObjectEditor/Private/MuCOE/SMutableParametersWidget.cpp
new file mode 100644
--- /dev/null
+++ b/Engine/Plugins/Experimental/Mutable/Source/CustomizableObjectEditor/Private/MuCOE/SMutableParametersWidget.cpp
@@ -0,0 +1,589 @@
+// Copyright Epic Games, Inc. All Rights Reserved.
+
+#include "SMutableParametersWidget.h"
+#include "MuCOE/CustomizableObjectEditorModule.h"
+#include "Styling/AppStyle.h"
+#include "Async/TaskGraphInterfaces.h"
+#include "Brushes/SlateImageBrush.h"
+#include "Widgets/Input/SNumericEntryBox.h"
+#include "Widgets/Input/SCheckBox.h"
+#include "Widgets/Input/STextComboBox.h"
+#include "Widgets/Input/SVectorInputBox.h"
+#include "Widgets/Colors/SColorPicker.h"
+#include "Widgets/Layout/SGridPanel.h"
+#include "MuT/TypeInfo.h"
+
+
+#define LOCTEXT_NAMESPACE "CustomizableObjectEditor"
+
+
+void SMutableParametersWidget::Construct(const FArguments& InArgs)
+{
+	OnParametersValueChanged = InArgs._OnParametersValueChanged;
+
+	ChildSlot
+	[
+		SNew(SVerticalBox)
+		+ SVerticalBox::Slot()
+		.AutoHeight()
+		.Padding(0, 4.0f)
+		.HAlign(HAlign_Fill)
+		[
+			SAssignNew(ParamBox, SVerticalBox)
+		]
+	];
+}
+
+
+void SMutableParametersWidget::SetParameters(const mu::ParametersPtr& InParameters)
+{
+	MutableParameters = InParameters;
+	bIsPendingUpdate = true;
+}
+
+
+void SMutableParametersWidget::Tick(const FGeometry& AllottedGeometry, const double InCurrentTime, const float InDeltaTime)
+{
+	SCompoundWidget::Tick(AllottedGeometry, InCurrentTime, InDeltaTime);
+
+	if (!bIsPendingUpdate)
+	{
+		return;
+	}
+
+	bIsPendingUpdate = false;
+
+	ParamBox->ClearChildren();
+
+	if (!MutableParameters)
+	{
+		return;
+	}
+
+	for ( int32 ParamIndex=0; ParamIndex< MutableParameters->GetCount(); ++ParamIndex )
+	{
+		FString ParamName = ANSI_TO_TCHAR(MutableParameters->GetName(ParamIndex));
+		TSharedPtr<SHorizontalBox> ParameterBox;
+
+		ParamBox->AddSlot()
+			.Padding(1.0f, 6.0f)
+			.HAlign(HAlign_Fill)
+			.VAlign(VAlign_Fill)
+			.AutoHeight()
+			[
+				SAssignNew(ParameterBox, SHorizontalBox)
+
+				+ SHorizontalBox::Slot()
+				.FillWidth(0.66f)
+				.VAlign(VAlign_Center)
+				[
+					SNew(STextBlock)
+					.Text(FText::FromString(ParamName))
+					.Justification(ETextJustify::Left)
+					.AutoWrapText(true)
+				]
+			];
+
+		switch (MutableParameters->GetType(ParamIndex) )
+		{
+		case mu::PARAMETER_TYPE::T_BOOL:
+			ParameterBox->AddSlot()
+				.Padding(4.0f)
+				.VAlign(VAlign_Center)
+				[
+					SNew(SCheckBox)
+					.Visibility(this, &SMutableParametersWidget::GetParameterVisibility, ParamIndex)
+					.IsChecked(this, &SMutableParametersWidget::GetBoolParameterValue, ParamIndex)
+					.OnCheckStateChanged(this, &SMutableParametersWidget::OnBoolParameterChanged, ParamIndex)
+				];
+			break;
+
+		case mu::PARAMETER_TYPE::T_FLOAT:
+			ParameterBox->AddSlot()
+				.Padding(4.0f)
+				.VAlign(VAlign_Center)
+				[
+					SNew(SNumericEntryBox<float>)
+					.Visibility(this, &SMutableParametersWidget::GetParameterVisibility, ParamIndex)
+					.AllowSpin(true)
+					.MinSliderValue(0.0f)
+					.MaxSliderValue(1.0f)
+					.Value(this, &SMutableParametersWidget::GetFloatParameterValue, ParamIndex)
+					.OnValueChanged(this, &SMutableParametersWidget::OnFloatParameterChanged, ParamIndex)
+				];
+			break;
+
+		case mu::PARAMETER_TYPE::T_COLOUR:
+			ParameterBox->AddSlot()
+				.Padding(4.0f)
+				.VAlign(VAlign_Center)
+				[
+					SNew(SColorBlock)
+					.Color(this, &SMutableParametersWidget::GetColorParameterValue, ParamIndex)
+					.ShowBackgroundForAlpha(false)
+					.AlphaDisplayMode(EColorBlockAlphaDisplayMode::Ignore)
+					.UseSRGB(false)
+					.OnMouseButtonDown(this, &SMutableParametersWidget::OnColorBlockMouseButtonDown, ParamIndex)
+					.Size(FVector2D(10.0f, 10.0f))
+				];
+			break;
+
+		case mu::PARAMETER_TYPE::T_INT:
+		{
+			// If we have a list of options, add a combo box too
+			TSharedPtr<SMutableTextSearchBox> ParamComboBox;
+			int32 ValueCount = MutableParameters->GetIntPossibleValueCount(ParamIndex);
+			if (ValueCount > 0)
+			{
+				FString ToolTipText = FString("None");
+				TArray<FString> OptionNamesAttribute;
+				int Value = MutableParameters->GetIntValue( ParamIndex );
+				int ValueIndex = MutableParameters->GetIntValueIndex( ParamIndex, Value );
+				for (int i = 0; i < ValueCount; ++i)
+				{
+					const char* ValueText = MutableParameters->GetIntPossibleValueName( ParamIndex, i );
+					OptionNamesAttribute.Add( FString(ANSI_TO_TCHAR(ValueText)) );				
+				}
+
+				ParameterBox->AddSlot()
+					.Padding(4.0f)
+					.VAlign(VAlign_Center)
+					[
+						SAssignNew(ParamComboBox, SMutableTextSearchBox)
+						.Visibility(this, &SMutableParametersWidget::GetParameterVisibility, ParamIndex)
+						.ToolTipText(FText::FromString(ToolTipText))
+						.PossibleSuggestions(OptionNamesAttribute)
+						.InitialText(FText::FromString(OptionNamesAttribute[ValueIndex]))
+						.MustMatchPossibleSuggestions(TAttribute<bool>(true))
+						.SuggestionListPlacement(EMenuPlacement::MenuPlacement_ComboBox)
+						.OnTextCommitted(this, &SMutableParametersWidget::OnIntParameterTextChanged, ParamIndex)
+					];
+			}
+
+			ParameterBox->AddSlot()
+				.Padding(4.0f)
+				.VAlign(VAlign_Center)
+				[
+					SNew(SNumericEntryBox<int32>)
+					.Visibility(this, &SMutableParametersWidget::GetParameterVisibility, ParamIndex)
+					.AllowSpin(true)
+					.MinSliderValue(0)
+					.MaxSliderValue(this, &SMutableParametersWidget::GetIntParameterValueMax, ParamIndex)
+					.Value(this, &SMutableParametersWidget::GetIntParameterValue, ParamIndex)
+					.OnValueChanged(this, &SMutableParametersWidget::OnIntParameterChanged, ParamIndex, ParamComboBox)
+				];
+			break;
+		}
+
+		case mu::PARAMETER_TYPE::T_PROJECTOR:
+			ParameterBox->AddSlot()
+				.Padding(4.0f)
+				.VAlign(VAlign_Center)
+				[
+					SNew(SGridPanel)
+					.FillColumn(0, 1.0f)
+
+					+SGridPanel::Slot(0,0)
+					.VAlign(VAlign_Center)
+					[
+						SNew(STextBlock)
+						.Text(FText::FromString(TEXT("Location")))
+						.Font(FAppStyle::GetFontStyle(TEXT("PropertyWindow.NormalFont")))
+					]
+
+					+SGridPanel::Slot(1, 0)
+					[
+						SNew(SNumericVectorInputBox<FVector::FReal>)
+						.X(this, &SMutableParametersWidget::GetProjectorLocation, EAxis::X, ParamIndex)
+						.Y(this, &SMutableParametersWidget::GetProjectorLocation, EAxis::Y, ParamIndex)
+						.Z(this, &SMutableParametersWidget::GetProjectorLocation, EAxis::Z, ParamIndex)
+						.bColorAxisLabels(true)
+						.OnXChanged(this, &SMutableParametersWidget::SetProjectorLocation, ETextCommit::Default, EAxis::X, false, ParamIndex)
+						.OnYChanged(this, &SMutableParametersWidget::SetProjectorLocation, ETextCommit::Default, EAxis::Y, false, ParamIndex)
+						.OnZChanged(this, &SMutableParametersWidget::SetProjectorLocation, ETextCommit::Default, EAxis::Z, false, ParamIndex)
+						.OnXCommitted(this, &SMutableParametersWidget::SetProjectorLocation, EAxis::X, true, ParamIndex)
+						.OnYCommitted(this, &SMutableParametersWidget::SetProjectorLocation, EAxis::Y, true, ParamIndex)
+						.OnZCommitted(this, &SMutableParametersWidget::SetProjectorLocation, EAxis::Z, true, ParamIndex)
+						.AllowSpin(true)
+						.SpinDelta(1)
+					]
+
+					+ SGridPanel::Slot(0, 1)
+					.VAlign(VAlign_Center)
+					[
+						SNew(STextBlock)
+						.Text(FText::FromString(TEXT("Scale")))
+						.Font(FAppStyle::GetFontStyle("PropertyWindow.NormalFont"))
+					]
+
+					+ SGridPanel::Slot(1, 1)
+					[
+						SNew(SNumericVectorInputBox<FVector::FReal>)
+						.X(this, &SMutableParametersWidget::GetProjectorScale, EAxis::X, ParamIndex)
+						.Y(this, &SMutableParametersWidget::GetProjectorScale, EAxis::Y, ParamIndex)
+						.Z(this, &SMutableParametersWidget::GetProjectorScale, EAxis::Z, ParamIndex)
+						.bColorAxisLabels(true)
+						.OnXChanged(this, &SMutableParametersWidget::SetProjectorScale, ETextCommit::Default, EAxis::X, false, ParamIndex)
+						.OnYChanged(this, &SMutableParametersWidget::SetProjectorScale, ETextCommit::Default, EAxis::Y, false, ParamIndex)
+						.OnZChanged(this, &SMutableParametersWidget::SetProjectorScale, ETextCommit::Default, EAxis::Z, false, ParamIndex)
+						.OnXCommitted(this, &SMutableParametersWidget::SetProjectorScale, EAxis::X, true, ParamIndex)
+						.OnYCommitted(this, &SMutableParametersWidget::SetProjectorScale, EAxis::Y, true, ParamIndex)
+						.OnZCommitted(this, &SMutableParametersWidget::SetProjectorScale, EAxis::Z, true, ParamIndex)
+						.AllowSpin(true)
+						.SpinDelta(1)
+					]
+
+				];
+			break;
+
+		default:
+			// Unsupported parameter type
+			ParameterBox->AddSlot()
+				.Padding(4.0f)
+				.VAlign(VAlign_Center)
+				[
+					SNew(STextBlock)
+					.Text(FText::FromString(TEXT("Unsupported parameter type.")))
+					.Justification(ETextJustify::Left)
+					.AutoWrapText(true)
+				];
+		}
+
+	}
+}
+
+
+EVisibility SMutableParametersWidget::GetParameterVisibility(int32 ParamIndex) const
+{
+	// \TODO: Parameter relevancy?
+	if (!MutableParameters
+		|| ParamIndex >= MutableParameters->GetCount() )
+	{
+		return EVisibility::HitTestInvisible;
+	}
+
+	return EVisibility::Visible;
+}
+
+
+ECheckBoxState SMutableParametersWidget::GetBoolParameterValue(int32 ParamIndex) const
+{
+	if ( !MutableParameters 
+		|| ParamIndex>=MutableParameters->GetCount()
+		|| MutableParameters->GetType(ParamIndex)!=mu::PARAMETER_TYPE::T_BOOL )
+	{
+		return ECheckBoxState::Undetermined;
+	}	
+
+	return MutableParameters->GetBoolValue(ParamIndex) ? ECheckBoxState::Checked : ECheckBoxState::Unchecked;
+}
+
+
+void SMutableParametersWidget::OnBoolParameterChanged(ECheckBoxState InCheckboxState, int32 ParamIndex)
+{
+	if (!MutableParameters
+		|| ParamIndex >= MutableParameters->GetCount()
+		|| MutableParameters->GetType(ParamIndex) != mu::PARAMETER_TYPE::T_BOOL)
+	{
+		return;
+	}
+
+	MutableParameters->SetBoolValue( ParamIndex, InCheckboxState==ECheckBoxState::Checked );
+
+	OnParametersValueChanged.ExecuteIfBound(ParamIndex);
+}
+
+
+TOptional<float> SMutableParametersWidget::GetFloatParameterValue(int32 ParamIndex) const
+{
+	if (!MutableParameters
+		|| ParamIndex >= MutableParameters->GetCount()
+		|| MutableParameters->GetType(ParamIndex) != mu::PARAMETER_TYPE::T_FLOAT)
+	{
+		return TOptional<float>();
+	}
+
+	return MutableParameters->GetFloatValue(ParamIndex);
+}
+
+
+void SMutableParametersWidget::OnFloatParameterChanged(float InValue, int32 ParamIndex)
+{
+	if (!MutableParameters
+		|| ParamIndex >= MutableParameters->GetCount()
+		|| MutableParameters->GetType(ParamIndex) != mu::PARAMETER_TYPE::T_FLOAT)
+	{
+		return;
+	}
+
+	MutableParameters->SetFloatValue(ParamIndex, InValue);
+
+	OnParametersValueChanged.ExecuteIfBound(ParamIndex);
+}
+
+
+TOptional<int32> SMutableParametersWidget::GetIntParameterValue(int32 ParamIndex) const
+{
+	if (!MutableParameters
+		|| ParamIndex >= MutableParameters->GetCount()
+		|| MutableParameters->GetType(ParamIndex) != mu::PARAMETER_TYPE::T_INT)
+	{
+		return TOptional<int32>();
+	}
+
+	int32 Value = MutableParameters->GetIntValue(ParamIndex);
+	int32 ValueIndex = MutableParameters->GetIntValueIndex(ParamIndex, Value);
+
+	return ValueIndex;
+}
+
+
+TOptional<int32> SMutableParametersWidget::GetIntParameterValueMax(int32 ParamIndex) const
+{
+	if (!MutableParameters
+		|| ParamIndex >= MutableParameters->GetCount()
+		|| MutableParameters->GetType(ParamIndex) != mu::PARAMETER_TYPE::T_INT)
+	{
+		return TOptional<int32>();
+	}
+
+	int32 ValueCount = MutableParameters->GetIntPossibleValueCount(ParamIndex);
+
+	if (!ValueCount)
+	{
+		ValueCount = 16;
+	}
+
+	return ValueCount -1;
+}
+
+
+void SMutableParametersWidget::OnIntParameterChanged(int32 InValue, int32 ParamIndex, TSharedPtr<SMutableTextSearchBox> Combo )
+{
+	if (!MutableParameters
+		|| ParamIndex >= MutableParameters->GetCount()
+		|| MutableParameters->GetType(ParamIndex) != mu::PARAMETER_TYPE::T_INT)
+	{
+		return;
+	}
+
+	// The slider value is actually the value index for the parameter
+	if (InValue >= 0 && InValue < MutableParameters->GetIntPossibleValueCount(ParamIndex))
+	{
+		int32 RealValue = MutableParameters->GetIntPossibleValue(ParamIndex, InValue);
+		MutableParameters->SetIntValue(ParamIndex, RealValue);
+
+		// Update the text combo if any
+		if (Combo)
+		{
+			FString Text = TEXT("Invalid");
+			Text = ANSI_TO_TCHAR(MutableParameters->GetIntPossibleValueName(ParamIndex, InValue));
+			Combo->SetText(FText::FromString(Text));
+		}
+
+		OnParametersValueChanged.ExecuteIfBound(ParamIndex);
+	}
+}
+
+
+void SMutableParametersWidget::OnIntParameterTextChanged(TSharedPtr<FString> Selection, ESelectInfo::Type, int32 ParamIndex)
+{
+	if (!Selection
+		|| !MutableParameters
+		|| ParamIndex >= MutableParameters->GetCount()
+		|| MutableParameters->GetType(ParamIndex) != mu::PARAMETER_TYPE::T_INT)
+	{
+		return;
+	}
+
+	int32 ValueIndex = MutableParameters->GetIntValueIndex(ParamIndex, TCHAR_TO_ANSI(**Selection));
+	if (ValueIndex >= 0 && ValueIndex < MutableParameters->GetIntPossibleValueCount(ParamIndex))
+	{
+		int32 RealValue = MutableParameters->GetIntPossibleValue(ParamIndex, ValueIndex);
+		MutableParameters->SetIntValue(ParamIndex, RealValue);
+
+		OnParametersValueChanged.ExecuteIfBound(ParamIndex);
+	}
+}
+
+
+FLinearColor SMutableParametersWidget::GetColorParameterValue(int32 ParamIndex) const
+{
+	if (!MutableParameters
+		|| ParamIndex >= MutableParameters->GetCount()
+		|| MutableParameters->GetType(ParamIndex) != mu::PARAMETER_TYPE::T_COLOUR)
+	{
+		return FLinearColor();
+	}
+
+	FLinearColor Result;
+	MutableParameters->GetColourValue(ParamIndex, &Result.R, &Result.G, &Result.B);
+	Result.A = 1.0f;
+
+	return Result;
+}
+
+
+FReply SMutableParametersWidget::OnColorBlockMouseButtonDown(const FGeometry& MyGeometry, const FPointerEvent& MouseEvent, int32 ParamIndex)
+{
+	if (MouseEvent.GetEffectingButton() != EKeys::LeftMouseButton)
+	{
+		return FReply::Unhandled();
+	}
+
+	FLinearColor col = GetColorParameterValue(ParamIndex);
+
+	TArray<FLinearColor*> LinearColorArray;
+	LinearColorArray.Add(&col);
+
+	FColorPickerArgs args;
+	args.bIsModal = true;
+	args.bUseAlpha = false;
+	args.bOnlyRefreshOnMouseUp = false;
+	args.InitialColorOverride = col;
+	args.OnColorCommitted = FOnLinearColorValueChanged::CreateSP(this, &SMutableParametersWidget::OnSetColorFromColorPicker, ParamIndex);
+	OpenColorPicker(args);
+
+	return FReply::Handled();
+}
+
+
+void SMutableParametersWidget::OnSetColorFromColorPicker(FLinearColor NewColor, int32 ParamIndex)
+{
+	if (!MutableParameters
+		|| ParamIndex >= MutableParameters->GetCount()
+		|| MutableParameters->GetType(ParamIndex) != mu::PARAMETER_TYPE::T_COLOUR)
+	{
+		return;
+	}
+
+	MutableParameters->SetColourValue(ParamIndex, NewColor.R, NewColor.G, NewColor.B);
+
+	OnParametersValueChanged.ExecuteIfBound(ParamIndex);
+}
+
+
+TOptional<FVector::FReal> SMutableParametersWidget::GetProjectorLocation(EAxis::Type Axis, int32 ParamIndex) const
+{ 
+	if (!MutableParameters
+		|| ParamIndex >= MutableParameters->GetCount()
+		|| MutableParameters->GetType(ParamIndex) != mu::PARAMETER_TYPE::T_PROJECTOR)
+	{
+		return 0.0f;
+	}
+
+	mu::PROJECTOR_TYPE Type;
+	float Pos[3];
+	float Dir[3];
+	float Up[3];
+	float Scale[3];
+	float Angle;
+	MutableParameters->GetProjectorValue(ParamIndex, &Type, 
+		&Pos[0], &Pos[1], &Pos[2],
+		&Dir[0], &Dir[1], &Dir[2],
+		&Up[0], &Up[1], &Up[2],
+		&Scale[0], &Scale[1], &Scale[2],
+		&Angle );
+
+	return Pos[ Axis-EAxis::X ];
+}
+
+
+void SMutableParametersWidget::SetProjectorLocation(FVector::FReal NewValue, ETextCommit::Type, EAxis::Type Axis, bool bCommitted, int32 ParamIndex)
+{
+	if (!MutableParameters
+		|| ParamIndex >= MutableParameters->GetCount()
+		|| MutableParameters->GetType(ParamIndex) != mu::PARAMETER_TYPE::T_PROJECTOR)
+	{
+		return;
+	}
+
+	mu::PROJECTOR_TYPE Type;
+	float Pos[3];
+	float Dir[3];
+	float Up[3];
+	float Scale[3];
+	float Angle;
+	MutableParameters->GetProjectorValue(ParamIndex, &Type,
+		&Pos[0], &Pos[1], &Pos[2],
+		&Dir[0], &Dir[1], &Dir[2],
+		&Up[0], &Up[1], &Up[2],
+		&Scale[0], &Scale[1], &Scale[2],
+		&Angle);
+
+	Pos[Axis - EAxis::X] = NewValue;
+
+	MutableParameters->SetProjectorValue(ParamIndex,
+		Pos[0], Pos[1], Pos[2],
+		Dir[0], Dir[1], Dir[2],
+		Up[0], Up[1], Up[2],
+		Scale[0], Scale[1], Scale[2],
+		Angle);
+
+	OnParametersValueChanged.ExecuteIfBound(ParamIndex);
+}
+
+
+TOptional<FVector::FReal> SMutableParametersWidget::GetProjectorScale(EAxis::Type Axis, int32 ParamIndex) const
+{
+	if (!MutableParameters
+		|| ParamIndex >= MutableParameters->GetCount()
+		|| MutableParameters->GetType(ParamIndex) != mu::PARAMETER_TYPE::T_PROJECTOR)
+	{
+		return 0.0f;
+	}
+
+	mu::PROJECTOR_TYPE Type;
+	float Pos[3];
+	float Dir[3];
+	float Up[3];
+	float Scale[3];
+	float Angle;
+	MutableParameters->GetProjectorValue(ParamIndex, &Type,
+		&Pos[0], &Pos[1], &Pos[2],
+		&Dir[0], &Dir[1], &Dir[2],
+		&Up[0], &Up[1], &Up[2],
+		&Scale[0], &Scale[1], &Scale[2],
+		&Angle);
+
+	return Scale[Axis - EAxis::X];
+}
+
+
+void SMutableParametersWidget::SetProjectorScale(FVector::FReal NewValue, ETextCommit::Type, EAxis::Type Axis, bool bCommitted, int32 ParamIndex)
+{
+	if (!MutableParameters
+		|| ParamIndex >= MutableParameters->GetCount()
+		|| MutableParameters->GetType(ParamIndex) != mu::PARAMETER_TYPE::T_PROJECTOR)
+	{
+		return;
+	}
+
+	mu::PROJECTOR_TYPE Type;
+	float Pos[3];
+	float Dir[3];
+	float Up[3];
+	float Scale[3];
+	float Angle;
+	MutableParameters->GetProjectorValue(ParamIndex, &Type,
+		&Pos[0], &Pos[1], &Pos[2],
+		&Dir[0], &Dir[1], &Dir[2],
+		&Up[0], &Up[1], &Up[2],
+		&Scale[0], &Scale[1], &Scale[2],
+		&Angle);
+
+	Scale[Axis - EAxis::X] = NewValue;
+
+	MutableParameters->SetProjectorValue(ParamIndex,
+		Pos[0], Pos[1], Pos[2],
+		Dir[0], Dir[1], Dir[2],
+		Up[0], Up[1], Up[2],
+		Scale[0], Scale[1], Scale[2],
+		Angle);
+
+	OnParametersValueChanged.ExecuteIfBound(ParamIndex);
+}
+
+
+#undef LOCTEXT_NAMESPACE
+
diff --git a/Engine/Plugins/Experimental/Mutable/Source/CustomizableObjectEditor/Private/MuCOE/SMutableParametersWidget.h b/Engine/Plugins/Experimental/Mutable/Source/CustomizableObjectEditor/Private/MuCOE/SMutableParametersWidget.h
new file mode 100644
--- /dev/null
+++ b/Engine/Plugins/Experimental/Mutable/Source/CustomizableObjectEditor/Private/MuCOE/SMutableParametersWidget.h
@@ -0,0 +1,82 @@
+// Copyright Epic Games, Inc. All Rights Reserved.
+
+#pragma once
+
+#include "SMutableTextSearchBox.h"
+#include "Widgets/SCompoundWidget.h"
+#include "Widgets/SBoxPanel.h"
+#include "Widgets/Input/SCheckBox.h"
+#include "Widgets/DeclarativeSyntaxSupport.h"
+#include "MuR/Parameters.h"
+#include "Math/Axis.h"
+
+
+/** Delegate that is executed when a parameter value changes */
+DECLARE_DELEGATE_OneParam(FOnMutableParameterValueChanged, int32);
+
+
+/**  This widget displays and allows edition of a set of parameter values for the lower level Mutable Parameters 
+* object. It is meant to be used for tools and debugging and it doesn't have any knowledge of Customizable Objects
+* and the additional parameter details stored there.
+*/
+class SMutableParametersWidget : public SCompoundWidget
+{
+public:
+
+	SLATE_BEGIN_ARGS(SMutableParametersWidget) {}
+
+		/** Parameters to show and edit. */
+		SLATE_ARGUMENT_DEFAULT(mu::ParametersPtr, Parameters) { nullptr };
+
+		/** Called when any parameter value has changed, with the parameter index as argument.  */
+		SLATE_EVENT(FOnMutableParameterValueChanged, OnParametersValueChanged)
+
+	SLATE_END_ARGS()
+
+	void Construct(const FArguments& InArgs);
+
+	// SWidget interface
+	void Tick(const FGeometry& AllottedGeometry, const double InCurrentTime, const float InDeltaTime) override;
+
+	// Own interface
+
+	/** Set the image to show in the widget. */
+	void SetParameters(const mu::ParametersPtr& InParameters);
+
+private:
+
+	/** Reference to the mutable parameters. */
+	mu::ParametersPtr MutableParameters;
+
+	/** Is true, the parameter object has changed and we need to update. */
+	bool bIsPendingUpdate = false;
+
+	/** Delegate called when the check box changes state */
+	FOnMutableParameterValueChanged OnParametersValueChanged;
+
+	/** Dynamically filled box with per-parameter widgets */
+	TSharedPtr<SVerticalBox> ParamBox;
+
+	/** Internal UI callbacks. */
+	EVisibility GetParameterVisibility(int32 ParamIndex) const;
+	ECheckBoxState GetBoolParameterValue(int32 ParamIndex) const;
+	void OnBoolParameterChanged(ECheckBoxState CheckboxState, int32 ParamIndex);
+	TOptional<float> GetFloatParameterValue(int32 ParamIndex) const;
+	void OnFloatParameterChanged(float Value, int32 ParamIndex);
+	FLinearColor GetColorParameterValue(int32 ParamIndex) const;
+	FReply OnColorBlockMouseButtonDown(const FGeometry& MyGeometry, const FPointerEvent& MouseEvent, int32 ParamIndex);
+	void OnSetColorFromColorPicker(FLinearColor NewColor, int32 ParamIndex);
+	TOptional<int32> GetIntParameterValue(int32 ParamIndex) const;
+	TOptional<int32> GetIntParameterValueMax(int32 ParamIndex) const;
+	void OnIntParameterChanged(int32 Value, int32 ParamIndex, TSharedPtr<SMutableTextSearchBox> );
+	void OnIntParameterTextChanged(TSharedPtr<FString> Selection, ESelectInfo::Type SelectInfo, int32 ParamIndex);
+
+	/** Projector UI callbacks */
+	TOptional<FVector::FReal> GetProjectorLocation(EAxis::Type Axis, int32 ParamIndex) const;
+	void SetProjectorLocation(FVector::FReal NewValue, ETextCommit::Type, EAxis::Type, bool bCommitted, int32 ParamIndex);
+	TOptional<FVector::FReal> GetProjectorScale(EAxis::Type Axis, int32 ParamIndex) const;
+	void SetProjectorScale(FVector::FReal NewValue, ETextCommit::Type, EAxis::Type, bool bCommitted, int32 ParamIndex);
+
+
+};
+
diff --git a/Engine/Plugins/Experimental/Mutable/Source/CustomizableObjectEditor/Private/MuCOE/SMutableTextSearchBox.cpp b/Engine/Plugins/Experimental/Mutable/Source/CustomizableObjectEditor/Private/MuCOE/SMutableTextSearchBox.cpp
new file mode 100644
--- /dev/null
+++ b/Engine/Plugins/Experimental/Mutable/Source/CustomizableObjectEditor/Private/MuCOE/SMutableTextSearchBox.cpp
@@ -0,0 +1,374 @@
+// Copyright Epic Games, Inc. All Rights Reserved.
+
+#include "SMutableTextSearchBox.h"
+
+#include "MuCO/CustomizableObject.h" // Required for LogMutable
+
+#include "Layout/WidgetPath.h"
+#include "Framework/Application/SlateApplication.h"
+#include "Widgets/Text/STextBlock.h"
+#include "Widgets/Input/SMenuAnchor.h"
+#include "Widgets/Views/SListView.h"
+#include "Widgets/Input/SSearchBox.h"
+#include "Widgets/Input/SButton.h"
+#include "Widgets/Input/SComboButton.h"
+#include "Widgets/Images/SImage.h"
+
+#include "Styling/AppStyle.h"
+
+void SMutableTextSearchBox::Construct( const FArguments& InArgs )
+{
+	OnTextChanged = InArgs._OnTextChanged;
+	OnTextCommitted = InArgs._OnTextCommitted;
+	OnKeyDownHandler = InArgs._OnKeyDownHandler;
+	PossibleSuggestions = InArgs._PossibleSuggestions;
+	PreCommittedText = InArgs._InitialText.Get();
+	bMustMatchPossibleSuggestions = InArgs._MustMatchPossibleSuggestions.Get();
+
+	ChildSlot
+		[
+			//SNew( SHorizontalBox )
+			//	+ SHorizontalBox::Slot()
+			//	.FillWidth( 1 )
+			//	[
+			SNew( SButton )
+				//.ButtonStyle( OurButtonStyle )
+				.ClickMethod( EButtonClickMethod::MouseDown )
+				.OnClicked( this, &SMutableTextSearchBox::OnButtonClicked )
+				//.ContentPadding( InArgs._ContentPadding )
+				//.ForegroundColor( InArgs._ForegroundColor )
+				.ButtonColorAndOpacity( FSlateColor(FLinearColor(1.f, 1.f, 1.f, 1.f)) )
+				.IsFocusable( true )
+				.Visibility(EVisibility::Visible)
+				.Content()
+				[
+					SAssignNew(SuggestionBox, SMenuAnchor)
+					.Placement( InArgs._SuggestionListPlacement )
+					.Padding(-2.0f)
+					[
+						SAssignNew(InputText, SSearchBox)
+						.InitialText(InArgs._InitialText)
+						.HintText(InArgs._HintText)
+						.OnTextChanged(this, &SMutableTextSearchBox::HandleTextChanged)
+						.OnTextCommitted(this, &SMutableTextSearchBox::HandleTextCommitted)
+						.SelectAllTextWhenFocused( true )
+						.DelayChangeNotificationsWhileTyping( InArgs._DelayChangeNotificationsWhileTyping )
+						.OnKeyDownHandler(this, &SMutableTextSearchBox::HandleKeyDown)
+						.Visibility(EVisibility::SelfHitTestInvisible)
+					]
+					.MenuContent
+						(
+							SNew(SBorder)
+							.BorderImage(FAppStyle::GetBrush("Menu.Background"))
+							[
+								SNew(SBox)
+								.WidthOverride(175)	// to enforce some minimum width, ideally we define the minimum, not a fixed width
+								//.HeightOverride(175) // avoids flickering, ideally this would be adaptive to the content without flickering
+								[
+									SAssignNew(SuggestionListView, SListView< TSharedPtr<FString> >)
+									.ListItemsSource(&Suggestions)
+									.SelectionMode( ESelectionMode::Single )							// Ideally the mouse over would not highlight while keyboard controls the UI
+									.OnGenerateRow(this, &SMutableTextSearchBox::MakeSuggestionListItemWidget)
+									.OnSelectionChanged( this, &SMutableTextSearchBox::OnSelectionChanged)
+									.OnListViewScrolled(this, &SMutableTextSearchBox::OnListViewScrolled)
+								]
+							]
+						)
+				]
+		];
+
+	InputText->SetVisibility(EVisibility::HitTestInvisible);
+}
+
+void SMutableTextSearchBox::OnListViewScrolled(double InScrollOffset)
+{	
+	if (!HasKeyboardFocus() && !SuggestionListView->IsUserScrolling())
+	{
+		FocusEditBox();
+	}
+}
+
+FReply SMutableTextSearchBox::OnButtonClicked()
+{
+	InputText->SetText(FText());
+	FReply Reply = FReply::Handled().SetUserFocus(InputText.ToSharedRef());
+	UpdateSuggestionList();
+
+	// Forward keyboard focus to our editable text widget
+	return Reply;
+}
+
+void SMutableTextSearchBox::SetText(const TAttribute< FText >& InNewText)
+{
+	InputText->SetText(InNewText);
+	PreCommittedText = InNewText.Get();
+}
+
+void SMutableTextSearchBox::SetError( const FText& InError )
+{
+	InputText->SetError(InError);
+}
+
+void SMutableTextSearchBox::SetError( const FString& InError )
+{
+	InputText->SetError(InError);
+}
+
+FReply SMutableTextSearchBox::OnPreviewKeyDown( const FGeometry& MyGeometry, const FKeyEvent& InKeyEvent )
+{
+	if ( SuggestionBox->IsOpen() && InKeyEvent.GetKey() == EKeys::Escape )
+	{
+		// Clear any selection first to prevent the currently selection being set in the text box
+		SuggestionListView->ClearSelection();
+		SuggestionBox->SetIsOpen(false, false);
+		return FReply::Handled();
+	}
+
+	return FReply::Unhandled();
+}
+
+FReply SMutableTextSearchBox::HandleKeyDown( const FGeometry& MyGeometry, const FKeyEvent& InKeyEvent )
+{
+	if ( SuggestionBox->IsOpen() && (InKeyEvent.GetKey() == EKeys::Up || InKeyEvent.GetKey() == EKeys::Down) )
+	{
+		const bool bSelectingUp = InKeyEvent.GetKey() == EKeys::Up;
+		TSharedPtr<FString> SelectedSuggestion = GetSelectedSuggestion();
+
+		if ( SelectedSuggestion.IsValid() )
+		{
+			// Find the selection index and select the previous or next one
+			int32 TargetIdx = INDEX_NONE;
+			for ( int32 SuggestionIdx = 0; SuggestionIdx < Suggestions.Num(); ++SuggestionIdx )
+			{
+				if ( Suggestions[SuggestionIdx] == SelectedSuggestion )
+				{
+					if ( bSelectingUp )
+					{
+						TargetIdx = SuggestionIdx - 1;
+					}
+					else
+					{
+						TargetIdx = SuggestionIdx + 1;
+					}
+					break;
+				}
+			}
+
+			if ( Suggestions.IsValidIndex(TargetIdx) )
+			{
+				SuggestionListView->SetSelection( Suggestions[TargetIdx] );
+				SuggestionListView->RequestScrollIntoView( Suggestions[TargetIdx] );
+			}
+		}
+		else if ( !bSelectingUp && Suggestions.Num() > 0 )
+		{
+			// Nothing selected and pressed down, select the first item
+			SuggestionListView->SetSelection( Suggestions[0] );
+		}
+
+		return FReply::Handled();
+	}
+
+	if (OnKeyDownHandler.IsBound())
+	{
+		return OnKeyDownHandler.Execute(MyGeometry, InKeyEvent);
+	}
+
+	return FReply::Unhandled();
+}
+
+bool SMutableTextSearchBox::SupportsKeyboardFocus() const
+{
+	return InputText->SupportsKeyboardFocus();
+}
+
+bool SMutableTextSearchBox::HasKeyboardFocus() const
+{
+	// Since keyboard focus is forwarded to our editable text, we will test it instead
+	return InputText->HasKeyboardFocus();
+}
+
+FReply SMutableTextSearchBox::OnFocusReceived( const FGeometry& MyGeometry, const FFocusEvent& InFocusEvent )
+{
+	UpdateSuggestionList();
+
+	// Forward keyboard focus to our editable text widget
+	return FReply::Handled().SetUserFocus(InputText.ToSharedRef(), InFocusEvent.GetCause());
+}
+
+void SMutableTextSearchBox::HandleTextChanged(const FText& NewText)
+{
+	OnTextChanged.ExecuteIfBound(NewText);
+	UpdateSuggestionList();
+}
+
+void SMutableTextSearchBox::HandleTextCommitted(const FText& NewText, ETextCommit::Type CommitType)
+{
+	if (SuggestionListView->IsUserScrolling())
+	{
+		return;
+	}
+
+	TSharedPtr<FString> SelectedSuggestion = GetSelectedSuggestion();
+
+	bool bCommitText = true;
+	FText CommittedText;
+	if ( SelectedSuggestion.IsValid() && CommitType != ETextCommit::OnCleared )
+	{
+		// Pressed selected a suggestion, set the text
+		CommittedText = FText::FromString( *SelectedSuggestion.Get() );
+	}
+	else
+	{
+		if ( CommitType == ETextCommit::OnCleared )
+		{
+			// Clear text when escape is pressed then commit an empty string
+			//CommittedText = FText::GetEmpty();
+			CommittedText = PreCommittedText;
+		}
+		else if( PossibleSuggestions.Get().Contains(NewText.ToString()) )
+		{
+			// If the text is a suggestion, set the text.
+			CommittedText = NewText;
+		}
+		else if( bMustMatchPossibleSuggestions )
+		{
+			// commit the original text if we have to match a suggestion
+			CommittedText = PreCommittedText;
+		}
+		else
+		{
+			// otherwise, set the typed text
+			CommittedText = NewText;
+		}	
+	}
+
+	// Set the text and execute the delegate
+	SetText(CommittedText);
+
+	TSharedPtr<FString> CommittedStringPtr = MakeShareable<FString>(new FString(CommittedText.ToString()));
+	OnTextCommitted.ExecuteIfBound(CommittedStringPtr, ESelectInfo::Direct);
+
+	if(CommitType != ETextCommit::Default)
+	{
+		// Clear the suggestion box if the user has navigated away or set their own text.
+		SuggestionListView->ClearSelection();
+		SuggestionBox->SetIsOpen(false, false);
+		//SuggestionListView->SetVisibility(EVisibility::Hidden);
+	}
+}
+
+void SMutableTextSearchBox::OnSelectionChanged( TSharedPtr<FString> NewValue, ESelectInfo::Type SelectInfo )
+{
+	// If the user clicked directly on an item to select it, then accept the choice and close the window
+	if( SelectInfo == ESelectInfo::OnMouseClick )
+	{
+		//SetText( FText::FromString( *NewValue.Get() ));
+		//SuggestionBox->SetIsOpen(false, false);
+		//FocusEditBox();
+		HandleTextCommitted(FText::FromString(*NewValue.Get()), ETextCommit::OnEnter);
+	}
+}
+
+
+TSharedRef<ITableRow> SMutableTextSearchBox::MakeSuggestionListItemWidget(TSharedPtr<FString> Text, const TSharedRef<STableViewBase>& OwnerTable)
+{
+	check(Text.IsValid());
+
+	return
+		SNew(STableRow< TSharedPtr<FString> >, OwnerTable)
+		[
+			SNew(STextBlock)
+			.Text(FText::FromString(*Text.Get()))
+			.HighlightText(this, &SMutableTextSearchBox::GetHighlightText)
+		];
+}
+
+FText SMutableTextSearchBox::GetHighlightText() const
+{
+	return InputText->GetText();
+}
+
+void SMutableTextSearchBox::UpdateSuggestionList()
+{
+	const FString TypedText = InputText->GetText().ToString();
+
+	Suggestions.Empty();
+
+	if ( TypedText.Len() > 0 )
+	{
+		TArray<FString> AllSuggestions = PossibleSuggestions.Get();
+
+		for ( auto SuggestionIt = AllSuggestions.CreateConstIterator(); SuggestionIt; ++SuggestionIt )
+		{
+			const FString& Suggestion = *SuggestionIt;
+			if ( Suggestion.Contains(TypedText))
+			{
+				Suggestions.Add( MakeShareable(new FString(Suggestion)) );
+			}
+		}
+
+		if ( Suggestions.Num() > 0 )
+		{
+			// At least one suggestion was found, open the menu
+			SuggestionBox->SetIsOpen(true, false);
+		}
+		else
+		{
+			// No suggestions were found, close the menu
+			SuggestionBox->SetIsOpen(false, false);
+		}
+	}
+	else //if (HasKeyboardFocus())
+	{
+		TArray<FString> AllSuggestions = PossibleSuggestions.Get();
+
+		for ( auto SuggestionIt = AllSuggestions.CreateConstIterator(); SuggestionIt; ++SuggestionIt )
+		{
+			const FString& Suggestion = *SuggestionIt;
+			Suggestions.Add( MakeShareable(new FString(Suggestion)) );
+		}
+
+		if ( Suggestions.Num() > 0 )
+		{
+			// At least one suggestion was found, open the menu
+			SuggestionBox->SetIsOpen(true, false);
+		}
+		else
+		{
+			// No suggestions were found, close the menu
+			SuggestionBox->SetIsOpen(false, false);
+		}
+	}
+	//else
+	//{
+	//	// No text was typed, close the menu
+	//	SuggestionBox->SetIsOpen(false, false);
+	//}
+
+	SuggestionListView->RequestListRefresh();
+}
+
+void SMutableTextSearchBox::FocusEditBox()
+{
+	FWidgetPath WidgetToFocusPath;
+	FSlateApplication::Get().GeneratePathToWidgetUnchecked( InputText.ToSharedRef(), WidgetToFocusPath );
+	FSlateApplication::Get().SetKeyboardFocus( WidgetToFocusPath, EFocusCause::SetDirectly );
+}
+
+TSharedPtr<FString> SMutableTextSearchBox::GetSelectedSuggestion()
+{
+	TSharedPtr<FString> SelectedSuggestion;
+	if ( SuggestionBox->IsOpen() )
+	{
+		const TArray< TSharedPtr<FString> >& SelectedSuggestionList = SuggestionListView->GetSelectedItems();
+		if ( SelectedSuggestionList.Num() > 0 )
+		{
+			// Selection mode is Single, so there should only be one suggestion at the most
+			SelectedSuggestion = SelectedSuggestionList[0];
+		}
+	}
+
+	return SelectedSuggestion;
+}
+
diff --git a/Engine/Plugins/Experimental/Mutable/Source/CustomizableObjectEditor/Private/MuCOE/SMutableTextSearchBox.h b/Engine/Plugins/Experimental/Mutable/Source/CustomizableObjectEditor/Private/MuCOE/SMutableTextSearchBox.h
new file mode 100644
--- /dev/null
+++ b/Engine/Plugins/Experimental/Mutable/Source/CustomizableObjectEditor/Private/MuCOE/SMutableTextSearchBox.h
@@ -0,0 +1,150 @@
+// Copyright Epic Games, Inc. All Rights Reserved.
+
+#pragma once
+
+
+#include "CoreMinimal.h"
+#include "SlateFwd.h"
+#include "Misc/Attribute.h"
+#include "Input/Reply.h"
+#include "Widgets/DeclarativeSyntaxSupport.h"
+#include "Widgets/SCompoundWidget.h"
+#include "Widgets/Views/STableViewBase.h"
+#include "Widgets/Views/STableRow.h"
+#include "Widgets/Input/SSearchBox.h"
+
+class SMenuAnchor;
+
+
+/**
+ * A widget to provide a search box with a filtered dropdown menu.
+ */
+
+typedef TSlateDelegates< TSharedPtr<FString> >::FOnSelectionChanged FOnTextSelectionCommitted;
+//DECLARE_DELEGATE_ThreeParams( FOnTextSelectionCommitted, const FText&, ETextCommit::Type, const FString& );
+
+class SMutableTextSearchBox : public SCompoundWidget
+{
+public:
+	SLATE_BEGIN_ARGS(SMutableTextSearchBox)
+		: _SuggestionListPlacement( MenuPlacement_BelowAnchor )
+		, _OnTextChanged()
+		, _OnTextCommitted()
+		, _InitialText()
+		, _HintText()
+		, _PossibleSuggestions(TArray<FString>())
+		, _DelayChangeNotificationsWhileTyping( false )
+		, _MustMatchPossibleSuggestions( false )
+	{}
+
+	/** Where to place the suggestion list */
+	SLATE_ARGUMENT( EMenuPlacement, SuggestionListPlacement )
+
+	/** Invoked whenever the text changes */
+	SLATE_EVENT( FOnTextChanged, OnTextChanged )
+
+	/** Invoked whenever the text is committed (e.g. user presses enter) */
+	SLATE_EVENT( FOnTextSelectionCommitted, OnTextCommitted )
+
+	/** Initial text to display for the search text */
+	SLATE_ATTRIBUTE( FText, InitialText )
+
+	/** Hint text to display for the search text when there is no value */
+	SLATE_ATTRIBUTE( FText, HintText )
+
+	/** All possible suggestions for the search text */
+	SLATE_ATTRIBUTE( TArray<FString>, PossibleSuggestions )
+
+	/** Whether the SearchBox should delay notifying listeners of text changed events until the user is done typing */
+	SLATE_ATTRIBUTE( bool, DelayChangeNotificationsWhileTyping )
+
+	/** Whether the SearchBox allows entries that don't match the possible suggestions */
+	SLATE_ATTRIBUTE( bool, MustMatchPossibleSuggestions )
+
+	/** Callback delegate to have first chance handling of the OnKeyDown event */
+	SLATE_EVENT( FOnKeyDown, OnKeyDownHandler )
+
+	SLATE_END_ARGS()
+
+	/** Constructs this widget with InArgs */
+	void Construct( const FArguments& InArgs );
+
+	/** Sets the text string currently being edited */
+	void SetText(const TAttribute< FText >& InNewText);
+
+	/** Set or clear the current error reporting information for this search box */
+	void SetError( const FText& InError );
+	void SetError( const FString& InError );
+
+	// SWidget implementation
+	virtual FReply OnPreviewKeyDown( const FGeometry& MyGeometry, const FKeyEvent& InKeyEvent ) override;
+	virtual bool SupportsKeyboardFocus() const override;
+	virtual bool HasKeyboardFocus() const override;
+	virtual FReply OnFocusReceived( const FGeometry& MyGeometry, const FFocusEvent& InFocusEvent ) override;
+
+private:
+	FReply OnButtonClicked();
+
+	/** First chance handler for key down events to the editable text widget */
+	FReply HandleKeyDown(const FGeometry& MyGeometry, const FKeyEvent& InKeyEvent);
+
+	/** Handler for when text in the editable text box changed */
+	void HandleTextChanged(const FText& NewText);
+
+	/** Handler for when text in the editable text box changed */
+	void HandleTextCommitted(const FText& NewText, ETextCommit::Type CommitType);
+
+	/** Called by SListView when the selection changes in the suggestion list */
+	void OnSelectionChanged( TSharedPtr<FString> NewValue, ESelectInfo::Type SelectInfo );
+
+	void OnListViewScrolled(double InScrollOffset);
+
+	/** Makes the widget for a suggestion message in the list view */
+	TSharedRef<ITableRow> MakeSuggestionListItemWidget(TSharedPtr<FString> Text, const TSharedRef<STableViewBase>& OwnerTable);
+
+	/** Gets the text to highlight in the suggestion list */
+	FText GetHighlightText() const;
+
+	/** Updates and shows or hides the suggestion list */
+	void UpdateSuggestionList();
+
+	/** Sets the focus to the InputText box */
+	void FocusEditBox();
+
+	/** Returns the currently selected suggestion */
+	TSharedPtr<FString> GetSelectedSuggestion();
+
+private:
+	/** The editable text field */
+	TSharedPtr< SSearchBox > InputText;
+
+	/** The the state of the text prior to being committed */
+	FText PreCommittedText;
+
+	/** Auto completion elements */
+	TSharedPtr< SMenuAnchor > SuggestionBox;
+
+	/** All suggestions stored in this widget for the list view */
+	TArray< TSharedPtr<FString> > Suggestions;
+
+	/** The list view for showing all suggestions */
+	TSharedPtr< SListView< TSharedPtr<FString> > > SuggestionListView;
+
+	/** Delegate for when text is changed in the edit box */
+	FOnTextChanged OnTextChanged;
+
+	/** Delegate for when text is changed in the edit box */
+	FOnTextSelectionCommitted OnTextCommitted;
+
+	/** Delegate for first chance handling for key down events */
+	FOnKeyDown OnKeyDownHandler;
+
+	/** All possible suggestions for the search text */
+	TAttribute< TArray<FString> > PossibleSuggestions;
+
+	/** Determines whether or not the committed text should match a suggestion */
+	bool bMustMatchPossibleSuggestions;
+
+	FString ParamName;
+};
+
