Moved the Mutable plugin files to live under a directory named after the module. This allows us to get rid of long include paths that contain public and private folder names in the path and fixes head

diff --git a/Engine/Plugins/Experimental/Mutable/Source/CustomizableObjectPopulationEditor/Private/MuCOPE/CustomizableObjectPopulationEditorModule.cpp b/Engine/Plugins/Experimental/Mutable/Source/CustomizableObjectPopulationEditor/Private/MuCOPE/CustomizableObjectPopulationEditorModule.cpp
new file mode 100644
--- /dev/null
+++ b/Engine/Plugins/Experimental/Mutable/Source/CustomizableObjectPopulationEditor/Private/MuCOPE/CustomizableObjectPopulationEditorModule.cpp
@@ -0,0 +1,144 @@
+// Copyright Epic Games, Inc. All Rights Reserved.
+
+#include "MuCOPE/CustomizableObjectPopulationEditorModule.h"
+#include "MuCOPE/CustomizableObjectPopulationEditor.h"
+#include "MuCO/CustomizableObject.h"	// For the LogMutable log category
+#include "MuCO/CustomizableObjectSystem.h"		// For defines related to memory function replacements.
+#include "MuCOE/CustomizableObjectIdentifierCustomization.h"
+#include "Modules/ModuleManager.h"
+#include "Interfaces/IPluginManager.h"
+#include "AssetToolsModule.h"
+#include "Styling/SlateStyle.h"
+#include "Brushes/SlateImageBrush.h"
+#include "MessageLogModule.h"
+#include "Logging/MessageLog.h"
+#include "MuCOE/CustomizableObjectEditorSettings.h"
+#include "ISettingsModule.h"
+#include "ISettingsSection.h"
+#include "AssetRegistry/AssetRegistryModule.h"
+#include "HAL/FileManager.h"
+#include "CoreGlobals.h"
+#include "Settings/ProjectPackagingSettings.h"
+#include "MuCOE/UnrealEditorPortabilityHelpers.h"
+#include "AssetTypeActions_CustomizableObjectPopulation.h"
+#include "AssetTypeActions_CustomizableObjectPopulationClass.h"
+#include "CustomizableObjectPopulationClassEditor.h"
+#include "MuCOPE/CustomizableObjectPopulationEditor.h"
+#include "CustomizableObjectPopulationEditorStyle.h"
+
+#include "CustomizableObjectPopulationClassDetails.h"
+
+// This is necessary to register the cook delegate.
+#include "CookOnTheSide/CookOnTheFlyServer.h"
+
+
+const FName CustomizableObjectPopulationEditorAppIdentifier = FName(TEXT("CustomizableObjectPopulationEditorApp"));
+const FName CustomizableObjectPopulationClassEditorAppIdentifier = FName(TEXT("CustomizableObjectPopulationClassEditorApp"));
+
+#define LOCTEXT_NAMESPACE "MutableSettings"
+
+/**
+ * StaticMesh editor module
+ */
+class FCustomizableObjectPopulationEditorModule : public ICustomizableObjectPopulationEditorModule
+{
+public:
+
+	// IModuleInterface interface
+	void StartupModule() override;
+	void ShutdownModule() override;
+
+	// ICustomizableObjectEditorModule interface
+	TSharedRef<ICustomizableObjectPopulationEditor> CreateCustomizableObjectPopulationEditor(const EToolkitMode::Type Mode, const TSharedPtr< IToolkitHost >& InitToolkitHost, UCustomizableObjectPopulation* CustomizablePopulation) override;
+	TSharedRef<ICustomizableObjectPopulationClassEditor> CreateCustomizableObjectPopulationClassEditor(const EToolkitMode::Type Mode, const TSharedPtr< IToolkitHost >& InitToolkitHost, UCustomizableObjectPopulationClass* CustomizablePopulationClass) override;
+
+	virtual TSharedPtr<FExtensibilityManager> GetCustomizableObjectPopulationEditorToolBarExtensibilityManager() override { return CustomizableObjectPopulationEditor_ToolBarExtensibilityManager; }
+
+private:
+
+	TSharedPtr<FExtensibilityManager> CustomizableObjectPopulationEditor_ToolBarExtensibilityManager;
+
+};
+
+IMPLEMENT_MODULE( FCustomizableObjectPopulationEditorModule, CustomizableObjectPopulationEditor );
+
+static void LogWarning(const char* msg)
+{
+	// Can be called from any thread.
+	UE_LOG(LogMutable, Warning, TEXT("%s"), ANSI_TO_TCHAR(msg) );
+}
+
+
+static void LogError(const char* msg)
+{
+	// Can be called from any thread.
+	UE_LOG(LogMutable, Error, TEXT("%s"), ANSI_TO_TCHAR(msg));
+}
+
+
+static void* CustomMalloc(std::size_t Size_t, uint32_t Alignment)
+{
+	return FMemory::Malloc(Size_t, Alignment);
+}
+
+
+static void CustomFree(void* mem)
+{
+	return FMemory::Free(mem);
+}
+
+
+void FCustomizableObjectPopulationEditorModule::StartupModule()
+{	
+	// Property views
+	FPropertyEditorModule& PropertyModule = FModuleManager::LoadModuleChecked<FPropertyEditorModule>("PropertyEditor");
+	
+	PropertyModule.RegisterCustomClassLayout("CustomizableObjectPopulationClass", FOnGetDetailCustomizationInstance::CreateStatic(&FCustomizableObjectPopulationClassDetails::MakeInstance));
+	
+	// Asset actions
+	FAssetToolsModule& AssetToolsModule = FModuleManager::LoadModuleChecked< FAssetToolsModule >( "AssetTools" );
+
+	TSharedPtr<FAssetTypeActions_CustomizableObjectPopulation> CustomizableObjectPopulationAssetTypeActions = MakeShareable(new FAssetTypeActions_CustomizableObjectPopulation);
+	AssetToolsModule.Get().RegisterAssetTypeActions(CustomizableObjectPopulationAssetTypeActions.ToSharedRef());
+	
+	TSharedPtr<FAssetTypeActions_CustomizableObjectPopulationClass> CustomizableObjectPopulationClassAssetTypeActions = MakeShareable(new FAssetTypeActions_CustomizableObjectPopulationClass);
+	AssetToolsModule.Get().RegisterAssetTypeActions(CustomizableObjectPopulationClassAssetTypeActions.ToSharedRef());
+
+	// Additional UI style
+	FCustomizableObjectPopulationEditorStyle::Initialize();
+
+	CustomizableObjectPopulationEditor_ToolBarExtensibilityManager = MakeShareable(new FExtensibilityManager);
+}
+
+
+void FCustomizableObjectPopulationEditorModule::ShutdownModule()
+{
+	if (FModuleManager::Get().IsModuleLoaded("PropertyEditor"))
+	{
+		FPropertyEditorModule& PropertyModule = FModuleManager::GetModuleChecked<FPropertyEditorModule>("PropertyEditor");
+		PropertyModule.UnregisterCustomClassLayout("CustomizableObjectPopulationClass");
+	}
+
+	CustomizableObjectPopulationEditor_ToolBarExtensibilityManager.Reset();
+
+	FCustomizableObjectPopulationEditorStyle::Shutdown();
+}
+
+
+TSharedRef<ICustomizableObjectPopulationEditor> FCustomizableObjectPopulationEditorModule::CreateCustomizableObjectPopulationEditor(const EToolkitMode::Type Mode, const TSharedPtr< IToolkitHost >& InitToolkitHost, UCustomizableObjectPopulation* CustomizablePopulation)
+{
+	TSharedRef<FCustomizableObjectPopulationEditor> NewCustomizableObjectPopulationEditor(new FCustomizableObjectPopulationEditor());
+	NewCustomizableObjectPopulationEditor->InitCustomizableObjectPopulationEditor(Mode, InitToolkitHost, CustomizablePopulation);
+	return NewCustomizableObjectPopulationEditor;
+}
+
+
+TSharedRef<ICustomizableObjectPopulationClassEditor> FCustomizableObjectPopulationEditorModule::CreateCustomizableObjectPopulationClassEditor(const EToolkitMode::Type Mode, const TSharedPtr< IToolkitHost >& InitToolkitHost, UCustomizableObjectPopulationClass* CustomizablePopulationClass)
+{
+	TSharedRef<FCustomizableObjectPopulationClassEditor> NewCustomizableObjectPopulationClassEditor(new FCustomizableObjectPopulationClassEditor());
+	NewCustomizableObjectPopulationClassEditor->InitCustomizableObjectPopulationClassEditor(Mode, InitToolkitHost, CustomizablePopulationClass);
+	return NewCustomizableObjectPopulationClassEditor;
+}
+
+
+#undef LOCTEXT_NAMESPACE
\ No newline at end of file
diff --git a/Engine/Plugins/Experimental/Mutable/Source/CustomizableObjectPopulationEditor/Public/MuCOPE/CustomizableObjectPopulationClassEditorActions.h b/Engine/Plugins/Experimental/Mutable/Source/CustomizableObjectPopulationEditor/Public/MuCOPE/CustomizableObjectPopulationClassEditorActions.h
new file mode 100644
--- /dev/null
+++ b/Engine/Plugins/Experimental/Mutable/Source/CustomizableObjectPopulationEditor/Public/MuCOPE/CustomizableObjectPopulationClassEditorActions.h
@@ -0,0 +1,39 @@
+// Copyright Epic Games, Inc. All Rights Reserved.
+
+#pragma once
+
+#include "Styling/AppStyle.h"
+#include "Framework/Commands/UICommandInfo.h"
+#include "Framework/Commands/Commands.h"
+#include "CustomizableObjectPopulationEditorStyle.h"
+
+
+class FCustomizableObjectPopulationClassEditorCommands : public TCommands<FCustomizableObjectPopulationClassEditorCommands>
+{
+public:
+
+	FCustomizableObjectPopulationClassEditorCommands() : TCommands<FCustomizableObjectPopulationClassEditorCommands>
+		(
+			"CustomizableObjectPopulationClassEditor", // Context name for fast lookup
+			NSLOCTEXT("Contexts", "CustomizableObjectPopulationClassEditor", "CustomizableObject Population Class Editor"), // Localized context name for displaying
+			NAME_None, // Parent
+			FCustomizableObjectPopulationEditorStyle::GetStyleSetName()
+			)
+	{
+	}
+
+	TSharedPtr< FUICommandInfo > Save;
+	TSharedPtr< FUICommandInfo > SaveCustomizableObject;
+	TSharedPtr< FUICommandInfo > OpenCustomizableObjectEditor;
+	TSharedPtr< FUICommandInfo > SaveEditorCurve;
+	TSharedPtr< FUICommandInfo > TestPopulationClass;
+	TSharedPtr< FUICommandInfo > GenerateInstances;
+	TSharedPtr< FUICommandInfo > InspectSelectedInstance;
+	TSharedPtr< FUICommandInfo > InspectSelectedSkeletalMesh;
+
+	/**
+	 * Initialize commands
+	 */
+	virtual void RegisterCommands() override;
+
+};
diff --git a/Engine/Plugins/Experimental/Mutable/Source/CustomizableObjectPopulationEditor/Public/MuCOPE/CustomizableObjectPopulationEditorActions.h b/Engine/Plugins/Experimental/Mutable/Source/CustomizableObjectPopulationEditor/Public/MuCOPE/CustomizableObjectPopulationEditorActions.h
new file mode 100644
--- /dev/null
+++ b/Engine/Plugins/Experimental/Mutable/Source/CustomizableObjectPopulationEditor/Public/MuCOPE/CustomizableObjectPopulationEditorActions.h
@@ -0,0 +1,35 @@
+// Copyright Epic Games, Inc. All Rights Reserved.
+
+#pragma once
+
+#include "Styling/AppStyle.h"
+#include "Framework/Commands/UICommandInfo.h"
+#include "Framework/Commands/Commands.h"
+#include "CustomizableObjectPopulationEditorStyle.h"
+
+class FCustomizableObjectPopulationEditorCommands : public TCommands<FCustomizableObjectPopulationEditorCommands>
+{
+
+public:
+	FCustomizableObjectPopulationEditorCommands() : TCommands<FCustomizableObjectPopulationEditorCommands>
+		(
+			"CustomizableObjectPopulationEditor", // Context name for fast lookup
+			NSLOCTEXT("Contexts", "CustomizableObjectPopulationEditor", "CustomizableObject Population Editor"), // Localized context name for displaying
+			NAME_None, // Parent
+			FCustomizableObjectPopulationEditorStyle::GetStyleSetName()
+			)
+	{
+	}
+
+	TSharedPtr< FUICommandInfo > Save;
+	TSharedPtr< FUICommandInfo > TestPopulation;
+	TSharedPtr< FUICommandInfo > GenerateInstances;
+	TSharedPtr< FUICommandInfo > InspectSelectedInstance;
+	TSharedPtr< FUICommandInfo > InspectSelectedSkeletalMesh;
+
+	/**
+	 * Initialize commands
+	 */
+	virtual void RegisterCommands() override;
+
+};
