Moved the Mutable plugin files to live under a directory named after the module. This allows us to get rid of long include paths that contain public and private folder names in the path and fixes head

diff --git a/Engine/Plugins/Experimental/Mutable/Source/MutableRuntime/Private/MuR/OpImageCrop.h b/Engine/Plugins/Experimental/Mutable/Source/MutableRuntime/Private/MuR/OpImageCrop.h
new file mode 100644
--- /dev/null
+++ b/Engine/Plugins/Experimental/Mutable/Source/MutableRuntime/Private/MuR/OpImageCrop.h
@@ -0,0 +1,86 @@
+// Copyright Epic Games, Inc. All Rights Reserved.
+
+#pragma once
+
+#include "MuR/ImagePrivate.h"
+#include "MuR/MutableMath.h"
+#include "MuR/Platform.h"
+#include "MuR/OpImagePixelFormat.h"
+
+
+namespace mu
+{
+
+	inline ImagePtr ImageCrop( int imageCompressionQuality, const Image* pBaseReal, const box< vec2<int> >& rect )
+	{
+		ImagePtrConst pBase = pBaseReal;
+
+		// RLE are not supported yet
+		if ( pBaseReal->GetFormat()==EImageFormat::IF_RGBA_UBYTE_RLE )
+		{
+            pBase = ImagePixelFormat( imageCompressionQuality, pBase.get(), EImageFormat::IF_RGBA_UBYTE );
+        }
+		else if ( pBaseReal->GetFormat()==EImageFormat::IF_RGB_UBYTE_RLE )
+		{
+            pBase = ImagePixelFormat( imageCompressionQuality, pBase.get(), EImageFormat::IF_RGB_UBYTE );
+        }
+		else if ( pBaseReal->GetFormat()==EImageFormat::IF_L_UBYTE_RLE )
+		{
+            pBase = ImagePixelFormat( imageCompressionQuality, pBase.get(), EImageFormat::IF_L_UBYTE );
+        }
+
+		check( pBase );
+
+		ImagePtr pCropped = new Image
+				(
+					rect.size[0], rect.size[1],
+					1,
+					pBase->GetFormat()
+					);
+
+//		check( rect.min[0]>=0 && rect.min[1]>=0 );
+//		check( rect.size[0]>=0 && rect.size[1]>=0 );
+//		check( rect.min[0]+rect.size[0]<=pBase->GetSizeX() );
+//		check( rect.min[1]+rect.size[1]<=pBase->GetSizeY() );
+		// TODO: better error control. This happens if some layouts are corrupt.
+		bool correct =
+				( rect.min[0]>=0 && rect.min[1]>=0 ) &&
+				( rect.size[0]>=0 && rect.size[1]>=0 ) &&
+				( rect.min[0]+rect.size[0]<=pBase->GetSizeX() ) &&
+				( rect.min[1]+rect.size[1]<=pBase->GetSizeY() );
+		if (!correct)
+		{
+			return pCropped;
+		}
+
+
+		const FImageFormatData& finfo = GetImageFormatData( pBase->GetFormat() );
+
+		// For now
+		check( finfo.m_pixelsPerBlockX == 1 );
+		check( finfo.m_pixelsPerBlockY == 1 );
+
+		checkf( rect.min[0] % finfo.m_pixelsPerBlockX == 0, TEXT("Rect must snap to blocks.") );
+		checkf( rect.min[1] % finfo.m_pixelsPerBlockY == 0, TEXT("Rect must snap to blocks.") );
+		checkf( rect.size[0] % finfo.m_pixelsPerBlockX == 0, TEXT("Rect must snap to blocks.") );
+		checkf( rect.size[1] % finfo.m_pixelsPerBlockY == 0, TEXT("Rect must snap to blocks.") );
+
+		int baseRowSize = finfo.m_bytesPerBlock * pBase->GetSizeX() / finfo.m_pixelsPerBlockX;
+		int cropRowSize = finfo.m_bytesPerBlock * rect.size[0] / finfo.m_pixelsPerBlockX;
+
+        const uint8_t* pBaseBuf = pBase->GetData();
+        uint8_t* pCropBuf = pCropped->GetData();
+
+		int skipPixels = pBase->GetSizeX() * rect.min[1] + rect.min[0];
+		pBaseBuf += finfo.m_bytesPerBlock * skipPixels / finfo.m_pixelsPerBlockX;
+		for ( int y=0; y<rect.size[1]; ++y )
+		{
+			memcpy( pCropBuf, pBaseBuf, cropRowSize );
+			pCropBuf += cropRowSize;
+			pBaseBuf += baseRowSize;
+		}
+
+		return pCropped;
+	}
+
+}
diff --git a/Engine/Plugins/Experimental/Mutable/Source/MutableRuntime/Private/MuR/OpImageDifference.h b/Engine/Plugins/Experimental/Mutable/Source/MutableRuntime/Private/MuR/OpImageDifference.h
new file mode 100644
--- /dev/null
+++ b/Engine/Plugins/Experimental/Mutable/Source/MutableRuntime/Private/MuR/OpImageDifference.h
@@ -0,0 +1,71 @@
+// Copyright Epic Games, Inc. All Rights Reserved.
+
+#pragma once
+
+#include "MuR/ImagePrivate.h"
+#include "MuR/Platform.h"
+
+
+namespace mu
+{
+
+	inline ImagePtr ImageDifference( const Image* pA, const Image* pB )
+	{
+		check( pA->GetSizeX() == pB->GetSizeX() );
+		check( pA->GetSizeY() == pB->GetSizeY() );
+		check( pA->GetFormat() == pB->GetFormat() );
+
+        ImagePtr pDest = new Image( pA->GetSizeX(), pA->GetSizeY(), pA->GetLODCount(), EImageFormat::IF_L_UBYTE );
+
+        uint8_t* pDestBuf = pDest->GetData();
+        const uint8_t* pABuf = pA->GetData();
+        const uint8_t* pBBuf = pB->GetData();
+
+		// Generic implementation
+		int pixelCount = (int)pA->CalculatePixelCount();
+
+		switch ( pA->GetFormat() )
+		{
+		case EImageFormat::IF_RGB_UBYTE:
+			for ( int i=0; i<pixelCount; ++i )
+			{
+                uint8_t diff = 0;
+				for ( int c=0; c<3; ++c )
+				{
+                    uint8_t a = pABuf[3*i+c];
+                    uint8_t b = pBBuf[3*i+c];
+					if ( abs(a-b)>3 )
+					{
+                        diff = std::max( diff, (uint8_t)255 );
+					}
+				}
+				pDestBuf[i] = diff;
+			}
+			break;
+
+        case EImageFormat::IF_BGRA_UBYTE:
+        case EImageFormat::IF_RGBA_UBYTE:
+            for ( int i=0; i<pixelCount; ++i )
+			{
+                uint8_t diff = 0;
+				for ( int c=0; c<3; ++c )
+				{
+                    uint8_t a = pABuf[4*i+c];
+                    uint8_t b = pBBuf[4*i+c];
+					if ( abs(a-b)>3 )
+					{
+                        diff = std::max( diff, (uint8_t)255 );
+					}
+				}
+				pDestBuf[i] = diff;
+			}
+			break;
+
+		default:
+			check(false);
+		}
+
+		return pDest;
+	}
+
+}
diff --git a/Engine/Plugins/Experimental/Mutable/Source/MutableRuntime/Private/MuR/OpImageDisplace.h b/Engine/Plugins/Experimental/Mutable/Source/MutableRuntime/Private/MuR/OpImageDisplace.h
new file mode 100644
--- /dev/null
+++ b/Engine/Plugins/Experimental/Mutable/Source/MutableRuntime/Private/MuR/OpImageDisplace.h
@@ -0,0 +1,330 @@
+// Copyright Epic Games, Inc. All Rights Reserved.
+
+#pragma once
+
+#include "MuR/ImagePrivate.h"
+#include "MuR/Platform.h"
+#include "ConvertData.h"
+#include "Async/ParallelFor.h"
+
+
+namespace mu
+{
+
+    inline uint8_t MutableEncodeOffset( int x, int y )
+	{
+        uint8_t c = uint8_t( ((7)<<4) | 7 );
+		if ( x<8 && x>-8 && y<8 && y>-8 )
+		{
+            c = uint8_t( ((x+7)<<4) | (y+7) );
+		}
+		return c;
+	}
+
+    inline void MutableDecodeOffset( uint8_t c, int& x, int& y )
+	{
+		x = int(c>>4) - 7;
+		y = int(c&0xf) - 7;
+	}
+
+	//---------------------------------------------------------------------------------------------
+	//!
+	//---------------------------------------------------------------------------------------------
+	inline void ImageMakeGrowMap( Image* pResult, const Image* pMask, int border )
+	{
+		MUTABLE_CPUPROFILER_SCOPE(ImageMakeGrowMap)
+
+		check( pMask->GetFormat() == EImageFormat::IF_L_UBYTE );
+		check( pResult->GetFormat() == EImageFormat::IF_L_UBYTE );
+		check( pResult->GetSizeX() == pMask->GetSizeX() );
+		check( pResult->GetSizeY() == pMask->GetSizeY() );
+
+		int sizeX = pResult->GetSizeX();
+		int sizeY = pResult->GetSizeY();
+		check( sizeX>0 && sizeY>0 );
+
+        if (sizeX<4 || sizeY<4)
+        {
+            return;
+        }
+
+		memset( pResult->GetData(), MutableEncodeOffset( 0, 0 ), sizeX*sizeY );
+
+		ImagePtrConst pThisMask = pMask;
+		for ( int b=0; b<border; ++b )
+		{
+			ImagePtr pNextMask = pThisMask->Clone();            
+
+			//for ( int y=0; y<sizeY; ++y )
+			const auto& ProcessRow = [
+				pThisMask, pNextMask, sizeX, sizeY, pResult
+			] (int32 y)
+			{
+				uint8_t* pNextData = pNextMask->GetData() + sizeX*y;
+
+				const uint8_t* pS0 = y > 0 ? pThisMask->GetData() + (y - 1) * sizeX : nullptr;
+				const uint8_t* pS1 = pThisMask->GetData() + y * sizeX;
+				bool bNotLastRow = y < sizeY - 1;
+				const uint8_t* pS2 = bNotLastRow ? pThisMask->GetData() + (y + 1) * sizeX : nullptr;
+
+				const uint8_t* pR0 = y > 0 ? pResult->GetData() + (y - 1) * sizeX : 0;
+				uint8_t* pR1 = pResult->GetData() + y * sizeX;
+
+				const uint8_t* pR2 = bNotLastRow ? pResult->GetData() + (y + 1) * sizeX : nullptr;
+
+				for (int x = 0; x < sizeX; ++x)
+				{
+					int dx, dy;
+					if (pS1[x])
+					{
+						*pNextData = 255;
+					}
+					else if (y > 0 && pS0[x])
+					{
+						MutableDecodeOffset(pR0[x], dx, dy);
+						pR1[x] = MutableEncodeOffset(dx, dy - 1);
+						*pNextData = 255;
+					}
+					else if (bNotLastRow && pS2 && pR2 && pS2[x])
+					{
+						MutableDecodeOffset(pR2[x], dx, dy);
+						pR1[x] = MutableEncodeOffset(dx, dy + 1);
+						*pNextData = 255;
+					}
+					else if (x < sizeX - 1 && pS1[x + 1])
+					{
+						MutableDecodeOffset(pR1[x + 1], dx, dy);
+						pR1[x] = MutableEncodeOffset(dx + 1, dy);
+						*pNextData = 255;
+					}
+					else if (x > 0 && pS1[x - 1])
+					{
+						MutableDecodeOffset(pR1[x - 1], dx, dy);
+						pR1[x] = MutableEncodeOffset(dx - 1, dy);
+						*pNextData = 255;
+					}
+
+					++pNextData;
+				}
+
+			};
+
+			ParallelFor(sizeY, ProcessRow);
+
+			pThisMask = pNextMask;
+		}
+
+	}
+
+
+	//---------------------------------------------------------------------------------------------
+	//!
+	//---------------------------------------------------------------------------------------------
+	inline void ImageDisplace( Image* pResult, const Image* pSource, const Image* pMap )
+	{
+		MUTABLE_CPUPROFILER_SCOPE(ImageDisplace)
+
+		check( pResult->GetFormat() == pSource->GetFormat()  );
+		check( pMap->GetFormat() == EImageFormat::IF_L_UBYTE
+						||
+						pMap->GetFormat() == EImageFormat::IF_L_UBYTE_RLE );
+		check( pResult->GetSizeX() == pSource->GetSizeX() );
+		check( pResult->GetSizeY() == pSource->GetSizeY() );
+		check( pResult->GetSizeX() == pMap->GetSizeX() );
+		check( pResult->GetSizeY() == pMap->GetSizeY() );
+
+		int sizeX = pResult->GetSizeX();
+		int sizeY = pResult->GetSizeY();
+
+        if (sizeX<4 || sizeY<4)
+        {
+            return;
+        }
+
+        const uint8_t* pSourceData = pSource->GetData();
+
+        auto mapFormat = pMap->GetFormat();
+        bool isUncompressed = ( mapFormat == EImageFormat::IF_L_UBYTE );
+
+        if ( isUncompressed )
+		{
+			switch (pResult->GetFormat())
+			{
+
+			case EImageFormat::IF_L_UBYTE:
+			{
+				//for ( int y=0; y<sizeY; ++y )
+				const auto& ProcessRow = [
+					pSourceData, pMap, pResult, sizeX, sizeY
+				] (int32 y)
+				{
+					const uint8_t* pMapData = pMap->GetData() + y * sizeX;
+					uint8_t* pResultData = pResult->GetData() + y * sizeX;
+
+					for (int x = 0; x < sizeX; ++x)
+					{
+						int dx, dy;
+						MutableDecodeOffset(*pMapData, dx, dy);
+						check(x + dx >= 0 && x + dx < sizeX);
+						check(y + dy >= 0 && y + dy < sizeY);
+
+						int offset = ((y + dy) * sizeX + (x + dx));
+						memcpy(pResultData, &pSourceData[offset], 1);
+
+						++pResultData;
+						++pMapData;
+					}
+
+				};
+
+				ParallelFor(sizeY, ProcessRow);
+				break;
+			}
+
+			case EImageFormat::IF_RGB_UBYTE:
+			{
+				//for (int y = 0; y < sizeY; ++y)
+				const auto& ProcessRow = [
+					pSourceData, pMap, pResult, sizeX, sizeY
+				] (int32 y)
+				{
+					const uint8_t* pMapData = pMap->GetData() + y * sizeX;
+					uint8_t* pResultData = pResult->GetData() + y * sizeX*3;
+
+					for (int x = 0; x < sizeX; ++x)
+					{
+						int dx, dy;
+						MutableDecodeOffset(*pMapData, dx, dy);
+						check(x + dx >= 0 && x + dx < sizeX);
+						check(y + dy >= 0 && y + dy < sizeY);
+
+						int offset = ((y + dy) * sizeX + (x + dx)) * 3;
+						memcpy(pResultData, &pSourceData[offset], 3);
+
+						pResultData += 3;
+						++pMapData;
+					}
+
+				};
+
+				ParallelFor(sizeY, ProcessRow);
+				break;
+			}
+
+			default:
+			{
+				// Generic
+				int pixelSize = GetImageFormatData( pResult->GetFormat() ).m_bytesPerBlock;
+				//for ( int y=0; y<sizeY; ++y )
+				const auto& ProcessRow = [
+					pSourceData, pMap, pResult, sizeX, sizeY, pixelSize
+				] (int32 y)
+				{
+					const uint8_t* pMapData = pMap->GetData() + y * sizeX;
+					uint8_t* pResultData = pResult->GetData() + y * sizeX * pixelSize;
+
+					for (int x = 0; x < sizeX; ++x)
+					{
+						int dx, dy;
+						MutableDecodeOffset(*pMapData, dx, dy);
+						check(x + dx >= 0 && x + dx < sizeX);
+						check(y + dy >= 0 && y + dy < sizeY);
+
+						int offset = ((y + dy) * sizeX + (x + dx)) * pixelSize;
+						memcpy(pResultData, &pSourceData[offset], pixelSize);
+
+						pResultData += pixelSize;
+						++pMapData;
+					}
+
+				};
+
+				ParallelFor(sizeY, ProcessRow);
+			}
+
+			}
+		}
+
+        else if (mapFormat==EImageFormat::IF_L_UBYTE_RLE )
+		{
+            int pixelSize = GetImageFormatData( pResult->GetFormat() ).m_bytesPerBlock;
+
+			const uint8_t* pMapData = pMap->GetData();
+			uint8_t* pResultData = pResult->GetData();
+
+			// TODO: ParallelFor this
+
+            for (int lod=0;lod<pMap->GetLODCount(); ++lod)
+            {
+                pMapData += 4+sizeY*sizeof(uint32_t);
+
+                for ( int y=0; y<sizeY; ++y )
+                {
+                    int x=0;
+                    while ( x<sizeX )
+                    {
+                        // Decode header
+                        uint16_t equal = *(const uint16_t*)pMapData;
+                        pMapData += 2;
+
+                        uint8_t different = *pMapData;
+                        ++pMapData;
+
+                        uint8_t equalPixel = *pMapData;
+                        ++pMapData;
+
+                        // Equal pixels
+                        {
+                            int dx,dy;
+                            MutableDecodeOffset( equalPixel, dx, dy );
+                            dx += x;
+                            dy += y;
+
+							if (dx >= 0 && dx < sizeX && dy >= 0 && dy < sizeY)
+							{
+								int offset = (dy * sizeX + dx) * pixelSize;
+								memcpy(pResultData, &pSourceData[offset], equal * pixelSize);
+							}
+
+                            pResultData += equal*pixelSize;
+                            x += equal;
+                        }
+
+
+                        // Different pixels
+                        for ( int i=0; i<different; ++i )
+                        {
+                            int dx,dy;
+                            MutableDecodeOffset( pMapData[i], dx, dy );
+                            dx += x;
+                            dy += y;
+
+							if (dx>=0 && dx<sizeX && dy>=0 && dy<sizeY)
+							{
+								int offset = (dy * sizeX + dx) * pixelSize;
+								memcpy(pResultData, &pSourceData[offset], pixelSize);
+							}
+
+                            pResultData += pixelSize;
+                            ++x;
+                        }
+
+                        pMapData += different;
+                    }
+                }
+
+                sizeX = FMath::DivideAndRoundUp(sizeX,2);
+                sizeY = FMath::DivideAndRoundUp(sizeY,2);
+            }
+		}
+
+		else
+		{
+			checkf( false, TEXT("Unsupported mask format.") );
+		}
+
+
+	}
+
+
+}
diff --git a/Engine/Plugins/Experimental/Mutable/Source/MutableRuntime/Private/MuR/OpImageDodge.h b/Engine/Plugins/Experimental/Mutable/Source/MutableRuntime/Private/MuR/OpImageDodge.h
new file mode 100644
--- /dev/null
+++ b/Engine/Plugins/Experimental/Mutable/Source/MutableRuntime/Private/MuR/OpImageDodge.h
@@ -0,0 +1,54 @@
+// Copyright Epic Games, Inc. All Rights Reserved.
+
+#pragma once
+
+#include "MuR/ImagePrivate.h"
+
+
+namespace mu
+{
+
+	inline unsigned DodgeChannelMasked( unsigned base, unsigned blended, unsigned mask )
+	{
+		// R = Base / (1-Blend)
+		unsigned dodge = ( base << 8 ) / ( 256 - blended );
+		unsigned masked = ( ( ( 255 - mask ) * base ) + ( mask * dodge ) ) >> 8;
+		return masked;
+	}
+
+	inline unsigned DodgeChannel( unsigned base, unsigned blended )
+	{
+		// R = Base / (1-Blend)
+		unsigned dodge = ( base << 8 ) / ( 256 - blended );
+		return dodge;
+	}
+
+
+    inline void ImageDodge( Image* pResult, const Image* pBase, const Image* pMask, vec3<float> col )
+	{
+        ImageLayerColour<DodgeChannelMasked,DodgeChannel, true>( pResult, pBase, pMask, col );
+	}
+
+
+    inline void ImageDodge( Image* pResult, const Image* pBase, const Image* pMask,
+                            const Image* pBlended,
+                            bool applyToAlpha )
+	{
+        ImageLayer<DodgeChannelMasked,DodgeChannel, true>
+                ( pResult, pBase, pMask, pBlended, applyToAlpha );
+	}
+
+
+    inline void ImageDodge( Image* pResult, const Image* pBase, vec3<float> col )
+	{
+        ImageLayerColour<DodgeChannel, true>( pResult, pBase, col );
+	}
+
+
+    inline void ImageDodge( Image* pResult, const Image* pBase, const Image* pBlended,
+                            bool applyToAlpha )
+	{
+        ImageLayer<DodgeChannel, true>( pResult, pBase, pBlended, applyToAlpha );
+	}
+
+}
diff --git a/Engine/Plugins/Experimental/Mutable/Source/MutableRuntime/Private/MuR/OpImageGradient.h b/Engine/Plugins/Experimental/Mutable/Source/MutableRuntime/Private/MuR/OpImageGradient.h
new file mode 100644
--- /dev/null
+++ b/Engine/Plugins/Experimental/Mutable/Source/MutableRuntime/Private/MuR/OpImageGradient.h
@@ -0,0 +1,43 @@
+// Copyright Epic Games, Inc. All Rights Reserved.
+
+#pragma once
+
+#include "MuR/ImagePrivate.h"
+#include "MuR/Platform.h"
+
+
+
+namespace mu
+{
+
+	inline ImagePtr ImageGradient( const vec3<float>& c0,
+								   const vec3<float>& c1,
+								   int sizeX, int sizeY )
+	{
+		ImagePtr pDest = new Image((uint16_t)sizeX, (uint16_t)sizeY, 1, EImageFormat::IF_RGB_UBYTE );
+
+        uint8_t* pDestBuf = pDest->GetData();
+
+		for ( int i=0; i<sizeX; ++i )
+		{
+			float delta = float(i)/float(sizeX-1);
+			vec3<float> c = c0 * (1.0f-delta) + c1 * delta;
+
+            uint8_t colour[4];
+            colour[0] = (uint8_t)std::max( 0, std::min( 255, int(c[0]*255.0f) ) );
+            colour[1] = (uint8_t)std::max( 0, std::min( 255, int(c[1]*255.0f) ) );
+            colour[2] = (uint8_t)std::max( 0, std::min( 255, int(c[2]*255.0f) ) );
+			colour[3] = 255;
+
+			for ( int j=0; j<sizeY; ++j )
+			{
+				pDestBuf[ (sizeX*j+i)*3 + 0 ] = colour[0];
+				pDestBuf[ (sizeX*j+i)*3 + 1 ] = colour[1];
+				pDestBuf[ (sizeX*j+i)*3 + 2 ] = colour[2];
+			}
+		}
+
+		return pDest;
+	}
+
+}
diff --git a/Engine/Plugins/Experimental/Mutable/Source/MutableRuntime/Private/MuR/OpImageGrow.h b/Engine/Plugins/Experimental/Mutable/Source/MutableRuntime/Private/MuR/OpImageGrow.h
new file mode 100644
--- /dev/null
+++ b/Engine/Plugins/Experimental/Mutable/Source/MutableRuntime/Private/MuR/OpImageGrow.h
@@ -0,0 +1,124 @@
+// Copyright Epic Games, Inc. All Rights Reserved.
+
+#pragma once
+
+#include "MuR/ImagePrivate.h"
+#include "MuR/Platform.h"
+#include "ConvertData.h"
+
+
+namespace mu
+{
+
+	//---------------------------------------------------------------------------------------------
+	//! Grow the image as if it was a bitmap
+	//---------------------------------------------------------------------------------------------
+	inline void ImageGrow( Image* pImage )
+	{
+		check( pImage->GetFormat()== EImageFormat::IF_L_UBYTE );
+
+		MUTABLE_CPUPROFILER_SCOPE(ImageGrow1)
+
+		int sizeX = pImage->GetSizeX();
+		int sizeY = pImage->GetSizeY();
+		check( sizeX>0 && sizeY>0 );
+
+		ImagePtrConst pSource = pImage->Clone();
+
+		int rowSize = sizeX*GetImageFormatData( pImage->GetFormat() ).m_bytesPerBlock;
+
+		for ( int y=0; y<sizeY; ++y )
+		{
+			int y0 = std::max( 0, std::min( sizeY-1, y-1 ) );
+			int y1 = std::max( 0, std::min( sizeY-1, y   ) );
+			int y2 = std::max( 0, std::min( sizeY-1, y+1 ) );
+
+            const uint8_t* pS0 = pSource->GetData() + y0 * rowSize;
+            const uint8_t* pS1 = pSource->GetData() + y1 * rowSize;
+            const uint8_t* pS2 = pSource->GetData() + y2 * rowSize;
+
+            uint8_t* pR = pImage->GetData() + y1 * rowSize;
+
+			if ( pS0[0] || pS2[0] || pS1[1] )
+			{
+				pR[0] = 255;
+			}
+
+			if ( pS0[sizeX-1] || pS2[sizeX-1] || pS1[sizeX-2] )
+			{
+				pR[sizeX-1] = 255;
+			}
+
+			for ( int x=1; x<sizeX-1; ++x )
+			{
+				if ( pS0[x] || pS2[x] || pS1[x-1] || pS1[x+1] )
+				{
+					pR[x] = 255;
+				}
+			}
+
+		}
+
+	}
+
+
+	//---------------------------------------------------------------------------------------------
+	//! Grow the image copying border colours, by using an external mask as reference
+	//---------------------------------------------------------------------------------------------
+	inline void ImageGrow( Image* pImage, const Image* pMask )
+	{
+		MUTABLE_CPUPROFILER_SCOPE(ImageGrow2)
+
+		check( pMask->GetFormat()== EImageFormat::IF_L_UBYTE );
+
+		int sizeX = pImage->GetSizeX();
+		int sizeY = pImage->GetSizeY();
+		check( sizeX>0 && sizeY>0 );
+
+        if (sizeX<2 || sizeY<2)
+        {
+            return;
+        }
+
+		int pixelSize = GetImageFormatData( pImage->GetFormat() ).m_bytesPerBlock;
+		int rowSize = sizeX*pixelSize;
+
+		for ( int y=0; y<sizeY; ++y )
+		{
+            const uint8_t* pS0 = y>0 ? pMask->GetData() + (y-1) * sizeX : nullptr;
+            const uint8_t* pS1 = pMask->GetData() + y * sizeX;
+            const uint8_t* pS2 = y<sizeY-1 ? pMask->GetData() + (y+1) * sizeX : nullptr;
+
+            const uint8_t* pR0 = y>0 ? pImage->GetData() + (y-1) * rowSize : nullptr;
+                  uint8_t* pR1 = pImage->GetData() + y * rowSize;
+            const uint8_t* pR2 = y<sizeY-1 ? pImage->GetData() + (y+1) * rowSize : nullptr;
+
+			for ( int x=0; x<sizeX; ++x )
+			{
+				if (pS1[x])
+				{
+					// Nothing to do.
+				}
+				else if ( y>0 && pS0[x] )
+				{
+					memcpy( &pR1[x*pixelSize], &pR0[x*pixelSize], pixelSize );
+				}
+				else if ( y<sizeY-1 && pS2[x] )
+				{
+					memcpy( &pR1[x*pixelSize], &pR2[x*pixelSize], pixelSize );
+				}
+				else if ( x<sizeX-1 && pS1[x+1] )
+				{
+					memcpy( &pR1[x*pixelSize], &pR1[(x+1)*pixelSize], pixelSize );
+				}
+				else if ( x>0 && pS1[x-1] )
+				{
+					memcpy( &pR1[x*pixelSize], &pR1[(x-1)*pixelSize], pixelSize );
+				}
+			}
+
+		}
+
+	}
+
+}
diff --git a/Engine/Plugins/Experimental/Mutable/Source/MutableRuntime/Private/MuR/OpImageHardLight.h b/Engine/Plugins/Experimental/Mutable/Source/MutableRuntime/Private/MuR/OpImageHardLight.h
new file mode 100644
--- /dev/null
+++ b/Engine/Plugins/Experimental/Mutable/Source/MutableRuntime/Private/MuR/OpImageHardLight.h
@@ -0,0 +1,61 @@
+// Copyright Epic Games, Inc. All Rights Reserved.
+
+#pragma once
+
+#include "MuR/ImagePrivate.h"
+
+namespace mu
+{
+
+	inline unsigned HardLightChannel( unsigned base, unsigned blended )
+	{
+		// gimp-like
+
+
+		// photoshop-like
+		// if (Blend > ½) R = 1 - (1-Base) × (1-2×(Blend-½))
+		// if (Blend <= ½) R = Base × (2×Blend)
+		unsigned hardlight = blended > 128
+				? 255 - ( ( ( 255 - base ) * ( 255 - 2*(blended - 128 ) ) ) >> 8 )
+				: ( base * (2 * blended) ) >> 8;
+
+		hardlight = std::min( 255u, hardlight );
+
+		return hardlight;
+	}
+
+	inline unsigned HardLightChannelMasked( unsigned base, unsigned blended, unsigned mask )
+	{
+		unsigned hardlight = HardLightChannel( base, blended );
+
+		unsigned masked = ( ( ( 255 - mask ) * base ) + ( mask * hardlight ) ) >> 8;
+		return masked;
+	}
+
+
+    inline void ImageHardLight( Image* pResult, const Image* pBase, vec3<float> col )
+	{
+        ImageLayerColour<HardLightChannel, true>( pResult, pBase, col );
+	}
+
+
+    inline void ImageHardLight( Image* pResult, const Image* pBase, const Image* pBlended,
+                                bool applyToAlpha )
+	{
+        ImageLayer<HardLightChannel, true>( pResult, pBase, pBlended, applyToAlpha );
+	}
+
+
+    inline void ImageHardLight( Image* pResult, const Image* pBase, const Image* pMask, vec3<float> col )
+	{
+        ImageLayerColour<HardLightChannelMasked,HardLightChannel,true>( pResult, pBase, pMask, col );
+	}
+
+
+    inline void ImageHardLight( Image* pResult, const Image* pBase, const Image* pMask, const Image* pBlended,
+                                bool applyToAlpha )
+	{
+        ImageLayer<HardLightChannelMasked,HardLightChannel,true>( pResult, pBase, pMask, pBlended, applyToAlpha );
+	}
+
+}
diff --git a/Engine/Plugins/Experimental/Mutable/Source/MutableRuntime/Private/MuR/OpImageInterpolate.h b/Engine/Plugins/Experimental/Mutable/Source/MutableRuntime/Private/MuR/OpImageInterpolate.h
new file mode 100644
--- /dev/null
+++ b/Engine/Plugins/Experimental/Mutable/Source/MutableRuntime/Private/MuR/OpImageInterpolate.h
@@ -0,0 +1,172 @@
+// Copyright Epic Games, Inc. All Rights Reserved.
+
+#pragma once
+
+#include "MuR/ImagePrivate.h"
+#include "MuR/Platform.h"
+
+namespace mu
+{
+
+	//---------------------------------------------------------------------------------------------
+	//!
+	//---------------------------------------------------------------------------------------------
+    inline void ImageInterpolate( Image* pDest,
+                                  const Image* pA,
+                                  const Image* pB,
+                                  float factor )
+	{
+		MUTABLE_CPUPROFILER_SCOPE(ImageInterpolate)
+
+        check( pA->GetSizeX() == pDest->GetSizeX() );
+        check( pA->GetSizeY() == pDest->GetSizeY() );
+        check( pA->GetFormat() == pDest->GetFormat() );
+        check( pA->GetSizeX() == pB->GetSizeX() );
+        check( pA->GetSizeY() == pB->GetSizeY() );
+        check( pA->GetFormat() == pB->GetFormat() );
+
+		// Clamp the factor
+		factor = std::max( 0.0f, std::min( 1.0f, factor ) );
+
+        uint8_t* pDestBuf = pDest->GetData();
+        const uint8_t* pABuf = pA->GetData();
+        const uint8_t* pBBuf = pB->GetData();
+
+		// Generic implementation
+		int pixelCount = (int)pA->CalculatePixelCount();
+
+		switch ( pA->GetFormat() )
+		{
+		case EImageFormat::IF_L_UBYTE:
+		{
+            uint32_t w_8 = (uint32_t)(factor*255);
+			for ( int i=0; i<pixelCount; ++i )
+			{
+                uint32_t a_8 = *pABuf++;
+                uint32_t b_8 = *pBBuf++;
+                uint32_t i_16 = a_8 * (255-w_8) + b_8 * w_8;
+                *pDestBuf++ = (uint8_t) ( i_16>>8 );
+			}
+			break;
+		}
+
+		case EImageFormat::IF_RGB_UBYTE:
+		{
+            uint32_t w_8 = (uint32_t)(factor*255);
+			for ( int i=0; i<pixelCount*3; ++i )
+			{
+                uint32_t a_8 = *pABuf++;
+                uint32_t b_8 = *pBBuf++;
+                uint32_t i_16 = a_8 * (255-w_8) + b_8 * w_8;
+                *pDestBuf++ = (uint8_t) ( i_16>>8 );
+			}
+			break;
+		}
+
+        case EImageFormat::IF_BGRA_UBYTE:
+        case EImageFormat::IF_RGBA_UBYTE:
+        {
+            uint32_t w_8 = (uint32_t)(factor*255);
+			for ( int i=0; i<pixelCount*4; ++i )
+			{
+                uint32_t a_8 = *pABuf++;
+                uint32_t b_8 = *pBBuf++;
+                uint32_t i_16 = a_8 * (255-w_8) + b_8 * w_8;
+                *pDestBuf++ = (uint8_t)(i_16>>8);
+			}
+			break;
+		}
+
+		default:
+			check(false);
+		}
+
+	}
+
+
+	//---------------------------------------------------------------------------------------------
+	//!
+	//---------------------------------------------------------------------------------------------
+	inline ImagePtr ImageInterpolate( const Image* p0, const Image* p1, const Image* p2,
+									  float factor1, float factor2 )
+	{
+		MUTABLE_CPUPROFILER_SCOPE(ImageInterpolate2)
+
+		check( p0->GetSizeX() == p1->GetSizeX() );
+		check( p0->GetSizeY() == p1->GetSizeY() );
+		check( p0->GetFormat() == p1->GetFormat() );
+		check( p0->GetSizeX() == p2->GetSizeX() );
+		check( p0->GetSizeY() == p2->GetSizeY() );
+		check( p0->GetFormat() == p2->GetFormat() );
+
+		// Clamp the factor
+		factor1 = std::max( 0.0f, std::min( 1.0f, factor1 ) );
+		factor2 = std::max( 0.0f, std::min( 1.0f-factor1, factor2 ) );
+		float factor0 = std::max( 0.0f, 1.0f-factor1-factor2 );
+
+        ImagePtr pDest = new Image( p0->GetSizeX(), p0->GetSizeY(),
+                                    p0->GetLODCount(),
+                                    p0->GetFormat() );
+
+        uint8_t* pDestBuf = pDest->GetData();
+        const uint8_t* p0Buf = p0->GetData();
+        const uint8_t* p1Buf = p1->GetData();
+        const uint8_t* p2Buf = p2->GetData();
+
+		// Generic implementation
+		int pixelCount = (int)p0->CalculatePixelCount();
+
+        uint32_t w0_8 = (uint32_t)(factor0*255);
+        uint32_t w1_8 = (uint32_t)(factor1*255);
+        uint32_t w2_8 = (uint32_t)(factor2*255);
+
+		switch ( p0->GetFormat() )
+		{
+		case EImageFormat::IF_L_UBYTE:
+		{
+			for ( int i=0; i<pixelCount; ++i )
+			{
+                uint32_t t0_8 = *p0Buf++;
+                uint32_t t1_8 = *p1Buf++;
+                uint32_t t2_8 = *p2Buf++;
+                uint32_t i_16 = t0_8 * w0_8 + t1_8 * w1_8 + t2_8 * w2_8;
+                *pDestBuf++ = (uint8_t) ( i_16>>8 );
+			}
+			break;
+		}
+
+		case EImageFormat::IF_RGB_UBYTE:
+		{
+			for ( int i=0; i<pixelCount*3; ++i )
+			{
+                uint32_t t0_8 = *p0Buf++;
+                uint32_t t1_8 = *p1Buf++;
+                uint32_t t2_8 = *p2Buf++;
+                uint32_t i_16 = t0_8 * w0_8 + t1_8 * w1_8 + t2_8 * w2_8;
+                *pDestBuf++ = (uint8_t) ( i_16>>8 );
+			}
+			break;
+		}
+
+        case EImageFormat::IF_BGRA_UBYTE:
+        case EImageFormat::IF_RGBA_UBYTE:
+        {
+			for ( int i=0; i<pixelCount*4; ++i )
+			{
+                uint32_t t0_8 = *p0Buf++;
+                uint32_t t1_8 = *p1Buf++;
+                uint32_t t2_8 = *p2Buf++;
+                uint32_t i_16 = t0_8 * w0_8 + t1_8 * w1_8 + t2_8 * w2_8;
+                *pDestBuf++ = (uint8_t) (i_16>>8);
+			}
+			break;
+		}
+
+		default:
+			check(false);
+		}
+
+		return pDest;
+	}
+
+}
diff --git a/Engine/Plugins/Experimental/Mutable/Source/MutableRuntime/Private/MuR/OpImageInvert.h b/Engine/Plugins/Experimental/Mutable/Source/MutableRuntime/Private/MuR/OpImageInvert.h
new file mode 100644
--- /dev/null
+++ b/Engine/Plugins/Experimental/Mutable/Source/MutableRuntime/Private/MuR/OpImageInvert.h
@@ -0,0 +1,68 @@
+// Copyright Epic Games, Inc. All Rights Reserved.
+
+#pragma once
+
+#include "MuR/ImagePrivate.h"
+#include "MuR/Platform.h"
+
+namespace mu
+{
+
+	//---------------------------------------------------------------------------------------------
+	//!
+	//---------------------------------------------------------------------------------------------
+	inline ImagePtr ImageInvert(const Image* pA)
+	{
+		MUTABLE_CPUPROFILER_SCOPE(ImageInvert)
+
+		ImagePtr pDest = new Image(pA->GetSizeX(), pA->GetSizeY(), pA->GetLODCount(), pA->GetFormat());
+
+		uint8_t* pDestBuf = pDest->GetData();
+		const uint8_t* pABuf = pA->GetData();
+
+		//Generic implementation
+		int pixelCount = (int)pA->CalculatePixelCount();
+
+		switch (pA->GetFormat())
+		{
+		case EImageFormat::IF_L_UBYTE:
+		{
+			for (int i = 0; i < pixelCount; ++i)
+			{
+				pDestBuf[i] = 255 - pABuf[i];
+			}
+			break;
+		}
+		
+		case EImageFormat::IF_RGB_UBYTE:
+		{
+			for (int i = 0; i < pixelCount; ++i)
+			{
+				for (int c = 0; c < 3; ++c)
+				{
+					pDestBuf[i * 3 + c] = 255 - pABuf[i * 3 + c];
+				}
+			}
+			break;
+		}
+		case EImageFormat::IF_RGBA_UBYTE:
+		case EImageFormat::IF_BGRA_UBYTE:
+		{
+			for (int i = 0; i < pixelCount; ++i)
+			{
+				for (int c = 0; c < 3; ++c)
+				{
+					pDestBuf[i * 4 + c] = 255 - pABuf[i * 4 + c];
+				}
+
+				pDestBuf[i * 4 + 3] = pABuf[i * 4 + 3];
+			}
+			break;
+		}
+		default:
+			check(false);			
+		}
+
+		return pDest;
+	}
+}
diff --git a/Engine/Plugins/Experimental/Mutable/Source/MutableRuntime/Private/MuR/OpImageLayer.h b/Engine/Plugins/Experimental/Mutable/Source/MutableRuntime/Private/MuR/OpImageLayer.h
new file mode 100644
--- /dev/null
+++ b/Engine/Plugins/Experimental/Mutable/Source/MutableRuntime/Private/MuR/OpImageLayer.h
@@ -0,0 +1,1729 @@
+// Copyright Epic Games, Inc. All Rights Reserved.
+
+#pragma once
+
+#include "MuR/ImagePrivate.h"
+#include "MuR/ImageRLE.h"
+#include "Async/ParallelFor.h"
+
+#include "Templates/UnrealTemplate.h"
+
+namespace mu
+{
+
+
+	//---------------------------------------------------------------------------------------------
+	//! Apply a blending function to an image with a colour source
+	//---------------------------------------------------------------------------------------------
+	template< unsigned (*BLEND_FUNC)(unsigned,unsigned),
+			  bool CLAMP >
+    inline void ImageLayerColour( Image* pResult, const Image* pBase, vec3<float> col )
+	{
+        check( pResult->GetFormat() == pBase->GetFormat() );
+        check( pResult->GetSizeX() == pBase->GetSizeX() );
+        check( pResult->GetSizeY() == pBase->GetSizeY() );
+        check( pResult->GetLODCount() == pBase->GetLODCount() );
+
+		EImageFormat baseFormat = pBase->GetFormat();
+
+        uint8* pDestBuf = pResult->GetData();
+        const uint8* pBaseBuf = pBase->GetData();
+
+		unsigned top[3];
+		top[0] = (unsigned)(255 * col[0]);
+		top[1] = (unsigned)(255 * col[1]);
+		top[2] = (unsigned)(255 * col[2]);
+
+		// Generic implementation
+		int32 pixelCount = pBase->CalculatePixelCount();
+
+		constexpr int PixelCountConcurrencyThreshold = 0xff;
+
+		switch (baseFormat)
+		{
+		case EImageFormat::IF_RGB_UBYTE:
+		{
+			const auto& ProcessPixel = [
+				pBaseBuf, pDestBuf, top
+			] (uint32 i)
+			{
+				for (int c = 0; c < 3; ++c)
+				{
+					unsigned base = pBaseBuf[3 * i + c];
+					unsigned result = BLEND_FUNC(base, top[c]);
+					if (CLAMP)
+					{
+						pDestBuf[3 * i + c] = (uint8)std::min(255u, result);
+					}
+					else
+					{
+						pDestBuf[3 * i + c] = (uint8)result;
+					}
+				}
+			};
+
+			if (pixelCount > PixelCountConcurrencyThreshold)
+			{
+				ParallelFor(pixelCount, ProcessPixel);
+			}
+			else
+			{
+				for ( int p=0; p< pixelCount; ++p )
+				{
+					ProcessPixel(p);
+				}
+			}
+			break;
+		}
+		case EImageFormat::IF_RGBA_UBYTE:
+		{
+			const auto& ProcessPixel = [
+				pBaseBuf, pDestBuf, top
+			] (uint32 i)
+			{
+				for (int c = 0; c < 3; ++c)
+				{
+					unsigned base = pBaseBuf[4 * i + c];
+					unsigned result = BLEND_FUNC(base, top[c]);
+					if (CLAMP)
+					{
+						pDestBuf[4 * i + c] = (uint8)std::min(255u, result);
+					}
+					else
+					{
+						pDestBuf[4 * i + c] = (uint8)result;
+					}
+				}
+				pDestBuf[4 * i + 3] = pBaseBuf[4 * i + 3];
+			};
+
+			if (pixelCount > PixelCountConcurrencyThreshold)
+			{
+				ParallelFor(pixelCount, ProcessPixel);
+			}
+			else
+			{
+				for (int p = 0; p < pixelCount; ++p)
+				{
+					ProcessPixel(p);
+				}
+			}
+			break;
+		}
+		case EImageFormat::IF_BGRA_UBYTE:
+		{
+			const auto& ProcessPixel = [
+					pBaseBuf, pDestBuf, top
+				] (uint32 i)
+				{
+					for (int c = 0; c < 3; ++c)
+					{
+						unsigned base = pBaseBuf[4 * i + c];
+						unsigned result = BLEND_FUNC(base, top[2 - c]);
+						if (CLAMP)
+						{
+							pDestBuf[4 * i + c] = (uint8)std::min(255u, result);
+						}
+						else
+						{
+							pDestBuf[4 * i + c] = (uint8)result;
+						}
+					}
+					pDestBuf[4 * i + 3] = pBaseBuf[4 * i + 3];
+				};
+
+			if (pixelCount > PixelCountConcurrencyThreshold)
+			{
+				ParallelFor(pixelCount, ProcessPixel);
+			}
+			else
+			{
+				for (int p = 0; p < pixelCount; ++p)
+				{
+					ProcessPixel(p);
+				}
+			}
+			break;
+		}
+		case EImageFormat::IF_L_UBYTE:
+		{
+			const auto& ProcessPixel = [
+					pBaseBuf, pDestBuf, top
+				] (uint32 i)
+				{
+					unsigned base = pBaseBuf[i];
+					unsigned result = BLEND_FUNC(base, top[0]);
+					if (CLAMP)
+					{
+						pDestBuf[i] = (uint8)std::min(255u, result);
+					}
+					else
+					{
+						pDestBuf[i] = (uint8)result;
+					}
+				};
+
+
+			if (pixelCount > PixelCountConcurrencyThreshold)
+			{
+				ParallelFor(pixelCount, ProcessPixel);
+			}
+			else
+			{
+				for (int p = 0; p < pixelCount; ++p)
+				{
+					ProcessPixel(p);
+				}
+			}
+			break;
+		}
+
+		default:
+			checkf(false, TEXT("Unsupported format."));
+			break;
+		}
+    }
+
+
+	//---------------------------------------------------------------------------------------------
+	template< unsigned (*BLEND_FUNC_MASKED)(unsigned,unsigned,unsigned),
+			  unsigned (*BLEND_FUNC)(unsigned,unsigned),
+			  bool CLAMP,
+			  int NC >
+    inline void BufferLayerColourFormat( uint8* pStartDestBuf,
+								   const Image* pBase,
+								   const Image* pMask,
+								   vec3<float> col )
+	{
+		unsigned top[3];
+		top[0] = (unsigned)(255 * col[0]);
+		top[1] = (unsigned)(255 * col[1]);
+		top[2] = (unsigned)(255 * col[2]);
+
+		uint8* pDestBuf = pStartDestBuf;
+        const uint8* pMaskBuf = pMask->GetData();
+        const uint8* pBaseBuf = pBase->GetData();
+		EImageFormat maskFormat = pMask->GetFormat();
+
+		int32 LODCount = pBase->GetLODCount();
+
+        bool isUncompressed = ( maskFormat == EImageFormat::IF_L_UBYTE );
+
+        if ( isUncompressed )
+        {
+            int32 pixelCount = pBase->CalculatePixelCount();
+            for ( int i=0; i<pixelCount; ++i )
+            {
+                unsigned mask = pMaskBuf[i];
+                for ( int c=0; c<NC && c<3; ++c )
+                {
+                    unsigned base = pBaseBuf[NC*i+c];
+                    unsigned result = BLEND_FUNC_MASKED( base, top[c], mask );
+                    if ( CLAMP )
+                    {
+                        pDestBuf[NC*i+c] = (uint8)FMath::Min( 255u, result );
+                    }
+                    else
+                    {
+                        pDestBuf[NC*i+c] = (uint8)result;
+                    }
+                }
+
+                constexpr bool isNC4 = (NC==4);
+                if ( isNC4 )
+                {
+                    pDestBuf[NC*i+3] = pBaseBuf[NC*i+3];
+                }
+            }
+        }
+        else if ( maskFormat==EImageFormat::IF_L_UBYTE_RLE )
+        {
+            int rows = pBase->GetSizeY();
+            int width = pBase->GetSizeX();
+
+			// debug test
+			ImagePtr TempDecodedMask = new Image(pMask->GetSizeX(), pMask->GetSizeY(), pMask->GetLODCount(), EImageFormat::IF_L_UBYTE);
+			UncompressRLE_L(pMask, TempDecodedMask.get());
+
+            for (int32 lod=0;lod<LODCount; ++lod)
+            {
+				// Skip mip size and row sizes.
+                pMaskBuf += sizeof(uint32) +rows*sizeof(uint32);
+
+                for ( int r=0; r<rows; ++r )
+                {
+                    const uint8* pDestRowEnd = pDestBuf + width*NC;
+                    while ( pDestBuf!=pDestRowEnd )
+                    {
+                        // Decode header
+						uint16 equal = 0;
+						FMemory::Memmove(&equal, pMaskBuf, sizeof(uint16));
+                        pMaskBuf += 2;
+
+                        uint8 different = *pMaskBuf;
+                        ++pMaskBuf;
+
+                        uint8 equalPixel = *pMaskBuf;
+                        ++pMaskBuf;
+
+                        // Equal pixels
+						check(pDestBuf + NC * equal <= pStartDestBuf + pBase->GetDataSize());
+                        if ( equalPixel==255 )
+                        {
+                            for ( int i=0; i<equal; ++i )
+                            {
+                                for ( int c=0; c<NC && c<3; ++c )
+                                {
+                                    unsigned base = pBaseBuf[NC*i+c];
+                                    unsigned result = BLEND_FUNC( base, top[c] );
+                                    if ( CLAMP )
+                                    {
+                                        pDestBuf[NC*i+c] = (uint8)FMath::Min( 255u, result );
+                                    }
+                                    else
+                                    {
+                                        pDestBuf[NC*i+c] = (uint8)result;
+                                    }
+                                }
+
+                                constexpr bool isNC4 = (NC==4);
+                                if ( isNC4 )
+                                {
+                                    pDestBuf[NC*i+3] = pBaseBuf[NC*i+3];
+                                }
+                            }
+                        }
+                        else if ( equalPixel>0 )
+                        {
+                            for ( int i=0; i<equal; ++i )
+                            {
+                                for ( int c=0; c<NC && c<3; ++c )
+                                {
+                                    unsigned base = pBaseBuf[NC*i+c];
+                                    unsigned result = BLEND_FUNC_MASKED( base, top[c], equalPixel );
+                                    if ( CLAMP )
+                                    {
+                                        pDestBuf[NC*i+c] = (uint8)FMath::Min( 255u, result );
+                                    }
+                                    else
+                                    {
+                                        pDestBuf[NC*i+c] = (uint8)result;
+                                    }
+                                }
+
+                                constexpr bool isNC4 = (NC==4);
+                                if ( isNC4 )
+                                {
+                                    pDestBuf[NC*i+3] = pBaseBuf[NC*i+3];
+                                }
+                            }
+                        }
+                        else
+                        {
+                            // It could happen if xxxxxOnBase
+                            if (pDestBuf!=pBaseBuf)
+                            {
+                                FMemory::Memmove( pDestBuf, pBaseBuf, NC*equal );
+                            }
+                        }
+                        pDestBuf += NC*equal;
+                        pBaseBuf += NC*equal;
+
+                        // Different pixels
+						check(pDestBuf + NC * different <= pStartDestBuf + pBase->GetDataSize());
+                        for ( int i=0; i<different; ++i )
+                        {
+                            for ( int c=0; c<NC && c<3; ++c )
+                            {
+                                unsigned mask = pMaskBuf[i];
+                                unsigned base = pBaseBuf[NC*i+c];
+                                unsigned result = BLEND_FUNC_MASKED( base, top[c], mask );
+                                if ( CLAMP )
+                                {
+                                    pDestBuf[NC*i+c] = (uint8)FMath::Min( 255u, result );
+                                }
+                                else
+                                {
+                                    pDestBuf[NC*i+c] = (uint8)result;
+                                }
+                            }
+
+                            constexpr bool isNC4 = (NC==4);
+                            if ( isNC4 )
+                            {
+                                pDestBuf[NC*i+3] = pBaseBuf[NC*i+3];
+                            }
+                        }
+
+                        pDestBuf += NC*different;
+                        pBaseBuf += NC*different;
+                        pMaskBuf += different;
+                    }
+                }
+
+                rows = FMath::DivideAndRoundUp(rows,2);
+                width = FMath::DivideAndRoundUp(width,2);
+            }
+        }
+        else
+        {
+            checkf( false, TEXT("Unsupported mask format.") );
+        }
+	}
+
+
+	//---------------------------------------------------------------------------------------------
+	template< unsigned (*BLEND_FUNC_MASKED)(unsigned,unsigned,unsigned),
+			  unsigned (*BLEND_FUNC)(unsigned,unsigned),
+			  bool CLAMP
+			  >
+    inline void BufferLayerColour( uint8* pDestBuf,
+								   const Image* pBase,
+								   const Image* pMask,
+								   vec3<float> col )
+	{
+		EImageFormat baseFormat = pBase->GetFormat();
+		if ( baseFormat==EImageFormat::IF_RGB_UBYTE )
+		{
+			BufferLayerColourFormat<BLEND_FUNC_MASKED, BLEND_FUNC, CLAMP, 3>( pDestBuf, pBase, pMask, col );
+		}
+        else if ( baseFormat==EImageFormat::IF_RGBA_UBYTE )
+        {
+            BufferLayerColourFormat<BLEND_FUNC_MASKED, BLEND_FUNC, CLAMP, 4>( pDestBuf, pBase, pMask, col );
+        }
+        else if ( baseFormat==EImageFormat::IF_BGRA_UBYTE )
+        {
+            float temp = col[0];
+            col[0] = col[2];
+            col[2] = temp;
+            BufferLayerColourFormat<BLEND_FUNC_MASKED, BLEND_FUNC, CLAMP, 4>( pDestBuf, pBase, pMask, col );
+        }
+        else if ( baseFormat==EImageFormat::IF_L_UBYTE )
+		{
+			BufferLayerColourFormat<BLEND_FUNC_MASKED, BLEND_FUNC, CLAMP, 1>( pDestBuf, pBase, pMask, col );
+		}
+		else
+		{
+			checkf( false, TEXT("Unsupported format.") );
+		}
+	}
+
+	//---------------------------------------------------------------------------------------------
+	//! Apply a blending function to an image with a colour source and a mask
+	//---------------------------------------------------------------------------------------------
+	template< unsigned (*BLEND_FUNC_MASKED)(unsigned,unsigned,unsigned),
+			  unsigned (*BLEND_FUNC)(unsigned,unsigned),
+			  bool CLAMP >
+    inline void ImageLayerColour( Image* pResult,
+                                  const Image* pBase,
+                                  const Image* pMask,
+                                  vec3<float> col )
+	{
+        check( pResult->GetFormat() == pBase->GetFormat() );
+        check( pResult->GetSizeX() == pBase->GetSizeX() );
+        check( pResult->GetSizeY() == pBase->GetSizeY() );
+        check( pResult->GetLODCount() == pBase->GetLODCount() );
+        check( pBase->GetSizeX() == pMask->GetSizeX() );
+		check( pBase->GetSizeY() == pMask->GetSizeY() );
+		check( pMask->GetFormat() == EImageFormat::IF_L_UBYTE
+						||
+						pMask->GetFormat() == EImageFormat::IF_L_UBYTE_RLE );
+
+        bool valid =
+                ( pResult->GetFormat() == pBase->GetFormat() )
+                &&
+                ( pResult->GetSizeX() == pBase->GetSizeX() )
+                &&
+                ( pResult->GetSizeY() == pBase->GetSizeY() )
+                &&
+                ( pResult->GetLODCount() == pBase->GetLODCount() )
+                &&
+                ( pBase->GetSizeX() == pMask->GetSizeX() )
+                &&
+                ( pBase->GetSizeY() == pMask->GetSizeY() )
+                &&
+                ( pMask->GetFormat() == EImageFormat::IF_L_UBYTE
+                                        ||
+                                        pMask->GetFormat() == EImageFormat::IF_L_UBYTE_RLE );
+        if (valid)
+        {
+            BufferLayerColour<BLEND_FUNC_MASKED, BLEND_FUNC, CLAMP>( pResult->GetData(), pBase, pMask, col );
+        }
+	}
+
+
+	//---------------------------------------------------------------------------------------------
+	//! Apply a blending function to an image with a colour source on the base image itself
+	//---------------------------------------------------------------------------------------------
+	template< unsigned (*BLEND_FUNC_MASKED)(unsigned,unsigned,unsigned),
+			  unsigned (*BLEND_FUNC)(unsigned,unsigned),
+			  bool CLAMP >
+	inline void ImageLayerColourOnBase( Image* pBase, const Image* pMask, vec3<float> col )
+	{
+		check( pBase->GetSizeX() == pMask->GetSizeX() );
+		check( pBase->GetSizeY() == pMask->GetSizeY() );
+		check( pMask->GetFormat() == EImageFormat::IF_L_UBYTE
+						||
+						pMask->GetFormat() == EImageFormat::IF_L_UBYTE_RLE );
+
+		BufferLayerColour<BLEND_FUNC_MASKED,BLEND_FUNC,CLAMP>( pBase->GetData(), pBase, pMask, col );
+	}
+
+
+	//---------------------------------------------------------------------------------------------
+	//! Apply a blending function to an image with another image as blending layer
+	//---------------------------------------------------------------------------------------------
+	template< unsigned (*BLEND_FUNC)(unsigned,unsigned), bool CLAMP, int NC >
+    inline void BufferLayer(uint8* pDestBuf, const Image* pBase, const Image* pBlended,
+                            bool applyToAlpha )
+	{
+		check( pBase->GetSizeX() == pBlended->GetSizeX() );
+		check( pBase->GetSizeY() == pBlended->GetSizeY() );
+		check( pBase->GetFormat() == pBlended->GetFormat() );
+
+        const uint8* pBaseBuf = pBase->GetData();
+        const uint8* pBlendedBuf = pBlended->GetData();
+
+		// Generic implementation
+		int32 pixelCount = pBase->CalculatePixelCount();
+
+		ParallelFor(pixelCount,
+			[
+				pBaseBuf, pBlendedBuf, pDestBuf, applyToAlpha
+			] (uint32 i)
+			{
+				if (applyToAlpha && NC > 3)
+				{
+					for (int c = 0; c < 3; ++c)
+					{
+						unsigned base = pBaseBuf[NC * i + c];
+						unsigned blended = pBlendedBuf[NC * i + c];
+						unsigned result = BLEND_FUNC(base, blended);
+						if (CLAMP)
+						{
+							pDestBuf[NC * i + c] = (uint8)FMath::Min(255u, result);
+						}
+						else
+						{
+							pDestBuf[NC * i + c] = (uint8)result;
+						}
+					}
+
+					pDestBuf[NC * i + 3] = pBaseBuf[NC * i + 3];
+				}
+				else
+				{
+					for (int c = 0; c < NC; ++c)
+					{
+						unsigned base = pBaseBuf[NC * i + c];
+						unsigned blended = pBlendedBuf[NC * i + c];
+						unsigned result = BLEND_FUNC(base, blended);
+						if (CLAMP)
+						{
+							pDestBuf[NC * i + c] = (uint8)FMath::Min(255u, result);
+						}
+						else
+						{
+							pDestBuf[NC * i + c] = (uint8)result;
+						}
+					}
+				}
+			});
+	}
+
+
+	//---------------------------------------------------------------------------------------------
+	//! Apply a blending function to an image with another image as blending layer
+	//---------------------------------------------------------------------------------------------
+	template< unsigned (*BLEND_FUNC)(unsigned,unsigned), bool CLAMP >
+    inline void ImageLayer( Image* pResult, const Image* pBase, const Image* pBlended,
+                            bool applyToAlpha )
+	{
+        check( pResult->GetFormat() == pBase->GetFormat() );
+        check( pResult->GetSizeX() == pBase->GetSizeX() );
+        check( pResult->GetSizeY() == pBase->GetSizeY() );
+        check( pResult->GetLODCount() == pBase->GetLODCount() );
+        check( pBase->GetSizeX() == pBlended->GetSizeX() );
+		check( pBase->GetSizeY() == pBlended->GetSizeY() );
+		check( pBase->GetFormat() == pBlended->GetFormat() );
+        check( pResult->GetLODCount() <= pBlended->GetLODCount() );
+
+		EImageFormat baseFormat = pBase->GetFormat();
+
+        uint8* pDestBuf = pResult->GetData();
+
+		if ( baseFormat==EImageFormat::IF_RGB_UBYTE )
+		{
+            BufferLayer< BLEND_FUNC, CLAMP, 3 > ( pDestBuf, pBase, pBlended, applyToAlpha );
+		}
+        else if ( baseFormat==EImageFormat::IF_RGBA_UBYTE || baseFormat==EImageFormat::IF_BGRA_UBYTE )
+		{
+            BufferLayer< BLEND_FUNC, CLAMP, 4 > ( pDestBuf, pBase, pBlended, applyToAlpha );
+		}
+		else if ( baseFormat==EImageFormat::IF_L_UBYTE )
+		{
+            BufferLayer< BLEND_FUNC, CLAMP, 1 > ( pDestBuf, pBase, pBlended, applyToAlpha );
+		}
+		else
+		{
+			checkf( false, TEXT("Unsupported format.") );
+		}
+	}
+
+
+	//---------------------------------------------------------------------------------------------
+	//! Apply a blending function to an image with another image as blending layer on the base image
+	//---------------------------------------------------------------------------------------------
+	template< unsigned (*BLEND_FUNC)(unsigned,unsigned), bool CLAMP >
+    inline void ImageLayerOnBase( Image* pBase, const Image* pBlended,
+                                  bool applyToAlpha )
+	{
+		check( pBase->GetSizeX() == pBlended->GetSizeX() );
+		check( pBase->GetSizeY() == pBlended->GetSizeY() );
+		check( pBase->GetFormat() == pBlended->GetFormat() );
+
+		EImageFormat baseFormat = pBase->GetFormat();
+        uint8* pDestBuf = pBase->GetData();
+
+		if ( baseFormat==EImageFormat::IF_RGB_UBYTE )
+		{
+            BufferLayer< BLEND_FUNC, CLAMP, 3 > ( pDestBuf, pBase, pBlended, applyToAlpha );
+		}
+        else if ( baseFormat==EImageFormat::IF_RGBA_UBYTE || baseFormat==EImageFormat::IF_BGRA_UBYTE )
+		{
+            BufferLayer< BLEND_FUNC, CLAMP, 4 > ( pDestBuf, pBase, pBlended, applyToAlpha );
+		}
+		else if ( baseFormat==EImageFormat::IF_L_UBYTE )
+		{
+            BufferLayer< BLEND_FUNC, CLAMP, 1 > ( pDestBuf, pBase, pBlended, applyToAlpha );
+		}
+		else
+		{
+			checkf( false, TEXT("Unsupported format.") );
+		}
+	}
+
+
+	//---------------------------------------------------------------------------------------------
+	template< unsigned (*BLEND_FUNC_MASKED)(unsigned,unsigned,unsigned),
+			  unsigned (*BLEND_FUNC)(unsigned,unsigned),
+			  bool CLAMP,
+			  int NC >
+    inline void BufferLayerFormat( uint8* pStartDestBuf,
+                                   const Image* pBase,
+                                   const Image* pMask,
+                                   const Image* pBlend,
+                                   bool applyToAlpha )
+	{
+        check(pBase->GetSizeX()==pMask->GetSizeX() && pBase->GetSizeY()==pMask->GetSizeY());
+        check(pBase->GetSizeX()==pBlend->GetSizeX() && pBase->GetSizeY()==pBlend->GetSizeY());
+        check(pBase->GetLODCount() <= pMask->GetLODCount() );
+        check(pBase->GetLODCount() <= pBlend->GetLODCount() );
+
+		uint8* pDestBuf = pStartDestBuf;
+        const uint8* pMaskBuf = pMask->GetData();
+        const uint8* pBaseBuf = pBase->GetData();
+        const uint8* pBlendedBuf = pBlend->GetData();
+
+		EImageFormat maskFormat = pMask->GetFormat();
+        bool isUncompressed = ( maskFormat == EImageFormat::IF_L_UBYTE );
+
+		// The base determines the number of lods to process.
+		int32 LODCount = pBase->GetLODCount();
+
+        if ( isUncompressed )
+		{
+			int32 pixelCount = pBase->CalculatePixelCount();
+
+            if (!applyToAlpha && NC>3)
+            {
+				ParallelFor(pixelCount,
+					[
+						pBaseBuf, pBlendedBuf, pMaskBuf, pDestBuf
+					] (uint32 i)
+					{
+						unsigned mask = pMaskBuf[i];
+						for (int c = 0; c < NC && c < 3; ++c)
+						{
+							unsigned base = pBaseBuf[NC * i + c];
+							unsigned blended = pBlendedBuf[NC * i + c];
+							unsigned result = BLEND_FUNC_MASKED(base, blended, mask);
+							if (CLAMP)
+							{
+								pDestBuf[NC * i + c] = (uint8)FMath::Min(255u, result);
+							}
+							else
+							{
+								pDestBuf[NC * i + c] = (uint8)result;
+							}
+						}
+
+						pDestBuf[NC * i + 3] = pBaseBuf[NC * i + 3];
+					});
+            }
+            else
+            {
+				ParallelFor(pixelCount,
+					[
+						pBaseBuf, pBlendedBuf, pMaskBuf, pDestBuf
+					] (uint32 i)
+					{
+						unsigned mask = pMaskBuf[i];
+						for (int c = 0; c < NC; ++c)
+						{
+							unsigned base = pBaseBuf[NC * i + c];
+							unsigned blended = pBlendedBuf[NC * i + c];
+							unsigned result = BLEND_FUNC_MASKED(base, blended, mask);
+							if (CLAMP)
+							{
+								pDestBuf[NC * i + c] = (uint8)FMath::Min(255u, result);
+							}
+							else
+							{
+								pDestBuf[NC * i + c] = (uint8)result;
+							}
+						}
+					});
+            }
+		}
+        else if ( maskFormat==EImageFormat::IF_L_UBYTE_RLE )
+		{
+			int rows = pBase->GetSizeY();
+			int width = pBase->GetSizeX();
+
+            for (int lod=0;lod<LODCount; ++lod)
+            {
+                pMaskBuf += 4+rows*sizeof(uint32);
+
+				// \todo: See how to handle mask buf here
+				//ParallelFor(rows,
+				//	[
+				//		pBaseBuf, pBlendedBuf, pMaskBuf, pDestBuf, width, applyToAlpha
+				//	] (uint32 r)
+				for (int r = 0; r < rows; ++r)
+				{						
+					//UE_LOG(LogMutableCore, Warning, "row: %d", r);
+
+					const uint8* pDestRowEnd = pDestBuf + width*NC;
+					while ( pDestBuf!=pDestRowEnd )
+					{
+						// Decode header
+						uint16 equal = *(const uint16*)pMaskBuf;
+						pMaskBuf += 2;
+
+						uint8 different = *pMaskBuf;
+						++pMaskBuf;
+
+						uint8 equalPixel = *pMaskBuf;
+						++pMaskBuf;
+
+						//UE_LOG(LogMutableCore, Warning, "    block: %d %d", (int)equal, (int)different );
+
+						// Equal pixels
+						check(pDestBuf + NC * equal <= pStartDestBuf + pBase->GetDataSize());
+						if ( equalPixel==255 )
+						{
+							for ( int i=0; i<equal; ++i )
+							{
+								if (!applyToAlpha && NC>3)
+								{
+									for ( int c=0; c<3; ++c )
+									{
+										unsigned base = pBaseBuf[NC*i+c];
+										unsigned blended = pBlendedBuf[NC*i+c];
+										unsigned result = BLEND_FUNC( base, blended );
+										if ( CLAMP )
+										{
+											pDestBuf[NC*i+c] = (uint8)FMath::Min( 255u, result );
+										}
+										else
+										{
+											pDestBuf[NC*i+c] = (uint8)result;
+										}
+									}
+
+									pDestBuf[NC*i+3] = (uint8)pBaseBuf[NC*i+3];
+								}
+								else
+								{
+									for ( int c=0; c<NC; ++c )
+									{
+										unsigned base = pBaseBuf[NC*i+c];
+										unsigned blended = pBlendedBuf[NC*i+c];
+										unsigned result = BLEND_FUNC( base, blended );
+										if ( CLAMP )
+										{
+											pDestBuf[NC*i+c] = (uint8)FMath::Min( 255u, result );
+										}
+										else
+										{
+											pDestBuf[NC*i+c] = (uint8)result;
+										}
+									}
+								}
+							}
+						}
+						else if ( equalPixel>0 )
+						{
+							for ( int i=0; i<equal; ++i )
+							{
+								if (!applyToAlpha && NC>3)
+								{
+									for ( int c=0; c<3; ++c )
+									{
+										unsigned base = pBaseBuf[NC*i+c];
+										unsigned blended = pBlendedBuf[NC*i+c];
+										unsigned result = BLEND_FUNC_MASKED( base, blended, equalPixel );
+										if ( CLAMP )
+										{
+											pDestBuf[NC*i+c] = (uint8)FMath::Min( 255u, result );
+										}
+										else
+										{
+											pDestBuf[NC*i+c] = (uint8)result;
+										}
+									}
+
+									pDestBuf[NC*i+3] = (uint8)pBaseBuf[NC*i+3];
+								}
+								else
+								{
+									for ( int c=0; c<NC; ++c )
+									{
+										unsigned base = pBaseBuf[NC*i+c];
+										unsigned blended = pBlendedBuf[NC*i+c];
+										unsigned result = BLEND_FUNC_MASKED( base, blended, equalPixel );
+										if ( CLAMP )
+										{
+											pDestBuf[NC*i+c] = (uint8)FMath::Min( 255u, result );
+										}
+										else
+										{
+											pDestBuf[NC*i+c] = (uint8)result;
+										}
+									}
+								}
+							}
+						}
+						else
+						{
+							// It could happen if xxxxxOnBase
+							if (pDestBuf!=pBaseBuf)
+							{
+								FMemory::Memmove( pDestBuf, pBaseBuf, NC*equal );
+							}
+						}
+						pDestBuf += NC*equal;
+						pBaseBuf += NC*equal;
+						pBlendedBuf += NC*equal;
+
+						// Different pixels
+						check(pDestBuf + NC * different <= pStartDestBuf + pBase->GetDataSize());
+						for ( int i=0; i<different; ++i )
+						{
+							if (!applyToAlpha && NC>3)
+							{
+								for ( int c=0; c<3; ++c )
+								{
+									unsigned mask = pMaskBuf[i];
+									unsigned base = pBaseBuf[NC*i+c];
+									unsigned blended = pBlendedBuf[NC*i+c];
+									unsigned result = BLEND_FUNC_MASKED( base, blended, mask );
+									if ( CLAMP )
+									{
+										pDestBuf[NC*i+c] = (uint8)std::min( 255u, result );
+									}
+									else
+									{
+										pDestBuf[NC*i+c] = (uint8)result;
+									}
+								}
+
+								pDestBuf[NC*i+3] = (uint8)pBaseBuf[NC*i+3];
+							}
+							else
+							{
+								for ( int c=0; c<NC; ++c )
+								{
+									unsigned mask = pMaskBuf[i];
+									unsigned base = pBaseBuf[NC*i+c];
+									unsigned blended = pBlendedBuf[NC*i+c];
+									unsigned result = BLEND_FUNC_MASKED( base, blended, mask );
+									if ( CLAMP )
+									{
+										pDestBuf[NC*i+c] = (uint8)std::min( 255u, result );
+									}
+									else
+									{
+										pDestBuf[NC*i+c] = (uint8)result;
+									}
+								}
+							}
+						}
+
+						pDestBuf += NC*different;
+						pBaseBuf += NC*different;
+						pBlendedBuf += NC*different;
+						pMaskBuf += different;
+					}
+				}
+
+                rows = FMath::DivideAndRoundUp(rows,2);
+                width = FMath::DivideAndRoundUp(width,2);
+            }
+		}
+		else
+		{
+			checkf( false, TEXT("Unsupported mask format.") );
+		}
+	}
+
+
+	//---------------------------------------------------------------------------------------------
+	template< unsigned (*BLEND_FUNC)(unsigned,unsigned),
+			  bool CLAMP,
+			  int NC >
+    inline void BufferLayerFormatStrideNoAlpha( uint8* pDestBuf,
+		int stride,
+		const Image* pMask,
+		const Image* pBlend, int32 LODCount )
+	{
+        const uint8* pMaskBuf = pMask->GetData();
+        const uint8* pBlendedBuf = pBlend->GetData();
+
+		EImageFormat maskFormat = pMask->GetFormat();
+        bool isUncompressed = ( maskFormat == EImageFormat::IF_L_UBYTE );
+
+        if ( isUncompressed )
+		{
+			int rowCount = pBlend->GetSizeY();
+			int pixelCount = pBlend->GetSizeX();
+			for ( int j=0; j<rowCount; ++j )
+			{
+				for ( int i=0; i<pixelCount; ++i )
+				{
+					unsigned mask = *pMaskBuf;
+					if (mask)
+					{
+						for ( int c=0; c<NC; ++c )
+						{
+							unsigned base = *pDestBuf;
+							unsigned blended = *pBlendedBuf;
+							unsigned result = BLEND_FUNC( base, blended );
+							if ( CLAMP )
+							{
+								*pDestBuf = (uint8)FMath::Min( 255u, result );
+							}
+							else
+							{
+								*pDestBuf = (uint8)result;
+							}
+							++pDestBuf;
+							++pBlendedBuf;
+						}
+					}
+					else
+					{
+						pDestBuf+=NC;
+						pBlendedBuf+=NC;
+					}
+					++pMaskBuf;
+				}
+
+				pDestBuf+=stride;
+			}
+		}
+        else if ( maskFormat==EImageFormat::IF_L_UBIT_RLE )
+		{
+			int rows = pMask->GetSizeY();
+			int width = pMask->GetSizeX();
+
+            for (int lod=0;lod< LODCount; ++lod)
+            {
+                pMaskBuf += 4+rows*sizeof(uint32);
+
+                for ( int r=0; r<rows; ++r )
+                {
+                    const uint8* pDestRowEnd = pDestBuf + width*NC;
+                    while ( pDestBuf!=pDestRowEnd )
+                    {
+                        // Decode header
+                        uint16 zeros = *(const uint16*)pMaskBuf;
+                        pMaskBuf += 2;
+
+                        uint16 ones = *(const uint16*)pMaskBuf;
+                        pMaskBuf += 2;
+
+                        // Skip
+                        pDestBuf += zeros*NC;
+                        pBlendedBuf += zeros*NC;
+
+                        // Copy
+                        FMemory::Memmove( pDestBuf, pBlendedBuf, ones*NC );
+
+                        pDestBuf += NC*ones;
+                        pBlendedBuf += NC*ones;
+                    }
+
+                    pDestBuf += stride;
+                }
+
+                rows = FMath::DivideAndRoundUp(rows,2);
+                width = FMath::DivideAndRoundUp(width,2);
+            }
+		}
+		else
+		{
+			checkf( false, TEXT("Unsupported mask format.") );
+		}
+	}
+
+
+	//---------------------------------------------------------------------------------------------
+	template< unsigned (*BLEND_FUNC_MASKED)(unsigned,unsigned,unsigned),
+			  unsigned (*BLEND_FUNC)(unsigned,unsigned),
+			  bool CLAMP >
+    inline void BufferLayer( uint8* pDestBuf,
+							 const Image* pBase,
+							 const Image* pMask,
+                             const Image* pBlend,
+                             bool applyToAlpha )
+	{
+		if ( pBase->GetFormat()==EImageFormat::IF_RGB_UBYTE )
+		{
+            BufferLayerFormat< BLEND_FUNC_MASKED, BLEND_FUNC, CLAMP, 3 >
+                    ( pDestBuf, pBase, pMask, pBlend, applyToAlpha ) ;
+		}
+        else if ( pBase->GetFormat()==EImageFormat::IF_RGBA_UBYTE || pBase->GetFormat()==EImageFormat::IF_BGRA_UBYTE )
+		{
+            BufferLayerFormat< BLEND_FUNC_MASKED, BLEND_FUNC, CLAMP, 4 >
+                    ( pDestBuf, pBase, pMask, pBlend, applyToAlpha ) ;
+		}
+		else if ( pBase->GetFormat()==EImageFormat::IF_L_UBYTE )
+		{
+            BufferLayerFormat< BLEND_FUNC_MASKED, BLEND_FUNC, CLAMP, 1 >
+                    ( pDestBuf, pBase, pMask, pBlend, applyToAlpha ) ;
+		}
+        else
+		{
+			checkf( false, TEXT("Unsupported format.") );
+		}
+	}
+
+
+	//---------------------------------------------------------------------------------------------
+	template< unsigned (*BLEND_FUNC)(unsigned,unsigned),
+			  bool CLAMP >
+    inline void BufferLayerStrideNoAlpha( uint8* pDestBuf, int stride, const Image* pMask, const Image* pBlend, int32 LODCount )
+	{
+		if ( pBlend->GetFormat()==EImageFormat::IF_RGB_UBYTE )
+		{
+			BufferLayerFormatStrideNoAlpha< BLEND_FUNC, CLAMP, 3 >
+					( pDestBuf, stride, pMask, pBlend, LODCount ) ;
+		}
+        else if ( pBlend->GetFormat()==EImageFormat::IF_RGBA_UBYTE || pBlend->GetFormat()==EImageFormat::IF_BGRA_UBYTE )
+		{
+			BufferLayerFormatStrideNoAlpha< BLEND_FUNC, CLAMP, 4 >
+					( pDestBuf, stride, pMask, pBlend, LODCount) ;
+		}
+		else if ( pBlend->GetFormat()==EImageFormat::IF_L_UBYTE )
+		{
+			BufferLayerFormatStrideNoAlpha< BLEND_FUNC, CLAMP, 1 >
+					( pDestBuf, stride, pMask, pBlend, LODCount) ;
+		}
+		else
+		{
+			checkf( false, TEXT("Unsupported format.") );
+		}
+	}
+
+
+	//---------------------------------------------------------------------------------------------
+	//! Apply a blending function to an image with another image as blending layer with a mask
+	//---------------------------------------------------------------------------------------------
+	template< unsigned (*BLEND_FUNC_MASKED)(unsigned,unsigned,unsigned),
+			  unsigned (*BLEND_FUNC)(unsigned,unsigned),
+			  bool CLAMP >
+    inline void ImageLayer( Image* pResult,
+                            const Image* pBase,
+                            const Image* pMask,
+                            const Image* pBlended,
+                            bool applyToAlpha )
+	{
+        check( pResult->GetFormat() == pBase->GetFormat() );
+        check( pResult->GetSizeX() == pBase->GetSizeX() );
+        check( pResult->GetSizeY() == pBase->GetSizeY() );
+        check( pResult->GetLODCount() == pBase->GetLODCount() );
+		check( pBase->GetSizeX() == pMask->GetSizeX() );
+		check( pBase->GetSizeY() == pMask->GetSizeY() );
+		check( pBase->GetSizeX() == pBlended->GetSizeX() );
+		check( pBase->GetSizeY() == pBlended->GetSizeY() );
+		check( pBase->GetFormat() == pBlended->GetFormat() );
+		check( pMask->GetFormat() == EImageFormat::IF_L_UBYTE
+						||
+						pMask->GetFormat() == EImageFormat::IF_L_UBYTE_RLE );
+        check( pBase->GetLODCount() <= pMask->GetLODCount() );
+        check( pBase->GetLODCount() <= pBlended->GetLODCount() );
+
+        BufferLayer<BLEND_FUNC_MASKED,BLEND_FUNC,CLAMP>( pResult->GetData(), pBase, pMask,
+                                                         pBlended, applyToAlpha );
+	}
+
+
+	//---------------------------------------------------------------------------------------------
+	//! Apply a blending function to an image with another image as blending layer, on the base
+	//! image itself
+	//---------------------------------------------------------------------------------------------
+	template< unsigned (*BLEND_FUNC_MASKED)(unsigned,unsigned,unsigned),
+			  unsigned (*BLEND_FUNC)(unsigned,unsigned),
+			  bool CLAMP >
+    inline void ImageLayerOnBase( Image* pBase, const Image* pMask, const Image* pBlended,
+                                  bool applyToAlpha )
+	{
+		check( pBase->GetSizeX() == pMask->GetSizeX() );
+		check( pBase->GetSizeY() == pMask->GetSizeY() );
+		check( pBase->GetSizeX() == pBlended->GetSizeX() );
+		check( pBase->GetSizeY() == pBlended->GetSizeY() );
+		check( pBase->GetFormat() == pBlended->GetFormat() );
+		check( pMask->GetFormat() == EImageFormat::IF_L_UBYTE
+						||
+						pMask->GetFormat() == EImageFormat::IF_L_UBYTE_RLE );
+        check( pBase->GetLODCount() <= pMask->GetLODCount() );
+        check( pBase->GetLODCount() <= pBlended->GetLODCount() );
+
+        BufferLayer<BLEND_FUNC_MASKED,BLEND_FUNC,CLAMP>( pBase->GetData(), pBase, pMask,
+                                                         pBlended, applyToAlpha );
+	}
+
+
+	//---------------------------------------------------------------------------------------------
+	//! Apply a blending function to an image with another image as blending layer, on a subrect of
+	//! the base image.
+	//! \warning this method applies the blending function to the alpha channel too
+	//! \warning this method uses the mask as a binary mask (>0)
+	//---------------------------------------------------------------------------------------------
+	template< unsigned (*BLEND_FUNC)(unsigned,unsigned),
+			  bool CLAMP >
+	inline void ImageLayerOnBaseNoAlpha( Image* pBase,
+										 const Image* pMask,
+										 const Image* pBlended,
+										 const box< vec2<int> >& rect )
+	{
+		check( pBase->GetSizeX() >= rect.min[0]+rect.size[0] );
+		check( pBase->GetSizeY() >= rect.min[1]+rect.size[1] );
+		check( pMask->GetSizeX() == pBlended->GetSizeX() );
+		check( pMask->GetSizeY() == pBlended->GetSizeY() );
+		check( pBase->GetFormat() == pBlended->GetFormat() );
+		check( pMask->GetFormat() == EImageFormat::IF_L_UBYTE
+						||
+						pMask->GetFormat() == EImageFormat::IF_L_UBYTE_RLE
+						||
+						pMask->GetFormat() == EImageFormat::IF_L_UBIT_RLE );
+        check( pBase->GetLODCount() <= pMask->GetLODCount() );
+        check( pBase->GetLODCount() <= pBlended->GetLODCount() );
+
+		int pixelSize = GetImageFormatData( pBase->GetFormat() ).m_bytesPerBlock;
+
+		int start = ( pBase->GetSizeX()*rect.min[1]+rect.min[0] ) * pixelSize;
+		int stride = ( pBase->GetSizeX() - rect.size[0] ) * pixelSize;
+		BufferLayerStrideNoAlpha<BLEND_FUNC,CLAMP>( pBase->GetData()+start, stride, pMask, pBlended, pBase->GetLODCount());
+	}
+
+
+	//---------------------------------------------------------------------------------------------
+	template< int NC >
+    inline void BufferTableColourFormat( uint8* pStartDestBuf,
+										 const Image* pBase,
+										 const Image* pMask,
+                                         const uint8* pTable0,
+                                         const uint8* pTable1,
+                                         const uint8* pTable2 )
+	{
+		uint8* pDestBuf = pStartDestBuf;
+        const uint8* pMaskBuf = pMask->GetData();
+        const uint8* pBaseBuf = pBase->GetData();
+		EImageFormat maskFormat = pMask->GetFormat();
+
+        bool isUncompressed = ( maskFormat == EImageFormat::IF_L_UBYTE );
+
+		// The base determines the lods to process.
+		int32 LODCount = pBase->GetLODCount();
+
+        if ( isUncompressed )
+		{
+			int pixelCount = (int)pBase->CalculatePixelCount();
+			for ( int i=0; i<pixelCount; ++i )
+			{
+				unsigned mask = pMaskBuf[i];
+
+				if (mask)
+				{
+					if (mask<255)
+					{
+                        pDestBuf[NC*i+0] = (uint8)( ( ( ( 255 - mask ) * pBaseBuf[NC*i+0] )
+                                             + ( mask * pTable0[ pBaseBuf[NC*i+0] ] ) ) >> 8 );
+                        pDestBuf[NC*i+1] = (uint8)( ( ( ( 255 - mask ) * pBaseBuf[NC*i+1] )
+                                             + ( mask * pTable1[ pBaseBuf[NC*i+1] ] ) ) >> 8 );
+                        pDestBuf[NC*i+2] = (uint8)( ( ( ( 255 - mask ) * pBaseBuf[NC*i+2] )
+                                             + ( mask * pTable2[ pBaseBuf[NC*i+2] ] ) ) >> 8 );
+					}
+					else
+					{
+						pDestBuf[NC*i+0] = pTable0[ pBaseBuf[NC*i+0] ];
+						pDestBuf[NC*i+1] = pTable1[ pBaseBuf[NC*i+1] ];
+						pDestBuf[NC*i+2] = pTable2[ pBaseBuf[NC*i+2] ];
+					}
+				}
+				else
+				{
+					pDestBuf[NC*i+0] = pBaseBuf[NC*i+0];
+					pDestBuf[NC*i+1] = pBaseBuf[NC*i+1];
+					pDestBuf[NC*i+2] = pBaseBuf[NC*i+2];
+				}
+
+                constexpr bool isNC4 = (NC==4);
+                if ( isNC4 )
+				{
+					pDestBuf[NC*i+3] = pBaseBuf[NC*i+3];
+				}
+			}
+		}
+		else if ( maskFormat==EImageFormat::IF_L_UBYTE_RLE )
+		{
+			int rows = pBase->GetSizeY();
+			int width = pBase->GetSizeX();
+
+            for (int32 lod=0;lod< LODCount; ++lod)
+            {
+                pMaskBuf += 4+rows*sizeof(uint32);
+
+                for ( int r=0; r<rows; ++r )
+                {
+                    const uint8* pDestRowEnd = pDestBuf + width*NC;
+                    while ( pDestBuf!=pDestRowEnd )
+                    {
+                        // Decode header
+                        uint16 equal = *(const uint16*)pMaskBuf;
+                        pMaskBuf += 2;
+
+                        uint8 different = *pMaskBuf;
+                        ++pMaskBuf;
+
+                        unsigned equalPixel = *pMaskBuf;
+                        ++pMaskBuf;
+
+                        // Equal pixels
+						check(pDestBuf + NC * equal <= pStartDestBuf + pBase->GetDataSize());
+                        if ( equalPixel==255 )
+                        {
+                            for ( int i=0; i<equal; ++i )
+                            {
+                                pDestBuf[NC*i+0] = pTable0[ pBaseBuf[NC*i+0] ];
+                                pDestBuf[NC*i+1] = pTable1[ pBaseBuf[NC*i+1] ];
+                                pDestBuf[NC*i+2] = pTable2[ pBaseBuf[NC*i+2] ];
+
+                                constexpr bool isNC4 = (NC==4);
+                                if ( isNC4 )
+                                {
+                                    pDestBuf[NC*i+3] = pBaseBuf[NC*i+3];
+                                }
+                            }
+                        }
+                        else if ( equalPixel>0 )
+                        {
+                            for ( int i=0; i<equal; ++i )
+                            {
+                                pDestBuf[NC*i+0] = (uint8)( ( ( ( 255 - equalPixel ) * pBaseBuf[NC*i+0] )
+                                                     + ( equalPixel * pTable0[ pBaseBuf[NC*i+0] ] ) ) >> 8 );
+                                pDestBuf[NC*i+1] = (uint8)( ( ( ( 255 - equalPixel ) * pBaseBuf[NC*i+1] )
+                                                     + ( equalPixel * pTable1[ pBaseBuf[NC*i+1] ] ) ) >> 8 );
+                                pDestBuf[NC*i+2] = (uint8)( ( ( ( 255 - equalPixel ) * pBaseBuf[NC*i+2] )
+                                                     + ( equalPixel * pTable2[ pBaseBuf[NC*i+2] ] ) ) >> 8 );
+
+                                constexpr bool isNC4 = (NC==4);
+                                if ( isNC4 )
+                                {
+                                    pDestBuf[NC*i+3] = pBaseBuf[NC*i+3];
+                                }
+                            }
+                        }
+                        else
+                        {
+                            // It could happen if xxxxxOnBase
+                            if (pDestBuf!=pBaseBuf)
+                            {
+								FMemory::Memmove( pDestBuf, pBaseBuf, NC*equal );
+                            }
+                        }
+                        pDestBuf += NC*equal;
+                        pBaseBuf += NC*equal;
+
+                        // Different pixels
+						check(pDestBuf + NC * different <= pStartDestBuf + pBase->GetDataSize());
+                        for ( int i=0; i<different; ++i )
+                        {
+                            unsigned mask = pMaskBuf[i];
+                            pDestBuf[NC*i+0] = (uint8)( ( ( ( 255 - mask ) * pBaseBuf[NC*i+0] )
+                                                 + ( mask * pTable0[ pBaseBuf[NC*i+0] ] ) ) >> 8 );
+                            pDestBuf[NC*i+1] = (uint8)( ( ( ( 255 - mask ) * pBaseBuf[NC*i+1] )
+                                                 + ( mask * pTable1[ pBaseBuf[NC*i+1] ] ) ) >> 8 );
+                            pDestBuf[NC*i+2] = (uint8)( ( ( ( 255 - mask ) * pBaseBuf[NC*i+2] )
+                                                 + ( mask * pTable2[ pBaseBuf[NC*i+2] ] ) ) >> 8 );
+
+                            constexpr bool isNC4 = (NC==4);
+                            if ( isNC4 )
+                            {
+                                pDestBuf[NC*i+3] = pBaseBuf[NC*i+3];
+                            }
+                        }
+
+                        pDestBuf += NC*different;
+                        pBaseBuf += NC*different;
+                        pMaskBuf += different;
+                    }
+                }
+
+                rows = FMath::DivideAndRoundUp(rows,2);
+                width = FMath::DivideAndRoundUp(width,2);
+            }
+		}
+		else
+		{
+			checkf( false, TEXT("Unsupported mask format.") );
+		}
+	}
+
+
+	//---------------------------------------------------------------------------------------------
+    inline void ImageTable( Image* pResult,
+                                const Image* pBase,
+								const Image* pMask,
+                                const uint8* pTable0,
+                                const uint8* pTable1,
+                                const uint8* pTable2
+								)
+	{
+        check( pResult->GetFormat() == pBase->GetFormat() );
+        check( pResult->GetSizeX() == pBase->GetSizeX() );
+        check( pResult->GetSizeY() == pBase->GetSizeY() );
+        check( pResult->GetLODCount() == pBase->GetLODCount() );
+        check( pBase->GetSizeX() == pMask->GetSizeX() );
+		check( pBase->GetSizeY() == pMask->GetSizeY() );
+		check( pMask->GetFormat() == EImageFormat::IF_L_UBYTE
+						||
+						pMask->GetFormat() == EImageFormat::IF_L_UBYTE_RLE );
+        check( pBase->GetLODCount() <= pMask->GetLODCount() );
+
+		EImageFormat baseFormat = pBase->GetFormat();
+		if ( baseFormat==EImageFormat::IF_RGB_UBYTE )
+		{
+            BufferTableColourFormat<3>( pResult->GetData(), pBase, pMask, pTable0, pTable1, pTable2 );
+		}
+        else if ( baseFormat==EImageFormat::IF_RGBA_UBYTE )
+        {
+            BufferTableColourFormat<4>( pResult->GetData(), pBase, pMask, pTable0, pTable1, pTable2 );
+        }
+        else if ( baseFormat==EImageFormat::IF_BGRA_UBYTE )
+        {
+            BufferTableColourFormat<4>( pResult->GetData(), pBase, pMask, pTable2, pTable1, pTable0 );
+        }
+        else
+		{
+			checkf( false, TEXT("Unsupported format.") );
+		}
+	}
+
+
+	template<size_t NC>
+	inline uint32 PackPixel(const uint8* pPixel)
+	{
+		static_assert(NC > 0 && NC <= 4);
+
+		uint32 PixelPack = 0;
+
+		// The compiler should be able to optimize this given that NC is a constant expression
+		FMemory::Memcpy( &PixelPack, pPixel, NC );
+
+		return PixelPack;
+	}
+
+	template<size_t NC>
+	inline void UnpackPixel(uint8* const pPixel, uint32 pixelData )
+	{
+		static_assert(NC > 0 && NC <= 4);
+
+		// The compiler should be able to optimize this given that NC is a constant expression
+		FMemory::Memcpy( pPixel, &pixelData, NC );
+	}
+
+	template< 
+			uint32 (*BLEND_FUNC)(uint32, uint32), 
+			size_t NC >
+	inline void BufferLayerCombineColour(uint8* pDestBuf, const Image* pBase, vec3<float> col)
+	{
+		//static_assert(NC > 0 && NC <= 4);
+
+		const uint8* pBaseBuf = pBase->GetData();
+
+		uint32 top = 
+			(uint32)(255.0f * col[0]) << 0 |
+			(uint32)(255.0f * col[1]) << 8 |
+			(uint32)(255.0f * col[2]) << 16;
+
+		int32 pixelCount = pBase->CalculatePixelCount();
+
+		for (int i = 0; i < pixelCount; ++i)
+		{
+			uint32 base = PackPixel<NC>(&pBaseBuf[NC * i]);
+
+			uint32 result = BLEND_FUNC(base, top);
+
+			UnpackPixel<NC>(&pDestBuf[NC * i], result);
+		}
+	}
+
+	template< 
+			uint32 (*BLEND_FUNC)(uint32, uint32), 
+			size_t NC >
+	inline void BufferLayerCombine(uint8* pDestBuf, const Image* pBase, const Image* pBlended)
+	{
+		//static_assert(NC > 0 && NC <= 4);
+
+		check(pBase->GetSizeX() == pBlended->GetSizeX());
+		check(pBase->GetSizeY() == pBlended->GetSizeY());
+		check(pBase->GetFormat() == pBlended->GetFormat());
+
+		const uint8* pBaseBuf = pBase->GetData();
+		const uint8* pBlendedBuf = pBlended->GetData();
+
+		int32 pixelCount = pBase->CalculatePixelCount();
+
+		for (int i = 0; i < pixelCount; ++i)
+		{
+			const uint32 base = PackPixel<NC>(&pBaseBuf[NC * i]);
+			const uint32 blend = PackPixel<NC>(&pBlendedBuf[NC * i]);
+
+			const uint32 result = BLEND_FUNC(base, blend);
+
+			UnpackPixel<NC>(&pDestBuf[NC * i], result);
+		}
+	}
+
+	template< 
+			uint32 (*BLEND_FUNC_MASKED)(uint32, uint32, uint32), 
+			size_t NC >
+	inline void BufferLayerCombine(uint8* pDestBuf, const Image* pBase, const Image* pMask, const Image* pBlended)
+	{
+		//static_assert(NC > 0 && NC <= 4);
+
+		check(pBase->GetSizeX() == pBlended->GetSizeX());
+		check(pBase->GetSizeY() == pBlended->GetSizeY());
+		check(pBase->GetSizeX() == pMask->GetSizeX());
+		check(pBase->GetSizeY() == pMask->GetSizeY());
+		check(pBase->GetFormat() == pBlended->GetFormat());
+
+		const uint8* pBaseBuf = pBase->GetData();
+		const uint8* pMaskBuf = pMask->GetData();
+		const uint8* pBlendedBuf = pBlended->GetData();
+
+		int32 pixelCount = pBase->CalculatePixelCount();
+
+		for (int i = 0; i < pixelCount; ++i)
+		{
+			const uint32 base = PackPixel<NC>(&pBaseBuf[NC * i]);
+			const uint32 blend = PackPixel<NC>(&pBlendedBuf[NC * i]);
+			const uint32 mask = PackPixel<1>(&pMaskBuf[i]);
+
+			const uint32 result = BLEND_FUNC_MASKED(base, blend, mask);
+
+			UnpackPixel<NC>(&pDestBuf[NC * i], result);
+		}
+	}
+
+	template< 
+			uint32 (*BLEND_FUNC_MASKED)(uint32, uint32, uint32), 
+			size_t NC >
+	inline void BufferLayerCombineColour(uint8* pDestBuf, const Image* pBase, const Image* pMask, vec3<float> col)
+	{
+		//static_assert(NC > 0 && NC <= 4);
+
+		check(pBase->GetSizeX() == pMask->GetSizeX());
+		check(pBase->GetSizeY() == pMask->GetSizeY());
+
+
+		const uint8* pBaseBuf = pBase->GetData();
+		const uint8* pMaskBuf = pMask->GetData();
+
+		uint32 top = 
+			(uint32)(255.0f * col[0]) << 0 |
+			(uint32)(255.0f * col[1]) << 8 |
+			(uint32)(255.0f * col[2]) << 16;
+
+
+		int32 pixelCount = pBase->CalculatePixelCount();
+
+		for (int i = 0; i < pixelCount; ++i)
+		{
+			const uint32 base = PackPixel<NC>(&pBaseBuf[NC * i]);
+			const uint32 mask = PackPixel<1>(&pMaskBuf[i]);
+
+			const uint32 result = BLEND_FUNC_MASKED(base, top, mask);
+
+			UnpackPixel<NC>(&pDestBuf[NC * i], result);
+		}
+	}
+
+	template< uint32 (*BLEND_FUNC)(uint32, uint32) >
+	inline void ImageLayerCombine(Image* pResult, const Image* pBase, const Image* pBlended)
+	{
+		check(pResult->GetFormat() == pBase->GetFormat());
+		check(pResult->GetSizeX() == pBase->GetSizeX());
+		check(pResult->GetSizeY() == pBase->GetSizeY());
+		check(pResult->GetLODCount() == pBase->GetLODCount());
+		check(pBase->GetSizeX() == pBlended->GetSizeX());
+		check(pBase->GetSizeY() == pBlended->GetSizeY());
+		check(pBase->GetFormat() == pBlended->GetFormat());
+		check(pResult->GetLODCount() <= pBlended->GetLODCount());
+
+		const EImageFormat baseFormat = pBase->GetFormat();
+
+		uint8* pDestBuf = pResult->GetData();
+		if (baseFormat == EImageFormat::IF_L_UBYTE)
+		{
+			BufferLayerCombine<BLEND_FUNC, 1 >(pDestBuf, pBase, pBlended);
+		}
+		else if (baseFormat == EImageFormat::IF_RGB_UBYTE)
+		{
+			BufferLayerCombine< BLEND_FUNC, 3 >(pDestBuf, pBase, pBlended);
+		}
+		else if (baseFormat == EImageFormat::IF_RGBA_UBYTE || baseFormat == EImageFormat::IF_BGRA_UBYTE)
+		{
+			// \todo: pass swizzle template argument if BGRA_UBYTE, not yet supported.
+			BufferLayerCombine< BLEND_FUNC, 4 >(pDestBuf, pBase, pBlended);
+		}
+		else
+		{
+			checkf(false, TEXT("Unsupported format."));
+		}
+	}
+
+	template< 
+		uint32 (*BLEND_FUNC)(uint32, uint32), 
+		uint32 (*BLEND_FUNC_MASKED)(uint32, uint32, uint32) >
+	inline void ImageLayerCombine(Image* pResult, const Image* pBase, const Image* pMask, const Image* pBlended)
+	{
+		check(pResult->GetFormat() == pBase->GetFormat());
+		check(pResult->GetSizeX() == pBase->GetSizeX());
+		check(pResult->GetSizeY() == pBase->GetSizeY());
+		check(pResult->GetLODCount() == pBase->GetLODCount());
+		check(pBase->GetSizeX() == pBlended->GetSizeX());
+		check(pBase->GetSizeY() == pBlended->GetSizeY());
+		check(pBase->GetFormat() == pBlended->GetFormat());
+		check(pResult->GetLODCount() <= pBlended->GetLODCount());
+
+		const EImageFormat baseFormat = pBase->GetFormat();
+
+		uint8* pDestBuf = pResult->GetData();
+
+		if (pMask->GetFormat() != EImageFormat::IF_L_UBYTE)
+		{
+			checkf(false, TEXT("Unsupported mask format."));
+
+			BufferLayerCombine< BLEND_FUNC, 1 >(pDestBuf, pBase, pBlended);
+		}
+
+		if (baseFormat == EImageFormat::IF_L_UBYTE)
+		{
+			BufferLayerCombine< BLEND_FUNC_MASKED, 1 >(pDestBuf, pBase, pMask, pBlended);
+		}
+		else if (baseFormat == EImageFormat::IF_RGB_UBYTE)
+		{
+			BufferLayerCombine< BLEND_FUNC_MASKED, 3 >(pDestBuf, pBase, pMask, pBlended);
+		}
+		else if (baseFormat == EImageFormat::IF_RGBA_UBYTE || baseFormat == EImageFormat::IF_BGRA_UBYTE)
+		{
+			// \todo: pass swizzle template argument if BGRA_UBYTE, not yet supported.
+			BufferLayerCombine< BLEND_FUNC_MASKED, 4 >(pDestBuf, pBase, pMask, pBlended);
+		}
+		else
+		{
+			checkf(false, TEXT("Unsupported format."));
+		}
+	}
+
+	template< uint32 (*BLEND_FUNC)(uint32, uint32)>
+	inline void ImageLayerCombineColour(Image* pResult, const Image* pBase, vec3<float> col)
+	{
+		check(pResult->GetFormat() == pBase->GetFormat());
+		check(pResult->GetSizeX() == pBase->GetSizeX());
+		check(pResult->GetSizeY() == pBase->GetSizeY());
+		check(pResult->GetLODCount() == pBase->GetLODCount());
+
+		const EImageFormat baseFormat = pBase->GetFormat();
+
+		uint8* pDestBuf = pResult->GetData();
+
+		if (baseFormat == EImageFormat::IF_L_UBYTE)
+		{
+			BufferLayerCombineColour< BLEND_FUNC, 1 >(pDestBuf, pBase, col);
+		}
+		else if (baseFormat == EImageFormat::IF_RGB_UBYTE)
+		{
+			BufferLayerCombineColour< BLEND_FUNC, 3 >(pDestBuf, pBase, col);
+		}
+		else if (baseFormat == EImageFormat::IF_RGBA_UBYTE || baseFormat == EImageFormat::IF_BGRA_UBYTE)
+		{
+			// \todo: pass swizzle template argument if BGRA_UBYTE, not yet supported.
+ 			BufferLayerCombineColour< BLEND_FUNC, 4 >(pDestBuf, pBase, col);
+		}
+		else
+		{
+			checkf(false, TEXT("Unsupported format."));
+		}
+	}
+
+	template< 
+		uint32 (*BLEND_FUNC)(uint32, uint32),
+		uint32 (*BLEND_FUNC_MASKED)(uint32, uint32, uint32) >
+	inline void ImageLayerCombineColour(Image* pResult, const Image* pBase, const Image* pMask, vec3<float> col)
+	{
+		check(pResult->GetFormat() == pBase->GetFormat());
+		check(pResult->GetSizeX() == pBase->GetSizeX());
+		check(pResult->GetSizeY() == pBase->GetSizeY());
+		check(pResult->GetLODCount() == pBase->GetLODCount());
+
+		const EImageFormat baseFormat = pBase->GetFormat();
+
+		uint8* pDestBuf = pResult->GetData();
+
+		if (pMask->GetFormat() != EImageFormat::IF_L_UBYTE)
+		{
+			checkf(false, TEXT("Unsupported mask format."));
+
+			BufferLayerCombineColour< BLEND_FUNC, 1 >(pDestBuf, pBase, col);
+		}
+
+		if (baseFormat == EImageFormat::IF_L_UBYTE)
+		{
+			BufferLayerCombineColour< BLEND_FUNC_MASKED, 1 >(pDestBuf, pBase, pMask, col);
+		}
+		else if (baseFormat == EImageFormat::IF_RGB_UBYTE)
+		{
+			BufferLayerCombineColour< BLEND_FUNC_MASKED, 3 >(pDestBuf, pBase, pMask, col);
+		}
+		else if (baseFormat == EImageFormat::IF_RGBA_UBYTE || baseFormat == EImageFormat::IF_BGRA_UBYTE)
+		{
+			// \todo: pass swizzle template argument if BGRA_UBYTE, not yet supported.
+ 			BufferLayerCombineColour< BLEND_FUNC_MASKED, 4 >(pDestBuf, pBase, pMask, col);
+		}
+		else
+		{
+			checkf(false, TEXT("Unsupported format."));
+		}
+	}
+
+	template< size_t NCBase, size_t NCBlend, class ImageCombineFn >
+	inline void BufferLayerCombineFunctor(uint8* pDestBuf, const Image* pBase, const Image* pBlended, ImageCombineFn&& ImageCombine)
+	{
+		//static_assert(NCBase > 0 && NCBase <= 4);
+		//static_assert(NCBlend > 0 && NCBlend <= 4);
+
+		check(pBase->GetSizeX() == pBlended->GetSizeX());
+		check(pBase->GetSizeY() == pBlended->GetSizeY());
+		check(pBase->GetLODCount() <= pBlended->GetLODCount());
+
+		const uint8* pBaseBuf = pBase->GetData();
+		const uint8* pBlendedBuf = pBlended->GetData();
+
+		int32 pixelCount = pBase->CalculatePixelCount();
+
+		for (int i = 0; i < pixelCount; ++i)
+		{
+			const uint32 base = PackPixel<NCBase>(&pBaseBuf[NCBase * i]);
+			const uint32 blend = PackPixel<NCBlend>(&pBlendedBuf[NCBlend * i]);
+
+			const uint32 result = ImageCombine(base, blend);
+
+			UnpackPixel<NCBase>(&pDestBuf[NCBase* i], result);
+		}
+	}
+	// Same functionality as above, in this case we use a functor which allows to pass user data. 
+	template< class ImageCombineFn >
+	inline void ImageLayerCombineFunctor(Image* pResult, const Image* pBase, const Image* pBlended, ImageCombineFn&& ImageCombine)
+	{
+		check(pResult->GetFormat() == pBase->GetFormat());
+		check(pResult->GetSizeX() == pBase->GetSizeX());
+		check(pResult->GetSizeY() == pBase->GetSizeY());
+		check(pResult->GetLODCount() == pBase->GetLODCount());
+		check(pBase->GetSizeX() == pBlended->GetSizeX());
+		check(pBase->GetSizeY() == pBlended->GetSizeY());
+		check(pResult->GetLODCount() <= pBlended->GetLODCount());
+
+		const EImageFormat baseFormat = pBase->GetFormat();
+		const EImageFormat blendFormat = pBlended->GetFormat();
+
+		uint8* pDestBuf = pResult->GetData();
+		if (baseFormat == EImageFormat::IF_L_UBYTE )
+		{
+			if (blendFormat == EImageFormat::IF_L_UBYTE )
+			{
+				BufferLayerCombineFunctor< 1, 1 >(pDestBuf, pBase, pBlended, Forward<ImageCombineFn>(ImageCombine));
+			}
+			else if (blendFormat == EImageFormat::IF_RGB_UBYTE )
+			{
+				BufferLayerCombineFunctor< 1, 3 >(pDestBuf, pBase, pBlended, Forward<ImageCombineFn>(ImageCombine));
+			}
+			else if (blendFormat == EImageFormat::IF_RGBA_UBYTE || blendFormat == EImageFormat::IF_BGRA_UBYTE )
+			{
+				BufferLayerCombineFunctor< 1, 4 >(pDestBuf, pBase, pBlended, Forward<ImageCombineFn>(ImageCombine));
+			}
+			else
+			{
+				checkf(false, TEXT("Unsupported format."));
+			}
+		}
+		else if (baseFormat == EImageFormat::IF_RGB_UBYTE)
+		{
+			if (blendFormat == EImageFormat::IF_L_UBYTE )
+			{
+				BufferLayerCombineFunctor< 3, 1 >(pDestBuf, pBase, pBlended, Forward<ImageCombineFn>(ImageCombine));
+			}
+			else if (blendFormat == EImageFormat::IF_RGB_UBYTE )
+			{
+				BufferLayerCombineFunctor< 3, 3 >(pDestBuf, pBase, pBlended, Forward<ImageCombineFn>(ImageCombine));
+			}
+			else if (blendFormat == EImageFormat::IF_RGBA_UBYTE || blendFormat == EImageFormat::IF_BGRA_UBYTE )
+			{
+				BufferLayerCombineFunctor< 3, 4 >(pDestBuf, pBase, pBlended, Forward<ImageCombineFn>(ImageCombine));
+			}
+			else
+			{
+				checkf(false, TEXT("Unsupported format."));
+			}
+		}
+		else if (baseFormat == EImageFormat::IF_RGBA_UBYTE || baseFormat == EImageFormat::IF_BGRA_UBYTE)
+		{
+			if (blendFormat == EImageFormat::IF_L_UBYTE )
+			{
+				BufferLayerCombineFunctor< 4, 1 >(pDestBuf, pBase, pBlended, Forward<ImageCombineFn>(ImageCombine));
+			}
+			else if (blendFormat == EImageFormat::IF_RGB_UBYTE )
+			{
+				BufferLayerCombineFunctor< 4, 3 >(pDestBuf, pBase, pBlended, Forward<ImageCombineFn>(ImageCombine));
+			}
+			else if (blendFormat == EImageFormat::IF_RGBA_UBYTE || blendFormat == EImageFormat::IF_BGRA_UBYTE )
+			{
+				BufferLayerCombineFunctor< 4, 4 >(pDestBuf, pBase, pBlended, Forward<ImageCombineFn>(ImageCombine));
+			}
+			else
+			{
+				checkf(false, TEXT("Unsupported format."));
+			}
+		}
+		else
+		{
+			checkf(false, TEXT("Unsupported format."));
+		}
+	}
+
+}
diff --git a/Engine/Plugins/Experimental/Mutable/Source/MutableRuntime/Private/MuR/OpImageLuminance.h b/Engine/Plugins/Experimental/Mutable/Source/MutableRuntime/Private/MuR/OpImageLuminance.h
new file mode 100644
--- /dev/null
+++ b/Engine/Plugins/Experimental/Mutable/Source/MutableRuntime/Private/MuR/OpImageLuminance.h
@@ -0,0 +1,61 @@
+// Copyright Epic Games, Inc. All Rights Reserved.
+
+#pragma once
+
+#include "MuR/ImagePrivate.h"
+#include "MuR/Platform.h"
+
+namespace mu
+{
+
+	inline ImagePtr ImageLuminance( const Image* pA )
+	{
+        ImagePtr pDest = new Image( pA->GetSizeX(), pA->GetSizeY(), pA->GetLODCount(), EImageFormat::IF_L_UBYTE );
+
+        uint8_t* pDestBuf = pDest->GetData();
+        const uint8_t* pABuf = pA->GetData();
+
+		// Generic implementation
+		int pixelCount = (int)pA->CalculatePixelCount();
+
+		switch ( pA->GetFormat() )
+		{
+		case EImageFormat::IF_RGB_UBYTE:
+		{
+			for ( int i=0; i<pixelCount; ++i )
+			{
+                uint32_t l_16 = 76 * pABuf[3*i+0] + 150 * pABuf[3*i+1] + 29 * pABuf[3*i+2];
+                pDestBuf[i] = (uint8_t)std::min( l_16>>8, 255u );
+			}
+			break;
+		}
+
+        case EImageFormat::IF_RGBA_UBYTE:
+        {
+            for ( int i=0; i<pixelCount; ++i )
+            {
+                uint32_t l_16 = 76 * pABuf[4*i+0] + 150 * pABuf[4*i+1] + 29 * pABuf[4*i+2];
+                pDestBuf[i] = (uint8_t)std::min( l_16>>8, 255u );
+            }
+            break;
+        }
+
+        case EImageFormat::IF_BGRA_UBYTE:
+        {
+            for ( int i=0; i<pixelCount; ++i )
+            {
+                uint32_t l_16 = 76 * pABuf[4*i+2] + 150 * pABuf[4*i+1] + 29 * pABuf[4*i+0];
+                pDestBuf[i] = (uint8_t)std::min( l_16>>8, 255u );
+            }
+            break;
+        }
+
+		default:
+			check(false);
+			break;
+		}
+
+		return pDest;
+	}
+
+}
