Moved the Mutable plugin files to live under a directory named after the module. This allows us to get rid of long include paths that contain public and private folder names in the path and fixes head

diff --git a/Engine/Plugins/Experimental/Mutable/Source/MutableRuntime/Private/MuR/OpImageMipmap.h b/Engine/Plugins/Experimental/Mutable/Source/MutableRuntime/Private/MuR/OpImageMipmap.h
new file mode 100644
--- /dev/null
+++ b/Engine/Plugins/Experimental/Mutable/Source/MutableRuntime/Private/MuR/OpImageMipmap.h
@@ -0,0 +1,716 @@
+// Copyright Epic Games, Inc. All Rights Reserved.
+
+#pragma once
+
+#include "MuR/ImagePrivate.h"
+#include "MuR/OpImagePixelFormat.h"
+#include "Async/ParallelFor.h"
+
+namespace mu
+{
+
+	namespace OpImageMipmap_Detail
+	{
+
+		struct SharpenKernelStorage8
+		{
+			constexpr static size_t SIZE = 8;
+			float m_storage[SIZE * SIZE];
+		};
+
+		template<size_t ITERS, size_t N>
+		inline void BlurKernel(float(&inOutKernelStorage)[N])
+		{
+			static_assert(N > 2);
+
+			float temp[N];
+
+			float* inBuf = &inOutKernelStorage[0];
+			float* outBuf = &temp[0];
+
+			for (int32_t i = 0; i < ITERS; ++i)
+			{
+				constexpr float oneOverThree = 1.0f / 3.0f;
+
+				outBuf[0] = (inBuf[0] + inBuf[1]) * oneOverThree;
+
+				for (int32_t j = 1; j < N - 1; ++j)
+				{
+					outBuf[j] = (inBuf[j - 1] + inBuf[j] + inBuf[j + 1]) * oneOverThree;
+				}
+
+				outBuf[N - 1] = (inBuf[N - 2] + inBuf[N - 1]) * oneOverThree;
+
+				// Ping-pong buffers.
+				std::swap(outBuf, inBuf);
+			}
+
+			if (ITERS % 2 == 1)
+			{
+				for (uint32_t i = 0; i < N; ++i)
+				{
+					inOutKernelStorage[i] = temp[i];
+				}
+			}
+		}
+
+		template<size_t N>
+		inline void BuildGaussianBlurKernel(float(&kernel)[N], float variance)
+		{
+			constexpr float sqrtTau = 2.5066282746310002f;
+			const auto normalDist = [sdSqrtTauInv = 1.0f / (std::sqrt(variance) * sqrtTau),
+				twoVInv = 1.0f / (2.0f * variance)]
+				(float x)
+				-> float
+			{
+				return std::exp(-x * x * twoVInv) * sdSqrtTauInv;
+			};
+
+			const float center = static_cast<float>(N) * 0.5f;
+			float sum = 0;
+			for (int32_t i = 0; i < N; ++i)
+			{
+				const float dist = normalDist(static_cast<float>(i) - center + 0.5f);
+				kernel[i] = dist;
+				sum += dist;
+			}
+
+			const float normFactor = 1.0f / sum;
+			for (int32_t i = 0; i < N; ++i)
+			{
+				kernel[i] *= normFactor;
+			}
+		}
+
+		template<size_t N>
+		inline void BuildSharpenKernel(float(&kernel)[N], float factor)
+		{
+			for (uint32_t i = 0; i < N; ++i)
+			{
+				kernel[i] = 0.0f;
+			}
+
+			kernel[(N >> 1) - 1] = factor * 0.5f;
+			kernel[(N >> 1)] = factor * 0.5f;
+		}
+
+		template<size_t N>
+		inline void BuildKernel2DFromSeparable(float(&kernel1D)[N], SharpenKernelStorage8& outKernel)
+		{
+			static_assert(N == SharpenKernelStorage8::SIZE);
+
+			for (int32_t y = 0; y < N; ++y)
+			{
+				for (int32_t x = 0; x < N; ++x)
+				{
+					outKernel.m_storage[y * N + x] = kernel1D[y] * kernel1D[x];
+				}
+			}
+		}
+
+		inline SharpenKernelStorage8 MakeMipGaussianSharpenKernel(float factor)
+		{
+			SharpenKernelStorage8 kernel;
+
+			float kernel1DStorage0[SharpenKernelStorage8::SIZE];
+			float kernel1DStorage1[SharpenKernelStorage8::SIZE];
+
+			// negative factors indicate blur where the -factor is the variance.
+			if (factor < 0)
+			{
+				BuildGaussianBlurKernel(kernel1DStorage0, -factor);
+
+				BuildKernel2DFromSeparable(kernel1DStorage0, kernel);
+
+				return kernel;
+			}
+
+			BuildSharpenKernel(kernel1DStorage0, 1.0f + factor);
+			BuildSharpenKernel(kernel1DStorage1, -factor);
+
+			BlurKernel<1>(kernel1DStorage0);
+			BlurKernel<3>(kernel1DStorage1);
+
+			for (int32_t i = 0; i < SharpenKernelStorage8::SIZE; ++i)
+			{
+				kernel1DStorage0[i] += kernel1DStorage1[i];
+			}
+
+			BuildKernel2DFromSeparable(kernel1DStorage0, kernel);
+
+			return kernel;
+		}
+
+		template<int32_t PIXEL_SIZE, typename CHANNEL_TYPE, EAddressMode AD_MODE = EAddressMode::AM_NONE>
+		class ImageAccessor
+		{
+			const vec2<int32_t> m_dim;
+			CHANNEL_TYPE const* const m_pImageBuf;
+
+		public:
+			ImageAccessor(const CHANNEL_TYPE* const pImageBuf, vec2<uint32_t> dim)
+				: m_dim(dim)
+				, m_pImageBuf(pImageBuf)
+			{
+			}
+
+			uint8_t operator()(int32_t x, int32_t y, int32_t c) const
+			{
+				if constexpr (AD_MODE == EAddressMode::AM_WRAP)
+				{
+					// (y & (m_dim.y() - 1)) * m_dim.y() + (x & (m_dim.x() - 1))
+					return m_pImageBuf[((y % m_dim.y()) * m_dim.x() + (x % m_dim.x())) * PIXEL_SIZE + c];
+				}
+				else if (AD_MODE == EAddressMode::AM_CLAMP)
+				{
+					return m_pImageBuf[(mu::clamp(y, 0, m_dim.y() - 1) * m_dim.x() +
+						mu::clamp(x, 0, m_dim.x() - 1))
+						* PIXEL_SIZE + c];
+				}
+				else if (AD_MODE == EAddressMode::AM_BLACK_BORDER)
+				{
+					return ((x < m_dim.x()) && (y < m_dim.y()) && (x >= 0) && (y >= 0))
+						? m_pImageBuf[(y * m_dim.x() + x) * PIXEL_SIZE + c]
+						: 0;
+				}
+				else
+				{
+					return m_pImageBuf[(y * m_dim.x() + x) * PIXEL_SIZE + c];
+				}
+			}
+		};
+
+
+		template<int32_t PIXEL_SIZE, EAddressMode AD, typename PIXEL_TYPE>
+		void GenerateMipSharpenedRegion(
+			const SharpenKernelStorage8& kernel,
+			const ImageAccessor<PIXEL_SIZE, PIXEL_TYPE, AD>& sourceAccessor,
+			PIXEL_TYPE* pDest, vec2<int32_t> destSize,
+			vec2<int32_t> regionStart, vec2<int32_t> regionEnd)
+		{
+			// Only kernels of size 8 supported.
+			constexpr int32_t kernelSize = SharpenKernelStorage8::SIZE;
+			static_assert(kernelSize == 8);
+
+			constexpr int32_t kernelCenter = (kernelSize >> 1) - 1;
+			for (int32_t y = regionStart.y(); y < regionEnd.y(); ++y)
+			{
+				for (int32_t x = regionStart.x(); x < regionEnd.x(); ++x)
+				{
+					float stagingStorage[kernelSize * kernelSize * PIXEL_SIZE];
+
+					const vec2<uint32_t> sourceCoord = vec2<uint32_t>(x << 1, y << 1);
+
+
+					// Convert Image kernel stamp to float. This conversion will be done multiple times.
+					// and should be done once before applying the kernel, or operate directly in unorm8/snorm8.
+					for (int32_t ky = 0; ky < kernelSize; ++ky)
+					{
+						for (int32_t kx = 0; kx < kernelSize; ++kx)
+						{
+							for (int32_t c = 0; c < PIXEL_SIZE; ++c)
+							{
+								stagingStorage[(ky * kernelSize + kx) + (kernelSize * kernelSize * c)] = static_cast<float>(
+									sourceAccessor(sourceCoord.x() + kx - kernelCenter,
+										sourceCoord.y() + ky - kernelCenter,
+										c)) / 255.0f;
+							}
+						}
+					}
+
+					// Apply kernel. Although the kernel is separable, the cost would not differ much from the not separable one.
+					// That is because to make it separable a partial convolution of the source is needed. 
+					// Since here we are only applying to half the image and the kernel used is small (8x8) the benefits are not 
+					// clear.
+
+					for (int32_t c = 0; c < PIXEL_SIZE; ++c)
+					{
+						float sampleValue = 0.0f;
+						float sum = 0.0f;
+						for (int32_t k = 0; k < kernelSize * kernelSize; ++k)
+						{
+							sampleValue += kernel.m_storage[k] * stagingStorage[k + (kernelSize * kernelSize * c)];
+							sum += kernel.m_storage[k];
+						}
+
+						pDest[(y * destSize.x() + x) * PIXEL_SIZE + c] = static_cast<uint8_t>(mu::clamp(sampleValue, 0.0f, 1.0f) * 255.0f);
+					}
+				}
+			}
+		}
+
+		template<size_t PIXEL_SIZE, EAddressMode AD_MODE>
+		inline void GenerateMipmapUint8Sharpen(int mips,
+			const uint8_t* pSource, uint8_t* pDest,
+			vec2<int32_t> sourceSize,
+			float sharpeningFactor)
+		{
+			SharpenKernelStorage8 kernel = MakeMipGaussianSharpenKernel(sharpeningFactor);
+
+			size_t totalDestSize = 0;
+
+			for (; mips >= 0; --mips)
+			{
+				ImageAccessor<PIXEL_SIZE, uint8_t, AD_MODE> sourceSampler(pSource, sourceSize);
+				ImageAccessor<PIXEL_SIZE, uint8_t, EAddressMode::AM_NONE> sourceSamplerAddressModeNone(pSource, sourceSize);
+
+				vec2<int32_t> destSize = vec2<int32_t>(FMath::DivideAndRoundUp(sourceSize.x(), 2), FMath::DivideAndRoundUp(sourceSize.y(), 2));
+
+				// Core image, will never sample outside the image so we don't need to care about address mode.
+				// Only sample source pixels where the kernel can be fully applied.
+				// Border cases are treated separately.
+				GenerateMipSharpenedRegion(
+					kernel,
+					sourceSamplerAddressModeNone,
+					pDest, destSize,
+					vec2<int32_t>(2, 2), vec2<int32_t>(destSize.x() - 2, destSize.y() - 2));
+
+				// Horizontal borders.
+				GenerateMipSharpenedRegion(
+					kernel,
+					sourceSampler,
+					pDest, destSize,
+					vec2<int32_t>(0, 0), vec2<int32_t>(destSize.x(), std::min(2, destSize.y())));
+
+				GenerateMipSharpenedRegion(
+					kernel,
+					sourceSampler,
+					pDest, destSize,
+					vec2<int32_t>(0, std::max(destSize.y() - 2, 0)), vec2<int32_t>(destSize.x(), destSize.y()));
+
+				// Vertical borders
+				GenerateMipSharpenedRegion(
+					kernel,
+					sourceSampler,
+					pDest, destSize,
+					vec2<int32_t>(0, std::min(2, destSize.y())), vec2<int32_t>(std::min(2, destSize.x()), destSize.y() - 2));
+
+				GenerateMipSharpenedRegion(
+					kernel,
+					sourceSampler,
+					pDest, destSize,
+					vec2<int32_t>(destSize.x() - 2, std::min(2, destSize.y())), vec2<int32_t>(destSize.x(), destSize.y() - 2));
+
+				sourceSize = destSize;
+
+				pSource = pDest;
+				pDest = pDest + destSize.x() * destSize.y() * PIXEL_SIZE;
+
+				totalDestSize += sourceSize.x() * sourceSize.y() * PIXEL_SIZE;
+			}
+		}
+
+		template<int32_t PIXEL_SIZE>
+		inline void GenerateMipmapUint8Unfiltered(
+			int mips,
+			const uint8_t* pSource, uint8_t* pDest,
+			vec2<int> sourceSize)
+		{
+			for (; mips >= 0; --mips)
+			{
+				vec2<int32_t> destSize = vec2<int32_t>(FMath::DivideAndRoundUp(sourceSize.x(), 2), FMath::DivideAndRoundUp(sourceSize.y(), 2));
+
+				for (int32_t y = 0; y < destSize.y(); ++y)
+				{
+					for (int32_t x = 0; x < destSize.x(); ++x)
+					{
+						for (int32_t c = 0; c < PIXEL_SIZE; ++c)
+						{
+							pDest[(y * destSize.x() + x) * PIXEL_SIZE + c] =
+								pSource[((y << 1) * sourceSize.x() + (x << 1)) * PIXEL_SIZE + c];
+						}
+					}
+				}
+
+				sourceSize = destSize;
+				pSource = pDest;
+				pDest = pDest + destSize.x() * destSize.y() * PIXEL_SIZE;
+			}
+		}
+
+		template<int32_t PIXEL_SIZE>
+		inline void GenerateMipmapsUint8SimpleAverage(
+			int mips,
+			const uint8_t* pSource, uint8_t* pDest,
+			vec2<int> sourceSize)
+		{
+			const uint8_t* pMipSource = pSource;
+			uint8_t* pMipDest = pDest;
+
+			vec2<int> destSize = sourceSize;
+			for (int m = 0; m < mips; ++m)
+			{
+				check(destSize[0] > 1 || destSize[1] > 1);
+
+				sourceSize = destSize;
+
+				int fullColumns = destSize[0] / 2;
+				bool strayColumn = (destSize[0] % 2) != 0;
+				int fullRows = destSize[1] / 2;
+				bool strayRow = (destSize[1] % 2) != 0;
+
+				destSize[0] = FMath::DivideAndRoundUp(destSize[0], 2);
+				destSize[1] = FMath::DivideAndRoundUp(destSize[1], 2);
+
+				int sourceStride = sourceSize[0] * PIXEL_SIZE;
+				int destStride = destSize[0] * PIXEL_SIZE;
+
+				const auto& ProcessRow = [
+					pMipDest, pMipSource, fullColumns, strayColumn, sourceStride, destStride
+				] (uint32 y)
+				{
+					const uint8_t* pSourceRow0 = pMipSource + 2 * y * sourceStride;
+					const uint8_t* pSourceRow1 = pSourceRow0 + sourceStride;
+					uint8_t* pDestRow = pMipDest + y * destStride;
+
+					for (int x = 0; x < fullColumns; ++x)
+					{
+						for (int c = 0; c < PIXEL_SIZE; ++c)
+						{
+							int p = pSourceRow0[c] + pSourceRow0[PIXEL_SIZE + c] + pSourceRow1[c] +
+								pSourceRow1[PIXEL_SIZE + c];
+							pDestRow[c] = (uint8_t)(p >> 2);
+						}
+
+						pSourceRow0 += 2 * PIXEL_SIZE;
+						pSourceRow1 += 2 * PIXEL_SIZE;
+						pDestRow += PIXEL_SIZE;
+					}
+
+					if (strayColumn)
+					{
+						for (int c = 0; c < PIXEL_SIZE; ++c)
+						{
+							int p = pSourceRow0[c] + pSourceRow1[c];
+							pDestRow[c] = (uint8_t)(p >> 1);
+						}
+					}
+				};
+
+					constexpr int PixelConcurrencyThreshold = 0xffff;
+					if (destSize[0] * destSize[1] < PixelConcurrencyThreshold)
+					{
+						for (int y = 0; y < fullRows; ++y)
+						{
+							ProcessRow(y);
+						}
+					}
+					else
+					{
+						ParallelFor(fullRows, ProcessRow);
+					}
+
+					if (strayRow)
+					{
+						const uint8_t* pSourceRow0 = pMipSource + 2 * fullRows * sourceStride;
+						const uint8_t* pSourceRow1 = pSourceRow0 + sourceStride;
+						uint8_t* pDestRow = pMipDest + fullRows * destStride;
+
+						for (int x = 0; x < fullColumns; ++x)
+						{
+							for (int c = 0; c < PIXEL_SIZE; ++c)
+							{
+								int p = pSourceRow0[c] + pSourceRow0[PIXEL_SIZE + c];
+								pDestRow[c] = (uint8_t)(p >> 1);
+							}
+
+							pSourceRow0 += 2 * PIXEL_SIZE;
+							pDestRow += PIXEL_SIZE;
+						}
+
+						if (strayColumn)
+						{
+							for (int c = 0; c < PIXEL_SIZE; ++c)
+							{
+								pDestRow[c] = pSourceRow0[c];
+							}
+						}
+					}
+
+					// Reset the source pointer for the next mip, to use the dest that we have just
+					// generated.
+					pMipSource = pMipDest;
+					pMipDest += destSize[0] * destSize[1] * PIXEL_SIZE;
+			}
+		}
+
+	} // namespace OpImageMipmap_Detail
+
+
+		//---------------------------------------------------------------------------------------------
+		//! Generate the mipmaps for byte-based images of whatever number of channels.
+		//! \param mips number of additional levels to build from the source.
+		//---------------------------------------------------------------------------------------------
+	template<int32_t PIXEL_SIZE>
+	inline void GenerateMipmapsUint8(int mips,
+		const uint8_t* pSource, uint8_t* pDest,
+		vec2<int32_t> sourceSize,
+		const FMipmapGenerationSettings& settings)
+	{
+		using namespace OpImageMipmap_Detail;
+
+		switch (settings.m_filterType)
+		{
+		case EMipmapFilterType::MFT_SimpleAverage:
+		{
+			GenerateMipmapsUint8SimpleAverage<PIXEL_SIZE>(mips, pSource, pDest, sourceSize);
+			break;
+		}
+		case EMipmapFilterType::MFT_Sharpen:
+		{
+			switch (settings.m_addressMode)
+			{
+			case EAddressMode::AM_CLAMP:
+			{
+				GenerateMipmapUint8Sharpen<PIXEL_SIZE, EAddressMode::AM_CLAMP>(
+					mips, pSource, pDest, sourceSize, settings.m_sharpenFactor);
+				break;
+			}
+			case EAddressMode::AM_WRAP:
+			{
+				GenerateMipmapUint8Sharpen<PIXEL_SIZE, EAddressMode::AM_WRAP>(
+					mips, pSource, pDest, sourceSize, settings.m_sharpenFactor);
+				break;
+			}
+			case EAddressMode::AM_BLACK_BORDER:
+			{
+				GenerateMipmapUint8Sharpen<PIXEL_SIZE, EAddressMode::AM_BLACK_BORDER>(
+					mips, pSource, pDest, sourceSize, settings.m_sharpenFactor);
+				break;
+			}
+			default:
+			{
+				GenerateMipmapUint8Sharpen<PIXEL_SIZE, EAddressMode::AM_NONE>(
+					mips, pSource, pDest, sourceSize, settings.m_sharpenFactor);
+				break;
+			}
+			}
+
+			break;
+		}
+		case EMipmapFilterType::MFT_Unfiltered:
+		{
+			GenerateMipmapUint8Unfiltered<PIXEL_SIZE>(mips, pSource, pDest, sourceSize);
+			break;
+		}
+		default:
+		{
+			check(false);
+			break;
+		}
+		}
+	}
+
+	struct SCRATCH_IMAGE_MIPMAP
+	{
+		ImagePtr pUncompressed;
+		ImagePtr pUncompressedMips;
+		ImagePtr pCompressedMips;
+	};
+
+
+	//---------------------------------------------------------------------------------------------
+	inline void ImageMipmap_PrepareScratch(Image* pDest, const Image* pBase, int levelCount,
+		SCRATCH_IMAGE_MIPMAP* scratch)
+	{
+		int startLevel = pBase->GetLODCount() - 1;
+
+		check(pDest->GetLODCount() == levelCount);
+		check(pDest->GetSizeX() == pBase->GetSizeX());
+		check(pDest->GetSizeY() == pBase->GetSizeY());
+		check(pDest->GetFormat() == pBase->GetFormat());
+		(void)pDest;
+
+		switch (pBase->GetFormat())
+		{
+		case EImageFormat::IF_L_UBYTE:
+		case EImageFormat::IF_RGB_UBYTE:
+		case EImageFormat::IF_RGBA_UBYTE:
+		case EImageFormat::IF_BGRA_UBYTE:
+			break;
+
+			// Bad cases: we need to decompress, mip and then recompress. It may be necessary to
+			// generate the latests mips after composing, or in unoptimised code.
+			// TODO: Make sure that the code optimisations avoid this cases generating separate
+			// operations to generate the mip tail.
+		case EImageFormat::IF_BC1:
+		case EImageFormat::IF_BC2:
+		case EImageFormat::IF_BC3:
+		case EImageFormat::IF_BC4:
+		case EImageFormat::IF_BC5:
+		case EImageFormat::IF_ASTC_4x4_RGB_LDR:
+		case EImageFormat::IF_ASTC_4x4_RGBA_LDR:
+		case EImageFormat::IF_ASTC_4x4_RG_LDR:
+		{
+			// Uncompress the last mip that we already have
+			vec2<int> uncompressedSize = pBase->CalculateMipSize(startLevel);
+			scratch->pUncompressed = new Image(
+				(uint16_t)uncompressedSize[0],
+				(uint16_t)uncompressedSize[1],
+				1,
+				EImageFormat::IF_RGBA_UBYTE);
+
+			// Generate the mipmaps from there on
+			scratch->pUncompressedMips = new Image(
+				(uint16_t)uncompressedSize[0],
+				(uint16_t)uncompressedSize[1],
+				levelCount - startLevel,
+				EImageFormat::IF_RGBA_UBYTE);
+
+			// Compress the mipmapped image
+			scratch->pCompressedMips = new Image(
+				(uint16_t)uncompressedSize[0],
+				(uint16_t)uncompressedSize[1],
+				scratch->pUncompressedMips->GetLODCount(),
+				pBase->GetFormat());
+
+			break;
+		}
+
+		case EImageFormat::IF_L_UBIT_RLE:
+		case EImageFormat::IF_L_UBYTE_RLE:
+		{
+			// Uncompress the last mip that we already have
+			vec2<int> uncompressedSize = pBase->CalculateMipSize(startLevel);
+			scratch->pUncompressed = new Image(
+				(uint16_t)uncompressedSize[0],
+				(uint16_t)uncompressedSize[1],
+				1,
+				EImageFormat::IF_L_UBYTE);
+
+			// Generate the mipmaps from there on
+			scratch->pUncompressedMips = new Image(
+				(uint16_t)uncompressedSize[0],
+				(uint16_t)uncompressedSize[1],
+				levelCount - startLevel,
+				EImageFormat::IF_L_UBYTE);
+
+			// Compress the mipmapped image
+			scratch->pCompressedMips = new Image(
+				(uint16_t)uncompressedSize[0],
+				(uint16_t)uncompressedSize[1],
+				scratch->pUncompressedMips->GetLODCount(),
+				pBase->GetFormat());
+
+			// Preallocate ample memory for the compressed data
+			scratch->pCompressedMips->m_data.SetNum(scratch->pUncompressedMips->GetDataSize());
+
+			// Preallocate ample memory for the destination data
+			pDest->m_data.SetNum(pBase->GetDataSize()
+				+
+				scratch->pUncompressedMips->GetDataSize());
+
+			break;
+		}
+
+		default:
+			checkf(false, TEXT("Format not implemented in mipmap generation."));
+		}
+	}
+
+
+	//---------------------------------------------------------------------------------------------
+	inline void ImageMipmap(int imageCompressionQuality, Image* pDest, const Image* pBase,
+		int levelCount, SCRATCH_IMAGE_MIPMAP* scratch,
+		FMipmapGenerationSettings& settings)
+	{
+		int startLevel = pBase->GetLODCount() - 1;
+
+		check(!(pBase->m_flags & Image::IF_CANNOT_BE_SCALED));
+
+		check(pDest->GetLODCount() == levelCount);
+		check(pDest->GetSizeX() == pBase->GetSizeX());
+		check(pDest->GetSizeY() == pBase->GetSizeY());
+		check(pDest->GetFormat() == pBase->GetFormat());
+
+		// Copy already-existing mips
+
+		// Calculate the data size, since the base could have more data allocated.
+		int32 mipDataSize = pBase->CalculateDataSize();
+		check(pDest->GetDataSize() >= mipDataSize);
+		FMemory::Memcpy(pDest->GetData(), pBase->GetData(), mipDataSize);
+
+		int mipsToBuild = levelCount - startLevel - 1;
+		if (!mipsToBuild)
+		{
+			return;
+		}
+
+		uint8_t* pSourceBuf = pDest->GetMipData(startLevel);
+		uint8_t* pDestBuf = pDest->GetMipData(startLevel + 1);
+
+		vec2<int> sourceSize = pDest->CalculateMipSize(startLevel);
+
+		switch (pBase->GetFormat())
+		{
+		case EImageFormat::IF_L_UBYTE:
+			GenerateMipmapsUint8<1>(levelCount - startLevel - 1, pSourceBuf, pDestBuf, sourceSize, settings);
+			break;
+
+		case EImageFormat::IF_RGB_UBYTE:
+			GenerateMipmapsUint8<3>(levelCount - startLevel - 1, pSourceBuf, pDestBuf, sourceSize, settings);
+			break;
+
+		case EImageFormat::IF_BGRA_UBYTE:
+		case EImageFormat::IF_RGBA_UBYTE:
+			GenerateMipmapsUint8<4>(levelCount - startLevel - 1, pSourceBuf, pDestBuf, sourceSize, settings);
+			break;
+
+			// Bad cases: we need to decompress, mip and then recompress. It may be necessary to
+			// generate the latests mips after composing, or in unoptimised code.
+			// TODO: Make sure that the code optimisations avoid this cases generating separate
+			// operations to generate the mip tail.
+		case EImageFormat::IF_BC1:
+		case EImageFormat::IF_BC2:
+		case EImageFormat::IF_BC3:
+		case EImageFormat::IF_BC4:
+		case EImageFormat::IF_BC5:
+		case EImageFormat::IF_ASTC_4x4_RGB_LDR:
+		case EImageFormat::IF_ASTC_4x4_RGBA_LDR:
+		case EImageFormat::IF_ASTC_4x4_RG_LDR:
+		case EImageFormat::IF_L_UBIT_RLE:
+		case EImageFormat::IF_L_UBYTE_RLE:
+		{
+			// Uncompress the last mip that we already have
+			bool bSuccess = ImagePixelFormatInPlace(imageCompressionQuality, scratch->pUncompressed.get(), pBase, startLevel);
+			check(bSuccess);
+
+			// Generate the mipmaps from there on
+			ImageMipmap(imageCompressionQuality, scratch->pUncompressedMips.get(),
+				scratch->pUncompressed.get(), levelCount - startLevel, 0, settings);
+
+			// Compress the mipmapped image
+			bSuccess = ImagePixelFormatInPlace(imageCompressionQuality, scratch->pCompressedMips.get(), scratch->pUncompressedMips.get());
+			int32 ExcessDataSize = scratch->pCompressedMips->GetDataSize();
+			while (!bSuccess)
+			{
+				// Bad case: this should almost never happen.
+				MUTABLE_CPUPROFILER_SCOPE(Mipmap_Recompression_OutOfSpace);
+
+				scratch->pCompressedMips->m_data.SetNum(ExcessDataSize);
+				bSuccess = ImagePixelFormatInPlace(imageCompressionQuality, scratch->pCompressedMips.get(), scratch->pUncompressedMips.get());
+				ExcessDataSize *= 4;
+			}
+
+			int32 FinalDestSize = pSourceBuf - pDest->GetData() + scratch->pCompressedMips->GetDataSize();
+
+			if (FinalDestSize > pDest->GetDataSize())
+			{
+				// Bad case: this should almost never happen.
+				pDest->m_data.SetNum(FinalDestSize);
+				pSourceBuf = pDest->GetMipData(startLevel);
+			}
+
+			// TODO: We are recompressing the source level, which we could avoid.
+			FMemory::Memcpy(pSourceBuf,
+				scratch->pCompressedMips->GetData(),
+				scratch->pCompressedMips->GetDataSize());
+			break;
+		}
+
+		default:
+			checkf(false, TEXT("Format not implemented in mipmap generation."));
+		}
+
+	}
+}
diff --git a/Engine/Plugins/Experimental/Mutable/Source/MutableRuntime/Private/MuR/OpImageMultiply.h b/Engine/Plugins/Experimental/Mutable/Source/MutableRuntime/Private/MuR/OpImageMultiply.h
new file mode 100644
--- /dev/null
+++ b/Engine/Plugins/Experimental/Mutable/Source/MutableRuntime/Private/MuR/OpImageMultiply.h
@@ -0,0 +1,72 @@
+// Copyright Epic Games, Inc. All Rights Reserved.
+
+#pragma once
+
+#include "MuR/ImagePrivate.h"
+#include "MuR/OpImageLayer.h"
+#include "MuR/Platform.h"
+
+namespace mu
+{
+
+	inline unsigned MultiplyChannelMasked( unsigned base, unsigned blended, unsigned mask )
+	{
+		unsigned multiply = ( base * blended ) >> 8;
+		unsigned masked = ( ( ( 255 - mask ) * base ) + ( mask * multiply ) ) >> 8;
+		return masked;
+	}
+
+	inline unsigned MultiplyChannel( unsigned base, unsigned blended )
+	{
+		unsigned multiply = ( base * blended ) >> 8;
+		return multiply;
+	}
+
+
+    inline void ImageMultiply( Image* pResult, const Image* pBase, const Image* pMask, vec3<float> col )
+	{
+        ImageLayerColour<MultiplyChannelMasked,MultiplyChannel,false>( pResult, pBase, pMask, col );
+	}
+
+
+    inline void ImageMultiply( Image* pResult, const Image* pBase, const Image* pMask,
+                               const Image* pBlended,
+                               bool applyToAlpha )
+	{
+        return ImageLayer<MultiplyChannelMasked,MultiplyChannel,false>
+                ( pResult, pBase, pMask, pBlended, applyToAlpha );
+	}
+
+
+    inline void ImageMultiply( Image* pResult, const Image* pBase, vec3<float> col )
+	{
+        return ImageLayerColour<MultiplyChannel,false>( pResult, pBase, col );
+	}
+
+
+    inline void ImageMultiply( Image* pResult, const Image* pBase, const Image* pBlended,
+                               bool applyToAlpha )
+	{
+        return ImageLayer<MultiplyChannel,false>( pResult, pBase, pBlended, applyToAlpha );
+	}
+
+
+	inline void ImageMultiplyOnBase( Image* pBase, const Image* pMask, vec3<float> col )
+	{
+		ImageLayerColourOnBase<MultiplyChannelMasked,MultiplyChannel, false>( pBase, pMask, col );
+	}
+
+    inline void ImageMultiplyOnBase( Image* pBase, const Image* pMask, const Image* pBlended,
+                                     bool applyToAlpha )
+	{
+        ImageLayerOnBase<MultiplyChannelMasked,MultiplyChannel, false>
+                ( pBase, pMask, pBlended, applyToAlpha );
+	}
+
+    inline void ImageMultiplyOnBase( Image* pBase, const Image* pBlended,
+                                     bool applyToAlpha )
+	{
+        ImageLayerOnBase<MultiplyChannel, false>( pBase, pBlended, applyToAlpha );
+	}
+
+}
diff --git a/Engine/Plugins/Experimental/Mutable/Source/MutableRuntime/Private/MuR/OpImageNormalCombine.h b/Engine/Plugins/Experimental/Mutable/Source/MutableRuntime/Private/MuR/OpImageNormalCombine.h
new file mode 100644
--- /dev/null
+++ b/Engine/Plugins/Experimental/Mutable/Source/MutableRuntime/Private/MuR/OpImageNormalCombine.h
@@ -0,0 +1,83 @@
+// Copyright Epic Games, Inc. All Rights Reserved.
+
+#pragma once
+
+#include "MuR/ImagePrivate.h"
+#include "MuR/OpImageLayer.h"
+
+namespace mu
+{
+	// Combine layer oprations get packed color.
+	inline unsigned CombineNormal(unsigned base, unsigned blended)
+	{
+		// See /Engine/Functions/Engine_MaterialFunctions02/Utility/BlendAngleCorrectedNormals.BlendAngleCorrectedNormals
+		// for the source of the effect.
+
+		const float baseRf = (static_cast<float>((base >> 0)  & 0xFF) / 255.0f) * 2.0f - 1.0f;
+		const float baseGf = (static_cast<float>((base >> 8)  & 0xFF) / 255.0f) * 2.0f - 1.0f;
+		const float baseBf = (static_cast<float>((base >> 16) & 0xFF) / 255.0f) * 2.0f; //One added to the b channel.
+
+		const float blendedRf = (static_cast<float>((blended >> 0)  & 0xFF) / 255.0f) * 2.0f - 1.0f;
+		const float blendedGf = (static_cast<float>((blended >> 8)  & 0xFF) / 255.0f) * 2.0f - 1.0f;
+		const float blendedBf = (static_cast<float>((blended >> 16) & 0xFF) / 255.0f) * 2.0f - 1.0f;
+
+		vec3f a = vec3f(baseRf, baseGf, baseBf);
+		vec3f b = vec3f(-blendedRf, -blendedGf, blendedBf);
+
+		// The original does not normalize but if not done results don't look good due to signal clipping.
+		vec3f n = normalise_approx(a*dot(a, b) - b*a.z());
+
+		return 
+			(mu::clamp(static_cast<unsigned>((n.x() + 1.0f) * 0.5f * 255.0f), 0u, 255u) << 0) |
+			(mu::clamp(static_cast<unsigned>((n.y() + 1.0f) * 0.5f * 255.0f), 0u, 255u) << 8) |
+			(mu::clamp(static_cast<unsigned>((n.z() + 1.0f) * 0.5f * 255.0f), 0u, 255u) << 16);
+	}
+
+	inline unsigned CombineNormalMasked(unsigned base, unsigned blended, unsigned mask)
+	{
+		// See /Engine/Functions/Engine_MaterialFunctions02/Utility/BlendAngleCorrectedNormals.BlendAngleCorrectedNormals
+		// for the source of the effect.
+
+		const float baseRf = (static_cast<float>((base >> 0)  & 0xFF) / 255.0f) * 2.0f - 1.0f;
+		const float baseGf = (static_cast<float>((base >> 8)  & 0xFF) / 255.0f) * 2.0f - 1.0f;
+		const float baseBf = (static_cast<float>((base >> 16) & 0xFF) / 255.0f) * 2.0f - 1.0f;
+
+		const float blendedRf = (static_cast<float>((blended >> 0)  & 0xFF) / 255.0f) * 2.0f - 1.0f;
+		const float blendedGf = (static_cast<float>((blended >> 8)  & 0xFF) / 255.0f) * 2.0f - 1.0f;
+		const float blendedBf = (static_cast<float>((blended >> 16) & 0xFF) / 255.0f) * 2.0f - 1.0f;
+
+		const float maskRf = (static_cast<float>(mask & 0xFF) / 255.0f);
+
+		vec3f a = vec3f(baseRf, baseGf, baseBf + 1.0f);
+		vec3f b = vec3f(-blendedRf, -blendedGf, blendedBf);
+
+		// The original does not normalize but if not done results don't look good due to signal clipping.
+		vec3f n = mu::normalise_approx( 
+					mu::lerp( vec3f(baseRf, baseGf, baseBf), a*dot(a, b) - b*a.z(), maskRf ) );
+
+		return 
+			(mu::clamp(static_cast<unsigned>((n.x() + 1.0f) * 0.5f * 255.0f), 0u, 255u) << 0) |
+			(mu::clamp(static_cast<unsigned>((n.y() + 1.0f) * 0.5f * 255.0f), 0u, 255u) << 8) |
+			(mu::clamp(static_cast<unsigned>((n.z() + 1.0f) * 0.5f * 255.0f), 0u, 255u) << 16);
+	}
+
+	inline void ImageNormalCombine(Image* pResult, const Image* pBase, vec3<float> col)
+	{
+		ImageLayerCombineColour<CombineNormal>(pResult, pBase, col);
+	}
+
+	inline void ImageNormalCombine(Image* pResult, const Image* pBase, const Image* pBlended)
+	{
+		ImageLayerCombine<CombineNormal>(pResult, pBase, pBlended);
+	}
+
+	inline void ImageNormalCombine(Image* pResult, const Image* pBase, const Image* pMask, vec3<float> col)
+	{
+		ImageLayerCombineColour<CombineNormal, CombineNormalMasked>(pResult, pBase, pMask, col);
+	}
+
+	inline void ImageNormalCombine(Image* pResult, const Image* pBase, const Image* pMask, const Image* pBlended)
+	{
+		ImageLayerCombine<CombineNormal, CombineNormalMasked>(pResult, pBase, pMask, pBlended);
+	}
+}
diff --git a/Engine/Plugins/Experimental/Mutable/Source/MutableRuntime/Private/MuR/OpImageOverlay.h b/Engine/Plugins/Experimental/Mutable/Source/MutableRuntime/Private/MuR/OpImageOverlay.h
new file mode 100644
--- /dev/null
+++ b/Engine/Plugins/Experimental/Mutable/Source/MutableRuntime/Private/MuR/OpImageOverlay.h
@@ -0,0 +1,61 @@
+// Copyright Epic Games, Inc. All Rights Reserved.
+
+#pragma once
+
+#include "MuR/ImagePrivate.h"
+#include "MuR/OpImageHardLight.h"
+
+
+namespace mu
+{
+	inline unsigned OverlayChannelMasked( unsigned base, unsigned blended, unsigned mask )
+	{
+		// Photoshop like
+		//return HardLightChannel( blended, base, mask );
+
+		// gimp-like
+		unsigned overlay = ( base * ( base + ( ( 2 * blended * ( 255 - base) ) >> 8  ) ) ) >> 8;
+		unsigned masked = ( ( ( 255 - mask ) * base ) + ( mask * overlay ) ) >> 8;
+		return masked;
+	}
+
+
+	inline unsigned OverlayChannel( unsigned base, unsigned blended )
+	{
+		// Photoshop like
+		//return HardLightChannel( blended, base, mask );
+
+		// gimp-like
+		unsigned overlay = ( base * ( base + ( ( 2 * blended * ( 255 - base) ) >> 8  ) ) ) >> 8;
+		return overlay;
+	}
+
+
+    inline void ImageOverlay( Image* pResult, const Image* pBase, const Image* pMask, vec3<float> col )
+	{
+        ImageLayerColour<OverlayChannelMasked,OverlayChannel,true>( pResult, pBase, pMask, col );
+	}
+
+
+    inline void ImageOverlay( Image* pResult, const Image* pBase, const Image* pMask,
+                              const Image* pBlended,
+                              bool applyToAlpha )
+	{
+        ImageLayer<OverlayChannelMasked,OverlayChannel,true>
+                ( pResult, pBase, pMask, pBlended, applyToAlpha );
+	}
+
+
+    inline void ImageOverlay( Image* pResult, const Image* pBase, vec3<float> col )
+	{
+        ImageLayerColour<OverlayChannel,true>( pResult, pBase, col );
+	}
+
+
+    inline void ImageOverlay( Image* pResult, const Image* pBase, const Image* pBlended,
+                              bool applyToAlpha )
+	{
+        ImageLayer<OverlayChannel,true>( pResult, pBase, pBlended, applyToAlpha );
+	}
+
+}
diff --git a/Engine/Plugins/Experimental/Mutable/Source/MutableRuntime/Private/MuR/OpImagePixelFormat.cpp b/Engine/Plugins/Experimental/Mutable/Source/MutableRuntime/Private/MuR/OpImagePixelFormat.cpp
new file mode 100644
--- /dev/null
+++ b/Engine/Plugins/Experimental/Mutable/Source/MutableRuntime/Private/MuR/OpImagePixelFormat.cpp
@@ -0,0 +1,1191 @@
+// Copyright Epic Games, Inc. All Rights Reserved.
+
+
+#include "MuR/OpImagePixelFormat.h"
+
+#include "ImageRLE.h"
+#include "BlockCompression/BlockCompression.h"
+
+#include "MuR/MutableTrace.h"
+
+
+namespace mu
+{
+
+
+//-------------------------------------------------------------------------------------------------
+ImagePtr ImagePixelFormat( int imageCompressionQuality, const Image* pBase,
+                               EImageFormat targetFormat, int onlyLOD )
+{
+	MUTABLE_CPUPROFILER_SCOPE(ImagePixelFormat);
+
+    vec2<int> resultSize;
+    int resultLODCount = 0;
+
+    if (onlyLOD==-1)
+    {
+        resultSize[0] = pBase->GetSizeX();
+        resultSize[1] = pBase->GetSizeY();
+        resultLODCount = pBase->GetLODCount();
+    }
+    else
+    {
+        resultSize = pBase->CalculateMipSize( onlyLOD );
+        resultLODCount = 1;
+    }
+
+    ImagePtr result = new Image( (uint16_t)resultSize[0], (uint16_t)resultSize[1], resultLODCount, targetFormat );
+	result->m_flags = pBase->m_flags;
+	bool bSuccess = ImagePixelFormatInPlace( imageCompressionQuality, result.get(), pBase, onlyLOD );
+	int32 OriginalDataSize = FMath::Max(result->m_data.Num(), pBase->m_data.Num());
+	int32 ExcessDataSize = OriginalDataSize * 4;
+	while (!bSuccess)
+    {
+		MUTABLE_CPUPROFILER_SCOPE(Recompression_OutOfSpace);
+
+        // Bad case, where the RLE compressed data requires more memory than the uncompressed data.
+		// We need to support it anyway for small mips or scaled images.
+		result->m_data.SetNum(ExcessDataSize);
+		bSuccess = ImagePixelFormatInPlace(imageCompressionQuality, result.get(), pBase, onlyLOD);
+		ExcessDataSize *= 4;
+	}
+    return result;
+}
+
+
+//-------------------------------------------------------------------------------------------------
+bool ImagePixelFormatInPlace( int imageCompressionQuality, Image* pResult, const Image* pBase,
+                                  int onlyLOD )
+{
+	MUTABLE_CPUPROFILER_SCOPE(ImagePixelFormatInPlace);
+
+    bool success = true;
+
+    vec2<int> resultSize;
+    int resultLODCount = 0;
+    const uint8_t* pBaseBuf = 0;
+    int pixelCount = 0;
+    int baseLOD = 0;
+
+    if (onlyLOD==-1)
+    {
+        pixelCount = (int)pBase->CalculatePixelCount();
+        pBaseBuf = pBase->GetData();
+        resultSize[0] = pBase->GetSizeX();
+        resultSize[1] = pBase->GetSizeY();
+        resultLODCount = pBase->GetLODCount();
+        baseLOD = 0;
+    }
+    else
+    {
+        pixelCount = (int)pBase->CalculatePixelCount( onlyLOD );
+        pBaseBuf = pBase->GetMipData( onlyLOD );
+        resultSize = pBase->CalculateMipSize( onlyLOD );
+        resultLODCount = 1;
+        baseLOD = onlyLOD;
+    }
+
+    check( pResult->GetSizeX()==resultSize[0] && pResult->GetSizeY()==resultSize[1] );
+    check( pResult->GetLODCount()==resultLODCount );
+
+    if ( pResult->GetFormat() == pBase->GetFormat()
+         && onlyLOD==-1 )
+    {
+        // Shouldn't really happen
+        pResult->m_data = pBase->m_data;
+    }
+    else
+    {
+        switch ( pResult->GetFormat() )
+        {
+
+        case EImageFormat::IF_L_UBYTE:
+        {
+            uint8_t* pDestBuf = pResult->GetData();
+
+            switch ( pBase->GetFormat() )
+            {
+
+            case EImageFormat::IF_L_UBYTE:
+            {
+				FMemory::Memcpy( pDestBuf, pBaseBuf, pixelCount );
+                break;
+            }
+
+            case EImageFormat::IF_L_UBYTE_RLE:
+            {
+                if( onlyLOD == -1 )
+                {
+                    UncompressRLE_L( pBase, pResult );
+                }
+                else
+                {
+                    UncompressRLE_L( resultSize[0], resultSize[1], pBaseBuf, pDestBuf );
+                }
+                break;
+            }
+
+            case EImageFormat::IF_L_UBIT_RLE:
+            {
+                if( onlyLOD == -1 )
+                {
+                    UncompressRLE_L1( pBase, pResult );
+                }
+                else
+                {
+                    UncompressRLE_L1( resultSize[0], resultSize[1], pBaseBuf, pDestBuf );
+                }
+                break;
+            }
+
+            case EImageFormat::IF_RGB_UBYTE:
+            {
+                for ( int i=0; i<pixelCount; ++i )
+                {
+                    unsigned result = 76 * pBaseBuf[3*i+0]
+                            + 150 * pBaseBuf[3*i+1]
+                            + 29 * pBaseBuf[3*i+2];
+                    pDestBuf[i] = (uint8_t)std::min( 255u, result >> 8 );
+                }
+                break;
+            }
+
+            case EImageFormat::IF_RGB_UBYTE_RLE:
+            {
+                ImagePtr pTempBase =
+                    ImagePixelFormat( imageCompressionQuality, pBase, EImageFormat::IF_RGB_UBYTE );
+                if (onlyLOD==-1)
+                {
+                    pBaseBuf = pTempBase->GetData();
+                }
+                else
+                {
+                    pBaseBuf = pTempBase->GetMipData( onlyLOD );
+                }
+                for ( int i=0; i<pixelCount; ++i )
+                {
+                    unsigned result = 76 * pBaseBuf[3*i+0]
+                            + 150 * pBaseBuf[3*i+1]
+                            + 29 * pBaseBuf[3*i+2];
+                    pDestBuf[i] = (uint8_t)std::min( 255u, result >> 8 );
+                }
+                break;
+            }
+
+            case EImageFormat::IF_RGBA_UBYTE:
+            {
+                for ( int i=0; i<pixelCount; ++i )
+                {
+                    unsigned result = 76 * pBaseBuf[4*i+0]
+                            + 150 * pBaseBuf[4*i+1]
+                            + 29 * pBaseBuf[4*i+2];
+                    pDestBuf[i] = (uint8_t)std::min( 255u, result >> 8 );
+                }
+                break;
+            }
+
+            case EImageFormat::IF_RGBA_UBYTE_RLE:
+            {
+                ImagePtr pTempBase =
+                    ImagePixelFormat( imageCompressionQuality, pBase, EImageFormat::IF_RGBA_UBYTE );
+                if (onlyLOD==-1)
+                {
+                    pBaseBuf = pTempBase->GetData();
+                }
+                else
+                {
+                    pBaseBuf = pTempBase->GetMipData( onlyLOD );
+                }
+                for ( int i=0; i<pixelCount; ++i )
+                {
+                    unsigned result = 76 * pBaseBuf[4*i+0]
+                            + 150 * pBaseBuf[4*i+1]
+                            + 29 * pBaseBuf[4*i+2];
+                    pDestBuf[i] = (uint8_t)std::min( 255u, result >> 8 );
+                }
+                break;
+            }
+
+            case EImageFormat::IF_BGRA_UBYTE:
+            {
+                for ( int i=0; i<pixelCount; ++i )
+                {
+                    unsigned result = 76 * pBaseBuf[4*i+2]
+                            + 150 * pBaseBuf[4*i+1]
+                            + 29 * pBaseBuf[4*i+0];
+                    pDestBuf[i] = (uint8_t)std::min( 255u, result >> 8 );
+                }
+                break;
+            }
+
+            case EImageFormat::IF_BC1:
+            case EImageFormat::IF_BC2:
+            case EImageFormat::IF_BC3:
+            {
+                for (int m = 0; m < resultLODCount; ++m)
+                {
+                    ImagePtr pTempBase =
+                        ImagePixelFormat( imageCompressionQuality, pBase, EImageFormat::IF_RGB_UBYTE );
+                    if (onlyLOD == -1)
+                    {
+                        pBaseBuf = pTempBase->GetData();
+                    }
+                    else
+                    {
+                        pBaseBuf = pTempBase->GetMipData(onlyLOD);
+                    }
+                    for (int i = 0; i < pixelCount; ++i)
+                    {
+                        unsigned result = 76 * pBaseBuf[3 * i + 0] + 150 * pBaseBuf[3 * i + 1] +
+                                          29 * pBaseBuf[3 * i + 2];
+                        pDestBuf[i] = (uint8_t)std::min(255u, result >> 8);
+                    }
+                }
+                break;
+            }
+
+            default:
+                // Case not implemented
+                check( false );
+
+            }
+
+            break;
+        }
+
+        case EImageFormat::IF_L_UBYTE_RLE:
+        {
+            switch ( pBase->GetFormat() )
+            {
+
+            case EImageFormat::IF_L_UBYTE:
+            {
+                check( onlyLOD == -1 );
+
+                // Try to compress
+				if (!pResult->m_data.Num())
+				{
+					// Allocate memory for the compressed data. TODO: Smaller?
+					pResult->m_data.SetNum(pBase->m_data.Num());
+				}
+
+                success = CompressRLE_L( pBase, pResult );
+
+#ifdef MUTABLE_DEBUG_RLE		
+				if (success)
+                {
+                    // verify
+                    ImagePtr pTest = ImagePixelFormat( imageCompressionQuality, pResult, EImageFormat::IF_L_UBYTE );
+					check(!FMemory::Memcmp(pTest->GetData(), pBase->GetData(), pBase->GetDataSize()));
+                }
+#endif
+
+                break;
+            }
+
+            case EImageFormat::IF_RGB_UBYTE:
+            case EImageFormat::IF_RGBA_UBYTE:
+            case EImageFormat::IF_BGRA_UBYTE:
+            {
+                check( onlyLOD == -1 );
+                ImagePtr pTemp = ImagePixelFormat( imageCompressionQuality, pBase, EImageFormat::IF_L_UBYTE );
+ 
+				// Try to compress
+				if (!pResult->m_data.Num())
+				{
+					// Allocate memory for the compressed data. TODO: Smaller?
+					pResult->m_data.SetNum(pTemp->m_data.Num());
+				}
+
+                success = CompressRLE_L( pTemp.get(), pResult );
+
+#ifdef MUTABLE_DEBUG_RLE		
+				if (success)
+                {
+                    // verify
+                    ImagePtr pTest = ImagePixelFormat( imageCompressionQuality, pResult, EImageFormat::IF_L_UBYTE );
+                    check( !FMemory::Memcmp(pTest->GetData(), pTemp->GetData(), pTemp->GetDataSize() ) );
+                }
+#endif
+
+                break;
+            }
+
+            default:
+                // Case not implemented
+                check( false );
+
+            }
+
+            break;
+        }
+
+        case EImageFormat::IF_L_UBIT_RLE:
+        {
+            switch ( pBase->GetFormat() )
+            {
+
+            case EImageFormat::IF_L_UBYTE:
+            {
+                check( onlyLOD == -1 );
+
+				// Allocate memory for the compressed data. TODO: Smaller?
+				if (!pResult->m_data.Num())
+				{
+					// Allocate memory for the compressed data. TODO: Smaller?
+					pResult->m_data.SetNum(pBase->m_data.Num());
+				}
+				
+				success = CompressRLE_L1( pBase, pResult );
+                
+#ifdef MUTABLE_DEBUG_RLE		
+				if (success)
+                {
+					// Verify
+                    ImagePtr pTest = ImagePixelFormat( imageCompressionQuality, pResult, EImageFormat::IF_L_UBYTE );
+                    check( !FMemory::Memcmp(pTest->GetData(), pBase->GetData(), pBase->GetDataSize() ) );
+                }
+#endif
+
+                break;
+            }
+
+            case EImageFormat::IF_RGB_UBYTE:
+            case EImageFormat::IF_RGBA_UBYTE:
+            case EImageFormat::IF_BGRA_UBYTE:
+            {
+                check( onlyLOD == -1 );
+                ImagePtr pTemp = ImagePixelFormat( imageCompressionQuality, pBase, EImageFormat::IF_L_UBYTE );
+				success = CompressRLE_L1( pTemp.get(), pResult );
+
+#ifdef MUTABLE_DEBUG_RLE		
+				if (success)
+				{
+					// Verify
+					ImagePtr pTest = ImagePixelFormat(imageCompressionQuality, pResult, EImageFormat::IF_L_UBYTE);
+					check(!FMemory::Memcmp(pTest->GetData(), pTemp->GetData(), pTemp->GetDataSize()));
+				}
+#endif
+
+                break;
+            }
+
+            default:
+                // Case not implemented
+                check( false );
+
+            }
+
+            break;
+        }
+
+        case EImageFormat::IF_RGB_UBYTE:
+        {
+            uint8_t* pDestBuf = pResult->GetData();
+
+            switch ( pBase->GetFormat() )
+            {
+
+            case EImageFormat::IF_L_UBYTE:
+            {
+                for ( int i=0; i<pixelCount; ++i )
+                {
+                    pDestBuf[i*3+0] = pBaseBuf[i];
+                    pDestBuf[i*3+1] = pBaseBuf[i];
+                    pDestBuf[i*3+2] = pBaseBuf[i];
+                }
+                break;
+            }
+
+            case EImageFormat::IF_RGB_UBYTE:
+            {
+				FMemory::Memcpy( pDestBuf, pBaseBuf, 3*pixelCount );
+                break;
+            }
+
+            case EImageFormat::IF_RGBA_UBYTE:
+            {
+                for ( int i=0; i<pixelCount; ++i )
+                {
+                    pDestBuf[i*3+0] = pBaseBuf[i*4+0];
+                    pDestBuf[i*3+1] = pBaseBuf[i*4+1];
+                    pDestBuf[i*3+2] = pBaseBuf[i*4+2];
+                }
+                break;
+            }
+
+            case EImageFormat::IF_BGRA_UBYTE:
+            {
+                for ( int i=0; i<pixelCount; ++i )
+                {
+                    pDestBuf[i*3+0] = pBaseBuf[i*4+2];
+                    pDestBuf[i*3+1] = pBaseBuf[i*4+1];
+                    pDestBuf[i*3+2] = pBaseBuf[i*4+0];
+                }
+                break;
+            }
+
+            case EImageFormat::IF_RGB_UBYTE_RLE:
+            {
+                UncompressRLE_RGB( pBase, pResult );
+                break;
+            }
+
+            case EImageFormat::IF_RGBA_UBYTE_RLE:
+            {
+                check( onlyLOD == -1 );
+                ImagePtr pTemp = ImagePixelFormat( imageCompressionQuality, pBase, EImageFormat::IF_RGBA_UBYTE );
+                const uint8_t* pTempBuf = pTemp->GetData();
+                for ( int i=0; i<pixelCount; ++i )
+                {
+                    pDestBuf[i*3+0] = pTempBuf[i*4+0];
+                    pDestBuf[i*3+1] = pTempBuf[i*4+1];
+                    pDestBuf[i*3+2] = pTempBuf[i*4+2];
+                }
+                break;
+            }
+
+            case EImageFormat::IF_BC1:
+            {
+                for ( int m=0; m<resultLODCount; ++m )
+                {
+                    vec2<int> mipSize = pResult->CalculateMipSize( m );
+                    miro::BC1_to_RGB( mipSize[0], mipSize[1],
+                                        pBase->GetMipData(baseLOD+m),
+                                        pResult->GetMipData(m) );
+                }
+                break;
+            }
+
+            case EImageFormat::IF_BC2:
+            {
+                for ( int m=0; m<resultLODCount; ++m )
+                {
+                    vec2<int> mipSize = pResult->CalculateMipSize( m );
+                    miro::BC2_to_RGB( mipSize[0], mipSize[1],
+                                        pBase->GetMipData(baseLOD+m),
+                                        pResult->GetMipData(m) );
+                }
+                break;
+            }
+
+            case EImageFormat::IF_BC3:
+            {
+                for ( int m=0; m<resultLODCount; ++m )
+                {
+                    vec2<int> mipSize = pResult->CalculateMipSize( m );
+                    miro::BC3_to_RGB( mipSize[0], mipSize[1],
+                                        pBase->GetMipData(baseLOD+m),
+                                        pResult->GetMipData(m) );
+                }
+                break;
+            }
+
+            case EImageFormat::IF_BC4:
+            {
+                for ( int m=0; m<resultLODCount; ++m )
+                {
+                    vec2<int> mipSize = pResult->CalculateMipSize( m );
+                    miro::BC4_to_RGB( mipSize[0], mipSize[1],
+                                      pBase->GetMipData(baseLOD+m),
+                                      pResult->GetMipData(m) );
+                }
+                break;
+            }
+
+            case EImageFormat::IF_BC5:
+            {
+                for ( int m=0; m<resultLODCount; ++m )
+                {
+                    vec2<int> mipSize = pResult->CalculateMipSize( m );
+                    miro::BC5_to_RGB( mipSize[0], mipSize[1],
+                                      pBase->GetMipData(baseLOD+m),
+                                      pResult->GetMipData(m) );
+                }
+                break;
+            }
+
+            case EImageFormat::IF_ASTC_4x4_RGB_LDR:
+            {
+                for ( int m=0; m<resultLODCount; ++m )
+                {
+                    vec2<int> mipSize = pResult->CalculateMipSize( m );
+                    miro::ASTC4x4RGBL_to_RGB( mipSize[0], mipSize[1],
+                                      pBase->GetMipData(baseLOD+m),
+                                      pResult->GetMipData(m) );
+                }
+                break;
+            }
+
+            case EImageFormat::IF_ASTC_4x4_RGBA_LDR:
+            {
+                for ( int m=0; m<resultLODCount; ++m )
+                {
+                    vec2<int> mipSize = pResult->CalculateMipSize( m );
+                    miro::ASTC4x4RGBAL_to_RGB( mipSize[0], mipSize[1],
+                                      pBase->GetMipData(baseLOD+m),
+                                      pResult->GetMipData(m) );
+                }
+                break;
+            }
+
+            case EImageFormat::IF_ASTC_4x4_RG_LDR:
+            {
+                for ( int m=0; m<resultLODCount; ++m )
+                {
+                    vec2<int> mipSize = pResult->CalculateMipSize( m );
+                    miro::ASTC4x4RGL_to_RGB( mipSize[0], mipSize[1],
+                                      pBase->GetMipData(baseLOD+m),
+                                      pResult->GetMipData(m) );
+                }
+                break;
+            }
+
+            default:
+                // Case not implemented
+                check( false );
+
+            }
+
+            break;
+        }
+
+        case EImageFormat::IF_BGRA_UBYTE:
+        {
+            // TODO: Optimise
+            pResult->m_format = EImageFormat::IF_RGBA_UBYTE;
+            bool bSuccess = ImagePixelFormatInPlace( imageCompressionQuality, pResult, pBase, onlyLOD );
+			check(bSuccess);
+            pResult->m_format = EImageFormat::IF_BGRA_UBYTE;
+
+            uint8_t* pDestBuf = pResult->GetData();
+            for ( int i=0; i<pixelCount; ++i )
+            {
+                uint8_t temp = pDestBuf[i*4+0];
+                pDestBuf[i*4+0] = pDestBuf[i*4+2];
+                pDestBuf[i*4+2] = temp;
+            }
+            break;
+        }
+
+        case EImageFormat::IF_RGBA_UBYTE:
+        {
+            uint8_t* pDestBuf = pResult->GetData();
+
+            switch ( pBase->GetFormat() )
+            {
+
+            case EImageFormat::IF_RGBA_UBYTE_RLE:
+            {
+                check( onlyLOD == -1 );
+                UncompressRLE_RGBA( pBase, pResult );
+                break;
+            }
+
+            case EImageFormat::IF_L_UBYTE:
+            {
+                for ( int i=0; i<pixelCount; ++i )
+                {
+                    pDestBuf[i*4+0] = pBaseBuf[i];
+                    pDestBuf[i*4+1] = pBaseBuf[i];
+                    pDestBuf[i*4+2] = pBaseBuf[i];
+                    pDestBuf[i*4+3] = 255;
+                }
+                break;
+            }
+
+            case EImageFormat::IF_RGB_UBYTE:
+            {
+                for ( int i=0; i<pixelCount; ++i )
+                {
+                    pDestBuf[i*4+0] = pBaseBuf[i*3+0];
+                    pDestBuf[i*4+1] = pBaseBuf[i*3+1];
+                    pDestBuf[i*4+2] = pBaseBuf[i*3+2];
+                    pDestBuf[i*4+3] = 255;
+                }
+                break;
+            }
+
+            case EImageFormat::IF_RGBA_UBYTE:
+            {
+				FMemory::Memcpy( pDestBuf, pBaseBuf, 4*pixelCount );
+                break;
+            }
+
+            case EImageFormat::IF_BGRA_UBYTE:
+            {
+                for ( int i=0; i<pixelCount; ++i )
+                {
+                    pDestBuf[i*4+0] = pBaseBuf[i*4+2];
+                    pDestBuf[i*4+1] = pBaseBuf[i*4+1];
+                    pDestBuf[i*4+2] = pBaseBuf[i*4+0];
+                    pDestBuf[i*4+3] = pBaseBuf[i*4+3];
+                }
+                break;
+            }
+
+            case EImageFormat::IF_BC1:
+            {
+                for ( int m=0; m<resultLODCount; ++m )
+                {
+                    vec2<int> mipSize = pResult->CalculateMipSize( m );
+                    miro::BC1_to_RGBA( mipSize[0], mipSize[1],
+                                        pBase->GetMipData(baseLOD+m),
+                                        pResult->GetMipData(m) );
+                }
+                break;
+            }
+
+            case EImageFormat::IF_BC2:
+            {
+                for ( int m=0; m<resultLODCount; ++m )
+                {
+                    vec2<int> mipSize = pResult->CalculateMipSize( m );
+                    miro::BC2_to_RGBA( mipSize[0], mipSize[1],
+                                        pBase->GetMipData(baseLOD+m),
+                                        pResult->GetMipData(m) );
+                }
+                break;
+            }
+
+            case EImageFormat::IF_BC3:
+            {
+                for ( int m=0; m<resultLODCount; ++m )
+                {
+                    vec2<int> mipSize = pResult->CalculateMipSize( m );
+                    miro::BC3_to_RGBA( mipSize[0], mipSize[1],
+                                        pBase->GetMipData(baseLOD+m),
+                                        pResult->GetMipData(m) );
+                }
+                break;
+            }
+
+            case EImageFormat::IF_BC4:
+            {
+                for ( int m=0; m<resultLODCount; ++m )
+                {
+                    vec2<int> mipSize = pResult->CalculateMipSize( m );
+                    miro::BC4_to_RGBA( mipSize[0], mipSize[1],
+                                       pBase->GetMipData(baseLOD+m),
+                                       pResult->GetMipData(m) );
+                }
+                break;
+            }
+
+            case EImageFormat::IF_BC5:
+            {
+                for ( int m=0; m<resultLODCount; ++m )
+                {
+                    vec2<int> mipSize = pResult->CalculateMipSize( m );
+                    miro::BC5_to_RGBA( mipSize[0], mipSize[1],
+                                        pBase->GetMipData(baseLOD+m),
+                                        pResult->GetMipData(m) );
+                }
+                break;
+            }
+
+            case EImageFormat::IF_ASTC_4x4_RGB_LDR:
+            {
+                for ( int m=0; m<resultLODCount; ++m )
+                {
+                    vec2<int> mipSize = pResult->CalculateMipSize( m );
+                    miro::ASTC4x4RGBL_to_RGBA( mipSize[0], mipSize[1],
+                                        pBase->GetMipData(baseLOD+m),
+                                        pResult->GetMipData(m) );
+                }
+                break;
+            }
+
+            case EImageFormat::IF_ASTC_4x4_RGBA_LDR:
+            {
+                for ( int m=0; m<resultLODCount; ++m )
+                {
+                    vec2<int> mipSize = pResult->CalculateMipSize( m );
+                    miro::ASTC4x4RGBAL_to_RGBA( mipSize[0], mipSize[1],
+                                        pBase->GetMipData(baseLOD+m),
+                                        pResult->GetMipData(m) );
+                }
+                break;
+            }
+
+            case EImageFormat::IF_ASTC_4x4_RG_LDR:
+            {
+                for ( int m=0; m<resultLODCount; ++m )
+                {
+                    vec2<int> mipSize = pResult->CalculateMipSize( m );
+                    miro::ASTC4x4RGL_to_RGBA( mipSize[0], mipSize[1],
+                                        pBase->GetMipData(baseLOD+m),
+                                        pResult->GetMipData(m) );
+                }
+                break;
+            }
+
+			case EImageFormat::IF_RGB_UBYTE_RLE:
+			{
+				check(onlyLOD == -1);
+				ImagePtr pTemp = ImagePixelFormat(imageCompressionQuality, pBase, EImageFormat::IF_RGB_UBYTE);
+				const uint8_t* pTempData = pTemp->GetData();
+				for (int i = 0; i < pixelCount; ++i)
+				{
+					pDestBuf[i * 4 + 0] = pTempData[i * 3 + 0];
+					pDestBuf[i * 4 + 1] = pTempData[i * 3 + 1];
+					pDestBuf[i * 4 + 2] = pTempData[i * 3 + 2];
+					pDestBuf[i * 4 + 3] = 255;
+				}
+				break;
+			}
+
+			case EImageFormat::IF_L_UBYTE_RLE:
+			{
+				check(onlyLOD == -1);
+				ImagePtr pTemp = ImagePixelFormat(imageCompressionQuality, pBase, EImageFormat::IF_L_UBYTE);
+				const uint8_t* pTempData = pTemp->GetData();
+				for (int i = 0; i < pixelCount; ++i)
+				{
+					pDestBuf[i * 4 + 0] = pTempData[i];
+					pDestBuf[i * 4 + 1] = pTempData[i];
+					pDestBuf[i * 4 + 2] = pTempData[i];
+					pDestBuf[i * 4 + 3] = 255;
+				}
+				break;
+			}
+
+            default:
+                // Case not implemented
+                check( false );
+
+            }
+
+            break;
+        }
+
+
+        case EImageFormat::IF_RGBA_UBYTE_RLE:
+        {
+            switch ( pBase->GetFormat() )
+            {
+
+            case EImageFormat::IF_RGBA_UBYTE:
+            {
+                check( onlyLOD == -1 );
+                CompressRLE_RGBA( pBase, pResult );
+
+                // TODO: TEST BLEH
+                //                ImagePtr pTest = ImagePixelFormat( imageCompressionQuality,
+                //                pResult, EImageFormat::IF_RGBA_UBYTE ); check(
+                //                !memcmp(pTest->GetData(), pBase->GetData(),
+                //                pBase->GetDataSize() ) );
+                break;
+            }
+
+            case EImageFormat::IF_RGB_UBYTE:
+            {
+                check( onlyLOD == -1 );
+
+                // \todo: optimise
+                ImagePtr pTemp = ImagePixelFormat( imageCompressionQuality, pBase, EImageFormat::IF_RGBA_UBYTE );
+                CompressRLE_RGBA( pTemp.get(), pResult );
+
+                // TODO: TEST BLEH
+                //                ImagePtr pTest = ImagePixelFormat( imageCompressionQuality,
+                //                pResult, EImageFormat::IF_RGBA_UBYTE ); check(
+                //                !memcmp(pTest->GetData(), pBase->GetData(),
+                //                pBase->GetDataSize() ) );
+                break;
+            }
+
+            case EImageFormat::IF_RGB_UBYTE_RLE:
+            {
+                check( onlyLOD == -1 );
+
+                // \todo: optimise
+                ImagePtr pTemp1 = ImagePixelFormat( imageCompressionQuality, pBase, EImageFormat::IF_RGB_UBYTE );
+                ImagePtr pTemp2 =
+                    ImagePixelFormat( imageCompressionQuality, pTemp1.get(), EImageFormat::IF_RGBA_UBYTE );
+                CompressRLE_RGBA( pTemp2.get(), pResult );
+                break;
+            }
+
+            default:
+                // Case not implemented
+                check( false );
+
+            }
+
+            break;
+        }
+
+
+        case EImageFormat::IF_RGB_UBYTE_RLE:
+        {
+            switch ( pBase->GetFormat() )
+            {
+
+            case EImageFormat::IF_RGB_UBYTE:
+            {
+                check( onlyLOD == -1 );
+                CompressRLE_RGB( pBase, pResult );
+
+                // TODO: TEST BLEH
+                // ImagePtr pTest = ImagePixelFormat( imageCompressionQuality, pDest.get(),
+                // EImageFormat::IF_RGB_UBYTE ); check( !memcmp(pTest->GetData(), pBase->GetData(),
+                // pBase->GetDataSize() ) );
+                break;
+            }
+
+            case EImageFormat::IF_RGBA_UBYTE_RLE:
+            {
+                check( onlyLOD == -1 );
+
+                // \todo: optimise
+                ImagePtr pTemp1 = ImagePixelFormat( imageCompressionQuality, pBase, EImageFormat::IF_RGBA_UBYTE );
+                ImagePtr pTemp2 =
+                    ImagePixelFormat( imageCompressionQuality, pTemp1.get(), EImageFormat::IF_RGB_UBYTE );
+                CompressRLE_RGB( pTemp2.get(), pResult );
+                break;
+            }
+
+            default:
+                // Case not implemented
+                check( false );
+
+            }
+
+            break;
+        }
+
+        case EImageFormat::IF_BC1:
+        {
+            switch ( pBase->GetFormat() )
+            {
+            case EImageFormat::IF_RGB_UBYTE:
+                for ( int m=0; m<resultLODCount; ++m )
+                {
+                    vec2<int> mipSize = pResult->CalculateMipSize( m );
+                    miro::RGB_to_BC1(
+                        mipSize[0], mipSize[1], pBase->GetMipData( baseLOD + m ),
+                        pResult->GetMipData( m ), imageCompressionQuality );
+                }
+                break;
+            case EImageFormat::IF_RGBA_UBYTE:
+                for ( int m=0; m<resultLODCount; ++m )
+                {
+                    vec2<int> mipSize = pResult->CalculateMipSize( m );
+                    miro::RGBA_to_BC1(
+                        mipSize[0], mipSize[1], pBase->GetMipData( baseLOD + m ),
+                        pResult->GetMipData( m ), imageCompressionQuality );
+                }
+                break;
+
+            case EImageFormat::IF_L_UBYTE:
+            {
+                ImagePtr pTemp = ImagePixelFormat( imageCompressionQuality, pBase, EImageFormat::IF_RGB_UBYTE );
+                for ( int m=0; m<resultLODCount; ++m )
+                {
+                    vec2<int> mipSize = pResult->CalculateMipSize( m );
+                    miro::RGB_to_BC1(
+                        mipSize[0], mipSize[1], pTemp->GetMipData( baseLOD + m ),
+                        pResult->GetMipData( m ), imageCompressionQuality );
+                }
+                break;
+            }
+
+            case EImageFormat::IF_BC3:
+            {
+                ImagePtr pTemp = ImagePixelFormat( imageCompressionQuality, pBase, EImageFormat::IF_RGBA_UBYTE );
+                for (int m = 0; m < resultLODCount; ++m)
+                {
+                    vec2<int> mipSize = pResult->CalculateMipSize(m);
+                    miro::RGBA_to_BC1(
+                        mipSize[0], mipSize[1], pTemp->GetMipData( baseLOD + m ),
+                        pResult->GetMipData( m ), imageCompressionQuality );
+                }
+                break;
+            }
+
+            default:
+                // Case not implemented
+                check( false );
+
+            }
+
+            break;
+        }
+
+        case EImageFormat::IF_BC2:
+        {
+            switch ( pBase->GetFormat() )
+            {
+            case EImageFormat::IF_RGBA_UBYTE:
+                for ( int m=0; m<resultLODCount; ++m )
+                {
+                    vec2<int> mipSize = pResult->CalculateMipSize( m );
+                    miro::RGBA_to_BC2(
+                        mipSize[0], mipSize[1], pBase->GetMipData( baseLOD + m ),
+                        pResult->GetMipData( m ), imageCompressionQuality );
+                }
+                break;
+
+            default:
+                // Case not implemented
+                check( false );
+
+            }
+
+            break;
+        }
+
+        case EImageFormat::IF_BC3:
+        {
+            switch ( pBase->GetFormat() )
+            {
+            case EImageFormat::IF_RGBA_UBYTE:
+                for ( int m=0; m<resultLODCount; ++m )
+                {
+                    vec2<int> mipSize = pResult->CalculateMipSize( m );
+                    miro::RGBA_to_BC3(
+                        mipSize[0], mipSize[1], pBase->GetMipData( baseLOD + m ),
+                        pResult->GetMipData( m ), imageCompressionQuality );
+                }
+                break;
+
+            case EImageFormat::IF_RGB_UBYTE:          
+                for ( int m=0; m<resultLODCount; ++m )
+                {
+                    vec2<int> mipSize = pResult->CalculateMipSize( m );
+                    miro::RGB_to_BC3(
+                        mipSize[0], mipSize[1], pBase->GetMipData( baseLOD + m ),
+                        pResult->GetMipData( m ), imageCompressionQuality );
+                }
+                break;
+
+            case EImageFormat::IF_BC1:
+            {
+                MUTABLE_CPUPROFILER_SCOPE(BC1toBC3);
+
+                for ( int m=0; m<resultLODCount; ++m )
+                {
+                    vec2<int> mipSize = pResult->CalculateMipSize( m );
+                    miro::BC1_to_BC3(
+                        mipSize[0], mipSize[1], pBase->GetMipData( baseLOD + m ),
+                        pResult->GetMipData( m ), imageCompressionQuality );
+                }
+                break;
+            }
+
+            default:
+                // Case not implemented
+                check( false );
+
+            }
+
+            break;
+        }
+
+        case EImageFormat::IF_BC4:
+        {
+            switch ( pBase->GetFormat() )
+            {
+            case EImageFormat::IF_L_UBYTE:
+                for ( int m=0; m<resultLODCount; ++m )
+                {
+                    vec2<int> mipSize = pResult->CalculateMipSize( m );
+                    miro::L_to_BC4(
+                        mipSize[0], mipSize[1], pBase->GetMipData( baseLOD + m ),
+                        pResult->GetMipData( m ), imageCompressionQuality );
+                }
+                break;
+
+            default:
+                // Case not implemented
+                check( false );
+
+            }
+
+            break;
+        }
+
+        case EImageFormat::IF_BC5:
+        {
+            switch ( pBase->GetFormat() )
+            {
+            case EImageFormat::IF_RGBA_UBYTE:
+                for ( int m=0; m<resultLODCount; ++m )
+                {
+                    vec2<int> mipSize = pResult->CalculateMipSize( m );
+                    miro::RGBA_to_BC5(
+                        mipSize[0], mipSize[1], pBase->GetMipData( baseLOD + m ),
+                        pResult->GetMipData( m ), imageCompressionQuality );
+                }
+                break;
+
+            case EImageFormat::IF_RGB_UBYTE:
+                for ( int m=0; m<resultLODCount; ++m )
+                {
+                    vec2<int> mipSize = pResult->CalculateMipSize( m );
+                    miro::RGB_to_BC5(
+                        mipSize[0], mipSize[1], pBase->GetMipData( baseLOD + m ),
+                        pResult->GetMipData( m ), imageCompressionQuality );
+                }
+                break;
+
+            default:
+                // Case not implemented
+                check( false );
+
+            }
+
+            break;
+        }
+
+        case EImageFormat::IF_ASTC_4x4_RGB_LDR:
+        {
+            switch ( pBase->GetFormat() )
+            {
+            case EImageFormat::IF_RGBA_UBYTE:
+                for ( int m=0; m<resultLODCount; ++m )
+                {
+                    vec2<int> mipSize = pResult->CalculateMipSize( m );
+                    miro::RGBA_to_ASTC4x4RGBL(
+                        mipSize[0], mipSize[1], pBase->GetMipData( baseLOD + m ),
+                        pResult->GetMipData( m ), imageCompressionQuality );
+                }
+                break;
+
+            case EImageFormat::IF_RGB_UBYTE:
+                for ( int m=0; m<resultLODCount; ++m )
+                {
+                    vec2<int> mipSize = pResult->CalculateMipSize( m );
+                    miro::RGB_to_ASTC4x4RGBL(
+                        mipSize[0], mipSize[1], pBase->GetMipData( baseLOD + m ),
+                        pResult->GetMipData( m ), imageCompressionQuality );
+                }
+                break;
+
+            case EImageFormat::IF_ASTC_4x4_RGBA_LDR:
+                for ( int m=0; m<resultLODCount; ++m )
+                {
+                    vec2<int> mipSize = pResult->CalculateMipSize( m );
+                    miro::ASTC4x4RGBAL_to_ASTC4x4RGBL( mipSize[0], mipSize[1],
+                                       pBase->GetMipData(baseLOD+m),
+                                       pResult->GetMipData(m) );
+                }
+                break;
+
+			case EImageFormat::IF_L_UBYTE:
+				for (int m = 0; m < resultLODCount; ++m)
+				{
+					vec2<int> mipSize = pResult->CalculateMipSize(m);
+					miro::L_to_ASTC4x4RGBL(mipSize[0], mipSize[1],
+						pBase->GetMipData(baseLOD + m),
+						pResult->GetMipData(m), imageCompressionQuality);
+				}
+				break;
+
+            default:
+                // Case not implemented
+                check( false );
+
+            }
+
+            break;
+        }
+
+        case EImageFormat::IF_ASTC_4x4_RGBA_LDR:
+        {
+            switch ( pBase->GetFormat() )
+            {
+            case EImageFormat::IF_RGBA_UBYTE:
+                for ( int m=0; m<resultLODCount; ++m )
+                {
+                    vec2<int> mipSize = pResult->CalculateMipSize( m );
+                    miro::RGBA_to_ASTC4x4RGBAL(
+                        mipSize[0], mipSize[1], pBase->GetMipData( baseLOD + m ),
+                        pResult->GetMipData( m ), imageCompressionQuality );
+                }
+                break;
+
+            case EImageFormat::IF_RGB_UBYTE:
+                for ( int m=0; m<resultLODCount; ++m )
+                {
+                    vec2<int> mipSize = pResult->CalculateMipSize( m );
+                    miro::RGB_to_ASTC4x4RGBAL(
+                        mipSize[0], mipSize[1], pBase->GetMipData( baseLOD + m ),
+                        pResult->GetMipData( m ), imageCompressionQuality );
+                }
+                break;
+
+            case EImageFormat::IF_ASTC_4x4_RGB_LDR:
+                for ( int m=0; m<resultLODCount; ++m )
+                {
+                    vec2<int> mipSize = pResult->CalculateMipSize( m );
+                    miro::ASTC4x4RGBL_to_ASTC4x4RGBAL( mipSize[0], mipSize[1],
+                                       pBase->GetMipData(baseLOD+m),
+                                       pResult->GetMipData(m) );
+                }
+                break;
+
+            default:
+                // Case not implemented
+                check( false );
+
+            }
+
+            break;
+        }
+
+        case EImageFormat::IF_ASTC_4x4_RG_LDR:
+        {
+            switch ( pBase->GetFormat() )
+            {
+            case EImageFormat::IF_RGBA_UBYTE:
+                for ( int m=0; m<resultLODCount; ++m )
+                {
+                    vec2<int> mipSize = pResult->CalculateMipSize( m );
+                    miro::RGBA_to_ASTC4x4RGL(
+                        mipSize[0], mipSize[1], pBase->GetMipData( baseLOD + m ),
+                        pResult->GetMipData( m ), imageCompressionQuality );
+                }
+                break;
+
+            case EImageFormat::IF_RGB_UBYTE:
+                for ( int m=0; m<resultLODCount; ++m )
+                {
+                    vec2<int> mipSize = pResult->CalculateMipSize( m );
+                    miro::RGB_to_ASTC4x4RGL(
+                        mipSize[0], mipSize[1], pBase->GetMipData( baseLOD + m ),
+                        pResult->GetMipData( m ), imageCompressionQuality );
+                }
+                break;
+
+            case EImageFormat::IF_ASTC_4x4_RG_LDR:
+                for ( int m=0; m<resultLODCount; ++m )
+                {
+                    vec2<int> mipSize = pResult->CalculateMipSize( m );
+                    // Hack that actually works because of block size.
+                    miro::ASTC4x4RGBAL_to_ASTC4x4RGBL( mipSize[0], mipSize[1],
+                                       pBase->GetMipData(baseLOD+m),
+                                       pResult->GetMipData(m) );
+                }
+                break;
+
+            default:
+                // Case not implemented
+                check( false );
+
+            }
+
+            break;
+        }
+
+        default:
+            // Case not implemented
+            check( false );
+        }
+    }
+
+    return success;
+}
+
+
+}
diff --git a/Engine/Plugins/Experimental/Mutable/Source/MutableRuntime/Private/MuR/OpImagePixelFormat.h b/Engine/Plugins/Experimental/Mutable/Source/MutableRuntime/Private/MuR/OpImagePixelFormat.h
new file mode 100644
--- /dev/null
+++ b/Engine/Plugins/Experimental/Mutable/Source/MutableRuntime/Private/MuR/OpImagePixelFormat.h
@@ -0,0 +1,27 @@
+// Copyright Epic Games, Inc. All Rights Reserved.
+
+#pragma once
+
+#include "MuR/ImagePrivate.h"
+
+
+namespace mu
+{
+    //---------------------------------------------------------------------------------------------
+    //! Convert an image to another pixel format.
+    //! \warning Not all format conversions are implemented.
+    //! \param onlyLOD If different than -1, only the specified lod level will be converted in the
+    //! returned image.
+    //! \return false if the conversion failed, usually because not enough memory was allocated in
+    //!     the result. This is only checked for RLE compression.
+    //---------------------------------------------------------------------------------------------
+	MUTABLERUNTIME_API extern bool ImagePixelFormatInPlace( int quality, Image* pResult, const Image* pBase,
+                                         int onlyLOD = -1 );
+
+    //---------------------------------------------------------------------------------------------
+    //! Wrapper for the above method that allocates the destination image.
+    //---------------------------------------------------------------------------------------------
+	MUTABLERUNTIME_API extern ImagePtr ImagePixelFormat( int quality, const Image* pBase, EImageFormat targetFormat,
+                                      int onlyLOD = -1 );
+
+}
