Moved the Mutable plugin files to live under a directory named after the module. This allows us to get rid of long include paths that contain public and private folder names in the path and fixes head

diff --git a/Engine/Plugins/Experimental/Mutable/Source/MutableTools/Private/MuT/NodeMeshInterpolate.cpp b/Engine/Plugins/Experimental/Mutable/Source/MutableTools/Private/MuT/NodeMeshInterpolate.cpp
new file mode 100644
--- /dev/null
+++ b/Engine/Plugins/Experimental/Mutable/Source/MutableTools/Private/MuT/NodeMeshInterpolate.cpp
@@ -0,0 +1,170 @@
+// Copyright Epic Games, Inc. All Rights Reserved.
+
+
+#include "MuT/NodeMeshInterpolate.h"
+#include "MuT/NodeMeshInterpolatePrivate.h"
+
+#include "MuT/NodeScalar.h"
+
+#include "MuR/MeshPrivate.h"
+
+
+#define NODE_INPUT_COUNT 	4
+
+
+namespace mu
+{
+
+	//---------------------------------------------------------------------------------------------
+	// Static initialisation
+	//---------------------------------------------------------------------------------------------
+	NODE_TYPE NodeMeshInterpolate::Private::s_type =
+			NODE_TYPE( "MeshInterpolate", NodeMesh::GetStaticType() );
+
+
+	//---------------------------------------------------------------------------------------------
+	//!
+	//---------------------------------------------------------------------------------------------
+
+	MUTABLE_IMPLEMENT_NODE( NodeMeshInterpolate, EType::Interpolate, Node, Node::EType::Mesh)
+
+
+	//---------------------------------------------------------------------------------------------
+	// Node Interface
+	//---------------------------------------------------------------------------------------------
+	int NodeMeshInterpolate::GetInputCount() const
+	{
+		return 1 + (int)m_pD->m_targets.size();
+	}
+
+
+	//---------------------------------------------------------------------------------------------
+	Node* NodeMeshInterpolate::GetInputNode( int i ) const
+	{
+		check( i>=0 && i<GetInputCount() );
+
+		Node* pResult = 0;
+
+		switch (i)
+		{
+		case 0:
+			pResult = m_pD->m_pFactor.get();
+			break;
+
+		default:
+			pResult = m_pD->m_targets[i-1].get();
+			break;
+		}
+
+		return pResult;
+	}
+
+
+	//---------------------------------------------------------------------------------------------
+	void NodeMeshInterpolate::SetInputNode( int i, NodePtr pNode )
+	{
+		check( i>=0 && i<GetInputCount() );
+
+		switch (i)
+		{
+		case 0:
+			m_pD->m_pFactor = dynamic_cast<NodeScalar*>(pNode.get());
+			break;
+
+		default:
+			m_pD->m_targets[i-1] = dynamic_cast<NodeMesh*>(pNode.get());
+			break;
+		}
+	}
+
+
+	//---------------------------------------------------------------------------------------------
+	// Own Interface
+	//---------------------------------------------------------------------------------------------
+	NodeScalarPtr NodeMeshInterpolate::GetFactor() const
+	{
+		return m_pD->m_pFactor.get();
+	}
+
+
+	//---------------------------------------------------------------------------------------------
+	void NodeMeshInterpolate::SetFactor( NodeScalarPtr pNode )
+	{
+		m_pD->m_pFactor = pNode;
+	}
+
+
+	//---------------------------------------------------------------------------------------------
+	NodeMeshPtr NodeMeshInterpolate::GetTarget( int t ) const
+	{
+		check( t>=0 && t<(int)m_pD->m_targets.size() );
+
+		return m_pD->m_targets[t].get();
+	}
+
+
+	//---------------------------------------------------------------------------------------------
+	void NodeMeshInterpolate::SetTarget( int t, NodeMeshPtr pNode )
+	{
+		check( t>=0 && t<(int)m_pD->m_targets.size() );
+		m_pD->m_targets[t] = pNode;
+	}
+
+
+	//---------------------------------------------------------------------------------------------
+	void NodeMeshInterpolate::SetTargetCount( int t )
+	{
+		m_pD->m_targets.resize(t);
+	}
+
+
+	//---------------------------------------------------------------------------------------------
+	int NodeMeshInterpolate::GetTargetCount() const
+	{
+		return (int)m_pD->m_targets.size();
+	}
+
+
+	//---------------------------------------------------------------------------------------------
+	void NodeMeshInterpolate::SetChannelCount( int t )
+	{
+		m_pD->m_channels.resize(t);
+	}
+
+
+	//---------------------------------------------------------------------------------------------
+	void NodeMeshInterpolate::SetChannel( int i, MESH_BUFFER_SEMANTIC semantic, int semanticIndex )
+	{
+		check( i>=0 && i<(int)m_pD->m_channels.size() );
+
+		if ( i>=0 && i<(int)m_pD->m_channels.size() )
+		{
+			m_pD->m_channels[i].semantic = semantic;
+			m_pD->m_channels[i].semanticIndex = semanticIndex;
+		}
+	}
+
+
+	//---------------------------------------------------------------------------------------------
+	//---------------------------------------------------------------------------------------------
+	//---------------------------------------------------------------------------------------------
+	NodeLayoutPtr NodeMeshInterpolate::Private::GetLayout( int index ) const
+	{
+		NodeLayoutPtr pResult;
+
+		// TODO: Substract layouts too? Usually they are ignored.
+		if ( m_targets.size()>0 && m_targets[0] )
+		{
+			NodeMesh::Private* pPrivate =
+					dynamic_cast<NodeMesh::Private*>( m_targets[0]->GetBasePrivate() );
+
+			pResult = pPrivate->GetLayout( index );
+		}
+
+		return pResult;
+	}
+
+
+}
+
+#undef NODE_INPUT_COUNT
diff --git a/Engine/Plugins/Experimental/Mutable/Source/MutableTools/Private/MuT/NodeMeshInterpolatePrivate.h b/Engine/Plugins/Experimental/Mutable/Source/MutableTools/Private/MuT/NodeMeshInterpolatePrivate.h
new file mode 100644
--- /dev/null
+++ b/Engine/Plugins/Experimental/Mutable/Source/MutableTools/Private/MuT/NodeMeshInterpolatePrivate.h
@@ -0,0 +1,106 @@
+// Copyright Epic Games, Inc. All Rights Reserved.
+
+#pragma once
+
+#include "MuT/NodeMeshInterpolate.h"
+#include "MuT/NodeScalar.h"
+#include "MuT/NodeMeshPrivate.h"
+#include "MuT/AST.h"
+
+#include "MuR/MeshPrivate.h"
+
+
+namespace mu
+{
+
+
+	class NodeMeshInterpolate::Private : public NodeMesh::Private
+	{
+	public:
+
+		MUTABLE_DEFINE_CONST_VISITABLE()
+
+	public:
+
+		static NODE_TYPE s_type;
+
+		NodeScalarPtr m_pFactor;
+		vector<NodeMeshPtr> m_targets;
+
+		//!
+		struct CHANNEL
+		{
+			CHANNEL()
+			{
+				semantic = MBS_NONE;
+				semanticIndex = 0;
+			}
+
+			CHANNEL( MESH_BUFFER_SEMANTIC asemantic,
+                     int32_t asemanticIndex )
+			{
+				semantic = asemantic;
+				semanticIndex = asemanticIndex;
+			}
+
+			MESH_BUFFER_SEMANTIC semantic;
+            int32_t semanticIndex;
+
+			//-------------------------------------------------------------------------------------
+			void Serialise( OutputArchive& arch ) const
+			{
+                const int32_t ver = 0;
+				arch << ver;
+
+				arch << semantic;
+				arch << semanticIndex;
+			}
+
+			//-------------------------------------------------------------------------------------
+			void Unserialise( InputArchive& arch )
+			{
+                int32_t ver = 0;
+				arch >> ver;
+				check(ver == 0);
+
+				arch >> semantic;
+				arch >> semanticIndex;
+			}
+		};
+
+		vector<CHANNEL> m_channels;
+
+		//!
+		void Serialise( OutputArchive& arch ) const
+		{
+            uint32_t ver = 1;
+			arch << ver;
+
+			arch << m_pFactor;
+			arch << m_targets;
+			arch << m_channels;
+		}
+
+		//!
+		void Unserialise( InputArchive& arch )
+		{
+            uint32_t ver;
+			arch >> ver;
+            check(ver==1);
+
+			arch >> m_pFactor;
+			arch >> m_targets;
+
+			if (ver>=1)
+			{
+				arch >> m_channels;
+			}
+		}
+
+		// NodeMesh::Private interface
+        NodeLayoutPtr GetLayout( int index ) const override;
+
+	};
+
+
+}
diff --git a/Engine/Plugins/Experimental/Mutable/Source/MutableTools/Private/MuT/NodeMeshMakeMorph.cpp b/Engine/Plugins/Experimental/Mutable/Source/MutableTools/Private/MuT/NodeMeshMakeMorph.cpp
new file mode 100644
--- /dev/null
+++ b/Engine/Plugins/Experimental/Mutable/Source/MutableTools/Private/MuT/NodeMeshMakeMorph.cpp
@@ -0,0 +1,111 @@
+// Copyright Epic Games, Inc. All Rights Reserved.
+
+
+#include "MuT/NodeMeshMakeMorph.h"
+#include "MuT/NodeMeshMakeMorphPrivate.h"
+#include "MuR/MeshPrivate.h"
+
+
+#define NODE_INPUT_COUNT 	2
+
+
+namespace mu
+{
+
+	//---------------------------------------------------------------------------------------------
+	// Static initialisation
+	//---------------------------------------------------------------------------------------------
+    NODE_TYPE NodeMeshMakeMorph::Private::s_type =
+			NODE_TYPE( "MeshMakeMorph", NodeMesh::GetStaticType() );
+
+
+	//---------------------------------------------------------------------------------------------
+	//!
+	//---------------------------------------------------------------------------------------------
+
+    MUTABLE_IMPLEMENT_NODE( NodeMeshMakeMorph, EType::MakeMorph, Node, Node::EType::Mesh)
+
+
+	//---------------------------------------------------------------------------------------------
+	// Node Interface
+	//---------------------------------------------------------------------------------------------
+    int NodeMeshMakeMorph::GetInputCount() const
+	{
+		return NODE_INPUT_COUNT;
+	}
+
+
+	//---------------------------------------------------------------------------------------------
+    Node* NodeMeshMakeMorph::GetInputNode( int i ) const
+	{
+		check( i>=0 && i<NODE_INPUT_COUNT );
+		return i == 0 ? m_pD->m_pBase.get() : m_pD->m_pTarget.get();
+	}
+
+
+	//---------------------------------------------------------------------------------------------
+    void NodeMeshMakeMorph::SetInputNode( int i, NodePtr pNode )
+	{
+		check( i>=0 && i<NODE_INPUT_COUNT );
+		if (i==0)
+		{
+			m_pD->m_pBase = dynamic_cast<NodeMesh*>( pNode.get() );
+		}
+		else
+		{
+			m_pD->m_pTarget = dynamic_cast<NodeMesh*>( pNode.get() );
+		}
+	}
+
+
+	//---------------------------------------------------------------------------------------------
+	// Own Interface
+	//---------------------------------------------------------------------------------------------
+    NodeMeshPtr NodeMeshMakeMorph::GetBase() const
+	{
+		return m_pD->m_pBase;
+	}
+
+
+	//---------------------------------------------------------------------------------------------
+    void NodeMeshMakeMorph::SetBase( NodeMesh* p )
+	{
+		m_pD->m_pBase = p;
+	}
+
+
+	//---------------------------------------------------------------------------------------------
+    NodeMeshPtr NodeMeshMakeMorph::GetTarget() const
+	{
+		return m_pD->m_pTarget;
+	}
+
+
+	//---------------------------------------------------------------------------------------------
+    void NodeMeshMakeMorph::SetTarget( NodeMesh* p )
+	{
+		m_pD->m_pTarget = p;
+	}
+
+
+	//---------------------------------------------------------------------------------------------
+    NodeLayoutPtr NodeMeshMakeMorph::Private::GetLayout( int index ) const
+	{
+		NodeLayoutPtr pResult;
+
+		// TODO: Substract layouts too? Usually they are ignored.
+		if ( m_pBase )
+		{
+			NodeMesh::Private* pPrivate =
+					dynamic_cast<NodeMesh::Private*>( m_pBase->GetBasePrivate() );
+
+			pResult = pPrivate->GetLayout( index );
+		}
+
+		return pResult;
+	}
+
+
+}
+
+#undef NODE_INPUT_COUNT
diff --git a/Engine/Plugins/Experimental/Mutable/Source/MutableTools/Private/MuT/NodeMeshMorph.cpp b/Engine/Plugins/Experimental/Mutable/Source/MutableTools/Private/MuT/NodeMeshMorph.cpp
new file mode 100644
--- /dev/null
+++ b/Engine/Plugins/Experimental/Mutable/Source/MutableTools/Private/MuT/NodeMeshMorph.cpp
@@ -0,0 +1,205 @@
+// Copyright Epic Games, Inc. All Rights Reserved.
+
+
+#include "MuT/NodeMeshMorph.h"
+#include "MuT/NodeMeshMorphPrivate.h"
+
+#include "MuT/NodeScalar.h"
+
+#include "MuR/MeshPrivate.h"
+
+
+
+namespace mu
+{
+
+	//---------------------------------------------------------------------------------------------
+	// Static initialisation
+	//---------------------------------------------------------------------------------------------
+	NODE_TYPE NodeMeshMorph::Private::s_type =
+			NODE_TYPE( "MeshMorph", NodeMesh::GetStaticType() );
+
+
+	//---------------------------------------------------------------------------------------------
+	//!
+	//---------------------------------------------------------------------------------------------
+
+	MUTABLE_IMPLEMENT_NODE( NodeMeshMorph, EType::Morph, Node, Node::EType::Mesh)
+
+
+	//---------------------------------------------------------------------------------------------
+	// Node Interface
+	//---------------------------------------------------------------------------------------------
+	int NodeMeshMorph::GetInputCount() const
+	{
+		return 2 + (int)m_pD->m_morphs.size();
+	}
+
+
+	//---------------------------------------------------------------------------------------------
+	Node* NodeMeshMorph::GetInputNode( int i ) const
+	{
+		check( i>=0 && i<GetInputCount() );
+
+		Node* pResult = 0;
+
+		switch (i)
+		{
+		case 0:
+			pResult = m_pD->m_pFactor.get();
+			break;
+
+		case 1:
+			pResult = m_pD->m_pBase.get();
+			break;
+
+		default:
+			pResult = m_pD->m_morphs[i-2].get();
+			break;
+		}
+
+		return pResult;
+	}
+
+
+	//---------------------------------------------------------------------------------------------
+	void NodeMeshMorph::SetInputNode( int i, NodePtr pNode )
+	{
+		check( i>=0 && i<GetInputCount() );
+
+		switch (i)
+		{
+		case 0:
+			m_pD->m_pFactor = dynamic_cast<NodeScalar*>(pNode.get());
+			break;
+
+		case 1:
+			m_pD->m_pBase = dynamic_cast<NodeMesh*>(pNode.get());
+			break;
+
+		default:
+			m_pD->m_morphs[i-2] = dynamic_cast<NodeMesh*>(pNode.get());
+			break;
+		}
+	}
+
+
+	//---------------------------------------------------------------------------------------------
+	// Own Interface
+	//---------------------------------------------------------------------------------------------
+	NodeScalarPtr NodeMeshMorph::GetFactor() const
+	{
+		return m_pD->m_pFactor.get();
+	}
+
+
+	//---------------------------------------------------------------------------------------------
+	void NodeMeshMorph::SetFactor( NodeScalarPtr pNode )
+	{
+		m_pD->m_pFactor = pNode;
+	}
+
+
+	//---------------------------------------------------------------------------------------------
+	NodeMeshPtr NodeMeshMorph::GetBase() const
+	{
+		return m_pD->m_pBase.get();
+	}
+
+
+	//---------------------------------------------------------------------------------------------
+	void NodeMeshMorph::SetBase( NodeMeshPtr pNode )
+	{
+		m_pD->m_pBase = pNode;
+	}
+
+
+	//---------------------------------------------------------------------------------------------
+	NodeMeshPtr NodeMeshMorph::GetMorph( int t ) const
+	{
+		check( t>=0 && t<(int)m_pD->m_morphs.size() );
+
+		return m_pD->m_morphs[t].get();
+	}
+
+
+	//---------------------------------------------------------------------------------------------
+	void NodeMeshMorph::SetMorph( int t, NodeMeshPtr pNode )
+	{
+		check( t>=0 && t<(int)m_pD->m_morphs.size() );
+		m_pD->m_morphs[t] = pNode;
+	}
+
+	//---------------------------------------------------------------------------------------------
+	void NodeMeshMorph::SetMorphCount( int t )
+	{
+		m_pD->m_morphs.resize(t);
+	}
+
+	//---------------------------------------------------------------------------------------------
+	int NodeMeshMorph::GetMorphCount() const
+	{
+		return (int)m_pD->m_morphs.size();
+	}
+
+	//---------------------------------------------------------------------------------------------
+	void NodeMeshMorph::SetMorphIndicesAreRelative( bool relative )
+	{
+		m_pD->m_vertexIndicesAreRelative = relative;
+	}
+
+	//---------------------------------------------------------------------------------------------
+	void NodeMeshMorph::SetReshapeSkeleton(bool bEnable)
+	{
+		m_pD->m_reshapeSkeleton = bEnable;
+	}	
+
+	//---------------------------------------------------------------------------------------------
+	void NodeMeshMorph::SetDeformAllBones(bool bEnable)
+	{
+		m_pD->m_deformAllBones = bEnable;
+	}
+
+	//---------------------------------------------------------------------------------------------
+	void NodeMeshMorph::SetDeformAllPhysics(bool bEnable)
+	{
+		m_pD->m_deformAllPhysics = bEnable;
+	}
+
+	//---------------------------------------------------------------------------------------------
+	void NodeMeshMorph::AddBoneToDeform(const char* BoneName)
+	{
+		m_pD->m_bonesToDeform.Emplace(BoneName);
+	}
+
+	void NodeMeshMorph::AddPhysicsBodyToDeform(const char* BoneName)
+	{
+		m_pD->m_physicsToDeform.Emplace(BoneName);
+	}
+
+	//---------------------------------------------------------------------------------------------
+	void NodeMeshMorph::SetReshapePhysicsVolumes(bool bEnable)
+	{
+		m_pD->m_reshapePhysicsVolumes = bEnable;
+	}
+
+	//---------------------------------------------------------------------------------------------
+	NodeLayoutPtr NodeMeshMorph::Private::GetLayout( int index ) const
+	{
+		NodeLayoutPtr pResult;
+
+		if ( m_pBase )
+		{
+			NodeMesh::Private* pPrivate =
+					dynamic_cast<NodeMesh::Private*>( m_pBase->GetBasePrivate() );
+
+			pResult = pPrivate->GetLayout( index );
+		}
+
+		return pResult;
+	}
+
+
+}
+
+
diff --git a/Engine/Plugins/Experimental/Mutable/Source/MutableTools/Private/MuT/NodeMeshMorphPrivate.h b/Engine/Plugins/Experimental/Mutable/Source/MutableTools/Private/MuT/NodeMeshMorphPrivate.h
new file mode 100644
--- /dev/null
+++ b/Engine/Plugins/Experimental/Mutable/Source/MutableTools/Private/MuT/NodeMeshMorphPrivate.h
@@ -0,0 +1,112 @@
+// Copyright Epic Games, Inc. All Rights Reserved.
+
+#pragma once
+
+#include "MuT/NodeMeshMorph.h"
+#include "MuT/NodeScalar.h"
+#include "MuT/NodeMeshPrivate.h"
+#include "MuT/AST.h"
+
+namespace mu
+{
+
+
+	class NodeMeshMorph::Private : public NodeMesh::Private
+	{
+	public:
+
+		MUTABLE_DEFINE_CONST_VISITABLE()
+
+	public:
+
+		static NODE_TYPE s_type;
+
+		NodeScalarPtr m_pFactor;
+		NodeMeshPtr m_pBase;
+		vector<NodeMeshPtr> m_morphs;
+        bool m_vertexIndicesAreRelative = false;
+		
+		bool m_reshapeSkeleton = false;
+		bool m_reshapePhysicsVolumes = false;
+		
+		bool m_deformAllBones = false;
+		bool m_deformAllPhysics = false;
+		
+		TArray<string> m_bonesToDeform;
+		TArray<string> m_physicsToDeform;
+
+        //!
+		void Serialise( OutputArchive& arch ) const
+		{
+            uint32_t ver = 3;
+			arch << ver;
+
+			arch << m_pFactor;
+			arch << m_pBase;
+			arch << m_morphs;
+            arch << m_vertexIndicesAreRelative;
+
+			arch << m_reshapeSkeleton;
+			arch << m_reshapePhysicsVolumes;
+			arch << m_bonesToDeform;
+
+			arch << m_deformAllBones;
+			arch << m_deformAllPhysics;
+			
+			arch << m_physicsToDeform;
+
+        }
+
+		//!
+		void Unserialise( InputArchive& arch )
+		{
+            uint32_t ver;
+			arch >> ver;
+			check(ver <= 3);
+
+			arch >> m_pFactor;
+			arch >> m_pBase;
+			arch >> m_morphs;
+            arch >> m_vertexIndicesAreRelative;
+	
+			if (ver >= 2)
+			{
+				arch >> m_reshapeSkeleton;
+				arch >> m_reshapePhysicsVolumes;
+	
+				// This repetition is needed, there was a bug where m_reshapePhysicsVolumes was serialized twice in ver 2.
+				if (ver == 2) 
+				{
+					arch >> m_reshapePhysicsVolumes;
+				}
+				arch >> m_bonesToDeform;			
+			}
+			else
+			{
+				m_reshapeSkeleton = false;
+				m_reshapePhysicsVolumes = false;
+				m_bonesToDeform.Empty();
+			}
+
+			if (ver >= 3)
+			{
+				arch >> m_deformAllBones;
+				arch >> m_deformAllPhysics;
+				arch >> m_physicsToDeform;
+			}
+			else
+			{
+				m_deformAllBones = false;
+				m_deformAllPhysics = false;
+				m_physicsToDeform.Empty();
+			}
+
+		}
+
+		// NodeMesh::Private interface
+        NodeLayoutPtr GetLayout( int index ) const override;
+
+	};
+
+
+}
diff --git a/Engine/Plugins/Experimental/Mutable/Source/MutableTools/Private/MuT/NodeMeshReshape.cpp b/Engine/Plugins/Experimental/Mutable/Source/MutableTools/Private/MuT/NodeMeshReshape.cpp
new file mode 100644
--- /dev/null
+++ b/Engine/Plugins/Experimental/Mutable/Source/MutableTools/Private/MuT/NodeMeshReshape.cpp
@@ -0,0 +1,179 @@
+// Copyright Epic Games, Inc. All Rights Reserved.
+
+#include "MuT/NodeMeshReshape.h"
+#include "MuT/NodeMeshReshapePrivate.h"
+
+#include "MuT/NodeScalar.h"
+
+#include "MuR/MeshPrivate.h"
+
+
+
+namespace mu
+{
+
+	//---------------------------------------------------------------------------------------------
+	// Static initialisation
+	//---------------------------------------------------------------------------------------------
+	NODE_TYPE NodeMeshReshape::Private::s_type =
+			NODE_TYPE( "MeshReshape", NodeMesh::GetStaticType() );
+
+
+	//---------------------------------------------------------------------------------------------
+	//!
+	//---------------------------------------------------------------------------------------------
+
+	MUTABLE_IMPLEMENT_NODE( NodeMeshReshape, EType::Reshape, Node, Node::EType::Mesh)
+
+
+	//---------------------------------------------------------------------------------------------
+	// Node Interface
+	//---------------------------------------------------------------------------------------------
+	int NodeMeshReshape::GetInputCount() const
+	{
+		return 3;
+	}
+
+
+	//---------------------------------------------------------------------------------------------
+	Node* NodeMeshReshape::GetInputNode( int i ) const
+	{
+		check( i>=0 && i<GetInputCount() );
+
+		Node* pResult = nullptr;
+
+		switch (i)
+		{
+		case 0: pResult = m_pD->m_pBaseMesh.get(); break;
+		case 1: pResult = m_pD->m_pBaseShape.get(); break;
+		case 2: pResult = m_pD->m_pTargetShape.get(); break;
+		default:
+			break;
+		}
+
+		return pResult;
+	}
+
+
+	//---------------------------------------------------------------------------------------------
+	void NodeMeshReshape::SetInputNode( int i, NodePtr pNode )
+	{
+		check( i>=0 && i<GetInputCount() );
+
+		switch (i)
+		{
+		case 0: m_pD->m_pBaseMesh = dynamic_cast<NodeMesh*>(pNode.get()); break;
+		case 1: m_pD->m_pBaseShape = dynamic_cast<NodeMesh*>(pNode.get()); break;
+		case 2: m_pD->m_pTargetShape = dynamic_cast<NodeMesh*>(pNode.get()); break;
+		default:
+			break;
+		}
+	}
+
+
+	//---------------------------------------------------------------------------------------------
+	// Own Interface
+	//---------------------------------------------------------------------------------------------
+	const Ptr<NodeMesh>& NodeMeshReshape::GetBaseMesh() const
+	{
+		return m_pD->m_pBaseMesh;
+	}
+
+
+	//---------------------------------------------------------------------------------------------
+	void NodeMeshReshape::SetBaseMesh(const Ptr<NodeMesh>& pNode)
+	{
+		m_pD->m_pBaseMesh = pNode;
+	}
+
+
+	//---------------------------------------------------------------------------------------------
+	const Ptr<NodeMesh>& NodeMeshReshape::GetBaseShape() const
+	{
+		return m_pD->m_pBaseShape;
+	}
+
+
+	//---------------------------------------------------------------------------------------------
+	void NodeMeshReshape::SetBaseShape(const Ptr<NodeMesh>& pNode)
+	{
+		m_pD->m_pBaseShape = pNode;
+	}
+
+
+	//---------------------------------------------------------------------------------------------
+	const Ptr<NodeMesh>& NodeMeshReshape::GetTargetShape() const
+	{
+		return m_pD->m_pTargetShape;
+	}
+
+
+	//---------------------------------------------------------------------------------------------
+	void NodeMeshReshape::SetTargetShape(const Ptr<NodeMesh>& pNode)
+	{
+		m_pD->m_pTargetShape = pNode;
+	}
+
+
+	//---------------------------------------------------------------------------------------------
+	void NodeMeshReshape::SetReshapeSkeleton(bool bEnable)
+	{
+		m_pD->m_reshapeSkeleton = bEnable;
+	}
+	
+
+	//---------------------------------------------------------------------------------------------
+	void NodeMeshReshape::SetEnableRigidParts(bool bEnable)
+	{
+		m_pD->m_enableRigidParts = bEnable;
+	}
+
+	//---------------------------------------------------------------------------------------------
+	void NodeMeshReshape::SetDeformAllBones(bool bEnable)
+	{
+		m_pD->m_deformAllBones = bEnable;
+	}
+
+	//---------------------------------------------------------------------------------------------
+	void NodeMeshReshape::SetDeformAllPhysics(bool bEnable)
+	{
+		m_pD->m_deformAllPhysics = bEnable;
+	}
+
+	//---------------------------------------------------------------------------------------------
+	void NodeMeshReshape::AddBoneToDeform(const char* BoneName)
+	{
+		m_pD->m_bonesToDeform.Emplace(BoneName);
+	}
+
+	//---------------------------------------------------------------------------------------------
+	void NodeMeshReshape::AddPhysicsBodyToDeform(const char* BoneName)
+	{
+		m_pD->m_physicsToDeform.Emplace(BoneName);
+	}
+
+	void NodeMeshReshape::SetReshapePhysicsVolumes(bool bEnable)
+	{
+		m_pD->m_reshapePhysicsVolumes = bEnable;
+	}
+	
+	//---------------------------------------------------------------------------------------------
+	Ptr<NodeLayout> NodeMeshReshape::Private::GetLayout( int index ) const
+	{
+		NodeLayoutPtr pResult;
+
+		if (m_pBaseMesh)
+		{
+			NodeMesh::Private* pPrivate =
+					dynamic_cast<NodeMesh::Private*>(m_pBaseMesh->GetBasePrivate() );
+
+			pResult = pPrivate->GetLayout( index );
+		}
+
+		return pResult;
+	}
+
+
+}
+
+
diff --git a/Engine/Plugins/Experimental/Mutable/Source/MutableTools/Private/MuT/NodeMeshReshapePrivate.h b/Engine/Plugins/Experimental/Mutable/Source/MutableTools/Private/MuT/NodeMeshReshapePrivate.h
new file mode 100644
--- /dev/null
+++ b/Engine/Plugins/Experimental/Mutable/Source/MutableTools/Private/MuT/NodeMeshReshapePrivate.h
@@ -0,0 +1,85 @@
+// Copyright Epic Games, Inc. All Rights Reserved.
+
+#pragma once
+
+#include "MuT/NodeMeshReshape.h"
+#include "MuT/NodeScalar.h"
+#include "MuT/NodeMeshPrivate.h"
+#include "MuT/AST.h"
+
+namespace mu
+{
+
+
+	class NodeMeshReshape::Private : public NodeMesh::Private
+	{
+	public:
+
+		MUTABLE_DEFINE_CONST_VISITABLE()
+
+	public:
+
+		static NODE_TYPE s_type;
+
+		Ptr<NodeMesh> m_pBaseMesh;
+		Ptr<NodeMesh> m_pBaseShape;
+		Ptr<NodeMesh> m_pTargetShape;
+		bool m_reshapeSkeleton = false;
+		bool m_enableRigidParts = false;
+		bool m_deformAllBones = false;
+		bool m_deformAllPhysics = false;
+		bool m_reshapePhysicsVolumes = false;
+		
+		TArray<string> m_bonesToDeform;
+		TArray<string> m_physicsToDeform;
+        //!
+		void Serialise( OutputArchive& arch ) const
+		{
+            uint32_t ver = 5;
+			arch << ver;
+
+			arch << m_pBaseMesh;
+			arch << m_pBaseShape;
+			arch << m_pTargetShape;
+			arch << m_reshapeSkeleton;
+			arch << m_enableRigidParts;
+			arch << m_deformAllBones;
+			arch << m_bonesToDeform;
+			arch << m_reshapePhysicsVolumes;
+			arch << m_deformAllPhysics;
+			arch << m_physicsToDeform;
+        }
+
+		//!
+		void Unserialise( InputArchive& arch )
+		{
+            uint32_t ver;
+			arch >> ver;
+			check(ver<=5);
+
+			arch >> m_pBaseMesh;
+			arch >> m_pBaseShape;
+			arch >> m_pTargetShape;
+			arch >> m_reshapeSkeleton;
+			arch >> m_enableRigidParts;
+			arch >> m_deformAllBones;
+			arch >> m_bonesToDeform;
+			
+		  	if (ver >= 4)
+		 	{
+				arch >> m_reshapePhysicsVolumes;
+		 	}
+
+			if (ver >= 5)
+			{
+				arch >> m_deformAllPhysics;
+				arch >> m_physicsToDeform;
+			}
+		}
+
+		// NodeMesh::Private interface
+        Ptr<NodeLayout> GetLayout( int index ) const override;
+
+	};
+
+}
diff --git a/Engine/Plugins/Experimental/Mutable/Source/MutableTools/Private/MuT/NodeMeshSubtract.cpp b/Engine/Plugins/Experimental/Mutable/Source/MutableTools/Private/MuT/NodeMeshSubtract.cpp
new file mode 100644
--- /dev/null
+++ b/Engine/Plugins/Experimental/Mutable/Source/MutableTools/Private/MuT/NodeMeshSubtract.cpp
@@ -0,0 +1,111 @@
+// Copyright Epic Games, Inc. All Rights Reserved.
+
+
+#include "MuT/NodeMeshSubtract.h"
+#include "MuT/NodeMeshSubtractPrivate.h"
+#include "MuR/MeshPrivate.h"
+
+
+#define NODE_INPUT_COUNT 	2
+
+
+namespace mu
+{
+
+	//---------------------------------------------------------------------------------------------
+	// Static initialisation
+	//---------------------------------------------------------------------------------------------
+	NODE_TYPE NodeMeshSubtract::Private::s_type =
+			NODE_TYPE( "MeshSubtract", NodeMesh::GetStaticType() );
+
+
+	//---------------------------------------------------------------------------------------------
+	//!
+	//---------------------------------------------------------------------------------------------
+
+	MUTABLE_IMPLEMENT_NODE( NodeMeshSubtract, EType::Subtract, Node, Node::EType::Mesh)
+
+
+	//---------------------------------------------------------------------------------------------
+	// Node Interface
+	//---------------------------------------------------------------------------------------------
+	int NodeMeshSubtract::GetInputCount() const
+	{
+		return NODE_INPUT_COUNT;
+	}
+
+
+	//---------------------------------------------------------------------------------------------
+	Node* NodeMeshSubtract::GetInputNode( int i ) const
+	{
+		check( i>=0 && i<NODE_INPUT_COUNT );
+		return i == 0 ? m_pD->m_pA.get() : m_pD->m_pB.get();
+	}
+
+
+	//---------------------------------------------------------------------------------------------
+	void NodeMeshSubtract::SetInputNode( int i, NodePtr pNode )
+	{
+		check( i>=0 && i<NODE_INPUT_COUNT );
+		if (i==0)
+		{
+			m_pD->m_pA = dynamic_cast<NodeMesh*>( pNode.get() );
+		}
+		else
+		{
+			m_pD->m_pB = dynamic_cast<NodeMesh*>( pNode.get() );
+		}
+	}
+
+
+	//---------------------------------------------------------------------------------------------
+	// Own Interface
+	//---------------------------------------------------------------------------------------------
+	NodeMeshPtr NodeMeshSubtract::GetA() const
+	{
+		return m_pD->m_pA;
+	}
+
+
+	//---------------------------------------------------------------------------------------------
+	void NodeMeshSubtract::SetA( NodeMesh* p )
+	{
+		m_pD->m_pA = p;
+	}
+
+
+	//---------------------------------------------------------------------------------------------
+	NodeMeshPtr NodeMeshSubtract::GetB() const
+	{
+		return m_pD->m_pB;
+	}
+
+
+	//---------------------------------------------------------------------------------------------
+	void NodeMeshSubtract::SetB( NodeMesh* p )
+	{
+		m_pD->m_pB = p;
+	}
+
+
+	//---------------------------------------------------------------------------------------------
+	NodeLayoutPtr NodeMeshSubtract::Private::GetLayout( int index ) const
+	{
+		NodeLayoutPtr pResult;
+
+		// TODO: Substract layouts too? Usually they are ignored.
+		if ( m_pA )
+		{
+			NodeMesh::Private* pPrivate =
+					dynamic_cast<NodeMesh::Private*>( m_pA->GetBasePrivate() );
+
+			pResult = pPrivate->GetLayout( index );
+		}
+
+		return pResult;
+	}
+
+
+}
+
+#undef NODE_INPUT_COUNT
diff --git a/Engine/Plugins/Experimental/Mutable/Source/MutableTools/Private/MuT/NodeMeshSubtractPrivate.h b/Engine/Plugins/Experimental/Mutable/Source/MutableTools/Private/MuT/NodeMeshSubtractPrivate.h
new file mode 100644
--- /dev/null
+++ b/Engine/Plugins/Experimental/Mutable/Source/MutableTools/Private/MuT/NodeMeshSubtractPrivate.h
@@ -0,0 +1,56 @@
+// Copyright Epic Games, Inc. All Rights Reserved.
+
+#pragma once
+
+#include "MuT/NodeMeshPrivate.h"
+#include "MuT/NodeMeshSubtract.h"
+#include "MuT/AST.h"
+
+#include "MuR/MemoryPrivate.h"
+
+
+namespace mu
+{
+
+
+	class NodeMeshSubtract::Private : public NodeMesh::Private
+	{
+	public:
+
+		MUTABLE_DEFINE_CONST_VISITABLE()
+
+	public:
+
+		static NODE_TYPE s_type;
+
+		NodeMeshPtr m_pA;
+		NodeMeshPtr m_pB;
+
+		//!
+		void Serialise( OutputArchive& arch ) const
+		{
+            uint32_t ver = 0;
+			arch << ver;
+
+			arch << m_pA;
+			arch << m_pB;
+		}
+
+		//!
+		void Unserialise( InputArchive& arch )
+		{
+            uint32_t ver;
+			arch >> ver;
+			check(ver==0);
+
+			arch >> m_pA;
+			arch >> m_pB;
+		}
+
+		// NodeMesh::Private interface
+        NodeLayoutPtr GetLayout( int index ) const override;
+
+	};
+
+
+}
diff --git a/Engine/Plugins/Experimental/Mutable/Source/MutableTools/Private/MuT/NodeMeshSwitch.cpp b/Engine/Plugins/Experimental/Mutable/Source/MutableTools/Private/MuT/NodeMeshSwitch.cpp
new file mode 100644
--- /dev/null
+++ b/Engine/Plugins/Experimental/Mutable/Source/MutableTools/Private/MuT/NodeMeshSwitch.cpp
@@ -0,0 +1,137 @@
+// Copyright Epic Games, Inc. All Rights Reserved.
+
+
+#include "MuT/NodeMeshSwitch.h"
+#include "MuT/NodeMeshSwitchPrivate.h"
+
+#include "MuT/NodeScalar.h"
+
+#include "MuR/MeshPrivate.h"
+
+
+namespace mu
+{
+
+	//---------------------------------------------------------------------------------------------
+	// Static initialisation
+	//---------------------------------------------------------------------------------------------
+	NODE_TYPE NodeMeshSwitch::Private::s_type =
+			NODE_TYPE( "MeshSwitch", NodeMesh::GetStaticType() );
+
+
+	//---------------------------------------------------------------------------------------------
+	//!
+	//---------------------------------------------------------------------------------------------
+
+	MUTABLE_IMPLEMENT_NODE( NodeMeshSwitch, EType::Switch, Node, Node::EType::Mesh)
+
+
+	//---------------------------------------------------------------------------------------------
+	// Node Interface
+	//---------------------------------------------------------------------------------------------
+	int NodeMeshSwitch::GetInputCount() const
+	{
+		return 1 + (int)m_pD->m_options.size();
+	}
+
+
+	//---------------------------------------------------------------------------------------------
+	Node* NodeMeshSwitch::GetInputNode( int i ) const
+	{
+		check( i>=0 && i<GetInputCount() );
+
+		Node* pResult = 0;
+
+		switch (i)
+		{
+		case 0:
+			pResult = m_pD->m_pParameter.get();
+			break;
+
+		default:
+			pResult = m_pD->m_options[i-1].get();
+			break;
+		}
+
+		return pResult;
+	}
+
+
+	//---------------------------------------------------------------------------------------------
+	void NodeMeshSwitch::SetInputNode( int i, NodePtr pNode )
+	{
+		check( i>=0 && i<GetInputCount() );
+
+		switch (i)
+		{
+		case 0:
+			m_pD->m_pParameter = dynamic_cast<NodeScalar*>(pNode.get());
+			break;
+
+		default:
+			m_pD->m_options[i-1] = dynamic_cast<NodeMesh*>(pNode.get());
+			break;
+		}
+	}
+
+
+	//---------------------------------------------------------------------------------------------
+	// Own Interface
+	//---------------------------------------------------------------------------------------------
+	NodeScalarPtr NodeMeshSwitch::GetParameter() const
+	{
+		return m_pD->m_pParameter.get();
+	}
+
+
+	//---------------------------------------------------------------------------------------------
+	void NodeMeshSwitch::SetParameter( NodeScalarPtr pNode )
+	{
+		m_pD->m_pParameter = pNode;
+	}
+
+
+	//---------------------------------------------------------------------------------------------
+	void NodeMeshSwitch::SetOptionCount( int t )
+	{
+		m_pD->m_options.resize(t);
+	}
+
+
+	//---------------------------------------------------------------------------------------------
+	NodeMeshPtr NodeMeshSwitch::GetOption( int t ) const
+	{
+		check( t>=0 && t<(int)m_pD->m_options.size() );
+		return m_pD->m_options[t].get();
+	}
+
+
+	//---------------------------------------------------------------------------------------------
+	void NodeMeshSwitch::SetOption( int t, NodeMeshPtr pNode )
+	{
+		check( t>=0 && t<(int)m_pD->m_options.size() );
+		m_pD->m_options[t] = pNode;
+	}
+
+
+	//---------------------------------------------------------------------------------------------
+	NodeLayoutPtr NodeMeshSwitch::Private::GetLayout( int index ) const
+	{
+		NodeLayoutPtr pResult;
+
+		if (m_options.size()>0 && m_options[0] )
+		{
+			NodeMesh::Private* pPrivate =
+					dynamic_cast<NodeMesh::Private*>( m_options[0]->GetBasePrivate() );
+
+			pResult = pPrivate->GetLayout( index );
+		}
+
+		return pResult;
+	}
+
+
+
+}
+
+
diff --git a/Engine/Plugins/Experimental/Mutable/Source/MutableTools/Private/MuT/NodeMeshSwitchPrivate.h b/Engine/Plugins/Experimental/Mutable/Source/MutableTools/Private/MuT/NodeMeshSwitchPrivate.h
new file mode 100644
--- /dev/null
+++ b/Engine/Plugins/Experimental/Mutable/Source/MutableTools/Private/MuT/NodeMeshSwitchPrivate.h
@@ -0,0 +1,55 @@
+// Copyright Epic Games, Inc. All Rights Reserved.
+
+#pragma once
+
+#include "MuT/NodeMeshSwitch.h"
+#include "MuT/NodeScalar.h"
+#include "MuT/NodeMeshPrivate.h"
+#include "MuT/AST.h"
+
+
+namespace mu
+{
+
+
+	class NodeMeshSwitch::Private : public NodeMesh::Private
+	{
+	public:
+
+		MUTABLE_DEFINE_CONST_VISITABLE()
+
+	public:
+
+		static NODE_TYPE s_type;
+
+		NodeScalarPtr m_pParameter;
+		vector<NodeMeshPtr> m_options;
+
+		//!
+		void Serialise( OutputArchive& arch ) const
+		{
+            uint32_t ver = 0;
+			arch << ver;
+
+			arch << m_pParameter;
+			arch << m_options;
+		}
+
+		//!
+		void Unserialise( InputArchive& arch )
+		{
+            uint32_t ver;
+			arch >> ver;
+			check(ver==0);
+
+			arch >> m_pParameter;
+			arch >> m_options;
+		}
+
+
+		// NodeMesh::Private interface
+        NodeLayoutPtr GetLayout( int index ) const override;
+	};
+
+
+}
diff --git a/Engine/Plugins/Experimental/Mutable/Source/MutableTools/Private/MuT/NodeMeshTable.cpp b/Engine/Plugins/Experimental/Mutable/Source/MutableTools/Private/MuT/NodeMeshTable.cpp
new file mode 100644
--- /dev/null
+++ b/Engine/Plugins/Experimental/Mutable/Source/MutableTools/Private/MuT/NodeMeshTable.cpp
@@ -0,0 +1,142 @@
+// Copyright Epic Games, Inc. All Rights Reserved.
+
+
+#include "MuT/NodeMeshTable.h"
+#include "MuT/NodeMeshTablePrivate.h"
+
+
+namespace mu
+{
+
+
+	//---------------------------------------------------------------------------------------------
+	// Static initialisation
+	//---------------------------------------------------------------------------------------------
+	NODE_TYPE NodeMeshTable::Private::s_type =
+			NODE_TYPE( "TableMesh", NodeMesh::GetStaticType() );
+
+
+	//---------------------------------------------------------------------------------------------
+	//!
+	//---------------------------------------------------------------------------------------------
+
+	MUTABLE_IMPLEMENT_NODE( NodeMeshTable, EType::Table, Node, Node::EType::Mesh);
+
+
+	//---------------------------------------------------------------------------------------------
+	// Node Interface
+	//---------------------------------------------------------------------------------------------
+	int NodeMeshTable::GetInputCount() const
+	{
+		return 0;
+	}
+
+
+	//---------------------------------------------------------------------------------------------
+	Node* NodeMeshTable::GetInputNode( int i ) const
+	{
+		check( i >=0 && i < 0 );
+        (void)i;
+        return 0;
+	}
+
+
+	//---------------------------------------------------------------------------------------------
+    void NodeMeshTable::SetInputNode( int i, NodePtr )
+	{
+		check( i >=0 && i < 0 );
+        (void)i;
+        //m_pD->m_pObject = dynamic_cast<NodeObject*>( pNode.get() );
+	}
+
+
+	//---------------------------------------------------------------------------------------------
+	// Own Interface
+	//---------------------------------------------------------------------------------------------
+	const char* NodeMeshTable::GetColumn() const
+	{
+		return m_pD->m_columnName.c_str();
+	}
+
+
+	//---------------------------------------------------------------------------------------------
+	void NodeMeshTable::SetColumn( const char* strName )
+	{
+		if (strName)
+		{
+			m_pD->m_columnName = strName;
+		}
+		else
+		{
+			m_pD->m_columnName = "";
+		}
+	}
+
+
+	//---------------------------------------------------------------------------------------------
+	void NodeMeshTable::SetParameterName( const char* strName )
+	{
+		m_pD->m_parameterName = strName;
+	}
+
+
+	//---------------------------------------------------------------------------------------------
+	void NodeMeshTable::SetTable( TablePtr pTable )
+	{
+		m_pD->m_pTable = pTable;
+	}
+
+
+	//---------------------------------------------------------------------------------------------
+	TablePtr NodeMeshTable::GetTable() const
+	{
+		return m_pD->m_pTable;
+	}
+
+
+	//---------------------------------------------------------------------------------------------
+	int NodeMeshTable::GetLayoutCount() const
+	{
+		return (int)m_pD->m_layouts.size();
+	}
+
+
+	//---------------------------------------------------------------------------------------------
+	void NodeMeshTable::SetLayoutCount( int i )
+	{
+		m_pD->m_layouts.resize( i );
+	}
+
+
+	//---------------------------------------------------------------------------------------------
+	NodeLayoutPtr NodeMeshTable::GetLayout( int i ) const
+	{
+		check( i>=0 && i<GetLayoutCount() );
+		return m_pD->GetLayout( i );
+	}
+
+
+	//---------------------------------------------------------------------------------------------
+	void NodeMeshTable::SetLayout( int i, NodeLayoutPtr pLayout )
+	{
+		check( i>=0 && i<GetLayoutCount() );
+		m_pD->m_layouts[i] = pLayout;
+	}
+
+
+	//---------------------------------------------------------------------------------------------
+	NodeLayoutPtr NodeMeshTable::Private::GetLayout( int i ) const
+	{
+		NodeLayoutPtr pResult;
+
+		if ( i>=0 && i<(int)m_layouts.size() )
+		{
+			pResult = m_layouts[i];
+		}
+
+		return pResult;
+	}
+
+}
+
+
diff --git a/Engine/Plugins/Experimental/Mutable/Source/MutableTools/Private/MuT/NodeMeshTablePrivate.h b/Engine/Plugins/Experimental/Mutable/Source/MutableTools/Private/MuT/NodeMeshTablePrivate.h
new file mode 100644
--- /dev/null
+++ b/Engine/Plugins/Experimental/Mutable/Source/MutableTools/Private/MuT/NodeMeshTablePrivate.h
@@ -0,0 +1,64 @@
+// Copyright Epic Games, Inc. All Rights Reserved.
+
+#pragma once
+
+#include "MuT/NodeMeshPrivate.h"
+#include "MuT/NodeMeshTable.h"
+#include "MuT/NodeLayout.h"
+#include "MuT/Table.h"
+#include "MuT/AST.h"
+
+#include "MuR/MemoryPrivate.h"
+
+
+namespace mu
+{
+
+
+	class NodeMeshTable::Private : public NodeMesh::Private
+	{
+	public:
+
+		MUTABLE_DEFINE_CONST_VISITABLE()
+
+	public:
+
+		static NODE_TYPE s_type;
+
+		string m_parameterName;
+		TablePtr m_pTable;
+		string m_columnName;
+
+		vector<NodeLayoutPtr> m_layouts;
+
+		//!
+		void Serialise( OutputArchive& arch ) const
+		{
+            uint32_t ver = 1;
+			arch << ver;
+
+			arch << m_parameterName;
+			arch << m_pTable;
+			arch << m_columnName;
+			arch << m_layouts;
+		}
+
+		//!
+		void Unserialise( InputArchive& arch )
+		{
+            uint32_t ver;
+			arch >> ver;
+            check(ver==1);
+
+			arch >> m_parameterName;
+			arch >> m_pTable;
+			arch >> m_columnName;
+			arch >> m_layouts;
+		}
+
+		// NodeMesh::Private interface
+        NodeLayoutPtr GetLayout( int index ) const override;
+
+	};
+
+}
diff --git a/Engine/Plugins/Experimental/Mutable/Source/MutableTools/Private/MuT/NodeMeshTangents.cpp b/Engine/Plugins/Experimental/Mutable/Source/MutableTools/Private/MuT/NodeMeshTangents.cpp
new file mode 100644
--- /dev/null
+++ b/Engine/Plugins/Experimental/Mutable/Source/MutableTools/Private/MuT/NodeMeshTangents.cpp
@@ -0,0 +1,94 @@
+// Copyright Epic Games, Inc. All Rights Reserved.
+
+
+#include "MuT/NodeMeshTangents.h"
+#include "MuT/NodeMeshTangentsPrivate.h"
+#include "MuR/MeshPrivate.h"
+
+
+#define NODE_INPUT_COUNT 	1
+
+
+namespace mu
+{
+
+	//---------------------------------------------------------------------------------------------
+	// Static initialisation
+	//---------------------------------------------------------------------------------------------
+	NODE_TYPE NodeMeshTangents::Private::s_type =
+			NODE_TYPE( "MeshTangents", NodeMesh::GetStaticType() );
+
+
+	//---------------------------------------------------------------------------------------------
+	//!
+	//---------------------------------------------------------------------------------------------
+
+	MUTABLE_IMPLEMENT_NODE( NodeMeshTangents, EType::Tangents, Node, Node::EType::Mesh)
+
+
+	//---------------------------------------------------------------------------------------------
+	// Node Interface
+	//---------------------------------------------------------------------------------------------
+	int NodeMeshTangents::GetInputCount() const
+	{
+		return NODE_INPUT_COUNT;
+	}
+
+
+	//---------------------------------------------------------------------------------------------
+	Node* NodeMeshTangents::GetInputNode( int i ) const
+	{
+		check( i>=0 && i<NODE_INPUT_COUNT );
+        (void)i;
+		return m_pD->m_pSource.get();
+	}
+
+
+	//---------------------------------------------------------------------------------------------
+	void NodeMeshTangents::SetInputNode( int i, NodePtr pNode )
+	{
+		check( i>=0 && i<NODE_INPUT_COUNT );
+        (void)i;
+        if (i==0)
+		{
+			m_pD->m_pSource = dynamic_cast<NodeMesh*>( pNode.get() );
+		}
+	}
+
+
+	//---------------------------------------------------------------------------------------------
+	// Own Interface
+	//---------------------------------------------------------------------------------------------
+	NodeMeshPtr NodeMeshTangents::GetSource() const
+	{
+		return m_pD->m_pSource;
+	}
+
+
+	//---------------------------------------------------------------------------------------------
+	void NodeMeshTangents::SetSource( NodeMesh* p )
+	{
+		m_pD->m_pSource = p;
+	}
+
+
+	//---------------------------------------------------------------------------------------------
+	NodeLayoutPtr NodeMeshTangents::Private::GetLayout( int index ) const
+	{
+		NodeLayoutPtr pResult;
+
+		if ( m_pSource )
+		{
+			NodeMesh::Private* pPrivate =
+					dynamic_cast<NodeMesh::Private*>( m_pSource->GetBasePrivate() );
+
+			pResult = pPrivate->GetLayout( index );
+		}
+
+		return pResult;
+	}
+
+
+}
+
+#undef NODE_INPUT_COUNT
diff --git a/Engine/Plugins/Experimental/Mutable/Source/MutableTools/Private/MuT/NodeMeshTangentsPrivate.h b/Engine/Plugins/Experimental/Mutable/Source/MutableTools/Private/MuT/NodeMeshTangentsPrivate.h
new file mode 100644
--- /dev/null
+++ b/Engine/Plugins/Experimental/Mutable/Source/MutableTools/Private/MuT/NodeMeshTangentsPrivate.h
@@ -0,0 +1,53 @@
+// Copyright Epic Games, Inc. All Rights Reserved.
+
+#pragma once
+
+#include "MuT/NodeMeshPrivate.h"
+#include "MuT/NodeMeshTangents.h"
+#include "MuT/AST.h"
+
+#include "MuR/MemoryPrivate.h"
+
+
+namespace mu
+{
+
+
+	class NodeMeshTangents::Private : public NodeMesh::Private
+	{
+	public:
+
+		MUTABLE_DEFINE_CONST_VISITABLE()
+
+	public:
+
+		static NODE_TYPE s_type;
+
+		NodeMeshPtr m_pSource;
+
+		//!
+		void Serialise( OutputArchive& arch ) const
+		{
+            uint32_t ver = 0;
+			arch << ver;
+
+			arch << m_pSource;
+		}
+
+		//!
+		void Unserialise( InputArchive& arch )
+		{
+            uint32_t ver;
+			arch >> ver;
+			check(ver==0);
+
+			arch >> m_pSource;
+		}
+
+		// NodeMesh::Private interface
+        NodeLayoutPtr GetLayout( int index ) const override;
+
+	};
+
+
+}
diff --git a/Engine/Plugins/Experimental/Mutable/Source/MutableTools/Private/MuT/NodeMeshTransform.cpp b/Engine/Plugins/Experimental/Mutable/Source/MutableTools/Private/MuT/NodeMeshTransform.cpp
new file mode 100644
--- /dev/null
+++ b/Engine/Plugins/Experimental/Mutable/Source/MutableTools/Private/MuT/NodeMeshTransform.cpp
@@ -0,0 +1,107 @@
+// Copyright Epic Games, Inc. All Rights Reserved.
+
+
+#include "MuT/NodeMeshTransform.h"
+#include "MuT/NodeMeshTransformPrivate.h"
+#include "MuR/MeshPrivate.h"
+
+
+#define NODE_INPUT_COUNT 	1
+
+
+namespace mu
+{
+
+	//---------------------------------------------------------------------------------------------
+	// Static initialisation
+	//---------------------------------------------------------------------------------------------
+    NODE_TYPE NodeMeshTransform::Private::s_type =
+            NODE_TYPE( "MeshTransform", NodeMesh::GetStaticType() );
+
+
+	//---------------------------------------------------------------------------------------------
+	//!
+	//---------------------------------------------------------------------------------------------
+
+    MUTABLE_IMPLEMENT_NODE( NodeMeshTransform, EType::Transform, Node, Node::EType::Mesh)
+
+
+	//---------------------------------------------------------------------------------------------
+	// Node Interface
+	//---------------------------------------------------------------------------------------------
+    int NodeMeshTransform::GetInputCount() const
+	{
+		return NODE_INPUT_COUNT;
+	}
+
+
+	//---------------------------------------------------------------------------------------------
+    Node* NodeMeshTransform::GetInputNode( int i ) const
+	{
+		check( i>=0 && i<NODE_INPUT_COUNT );
+        (void)i;
+        return m_pD->m_pSource.get();
+	}
+
+
+	//---------------------------------------------------------------------------------------------
+    void NodeMeshTransform::SetInputNode( int i, NodePtr pNode )
+	{
+		check( i>=0 && i<NODE_INPUT_COUNT );
+		if (i==0)
+		{
+			m_pD->m_pSource = dynamic_cast<NodeMesh*>( pNode.get() );
+		}
+	}
+
+
+	//---------------------------------------------------------------------------------------------
+	// Own Interface
+	//---------------------------------------------------------------------------------------------
+    NodeMeshPtr NodeMeshTransform::GetSource() const
+	{
+		return m_pD->m_pSource;
+	}
+
+
+    //---------------------------------------------------------------------------------------------
+    void NodeMeshTransform::SetSource( NodeMesh* p )
+    {
+        m_pD->m_pSource = p;
+    }
+
+
+    //---------------------------------------------------------------------------------------------
+    void NodeMeshTransform::SetTransform( const float* mat )
+    {
+        memcpy( &m_pD->m_transform[0][0], mat, 16*sizeof(float) );
+    }
+
+
+    //---------------------------------------------------------------------------------------------
+    void NodeMeshTransform::GetTransform( float* mat ) const
+    {
+        memcpy( mat, &m_pD->m_transform[0][0], 16*sizeof(float) );
+    }
+
+
+	//---------------------------------------------------------------------------------------------
+    NodeLayoutPtr NodeMeshTransform::Private::GetLayout( int index ) const
+	{
+		NodeLayoutPtr pResult;
+
+		if ( m_pSource )
+		{
+			NodeMesh::Private* pPrivate =
+					dynamic_cast<NodeMesh::Private*>( m_pSource->GetBasePrivate() );
+
+			pResult = pPrivate->GetLayout( index );
+		}
+
+		return pResult;
+	}
+
+
+}
+
+#undef NODE_INPUT_COUNT
diff --git a/Engine/Plugins/Experimental/Mutable/Source/MutableTools/Private/MuT/NodeMeshTransformPrivate.h b/Engine/Plugins/Experimental/Mutable/Source/MutableTools/Private/MuT/NodeMeshTransformPrivate.h
new file mode 100644
--- /dev/null
+++ b/Engine/Plugins/Experimental/Mutable/Source/MutableTools/Private/MuT/NodeMeshTransformPrivate.h
@@ -0,0 +1,57 @@
+// Copyright Epic Games, Inc. All Rights Reserved.
+
+#pragma once
+
+#include "MuT/NodeMeshPrivate.h"
+#include "MuT/NodeMeshTransform.h"
+#include "MuT/AST.h"
+
+#include "MuR/MemoryPrivate.h"
+#include "MuR/MutableMath.h"
+
+
+namespace mu
+{
+
+
+    class NodeMeshTransform::Private : public NodeMesh::Private
+	{
+	public:
+
+		MUTABLE_DEFINE_CONST_VISITABLE()
+
+	public:
+
+		static NODE_TYPE s_type;
+
+		NodeMeshPtr m_pSource;
+        mat4f m_transform;
+
+		//!
+		void Serialise( OutputArchive& arch ) const
+		{
+            uint32_t ver = 0;
+			arch << ver;
+
+			arch << m_pSource;
+            arch << m_transform;
+		}
+
+		//!
+		void Unserialise( InputArchive& arch )
+		{
+            uint32_t ver;
+			arch >> ver;
+			check(ver==0);
+
+			arch >> m_pSource;
+            arch >> m_transform;
+		}
+
+		// NodeMesh::Private interface
+        NodeLayoutPtr GetLayout( int index ) const override;
+
+	};
+
+
+}
diff --git a/Engine/Plugins/Experimental/Mutable/Source/MutableTools/Private/MuT/NodeMeshVariation.cpp b/Engine/Plugins/Experimental/Mutable/Source/MutableTools/Private/MuT/NodeMeshVariation.cpp
new file mode 100644
--- /dev/null
+++ b/Engine/Plugins/Experimental/Mutable/Source/MutableTools/Private/MuT/NodeMeshVariation.cpp
@@ -0,0 +1,115 @@
+// Copyright Epic Games, Inc. All Rights Reserved.
+
+
+#include "MuT/NodeMeshVariation.h"
+#include "MuT/NodeMeshVariationPrivate.h"
+
+
+namespace mu
+{
+
+    //---------------------------------------------------------------------------------------------
+    // Static initialisation
+    //---------------------------------------------------------------------------------------------
+    NODE_TYPE NodeMeshVariation::Private::s_type =
+        NODE_TYPE( "MeshVariation", NodeMesh::GetStaticType() );
+
+
+    //---------------------------------------------------------------------------------------------
+    //---------------------------------------------------------------------------------------------
+    //---------------------------------------------------------------------------------------------
+
+    MUTABLE_IMPLEMENT_NODE( NodeMeshVariation, EType::Variation, Node, Node::EType::Mesh)
+
+
+    //---------------------------------------------------------------------------------------------
+    int NodeMeshVariation::GetInputCount() const { return 1 + int( m_pD->m_variations.size() ); }
+
+
+    //---------------------------------------------------------------------------------------------
+    Node* NodeMeshVariation::GetInputNode( int i ) const
+    {
+        check( i >= 0 && i < GetInputCount() );
+
+        if ( i == 0 )
+        {
+            return m_pD->m_defaultMesh.get();
+        }
+        i -= 1;
+
+        if ( i < int( m_pD->m_variations.size() ) )
+        {
+            return m_pD->m_variations[i].m_mesh.get();
+        }
+        i -= int( m_pD->m_variations.size() );
+
+        return nullptr;
+    }
+
+
+    //---------------------------------------------------------------------------------------------
+    void NodeMeshVariation::SetInputNode( int i, NodePtr pNode )
+    {
+        check( i >= 0 && i < GetInputCount() );
+
+        if ( i == 0 )
+        {
+            m_pD->m_defaultMesh = dynamic_cast<NodeMesh*>( pNode.get() );
+            return;
+        }
+
+        i -= 1;
+        if ( i < int( m_pD->m_variations.size() ) )
+        {
+
+            m_pD->m_variations[i].m_mesh = dynamic_cast<NodeMesh*>( pNode.get() );
+            return;
+        }
+        i -= (int)m_pD->m_variations.size();
+    }
+
+
+    //---------------------------------------------------------------------------------------------
+    // Own Interface
+    //---------------------------------------------------------------------------------------------
+    void NodeMeshVariation::SetDefaultMesh( NodeMesh* p ) { m_pD->m_defaultMesh = p; }
+
+
+    //---------------------------------------------------------------------------------------------
+    int NodeMeshVariation::GetVariationCount() const { return int( m_pD->m_variations.size() ); }
+
+
+    //---------------------------------------------------------------------------------------------
+    void NodeMeshVariation::SetVariationCount( int num )
+    {
+        check( num >= 0 );
+        m_pD->m_variations.resize( num );
+    }
+
+    //---------------------------------------------------------------------------------------------
+    void NodeMeshVariation::SetVariationTag( int index, const char* strTag )
+    {
+        check( index >= 0 && index < (int)m_pD->m_variations.size() );
+        check( strTag );
+
+        if ( strTag )
+        {
+            m_pD->m_variations[index].m_tag = strTag;
+        }
+        else
+        {
+            m_pD->m_variations[index].m_tag = "";
+        }
+    }
+
+
+    //---------------------------------------------------------------------------------------------
+    void NodeMeshVariation::SetVariationMesh( int index, NodeMesh* pNode )
+    {
+        check( index >= 0 && index < (int)m_pD->m_variations.size() );
+
+        m_pD->m_variations[index].m_mesh = pNode;
+    }
+
+
+} // namespace mu
diff --git a/Engine/Plugins/Experimental/Mutable/Source/MutableTools/Private/MuT/NodeMeshVariationPrivate.h b/Engine/Plugins/Experimental/Mutable/Source/MutableTools/Private/MuT/NodeMeshVariationPrivate.h
new file mode 100644
--- /dev/null
+++ b/Engine/Plugins/Experimental/Mutable/Source/MutableTools/Private/MuT/NodeMeshVariationPrivate.h
@@ -0,0 +1,79 @@
+// Copyright Epic Games, Inc. All Rights Reserved.
+
+#pragma once
+
+#include "MuT/NodePrivate.h"
+#include "MuT/NodeMeshVariation.h"
+
+#include "MuR/MemoryPrivate.h"
+
+
+namespace mu
+{
+
+    //---------------------------------------------------------------------------------------------
+    //!
+    //---------------------------------------------------------------------------------------------
+    class NodeMeshVariation::Private : public Node::Private
+    {
+    public:
+        MUTABLE_DEFINE_CONST_VISITABLE()
+
+    public:
+        Private() {}
+
+        static NODE_TYPE s_type;
+
+        NodeMeshPtr m_defaultMesh;
+
+        struct VARIATION
+        {
+            NodeMeshPtr m_mesh;
+            string m_tag;
+
+            //!
+            void Serialise( OutputArchive& arch ) const
+            {
+                uint32_t ver = 1;
+                arch << ver;
+
+                arch << m_tag;
+                arch << m_mesh;
+            }
+
+            void Unserialise( InputArchive& arch )
+            {
+                uint32_t ver = 0;
+                arch >> ver;
+                check( ver == 1 );
+
+                arch >> m_tag;
+                arch >> m_mesh;
+            }
+        };
+
+        vector<VARIATION> m_variations;
+
+        //!
+        void Serialise( OutputArchive& arch ) const
+        {
+            uint32_t ver = 1;
+            arch << ver;
+
+            arch << m_defaultMesh;
+            arch << m_variations;
+        }
+
+        //!
+        void Unserialise( InputArchive& arch )
+        {
+            uint32_t ver;
+            arch >> ver;
+			check(ver == 1 );
+
+            arch >> m_defaultMesh;
+            arch >> m_variations;
+        }
+    };
+
+} // namespace mu
diff --git a/Engine/Plugins/Experimental/Mutable/Source/MutableTools/Private/MuT/NodeModifier.cpp b/Engine/Plugins/Experimental/Mutable/Source/MutableTools/Private/MuT/NodeModifier.cpp
new file mode 100644
--- /dev/null
+++ b/Engine/Plugins/Experimental/Mutable/Source/MutableTools/Private/MuT/NodeModifier.cpp
@@ -0,0 +1,127 @@
+// Copyright Epic Games, Inc. All Rights Reserved.
+
+
+#include "MuT/NodeModifier.h"
+#include "NodeModifierPrivate.h"
+
+#include "MuT/NodeModifierMeshClipMorphPlane.h"
+#include "MuT/NodeModifierMeshClipWithMesh.h"
+#include "MuT/NodeModifierMeshClipDeform.h"
+
+
+namespace mu
+{
+
+
+	//---------------------------------------------------------------------------------------------
+	// Static initialisation
+	//---------------------------------------------------------------------------------------------
+	static NODE_TYPE s_nodeModifierType =
+			NODE_TYPE( "NodeModifier", Node::GetStaticType() );
+
+
+	//---------------------------------------------------------------------------------------------
+	//---------------------------------------------------------------------------------------------
+	//---------------------------------------------------------------------------------------------
+	const NODE_TYPE* NodeModifier::GetType() const
+	{
+		return GetStaticType();
+	}
+
+
+	//---------------------------------------------------------------------------------------------
+	const NODE_TYPE* NodeModifier::GetStaticType()
+	{
+          return &s_nodeModifierType;
+        }
+
+
+	//---------------------------------------------------------------------------------------------
+	void NodeModifier::Serialise( const NodeModifier* p, OutputArchive& arch )
+	{
+        uint32_t ver = 0;
+		arch << ver;
+
+#define SERIALISE_CHILDREN( C, ID ) \
+        ( const C* pTyped##ID = dynamic_cast<const C*>(p) )			\
+        {                                                           \
+            arch << (uint32_t)ID;									\
+            C::Serialise( pTyped##ID, arch );						\
+		}
+
+		if SERIALISE_CHILDREN(NodeModifierMeshClipMorphPlane, 			0 )
+		else if SERIALISE_CHILDREN(NodeModifierMeshClipWithMesh, 		1 )
+		else if SERIALISE_CHILDREN(NodeModifierMeshClipDeform, 		    2 )
+		else check(false);
+
+#undef SERIALISE_CHILDREN
+    }
+
+        
+	//---------------------------------------------------------------------------------------------
+	NodeModifierPtr NodeModifier::StaticUnserialise( InputArchive& arch )
+	{
+        uint32_t ver;
+		arch >> ver;
+		check( ver == 0 );
+
+        uint32_t id;
+		arch >> id;
+
+		switch (id)
+		{
+		case 0 :  return NodeModifierMeshClipMorphPlane::StaticUnserialise( arch ); break;
+		case 1 :  return NodeModifierMeshClipWithMesh::StaticUnserialise( arch ); break;
+		case 2 :  return NodeModifierMeshClipDeform::StaticUnserialise( arch ); break;
+		default : check(false);
+		}
+
+		return 0;
+	}
+
+	//---------------------------------------------------------------------------------------------
+	void NodeModifier::AddTag(const char* tagName)
+	{
+		NodeModifier::Private* pD = dynamic_cast<NodeModifier::Private*>(GetBasePrivate());
+		check(pD);
+
+		pD->m_tags.push_back(tagName);
+	}
+
+
+    //---------------------------------------------------------------------------------------------
+    int NodeModifier::GetTagCount() const
+    {
+        NodeModifier::Private* pD = dynamic_cast<NodeModifier::Private*>(GetBasePrivate());
+        check(pD);
+
+        return int( pD->m_tags.size());
+    }
+
+
+    //---------------------------------------------------------------------------------------------
+    const char* NodeModifier::GetTag( int i ) const
+    {
+        NodeModifier::Private* pD = dynamic_cast<NodeModifier::Private*>(GetBasePrivate());
+        check(pD);
+
+        if (i>=0 && i<GetTagCount())
+        {
+            return pD->m_tags[i].c_str();
+        }
+        return nullptr;
+    }
+
+
+	//---------------------------------------------------------------------------------------------
+	void NodeModifier::SetStage(bool bBeforeNormalOperation)
+	{
+		NodeModifier::Private* pD = dynamic_cast<NodeModifier::Private*>(GetBasePrivate());
+		check(pD);
+
+		pD->m_applyBeforeNormalOperations = bBeforeNormalOperation; 
+	}
+
+}
+
+
diff --git a/Engine/Plugins/Experimental/Mutable/Source/MutableTools/Private/MuT/NodeModifierMeshClipDeform.cpp b/Engine/Plugins/Experimental/Mutable/Source/MutableTools/Private/MuT/NodeModifierMeshClipDeform.cpp
new file mode 100644
--- /dev/null
+++ b/Engine/Plugins/Experimental/Mutable/Source/MutableTools/Private/MuT/NodeModifierMeshClipDeform.cpp
@@ -0,0 +1,66 @@
+// Copyright Epic Games, Inc. All Rights Reserved.
+
+
+#include "MuT/NodeModifierMeshClipDeform.h"
+#include "NodeModifierMeshClipDeformPrivate.h"
+#include "MuR/MeshPrivate.h"
+
+
+namespace mu
+{
+
+	//---------------------------------------------------------------------------------------------
+	// Static initialisation
+	//---------------------------------------------------------------------------------------------
+	NODE_TYPE NodeModifierMeshClipDeform::Private::s_type =
+			NODE_TYPE( "NodeModifierMeshClipDeform", NodeModifier::GetStaticType() );
+
+
+	//---------------------------------------------------------------------------------------------
+	//!
+	//---------------------------------------------------------------------------------------------
+
+	MUTABLE_IMPLEMENT_NODE( NodeModifierMeshClipDeform, EType::MeshClipDeform, Node, Node::EType::Modifier)
+
+
+	//---------------------------------------------------------------------------------------------
+	// Node Interface
+	//---------------------------------------------------------------------------------------------
+	int NodeModifierMeshClipDeform::GetInputCount() const
+	{
+		return 0;
+	}
+
+
+	//---------------------------------------------------------------------------------------------
+	Node* NodeModifierMeshClipDeform::GetInputNode( int i ) const
+	{
+		check( i>=0 && i< GetInputCount());
+		(void)i;
+		return nullptr;
+	}
+
+
+	//---------------------------------------------------------------------------------------------
+	void NodeModifierMeshClipDeform::SetInputNode( int32 i, NodePtr )
+	{
+		check( i>=0 && i< GetInputCount());
+		(void)i;
+	}
+
+
+	//---------------------------------------------------------------------------------------------
+	// Own Interface
+	//---------------------------------------------------------------------------------------------
+	void NodeModifierMeshClipDeform::SetClipMesh(NodeMesh* ClipMesh)
+	{
+		m_pD->ClipMesh = ClipMesh;
+	}
+
+	void NodeModifierMeshClipDeform::SetBindingMethod(EShapeBindingMethod BindingMethod)
+	{
+		check(BindingMethod != EShapeBindingMethod::ReshapeClosestProject);
+		m_pD->BindingMethod = BindingMethod;
+	}
+
+}
diff --git a/Engine/Plugins/Experimental/Mutable/Source/MutableTools/Private/MuT/NodeModifierMeshClipMorphPlane.cpp b/Engine/Plugins/Experimental/Mutable/Source/MutableTools/Private/MuT/NodeModifierMeshClipMorphPlane.cpp
new file mode 100644
--- /dev/null
+++ b/Engine/Plugins/Experimental/Mutable/Source/MutableTools/Private/MuT/NodeModifierMeshClipMorphPlane.cpp
@@ -0,0 +1,100 @@
+// Copyright Epic Games, Inc. All Rights Reserved.
+
+
+#include "MuT/NodeModifierMeshClipMorphPlane.h"
+#include "NodeModifierMeshClipMorphPlanePrivate.h"
+#include "MuR/MeshPrivate.h"
+
+
+#define NODE_INPUT_COUNT 	0
+
+
+namespace mu
+{
+
+	//---------------------------------------------------------------------------------------------
+	// Static initialisation
+	//---------------------------------------------------------------------------------------------
+    NODE_TYPE NodeModifierMeshClipMorphPlane::Private::s_type =
+            NODE_TYPE( "ModifierMeshClipMorphPlane", NodeModifier::GetStaticType() );
+
+
+	//---------------------------------------------------------------------------------------------
+	//!
+	//---------------------------------------------------------------------------------------------
+
+    MUTABLE_IMPLEMENT_NODE( NodeModifierMeshClipMorphPlane, EType::MeshClipMorphPlane, Node, Node::EType::Modifier)
+
+
+	//---------------------------------------------------------------------------------------------
+	// Node Interface
+	//---------------------------------------------------------------------------------------------
+    int NodeModifierMeshClipMorphPlane::GetInputCount() const
+	{
+		return NODE_INPUT_COUNT;
+	}
+
+
+	//---------------------------------------------------------------------------------------------
+    Node* NodeModifierMeshClipMorphPlane::GetInputNode( int i ) const
+	{
+		check( i>=0 && i<NODE_INPUT_COUNT );
+        (void)i;
+		return nullptr;
+	}
+
+
+	//---------------------------------------------------------------------------------------------
+    void NodeModifierMeshClipMorphPlane::SetInputNode( int i, NodePtr )
+	{
+		check( i>=0 && i<NODE_INPUT_COUNT );
+        (void)i;
+    }
+
+
+	//---------------------------------------------------------------------------------------------
+	// Own Interface
+	//---------------------------------------------------------------------------------------------
+    void NodeModifierMeshClipMorphPlane::SetPlane( float centerX, float centerY, float centerZ,
+                                                   float normalX, float normalY, float normalZ)
+	{
+		m_pD->m_origin = vec3f(centerX, centerY, centerZ);
+		m_pD->m_normal = vec3f(normalX, normalY, normalZ);
+	}
+
+
+	//---------------------------------------------------------------------------------------------
+	void NodeModifierMeshClipMorphPlane::SetParams(float dist, float factor)
+	{
+		m_pD->m_dist = dist;
+		m_pD->m_factor = factor;
+	}
+
+	//---------------------------------------------------------------------------------------------
+	void NodeModifierMeshClipMorphPlane::SetMorphEllipse(float radius1, float radius2, float rotation)
+	{
+		m_pD->m_radius1 = radius1;
+		m_pD->m_radius2 = radius2;
+		m_pD->m_rotation = rotation;
+	}
+
+	//---------------------------------------------------------------------------------------------
+	void NodeModifierMeshClipMorphPlane::SetVertexSelectionBox(float centerX, float centerY, float centerZ, float radiusX, float radiusY, float radiusZ)
+	{
+		m_pD->m_vertexSelectionType = Private::VS_SHAPE;
+		m_pD->m_selectionBoxOrigin = vec3f(centerX, centerY, centerZ);
+		m_pD->m_selectionBoxRadius = vec3f(radiusX, radiusY, radiusZ);
+	}
+
+	//---------------------------------------------------------------------------------------------
+	void NodeModifierMeshClipMorphPlane::SetVertexSelectionBone(const char* strBoneName, float maxEffectRadius)
+	{
+		m_pD->m_vertexSelectionType = Private::VS_BONE_HIERARCHY;
+		m_pD->m_vertexSelectionBone = strBoneName;
+		m_pD->m_maxEffectRadius = maxEffectRadius;
+	}
+
+}
+
+#undef NODE_INPUT_COUNT
+
diff --git a/Engine/Plugins/Experimental/Mutable/Source/MutableTools/Private/MuT/NodeModifierMeshClipWithMesh.cpp b/Engine/Plugins/Experimental/Mutable/Source/MutableTools/Private/MuT/NodeModifierMeshClipWithMesh.cpp
new file mode 100644
--- /dev/null
+++ b/Engine/Plugins/Experimental/Mutable/Source/MutableTools/Private/MuT/NodeModifierMeshClipWithMesh.cpp
@@ -0,0 +1,60 @@
+// Copyright Epic Games, Inc. All Rights Reserved.
+
+
+#include "MuT/NodeModifierMeshClipWithMesh.h"
+#include "NodeModifierMeshClipWithMeshPrivate.h"
+#include "MuR/MeshPrivate.h"
+
+
+namespace mu
+{
+
+	//---------------------------------------------------------------------------------------------
+	// Static initialisation
+	//---------------------------------------------------------------------------------------------
+    NODE_TYPE NodeModifierMeshClipWithMesh::Private::s_type =
+            NODE_TYPE( "NodeModifierMeshClipWithMesh", NodeModifier::GetStaticType() );
+
+
+	//---------------------------------------------------------------------------------------------
+	//!
+	//---------------------------------------------------------------------------------------------
+
+    MUTABLE_IMPLEMENT_NODE( NodeModifierMeshClipWithMesh, EType::MeshClipWithMesh, Node, Node::EType::Modifier)
+
+
+	//---------------------------------------------------------------------------------------------
+	// Node Interface
+	//---------------------------------------------------------------------------------------------
+    int NodeModifierMeshClipWithMesh::GetInputCount() const
+	{
+		return 0;
+	}
+
+
+	//---------------------------------------------------------------------------------------------
+    Node* NodeModifierMeshClipWithMesh::GetInputNode( int i ) const
+	{
+		check( i>=0 && i< GetInputCount());
+        (void)i;
+        return nullptr;
+	}
+
+
+	//---------------------------------------------------------------------------------------------
+    void NodeModifierMeshClipWithMesh::SetInputNode( int i, NodePtr )
+	{
+		check( i>=0 && i< GetInputCount());
+        (void)i;
+    }
+
+
+	//---------------------------------------------------------------------------------------------
+	// Own Interface
+	//---------------------------------------------------------------------------------------------
+	void NodeModifierMeshClipWithMesh::SetClipMesh(NodeMesh* clipMesh)
+	{
+		m_pD->m_clipMesh = clipMesh;
+	}
+
+}
diff --git a/Engine/Plugins/Experimental/Mutable/Source/MutableTools/Private/MuT/NodeModifierMeshClipWithMeshPrivate.h b/Engine/Plugins/Experimental/Mutable/Source/MutableTools/Private/MuT/NodeModifierMeshClipWithMeshPrivate.h
new file mode 100644
--- /dev/null
+++ b/Engine/Plugins/Experimental/Mutable/Source/MutableTools/Private/MuT/NodeModifierMeshClipWithMeshPrivate.h
@@ -0,0 +1,60 @@
+// Copyright Epic Games, Inc. All Rights Reserved.
+
+#pragma once
+
+#include "NodeModifierPrivate.h"
+#include "MuT/NodeModifierMeshClipWithMesh.h"
+#include "MuT/NodeMesh.h"
+#include "MuT/AST.h"
+
+#include "MuR/MemoryPrivate.h"
+#include "MuR/MutableMath.h"
+
+
+namespace mu
+{
+
+    class NodeModifierMeshClipWithMesh::Private : public NodeModifier::Private
+	{
+	public:
+
+		MUTABLE_DEFINE_CONST_VISITABLE()
+
+	public:
+
+		Private()
+		{
+		}
+
+		static NODE_TYPE s_type;
+
+		//! 
+		NodeMeshPtr m_clipMesh;
+
+		//!
+		void Serialise( OutputArchive& arch ) const
+		{
+			NodeModifier::Private::Serialise(arch);
+
+            uint32_t ver = 0;
+			arch << ver;
+
+			arch << m_clipMesh;
+		}
+
+		//!
+		void Unserialise( InputArchive& arch )
+		{
+			NodeModifier::Private::Unserialise( arch );
+			
+            uint32_t ver;
+			arch >> ver;
+            check(ver<=0);
+
+			arch >> m_clipMesh;
+		}
+
+	};
+
+
+}
diff --git a/Engine/Plugins/Experimental/Mutable/Source/MutableTools/Private/MuT/NodeModifierPrivate.h b/Engine/Plugins/Experimental/Mutable/Source/MutableTools/Private/MuT/NodeModifierPrivate.h
new file mode 100644
--- /dev/null
+++ b/Engine/Plugins/Experimental/Mutable/Source/MutableTools/Private/MuT/NodeModifierPrivate.h
@@ -0,0 +1,47 @@
+// Copyright Epic Games, Inc. All Rights Reserved.
+
+#pragma once
+
+#include "MuT/NodePrivate.h"
+#include "MuT/NodeModifier.h"
+
+#include "MuR/MemoryPrivate.h"
+
+
+namespace mu
+{
+
+	class NodeModifier::Private : public Node::Private
+	{
+	public:
+
+		// Tags that target surface need to have enabled to receive this modifier.
+		vector<mu::string> m_tags;
+
+		// Wether the modifier has to be applied after the normal node operations or before
+		bool m_applyBeforeNormalOperations = true;
+
+		//!
+		void Serialise(OutputArchive& arch) const
+		{
+            uint32_t ver = 2;
+			arch << ver;
+
+			arch << m_tags;
+			arch << m_applyBeforeNormalOperations;
+		}
+
+		//!
+		void Unserialise(InputArchive& arch)
+		{
+            uint32_t ver;
+			arch >> ver;
+			check(ver == 2);
+
+			arch >> m_tags;
+			arch >> m_applyBeforeNormalOperations;
+		}
+
+	};
+
+}
diff --git a/Engine/Plugins/Experimental/Mutable/Source/MutableTools/Private/MuT/NodeObject.cpp b/Engine/Plugins/Experimental/Mutable/Source/MutableTools/Private/MuT/NodeObject.cpp
new file mode 100644
--- /dev/null
+++ b/Engine/Plugins/Experimental/Mutable/Source/MutableTools/Private/MuT/NodeObject.cpp
@@ -0,0 +1,71 @@
+// Copyright Epic Games, Inc. All Rights Reserved.
+
+
+#include "MuT/NodeObject.h"
+#include "MuT/NodeObjectPrivate.h"
+
+#include "MuT/NodeObjectNew.h"
+#include "MuT/NodeObjectGroup.h"
+#include "MuT/NodeObjectState.h"
+
+
+namespace mu
+{
+
+
+	//---------------------------------------------------------------------------------------------
+	// Static initialisation
+	//---------------------------------------------------------------------------------------------
+	NODE_TYPE s_nodeObjectType = NODE_TYPE( "NodeObject", NodeObject::GetStaticType() );
+
+
+	//---------------------------------------------------------------------------------------------
+	void NodeObject::Serialise( const NodeObject* p, OutputArchive& arch )
+	{
+        uint32_t ver = 0;
+		arch << ver;
+
+		arch << uint32_t(p->Type);
+		p->SerialiseWrapper(arch);
+    }
+
+    //---------------------------------------------------------------------------------------------
+	NodeObjectPtr NodeObject::StaticUnserialise( InputArchive& arch )
+	{
+        uint32_t ver;
+		arch >> ver;
+		check( ver == 0 );
+
+        uint32_t id;
+		arch >> id;
+
+		switch (id)
+		{
+		case 0 :  return NodeObjectNew::StaticUnserialise( arch ); break;
+		case 1 :  return NodeObjectGroup::StaticUnserialise( arch ); break;
+        case 2 :  return NodeObjectState::StaticUnserialise( arch ); break;
+        //case 3 :  return NodeObjectTransform::StaticUnserialise( arch ); break;
+        default : check(false);
+		}
+
+		return 0;
+	}
+
+
+	//---------------------------------------------------------------------------------------------
+	const NODE_TYPE* NodeObject::GetType() const
+	{
+		return GetStaticType();
+	}
+
+
+	//---------------------------------------------------------------------------------------------
+	const NODE_TYPE* NodeObject::GetStaticType()
+	{
+		return &s_nodeObjectType;
+	}
+
+
+}
+
+
diff --git a/Engine/Plugins/Experimental/Mutable/Source/MutableTools/Private/MuT/NodeObjectGroup.cpp b/Engine/Plugins/Experimental/Mutable/Source/MutableTools/Private/MuT/NodeObjectGroup.cpp
new file mode 100644
--- /dev/null
+++ b/Engine/Plugins/Experimental/Mutable/Source/MutableTools/Private/MuT/NodeObjectGroup.cpp
@@ -0,0 +1,178 @@
+// Copyright Epic Games, Inc. All Rights Reserved.
+
+
+#include "MuT/NodeObjectGroup.h"
+#include "NodeObjectGroupPrivate.h"
+
+#include "MuT/NodeLayout.h"
+
+
+namespace mu
+{
+
+
+	//---------------------------------------------------------------------------------------------
+	// Static initialisation
+	//---------------------------------------------------------------------------------------------
+	NODE_TYPE NodeObjectGroup::Private::s_type =
+            NODE_TYPE( "ObjectGroup", NodeObjectGroup::GetStaticType() );
+
+
+	//---------------------------------------------------------------------------------------------
+	//!
+	//---------------------------------------------------------------------------------------------
+
+	MUTABLE_IMPLEMENT_NODE( NodeObjectGroup, EType::Group, Node, Node::EType::Object)
+
+
+	//---------------------------------------------------------------------------------------------
+	// Node Interface
+	//---------------------------------------------------------------------------------------------
+	int NodeObjectGroup::GetInputCount() const
+	{
+		return (int)m_pD->m_children.size();
+	}
+
+
+	//---------------------------------------------------------------------------------------------
+	Node* NodeObjectGroup::GetInputNode( int i ) const
+	{
+		check( i>=0 && i<GetInputCount() );
+
+		Node* pResult = 0;
+
+		if ( i<(int)m_pD->m_children.size() )
+		{
+			pResult = m_pD->m_children[i].get();
+		}
+
+		return pResult;
+	}
+
+
+	//---------------------------------------------------------------------------------------------
+	void NodeObjectGroup::SetInputNode( int i, NodePtr pNode )
+	{
+		check( i>=0 && i<GetInputCount() );
+
+		if ( i<(int)m_pD->m_children.size() )
+		{
+			m_pD->m_children[i] = dynamic_cast<NodeObject*>( pNode.get() );
+		}
+	}
+
+
+	//---------------------------------------------------------------------------------------------
+	// Own Interface
+	//---------------------------------------------------------------------------------------------
+	const char* NodeObjectGroup::GetName() const
+	{
+		return m_pD->m_name.c_str();
+	}
+
+
+	//---------------------------------------------------------------------------------------------
+	void NodeObjectGroup::SetName( const char* strName )
+	{
+		if( strName )
+		{
+			m_pD->m_name = strName;
+		}
+		else
+		{
+			m_pD->m_name = "";
+		}
+	}
+
+
+	const char* NodeObjectGroup::GetUid() const
+	{
+		return m_pD->m_uid.c_str();
+	}
+
+
+	//---------------------------------------------------------------------------------------------
+	void NodeObjectGroup::SetUid( const char* strUid )
+	{
+		if( strUid )
+		{
+			m_pD->m_uid = strUid;
+		}
+		else
+		{
+			m_pD->m_uid = "";
+		}
+	}
+
+
+	//---------------------------------------------------------------------------------------------
+	NodeObjectGroup::CHILD_SELECTION NodeObjectGroup::GetSelectionType() const
+	{
+		return m_pD->m_type;
+	}
+
+
+	//---------------------------------------------------------------------------------------------
+	void NodeObjectGroup::SetSelectionType( CHILD_SELECTION t )
+	{
+		m_pD->m_type = t;
+	}
+
+
+	//---------------------------------------------------------------------------------------------
+	int NodeObjectGroup::GetChildCount() const
+	{
+		return (int)m_pD->m_children.size();
+	}
+
+
+	//---------------------------------------------------------------------------------------------
+	void NodeObjectGroup::SetChildCount( int num )
+	{
+		check( num >=0 );
+		m_pD->m_children.resize( num );
+	}
+
+
+	//---------------------------------------------------------------------------------------------
+	NodeObjectPtr NodeObjectGroup::GetChild( int index ) const
+	{
+		check( index >=0 && index < (int)m_pD->m_children.size() );
+
+		return m_pD->m_children[ index ].get();
+	}
+
+
+	//---------------------------------------------------------------------------------------------
+	void NodeObjectGroup::SetChild( int index, NodeObjectPtr pObject )
+	{
+		check( index >=0 && index < (int)m_pD->m_children.size() );
+
+		m_pD->m_children[ index ] = pObject;
+	}
+
+
+	//---------------------------------------------------------------------------------------------
+    NodeLayoutPtr NodeObjectGroup::Private::GetLayout(int lod, int component, int surface, int texture ) const
+	{
+		NodeLayoutPtr pLayout;
+
+		for ( size_t i=0; !pLayout && i<m_children.size(); ++i )
+		{
+			if (m_children[i])
+			{
+				NodeObject::Private* pPrivate =
+						dynamic_cast<NodeObject::Private*>( m_children[i]->GetBasePrivate() );
+
+                pLayout = pPrivate->GetLayout( lod, component, surface, texture );
+			}
+		}
+
+		// TODO: layout index
+		return pLayout;
+	}
+
+
+}
+
+
diff --git a/Engine/Plugins/Experimental/Mutable/Source/MutableTools/Private/MuT/NodeObjectGroupPrivate.h b/Engine/Plugins/Experimental/Mutable/Source/MutableTools/Private/MuT/NodeObjectGroupPrivate.h
new file mode 100644
--- /dev/null
+++ b/Engine/Plugins/Experimental/Mutable/Source/MutableTools/Private/MuT/NodeObjectGroupPrivate.h
@@ -0,0 +1,70 @@
+// Copyright Epic Games, Inc. All Rights Reserved.
+
+#pragma once
+
+#include "MuT/NodeObjectPrivate.h"
+
+#include "MuT/NodeObjectGroup.h"
+#include "MuT/NodeLayout.h"
+
+#include "MuR/MemoryPrivate.h"
+
+
+namespace mu
+{
+
+	MUTABLE_DEFINE_ENUM_SERIALISABLE( NodeObjectGroup::CHILD_SELECTION )
+
+
+	class NodeObjectGroup::Private : public NodeObject::Private
+	{
+	public:
+
+		MUTABLE_DEFINE_CONST_VISITABLE()
+
+	public:
+
+		static NODE_TYPE s_type;
+
+		string m_name;
+		string m_uid;
+
+		CHILD_SELECTION m_type;
+
+		//! Set the child selection type
+		void SetSelectionType( CHILD_SELECTION );
+
+		vector<NodeObjectPtr> m_children;
+
+		//!
+		void Serialise( OutputArchive& arch ) const
+		{
+            uint32_t ver = 1;
+			arch << ver;
+
+			arch << m_type;
+			arch << m_name;
+			arch << m_uid;
+			arch << m_children;
+		}
+
+		//!
+		void Unserialise( InputArchive& arch )
+		{
+            uint32_t ver;
+			arch >> ver;
+            check(ver==1);
+
+			arch >> m_type;
+			arch >> m_name;
+            arch >> m_uid;
+			arch >> m_children;
+		}
+
+
+        // NodeObject::Private interface
+        NodeLayoutPtr GetLayout( int lod, int component, int surface, int texture ) const override;
+
+	};
+
+}
diff --git a/Engine/Plugins/Experimental/Mutable/Source/MutableTools/Private/MuT/NodeObjectNew.cpp b/Engine/Plugins/Experimental/Mutable/Source/MutableTools/Private/MuT/NodeObjectNew.cpp
new file mode 100644
--- /dev/null
+++ b/Engine/Plugins/Experimental/Mutable/Source/MutableTools/Private/MuT/NodeObjectNew.cpp
@@ -0,0 +1,329 @@
+// Copyright Epic Games, Inc. All Rights Reserved.
+
+
+#include "MuT/NodeObjectNew.h"
+#include "MuT/NodeObjectNewPrivate.h"
+
+#include "MuT/NodeMeshPrivate.h"
+#include "MuT/NodeLayout.h"
+#include "MuT/NodeMeshConstant.h"
+#include "MuT/NodeComponentNew.h"
+#include "MuT/NodeSurfaceNew.h"
+
+
+namespace mu
+{
+
+
+	//---------------------------------------------------------------------------------------------
+	// Static initialisation
+	//---------------------------------------------------------------------------------------------
+	NODE_TYPE NodeObjectNew::Private::s_type =
+			NODE_TYPE( "Object", NodeObjectNew::GetStaticType() );
+
+
+	//---------------------------------------------------------------------------------------------
+	//!
+	//---------------------------------------------------------------------------------------------
+
+	MUTABLE_IMPLEMENT_NODE( NodeObjectNew, EType::New, Node, Node::EType::Object )
+
+
+	//---------------------------------------------------------------------------------------------
+	// Node Interface
+	//---------------------------------------------------------------------------------------------
+	int NodeObjectNew::GetInputCount() const
+	{
+		return (int)(m_pD->m_lods.size() + m_pD->m_children.size());
+	}
+
+
+	//---------------------------------------------------------------------------------------------
+	Node* NodeObjectNew::GetInputNode( int i ) const
+	{
+		check( i>=0 && i<GetInputCount() );
+
+		Node* pResult = 0;
+
+		if ( i<(int)m_pD->m_lods.size() )
+		{
+			pResult = m_pD->m_lods[i].get();
+		}
+		else
+		{
+			i -= (int)m_pD->m_lods.size();
+			pResult = m_pD->m_children[i].get();
+		}
+
+		return pResult;
+	}
+
+
+	//---------------------------------------------------------------------------------------------
+	void NodeObjectNew::SetInputNode( int i, NodePtr pNode )
+	{
+		check( i>=0 && i<GetInputCount() );
+
+		if ( i<(int)m_pD->m_lods.size() )
+		{
+			m_pD->m_lods[i] = dynamic_cast<NodeLOD*>( pNode.get() );
+		}
+		else
+		{
+			i -= (int)m_pD->m_lods.size();
+			m_pD->m_children[i] = dynamic_cast<NodeObject*>( pNode.get() );
+		}
+	}
+
+
+	//---------------------------------------------------------------------------------------------
+	// Own Interface
+	//---------------------------------------------------------------------------------------------
+	const char* NodeObjectNew::GetName() const
+	{
+		return m_pD->m_name.c_str();
+	}
+
+
+	//---------------------------------------------------------------------------------------------
+	void NodeObjectNew::SetName( const char* strName )
+	{
+		if( strName )
+		{
+			m_pD->m_name = strName;
+		}
+		else
+		{
+			m_pD->m_name = "";
+		}
+	}
+
+
+	const char* NodeObjectNew::GetUid() const
+	{
+		return m_pD->m_uid.c_str();
+	}
+
+
+	//---------------------------------------------------------------------------------------------
+	void NodeObjectNew::SetUid( const char* strUid )
+	{
+		if( strUid )
+		{
+			m_pD->m_uid = strUid;
+		}
+		else
+		{
+			m_pD->m_uid = "";
+		}
+	}
+
+
+	//---------------------------------------------------------------------------------------------
+	int NodeObjectNew::GetLODCount() const
+	{
+		return (int)m_pD->m_lods.size();
+	}
+
+
+	//---------------------------------------------------------------------------------------------
+	void NodeObjectNew::SetLODCount( int num )
+	{
+		check( num >=0 );
+		m_pD->m_lods.resize( num );
+	}
+
+
+	//---------------------------------------------------------------------------------------------
+	NodeLODPtr NodeObjectNew::GetLOD( int index ) const
+	{
+		check( index >=0 && index < (int)m_pD->m_lods.size() );
+
+		return m_pD->m_lods[ index ].get();
+	}
+
+
+	//---------------------------------------------------------------------------------------------
+	void NodeObjectNew::SetLOD( int index, NodeLODPtr pLOD )
+	{
+		check( index >=0 && index < (int)m_pD->m_lods.size() );
+
+		m_pD->m_lods[ index ] = pLOD;
+	}
+
+
+	//---------------------------------------------------------------------------------------------
+	int NodeObjectNew::GetChildCount() const
+	{
+		return (int)m_pD->m_children.size();
+	}
+
+
+	//---------------------------------------------------------------------------------------------
+	void NodeObjectNew::SetChildCount( int num )
+	{
+		check( num >=0 );
+		m_pD->m_children.resize( num );
+	}
+
+
+	//---------------------------------------------------------------------------------------------
+	NodeObjectPtr NodeObjectNew::GetChild( int index ) const
+	{
+		check( index >=0 && index < (int)m_pD->m_children.size() );
+
+		return m_pD->m_children[ index ].get();
+	}
+
+
+	//---------------------------------------------------------------------------------------------
+	void NodeObjectNew::SetChild( int index, NodeObjectPtr pObject )
+	{
+		check( index >=0 && index < (int)m_pD->m_children.size() );
+
+		m_pD->m_children[ index ] = pObject;
+	}
+
+
+	//---------------------------------------------------------------------------------------------
+    NodeLayoutPtr NodeObjectNew::Private::GetLayout( int lod,
+                                                     int component,
+                                                     int surface,
+                                                     int //texture
+                                                     ) const
+	{
+		check( lod>=0 && lod<(int)m_lods.size() );
+
+		NodeLayoutPtr pLayout;
+
+		const NodeComponentNew* pComp =
+			dynamic_cast<const NodeComponentNew*> ( m_lods[lod]->GetComponent( component ).get() );
+
+        if (pComp)
+        {
+            const NodeSurfaceNew* pSurface =
+                    dynamic_cast<const NodeSurfaceNew*> ( pComp->GetSurface( surface ) );
+
+            if ( pSurface )
+            {
+                // TODO: Look for the layout index of the given texture
+                // TODO: Multiple meshes
+                if ( pSurface->GetMeshCount()>0 )
+                {
+                     const NodeMesh* pMesh =
+                             dynamic_cast<const NodeMesh*> ( pSurface->GetMesh( 0 ).get() );
+
+                    if ( pMesh )
+                    {
+                        NodeMesh::Private* pPrivate =
+                                dynamic_cast<NodeMesh::Private*>( pMesh->GetBasePrivate() );
+
+                        pLayout = pPrivate->GetLayout( 0 );
+                    }
+                }
+            }
+        }
+
+		// TODO: layout index
+		return pLayout;
+	}
+
+
+    //---------------------------------------------------------------------------------------------
+    bool NodeObjectNew::Private::HasComponent( const NodeComponent* pComponent ) const
+    {
+        bool found = false;
+        for ( size_t l=0; !found && l<m_lods.size(); ++l )
+        {
+            found = std::find( m_lods[l]->GetPrivate()->m_components.begin(),
+                               m_lods[l]->GetPrivate()->m_components.end(),
+                               pComponent )
+                    != m_lods[l]->GetPrivate()->m_components.end();
+        }
+
+        return found;
+    }
+
+
+	//---------------------------------------------------------------------------------------------
+	int NodeObjectNew::GetStateCount() const
+	{
+		return (int)m_pD->m_states.size();
+	}
+
+
+	//---------------------------------------------------------------------------------------------
+	void NodeObjectNew::SetStateCount( int c )
+	{
+		m_pD->m_states.resize( c );
+	}
+
+
+	//---------------------------------------------------------------------------------------------
+	const char* NodeObjectNew::GetStateName( int s ) const
+	{
+		check( s>=0 && s<GetStateCount() );
+		return m_pD->m_states[s].m_name.c_str();
+	}
+
+
+	//---------------------------------------------------------------------------------------------
+	void NodeObjectNew::SetStateName( int s, const char* n )
+	{
+		check( s>=0 && s<GetStateCount() );
+		m_pD->m_states[s].m_name = n;
+	}
+
+
+	//---------------------------------------------------------------------------------------------
+	bool NodeObjectNew::HasStateParam( int s, const char* param ) const
+	{
+		check( s>=0 && s<GetStateCount() );
+		return std::find( m_pD->m_states[s].m_runtimeParams.begin(),
+						  m_pD->m_states[s].m_runtimeParams.end(),
+						  param )
+				!=
+				m_pD->m_states[s].m_runtimeParams.end();
+	}
+
+
+	//---------------------------------------------------------------------------------------------
+	void NodeObjectNew::AddStateParam( int s, const char* param )
+	{
+		check( s>=0 && s<GetStateCount() );
+
+		if (!HasStateParam(s,param))
+		{
+			m_pD->m_states[s].m_runtimeParams.push_back( param );
+		}
+	}
+
+
+	//---------------------------------------------------------------------------------------------
+	void NodeObjectNew::RemoveStateParam( int s, const char* param )
+	{
+		check( s>=0 && s<GetStateCount() );
+
+		vector<string>::iterator it = std::find( m_pD->m_states[s].m_runtimeParams.begin(),
+														   m_pD->m_states[s].m_runtimeParams.end(),
+														   param );
+		if ( it != m_pD->m_states[s].m_runtimeParams.end() )
+		{
+			m_pD->m_states[s].m_runtimeParams.erase( it );
+		}
+	}
+
+
+    //---------------------------------------------------------------------------------------------
+    void NodeObjectNew::SetStateProperties( int s, bool avoidRuntimeCompression, bool onlyFirstLOD, int firstLOD )
+    {
+        check( s>=0 && s<GetStateCount() );
+
+        m_pD->m_states[s].m_optimisation.m_avoidRuntimeCompression = avoidRuntimeCompression;
+        m_pD->m_states[s].m_optimisation.m_onlyFirstLOD = onlyFirstLOD;
+        m_pD->m_states[s].m_optimisation.m_firstLOD = firstLOD;
+    }
+
+}
+
+
diff --git a/Engine/Plugins/Experimental/Mutable/Source/MutableTools/Private/MuT/NodeObjectNewPrivate.h b/Engine/Plugins/Experimental/Mutable/Source/MutableTools/Private/MuT/NodeObjectNewPrivate.h
new file mode 100644
--- /dev/null
+++ b/Engine/Plugins/Experimental/Mutable/Source/MutableTools/Private/MuT/NodeObjectNewPrivate.h
@@ -0,0 +1,76 @@
+// Copyright Epic Games, Inc. All Rights Reserved.
+
+#pragma once
+
+#include "MuT/NodeObjectPrivate.h"
+#include "NodeLODPrivate.h"
+#include "MuT/CompilerPrivate.h"
+
+#include "MuT/NodeObjectNew.h"
+#include "MuT/NodeLayout.h"
+
+#include "MuR/MemoryPrivate.h"
+
+
+namespace mu
+{
+
+
+	class NodeObjectNew::Private : public NodeObject::Private
+	{
+	public:
+
+		MUTABLE_DEFINE_CONST_VISITABLE()
+
+	public:
+
+		static NODE_TYPE s_type;
+
+		string m_name;
+		string m_uid;
+
+		vector<NodeLODPtr> m_lods;
+
+		vector<NodeObjectPtr> m_children;
+
+
+		//! List of states
+        vector<OBJECT_STATE> m_states;
+
+
+		//!
+		void Serialise( OutputArchive& arch ) const
+		{
+            uint32_t ver = 2;
+			arch << ver;
+
+			arch << m_name;
+			arch << m_uid;
+			arch << m_lods;
+			arch << m_children;
+			arch << m_states;
+		}
+
+		//!
+		void Unserialise( InputArchive& arch )
+		{
+            uint32_t ver;
+			arch >> ver;
+            check( ver==2 );
+
+			arch >> m_name;
+            arch >> m_uid;
+			arch >> m_lods;
+			arch >> m_children;
+			arch >> m_states;
+		}
+
+        //! Return true if the given component is set in any lod of this object.
+        bool HasComponent( const NodeComponent* pComponent ) const;
+
+        // NodeObject::Private interface
+        NodeLayoutPtr GetLayout( int lod, int component, int surface, int texture ) const override;
+
+	};
+
+}
