Moved the Mutable plugin files to live under a directory named after the module. This allows us to get rid of long include paths that contain public and private folder names in the path and fixes head

diff --git a/Engine/Plugins/Experimental/Mutable/Source/MutableTools/Private/MuT/NodeObjectState.cpp b/Engine/Plugins/Experimental/Mutable/Source/MutableTools/Private/MuT/NodeObjectState.cpp
new file mode 100644
--- /dev/null
+++ b/Engine/Plugins/Experimental/Mutable/Source/MutableTools/Private/MuT/NodeObjectState.cpp
@@ -0,0 +1,195 @@
+// Copyright Epic Games, Inc. All Rights Reserved.
+
+
+#include "MuT/NodeObjectState.h"
+#include "MuT/NodeObjectStatePrivate.h"
+
+
+namespace mu
+{
+
+
+	//---------------------------------------------------------------------------------------------
+	// Static initialisation
+	//---------------------------------------------------------------------------------------------
+	NODE_TYPE NodeObjectState::Private::s_type =
+			NODE_TYPE( "ObjectState", NodeObjectState::GetStaticType() );
+
+
+	//---------------------------------------------------------------------------------------------
+    NodeLayoutPtr NodeObjectState::Private::GetLayout(int lod, int component, int surface, int texture ) const
+	{
+		NodeLayoutPtr pResult;
+
+		if ( m_pSource )
+		{
+			const NodeObject::Private* pPrivate = dynamic_cast<const NodeObject::Private*>
+					( m_pSource->GetBasePrivate() );
+            pResult = pPrivate->GetLayout( lod, component, surface, texture );
+		}
+
+		return pResult;
+	}
+
+
+	//---------------------------------------------------------------------------------------------
+	//!
+	//---------------------------------------------------------------------------------------------
+
+	MUTABLE_IMPLEMENT_NODE( NodeObjectState, NodeObject::EType::State, Node, Node::EType::Object)
+
+
+	//---------------------------------------------------------------------------------------------
+	// Node Interface
+	//---------------------------------------------------------------------------------------------
+	int NodeObjectState::GetInputCount() const
+	{
+		return 2;
+	}
+
+
+	//---------------------------------------------------------------------------------------------
+	Node* NodeObjectState::GetInputNode( int i ) const
+	{
+		check( i>=0 && i<GetInputCount() );
+
+		Node* pResult = 0;
+
+		if (i==0)
+		{
+			pResult = m_pD->m_pSource.get();
+		}
+		else
+		{
+			pResult = m_pD->m_pRoot.get();
+		}
+
+		return pResult;
+	}
+
+
+	//---------------------------------------------------------------------------------------------
+	void NodeObjectState::SetInputNode( int i, NodePtr pNode )
+	{
+		check( i>=0 && i<GetInputCount() );
+
+		if (i==0)
+		{
+			m_pD->m_pSource = dynamic_cast<NodeObject*>( pNode.get() );
+		}
+		else
+		{
+			m_pD->m_pRoot = dynamic_cast<NodeObject*>( pNode.get() );
+		}
+	}
+
+
+	//---------------------------------------------------------------------------------------------
+	// NodeObject Interface
+	//---------------------------------------------------------------------------------------------
+	const char* NodeObjectState::GetName() const
+	{
+		const char* strRes = "";
+		return strRes;
+	}
+
+
+	//---------------------------------------------------------------------------------------------
+    void NodeObjectState::SetName( const char* )
+	{
+	}
+
+
+	const char* NodeObjectState::GetUid() const
+	{
+		const char* strRes = "";
+		return strRes;
+	}
+
+
+	//---------------------------------------------------------------------------------------------
+	void NodeObjectState::SetUid( const char* )
+	{
+	}
+
+
+	//---------------------------------------------------------------------------------------------
+	// Own Interface
+	//---------------------------------------------------------------------------------------------
+	NodeObjectPtr NodeObjectState::GetSource() const
+	{
+		return m_pD->m_pSource.get();
+	}
+
+
+	//---------------------------------------------------------------------------------------------
+	void NodeObjectState::SetSource( NodeObjectPtr pSource )
+	{
+		m_pD->m_pSource = pSource;
+	}
+
+
+	//---------------------------------------------------------------------------------------------
+	NodeObjectPtr NodeObjectState::GetStateRoot() const
+	{
+		return m_pD->m_pRoot.get();
+	}
+
+
+	//---------------------------------------------------------------------------------------------
+	void NodeObjectState::SetStateRoot( NodeObjectPtr pState )
+	{
+		m_pD->m_pRoot = pState;
+	}
+
+
+	//---------------------------------------------------------------------------------------------
+	const char* NodeObjectState::GetStateName() const
+	{
+		return m_pD->m_state.m_name.c_str();
+	}
+
+
+	//---------------------------------------------------------------------------------------------
+	void NodeObjectState::SetStateName( const char* n )
+	{
+		m_pD->m_state.m_name = n;
+	}
+
+
+	//---------------------------------------------------------------------------------------------
+	bool NodeObjectState::HasStateParam( const char* param ) const
+	{
+		return std::find( m_pD->m_state.m_runtimeParams.begin(),
+						  m_pD->m_state.m_runtimeParams.end(),
+						  param )
+				!=
+				m_pD->m_state.m_runtimeParams.end();
+	}
+
+
+	//---------------------------------------------------------------------------------------------
+	void NodeObjectState::AddStateParam( const char* param )
+	{
+		if (!HasStateParam(param))
+		{
+			m_pD->m_state.m_runtimeParams.push_back( param );
+		}
+	}
+
+
+	//---------------------------------------------------------------------------------------------
+	void NodeObjectState::RemoveStateParam( const char* param )
+	{
+		vector<string>::iterator it = std::find( m_pD->m_state.m_runtimeParams.begin(),
+														   m_pD->m_state.m_runtimeParams.end(),
+														   param );
+		if ( it != m_pD->m_state.m_runtimeParams.end() )
+		{
+			m_pD->m_state.m_runtimeParams.erase( it );
+		}
+	}
+
+}
+
+
diff --git a/Engine/Plugins/Experimental/Mutable/Source/MutableTools/Private/MuT/NodeObjectStatePrivate.h b/Engine/Plugins/Experimental/Mutable/Source/MutableTools/Private/MuT/NodeObjectStatePrivate.h
new file mode 100644
--- /dev/null
+++ b/Engine/Plugins/Experimental/Mutable/Source/MutableTools/Private/MuT/NodeObjectStatePrivate.h
@@ -0,0 +1,62 @@
+// Copyright Epic Games, Inc. All Rights Reserved.
+
+
+#pragma once
+
+#include "MuT/NodeObjectPrivate.h"
+#include "MuT/NodeObjectState.h"
+#include "MuT/CompilerPrivate.h"
+
+
+namespace mu
+{
+
+	class NodeObjectState::Private : public NodeObject::Private
+	{
+	public:
+
+		MUTABLE_DEFINE_CONST_VISITABLE()
+
+	public:
+
+
+        // NodeObject::Private interface
+        NodeLayoutPtr GetLayout( int lod, int component, int surface, int texture ) const override;
+
+
+		//-----------------------------------------------------------------------------------------
+		// Own interface
+		//-----------------------------------------------------------------------------------------
+		static NODE_TYPE s_type;
+
+        OBJECT_STATE m_state;
+
+		NodeObjectPtr m_pSource;
+		NodeObjectPtr m_pRoot;
+
+		//!
+		void Serialise( OutputArchive& arch ) const
+		{
+            uint32_t ver = 0;
+			arch << ver;
+
+			arch << m_state;
+			arch << m_pSource;
+			arch << m_pRoot;
+		}
+
+		//!
+		void Unserialise( InputArchive& arch )
+		{
+            uint32_t ver;
+			arch >> ver;
+			check(ver==0);
+
+			arch >> m_state;
+			arch >> m_pSource;
+			arch >> m_pRoot;
+		}
+	};
+
+}
+
diff --git a/Engine/Plugins/Experimental/Mutable/Source/MutableTools/Private/MuT/NodePatchImage.cpp b/Engine/Plugins/Experimental/Mutable/Source/MutableTools/Private/MuT/NodePatchImage.cpp
new file mode 100644
--- /dev/null
+++ b/Engine/Plugins/Experimental/Mutable/Source/MutableTools/Private/MuT/NodePatchImage.cpp
@@ -0,0 +1,160 @@
+// Copyright Epic Games, Inc. All Rights Reserved.
+
+
+#include "MuT/NodePatchImage.h"
+#include "NodePatchImagePrivate.h"
+
+#include "MuT/NodeImage.h"
+
+
+namespace mu
+{
+
+
+    //---------------------------------------------------------------------------------------------
+    // Static initialisation
+    //---------------------------------------------------------------------------------------------
+    NODE_TYPE NodePatchImage::Private::s_type =
+            NODE_TYPE( "PatchTexture", NodePatchImage::GetStaticType() );
+
+
+    //---------------------------------------------------------------------------------------------
+    //!
+    //---------------------------------------------------------------------------------------------
+
+    MUTABLE_IMPLEMENT_NODE( NodePatchImage, EType::PatchImage, Node, Node::EType::PatchImage);
+
+
+    //---------------------------------------------------------------------------------------------
+    // Node Interface
+    //---------------------------------------------------------------------------------------------
+    int NodePatchImage::GetInputCount() const
+    {
+        return 2;
+    }
+
+
+    //---------------------------------------------------------------------------------------------
+    Node* NodePatchImage::GetInputNode( int i ) const
+    {
+        check( i >=0 && i < GetInputCount());
+
+        Node* pResult = 0;
+
+        if ( i==0 )
+        {
+            pResult = m_pD->m_pImage.get();
+        }
+        else
+        {
+            pResult = m_pD->m_pMask.get();
+        }
+
+        return pResult;
+    }
+
+
+    //---------------------------------------------------------------------------------------------
+    void NodePatchImage::SetInputNode( int i, NodePtr pNode )
+    {
+        check( i >=0 && i < GetInputCount());
+
+        if ( i==0 )
+        {
+            m_pD->m_pImage = dynamic_cast<NodeImage*>( pNode.get() );
+        }
+        else
+        {
+            m_pD->m_pMask = dynamic_cast<NodeImage*>( pNode.get() );
+        }
+    }
+
+
+    //---------------------------------------------------------------------------------------------
+    // Own Interface
+    //---------------------------------------------------------------------------------------------
+    void NodePatchImage::SetImage( NodeImagePtr pImage )
+    {
+        m_pD->m_pImage = pImage;
+    }
+
+
+    //---------------------------------------------------------------------------------------------
+    NodeImagePtr NodePatchImage::GetImage() const
+    {
+        return m_pD->m_pImage.get();
+    }
+
+
+    //---------------------------------------------------------------------------------------------
+    void NodePatchImage::SetMask( NodeImagePtr pMask )
+    {
+        m_pD->m_pMask = pMask;
+    }
+
+
+    //---------------------------------------------------------------------------------------------
+    NodeImagePtr NodePatchImage::GetMask() const
+    {
+        return m_pD->m_pMask.get();
+    }
+
+
+    //---------------------------------------------------------------------------------------------
+    void NodePatchImage::SetBlockCount( int c )
+    {
+        m_pD->m_blocks.resize( c );
+    }
+
+
+    //---------------------------------------------------------------------------------------------
+    int NodePatchImage::GetBlockCount() const
+    {
+        return (int)m_pD->m_blocks.size();
+    }
+
+
+    //---------------------------------------------------------------------------------------------
+    void NodePatchImage::SetBlock( int index, int block )
+    {
+        check( index>=0 && index<(int)m_pD->m_blocks.size() );
+        m_pD->m_blocks[ index ] = block;
+    }
+
+
+    //---------------------------------------------------------------------------------------------
+    int NodePatchImage::GetBlock( int index ) const
+    {
+        check( index>=0 && index<(int)m_pD->m_blocks.size() );
+        return m_pD->m_blocks[ index ];
+    }
+
+
+    //---------------------------------------------------------------------------------------------
+	EBlendType NodePatchImage::GetBlendType() const
+    {
+        return m_pD->m_blendType;
+    }
+
+
+    //---------------------------------------------------------------------------------------------
+    void NodePatchImage::SetBlendType(EBlendType t)
+    {
+        m_pD->m_blendType = t;
+    }
+
+
+    //---------------------------------------------------------------------------------------------
+    bool NodePatchImage::GetApplyToAlphaChannel() const
+    {
+        return m_pD->m_applyToAlpha;
+    }
+
+
+    //---------------------------------------------------------------------------------------------
+    void NodePatchImage::SetApplyToAlphaChannel(bool a )
+    {
+        m_pD->m_applyToAlpha = a;
+    }
+
+}
diff --git a/Engine/Plugins/Experimental/Mutable/Source/MutableTools/Private/MuT/NodePatchMesh.cpp b/Engine/Plugins/Experimental/Mutable/Source/MutableTools/Private/MuT/NodePatchMesh.cpp
new file mode 100644
--- /dev/null
+++ b/Engine/Plugins/Experimental/Mutable/Source/MutableTools/Private/MuT/NodePatchMesh.cpp
@@ -0,0 +1,108 @@
+// Copyright Epic Games, Inc. All Rights Reserved.
+
+
+#include "MuT/NodePatchMesh.h"
+#include "NodePatchMeshPrivate.h"
+#include "MuR/MeshPrivate.h"
+
+
+namespace mu
+{
+
+	//---------------------------------------------------------------------------------------------
+	// Static initialisation
+	//---------------------------------------------------------------------------------------------
+	NODE_TYPE NodePatchMesh::Private::s_type =
+			NODE_TYPE( "MeshPatch", Node::GetStaticType() );
+
+
+	//---------------------------------------------------------------------------------------------
+	//!
+	//---------------------------------------------------------------------------------------------
+
+	MUTABLE_IMPLEMENT_NODE( NodePatchMesh, EType::PatchMesh, Node, Node::EType::PatchMesh);
+
+
+	//---------------------------------------------------------------------------------------------
+	// Node Interface
+	//---------------------------------------------------------------------------------------------
+	int NodePatchMesh::GetInputCount() const
+	{
+		return 2;
+	}
+
+
+	//---------------------------------------------------------------------------------------------
+	Node* NodePatchMesh::GetInputNode( int i ) const
+	{
+		Node* pResult = 0;
+
+		if( i==0 )
+		{
+			pResult = m_pD->m_pRemove.get();
+		}
+		else if (i==1)
+		{
+			pResult = m_pD->m_pAdd.get();
+		}
+		else
+		{
+			check( false );
+		}
+
+		return pResult;
+	}
+
+
+	//---------------------------------------------------------------------------------------------
+	void NodePatchMesh::SetInputNode( int i, NodePtr pNode )
+	{
+		if( i==0 )
+		{
+			m_pD->m_pRemove = dynamic_cast<NodeMesh*>( pNode.get() );
+		}
+		else if (i==1)
+		{
+			m_pD->m_pAdd = dynamic_cast<NodeMesh*>( pNode.get() );
+		}
+		else
+		{
+			check( false );
+		}
+	}
+
+
+	//---------------------------------------------------------------------------------------------
+	// Own Interface
+	//---------------------------------------------------------------------------------------------
+    NodeMesh* NodePatchMesh::GetRemove() const
+	{
+        return m_pD->m_pRemove.get();
+	}
+
+
+	//---------------------------------------------------------------------------------------------
+    void NodePatchMesh::SetRemove( NodeMesh* pRemove )
+	{
+		m_pD->m_pRemove = pRemove;
+	}
+
+
+	//---------------------------------------------------------------------------------------------
+    NodeMesh* NodePatchMesh::GetAdd() const
+	{
+        return m_pD->m_pAdd.get();
+	}
+
+
+	//---------------------------------------------------------------------------------------------
+    void NodePatchMesh::SetAdd( NodeMesh* pAdd )
+	{
+		m_pD->m_pAdd = pAdd;
+	}
+
+
+
+}
+
+
diff --git a/Engine/Plugins/Experimental/Mutable/Source/MutableTools/Private/MuT/NodePrivate.h b/Engine/Plugins/Experimental/Mutable/Source/MutableTools/Private/MuT/NodePrivate.h
new file mode 100644
--- /dev/null
+++ b/Engine/Plugins/Experimental/Mutable/Source/MutableTools/Private/MuT/NodePrivate.h
@@ -0,0 +1,115 @@
+// Copyright Epic Games, Inc. All Rights Reserved.
+
+
+#pragma once
+
+#include "MuT/Node.h"
+#include "Visitor.h"
+#include "MuT/AST.h"
+#include "MuR/Operations.h"
+#include "MuR/MemoryPrivate.h"
+
+#include <memory>
+
+namespace mu
+{
+
+
+    //!
+    class Node::Private :
+        public Base,
+        public BaseVisitable<Ptr<class ASTOp>,DefaultCatchAll,true>
+    {
+    public:
+
+        //!
+        static void CloneChildren( const Node* pSource, NodePtr pTarget, NodeMapPtr pMap );
+
+        //! This is an opaque context used to attach to reported error messages.
+		const void* m_errorContext = nullptr;
+
+		//! Generic pointer to the node owning this private.
+		const Node* m_pNode = nullptr;
+    };
+
+
+    class NodeMap::Private : public Base
+    {
+    public:
+
+        map< const void*, NodePtr > m_map;
+
+    };
+
+
+#define MUTABLE_IMPLEMENT_NODE( N, T, PN, PT )				\
+    N::N()													\
+    {														\
+		Type = T;											\
+        m_pD = new Private();								\
+		m_pD->m_pNode = this;								\
+		PN::Type = PT;										\
+	}														\
+                                                            \
+    N::~N()													\
+    {														\
+        check( m_pD );										\
+        delete m_pD;										\
+        m_pD = nullptr;										\
+    }														\
+                                                            \
+    void N::Serialise( const N* p, OutputArchive& arch )	\
+    {														\
+        arch << *p->m_pD;									\
+    }														\
+                                                            \
+	void N::SerialiseWrapper( OutputArchive& arch ) const	\
+	{														\
+		N::Serialise(this, arch);							\
+	}														\
+                                                            \
+    Ptr<N> N::StaticUnserialise( InputArchive& arch )				\
+    {														\
+        Ptr<N> pResult = new N();							\
+        arch >> *pResult->m_pD;								\
+        return pResult;										\
+    }														\
+                                                            \
+    NodePtr N::Clone( NodeMapPtr pMap ) const				\
+    {														\
+        Ptr<N> pCloned = new N();							\
+                                                            \
+        *pCloned->m_pD = *m_pD;								\
+                                                            \
+        if (!pMap)											\
+        {													\
+            pMap = new NodeMap;								\
+        }													\
+        pMap->Add( this, pCloned );							\
+                                                            \
+        Private::CloneChildren( this, pCloned, pMap );		\
+                                                            \
+        return pCloned;										\
+    }														\
+                                                            \
+    N::Private* N::GetPrivate() const						\
+    {														\
+        return m_pD;										\
+    }														\
+                                                            \
+    Node::Private* N::GetBasePrivate() const				\
+    {														\
+        return m_pD;										\
+    }														\
+                                                            \
+    const NODE_TYPE* N::GetType() const						\
+    {														\
+        return GetStaticType();								\
+    }														\
+                                                            \
+    const NODE_TYPE* N::GetStaticType()						\
+    {														\
+        return &Private::s_type;							\
+    }
+
+}
diff --git a/Engine/Plugins/Experimental/Mutable/Source/MutableTools/Private/MuT/NodeProjector.cpp b/Engine/Plugins/Experimental/Mutable/Source/MutableTools/Private/MuT/NodeProjector.cpp
new file mode 100644
--- /dev/null
+++ b/Engine/Plugins/Experimental/Mutable/Source/MutableTools/Private/MuT/NodeProjector.cpp
@@ -0,0 +1,67 @@
+// Copyright Epic Games, Inc. All Rights Reserved.
+
+
+#include "MuT/NodeProjector.h"
+
+#include "MuT/Node.h"
+#include "MuT/NodePrivate.h"
+
+
+namespace mu
+{
+
+	//---------------------------------------------------------------------------------------------
+	// Static initialisation
+	//---------------------------------------------------------------------------------------------
+	static NODE_TYPE s_nodeProjectorType = NODE_TYPE( "NodeProjector", Node::GetStaticType() );
+
+
+	//---------------------------------------------------------------------------------------------
+	void NodeProjector::Serialise( const NodeProjector* p, OutputArchive& arch )
+	{
+        uint32_t ver = 0;
+		arch << ver;
+
+		arch << uint32_t(p->Type);
+		p->SerialiseWrapper(arch);
+	}
+
+        
+	//---------------------------------------------------------------------------------------------
+	NodeProjectorPtr NodeProjector::StaticUnserialise( InputArchive& arch )
+	{
+        uint32_t ver;
+		arch >> ver;
+		check( ver == 0 );
+
+        uint32_t id;
+		arch >> id;
+
+		switch (id)
+		{
+		case 0 :  return NodeProjectorConstant::StaticUnserialise( arch ); break;
+		case 1 :  return NodeProjectorParameter::StaticUnserialise( arch ); break;
+		default : check(false);
+		}
+
+		return 0;
+	}
+
+
+	//---------------------------------------------------------------------------------------------
+	const NODE_TYPE* NodeProjector::GetType() const
+	{
+		return GetStaticType();
+	}
+
+
+	//---------------------------------------------------------------------------------------------
+	const NODE_TYPE* NodeProjector::GetStaticType()
+	{
+		return &s_nodeProjectorType;
+	}
+
+
+}
+
+
diff --git a/Engine/Plugins/Experimental/Mutable/Source/MutableTools/Private/MuT/NodeProjectorConstant.cpp b/Engine/Plugins/Experimental/Mutable/Source/MutableTools/Private/MuT/NodeProjectorConstant.cpp
new file mode 100644
--- /dev/null
+++ b/Engine/Plugins/Experimental/Mutable/Source/MutableTools/Private/MuT/NodeProjectorConstant.cpp
@@ -0,0 +1,101 @@
+// Copyright Epic Games, Inc. All Rights Reserved.
+
+
+#include "NodeProjectorPrivate.h"
+
+#include "MuR/MutableMath.h"
+
+
+namespace mu
+{
+
+	//---------------------------------------------------------------------------------------------
+	// Static initialisation
+	//---------------------------------------------------------------------------------------------
+	NODE_TYPE NodeProjectorConstant::Private::s_type =
+			NODE_TYPE( "ProjectorConstant", NodeProjector::GetStaticType() );
+
+
+	//---------------------------------------------------------------------------------------------
+	//!
+	//---------------------------------------------------------------------------------------------
+
+	MUTABLE_IMPLEMENT_NODE( NodeProjectorConstant, EType::Constant, Node, Node::EType::Projector)
+
+
+	//---------------------------------------------------------------------------------------------
+	// Node Interface
+	//---------------------------------------------------------------------------------------------
+	int NodeProjectorConstant::GetInputCount() const
+	{
+		return 0;
+	}
+
+
+	//---------------------------------------------------------------------------------------------
+    Node* NodeProjectorConstant::GetInputNode( int ) const
+	{
+		check( false );
+		return 0;
+	}
+
+
+	//---------------------------------------------------------------------------------------------
+    void NodeProjectorConstant::SetInputNode( int, NodePtr )
+	{
+		check( false );
+	}
+
+
+	//---------------------------------------------------------------------------------------------
+	// Own Interface
+	//---------------------------------------------------------------------------------------------
+    void NodeProjectorConstant::GetValue( PROJECTOR_TYPE* pType,
+                                          float* pPosX, float* pPosY, float* pPosZ,
+										  float* pDirX, float* pDirY, float* pDirZ,
+										  float* pUpX, float* pUpY, float* pUpZ,
+                                          float* pScaleU, float* pScaleV, float* pScaleW,
+                                          float* pProjectionAngle ) const
+	{
+        if (pType) *pType = m_pD->m_type;
+
+        if (pPosX) *pPosX = m_pD->m_position[0];
+        if (pPosY) *pPosY = m_pD->m_position[1];
+        if (pPosZ) *pPosZ = m_pD->m_position[2];
+
+        if (pDirX) *pDirX = m_pD->m_direction[0];
+        if (pDirY) *pDirY = m_pD->m_direction[1];
+        if (pDirZ) *pDirZ = m_pD->m_direction[2];
+
+        if (pUpX) *pUpX = m_pD->m_up[0];
+        if (pUpY) *pUpY = m_pD->m_up[1];
+        if (pUpZ) *pUpZ = m_pD->m_up[2];
+
+        if (pScaleU) *pScaleU = m_pD->m_scale[0];
+        if (pScaleV) *pScaleV = m_pD->m_scale[1];
+        if (pScaleW) *pScaleW = m_pD->m_scale[2];
+
+        if (pProjectionAngle) *pProjectionAngle = m_pD->m_projectionAngle;
+    }
+
+
+	//---------------------------------------------------------------------------------------------
+    void NodeProjectorConstant::SetValue( PROJECTOR_TYPE type,
+                                          float posX, float posY, float posZ,
+										  float dirX, float dirY, float dirZ,
+										  float upX, float upY, float upZ,
+                                          float scaleU, float scaleV, float scaleW,
+                                          float projectionAngle )
+	{
+        m_pD->m_type = type;
+		m_pD->m_position = vec3<float>( posX, posY, posZ );
+		m_pD->m_direction = vec3<float>( dirX, dirY, dirZ );
+		m_pD->m_up = vec3<float>( upX, upY, upZ );
+        m_pD->m_scale = vec3<float>( scaleU, scaleV, scaleW );
+        m_pD->m_projectionAngle = projectionAngle;
+	}
+
+
+}
+
+
diff --git a/Engine/Plugins/Experimental/Mutable/Source/MutableTools/Private/MuT/NodeProjectorParameter.cpp b/Engine/Plugins/Experimental/Mutable/Source/MutableTools/Private/MuT/NodeProjectorParameter.cpp
new file mode 100644
--- /dev/null
+++ b/Engine/Plugins/Experimental/Mutable/Source/MutableTools/Private/MuT/NodeProjectorParameter.cpp
@@ -0,0 +1,169 @@
+// Copyright Epic Games, Inc. All Rights Reserved.
+
+
+#include "NodeProjectorPrivate.h"
+
+#include "MuR/MutableMath.h"
+#include "MuR/MemoryPrivate.h"
+
+
+namespace mu
+{
+
+	//---------------------------------------------------------------------------------------------
+	// Static initialisation
+	//---------------------------------------------------------------------------------------------
+	NODE_TYPE NodeProjectorParameter::Private::s_type =
+			NODE_TYPE( "ProjectorParameter", NodeProjector::GetStaticType() );
+
+
+	//---------------------------------------------------------------------------------------------
+	//!
+	//---------------------------------------------------------------------------------------------
+
+	MUTABLE_IMPLEMENT_NODE( NodeProjectorParameter, EType::Parameter, Node, Node::EType::Projector)
+
+
+	//---------------------------------------------------------------------------------------------
+	// Node Interface
+	//---------------------------------------------------------------------------------------------
+	int NodeProjectorParameter::GetInputCount() const
+	{
+        return int( m_pD->m_ranges.size() );
+    }
+
+
+	//---------------------------------------------------------------------------------------------
+    Node* NodeProjectorParameter::GetInputNode( int i ) const
+	{
+        check( i<GetInputCount() );
+        if (i<GetInputCount())
+        {
+            return m_pD->m_ranges[i].get();
+        }
+        return nullptr;
+    }
+
+
+	//---------------------------------------------------------------------------------------------
+    void NodeProjectorParameter::SetInputNode( int i, NodePtr n )
+	{
+        check( i<GetInputCount() );
+        if (i<GetInputCount())
+        {
+            m_pD->m_ranges[i] = dynamic_cast<NodeRange*>(n.get());
+        }
+    }
+
+
+	//---------------------------------------------------------------------------------------------
+	// Own Interface
+	//---------------------------------------------------------------------------------------------
+	const char* NodeProjectorParameter::GetName() const
+	{
+		return m_pD->m_name.c_str();
+	}
+
+
+	//---------------------------------------------------------------------------------------------
+	void NodeProjectorParameter::SetName( const char* strName )
+	{
+		if ( strName )
+		{
+			m_pD->m_name = strName;
+		}
+		else
+		{
+			m_pD->m_name = "";
+		}
+	}
+
+
+	const char* NodeProjectorParameter::GetUid() const
+	{
+		return m_pD->m_uid.c_str();
+	}
+
+
+	//---------------------------------------------------------------------------------------------
+	void NodeProjectorParameter::SetUid( const char* strUid )
+	{
+		if ( strUid )
+		{
+			m_pD->m_uid = strUid;
+		}
+		else
+		{
+			m_pD->m_uid = "";
+		}
+	}
+
+
+	//---------------------------------------------------------------------------------------------
+    void NodeProjectorParameter::GetDefaultValue( PROJECTOR_TYPE* pType,
+                                                  float* pPosX, float* pPosY, float* pPosZ,
+												  float* pDirX, float* pDirY, float* pDirZ,
+												  float* pUpX, float* pUpY, float* pUpZ,
+                                                  float* pScaleU, float* pScaleV, float* pScaleW,
+                                                  float* pProjectionAngle ) const
+	{
+        if (pType) *pType = m_pD->m_type;
+
+        if (pPosX) *pPosX = m_pD->m_position[0];
+        if (pPosY) *pPosY = m_pD->m_position[1];
+        if (pPosZ) *pPosZ = m_pD->m_position[2];
+
+        if (pDirX) *pDirX = m_pD->m_direction[0];
+        if (pDirY) *pDirY = m_pD->m_direction[1];
+        if (pDirZ) *pDirZ = m_pD->m_direction[2];
+
+        if (pUpX) *pUpX = m_pD->m_up[0];
+        if (pUpY) *pUpY = m_pD->m_up[1];
+        if (pUpZ) *pUpZ = m_pD->m_up[2];
+
+        if (pScaleU) *pScaleU = m_pD->m_scale[0];
+        if (pScaleV) *pScaleV = m_pD->m_scale[1];
+        if (pScaleW) *pScaleW = m_pD->m_scale[2];
+
+        if (pProjectionAngle) *pProjectionAngle = m_pD->m_projectionAngle;
+    }
+
+
+	//---------------------------------------------------------------------------------------------
+    void NodeProjectorParameter::SetDefaultValue( PROJECTOR_TYPE type,
+                                                  float posX, float posY, float posZ,
+												  float dirX, float dirY, float dirZ,
+												  float upX, float upY, float upZ,
+                                                  float scaleU, float scaleV, float scaleW,
+                                                  float projectionAngle )
+	{
+        m_pD->m_type = type;
+        m_pD->m_position = vec3<float>( posX, posY, posZ );
+		m_pD->m_direction = vec3<float>( dirX, dirY, dirZ );
+		m_pD->m_up = vec3<float>( upX, upY, upZ );
+        m_pD->m_scale = vec3<float>( scaleU, scaleV, scaleW );
+        m_pD->m_projectionAngle = projectionAngle;
+    }
+
+
+    //---------------------------------------------------------------------------------------------
+    void NodeProjectorParameter::SetRangeCount( int i )
+    {
+        check(i>=0);
+        m_pD->m_ranges.resize(i);
+    }
+
+
+    //---------------------------------------------------------------------------------------------
+    void NodeProjectorParameter::SetRange( int i, NodeRangePtr pRange )
+    {
+        check( i>=0 && i<int(m_pD->m_ranges.size()) );
+        if ( i>=0 && i<int(m_pD->m_ranges.size()) )
+        {
+            m_pD->m_ranges[i] = pRange;
+        }
+    }
+
+}
+
+
diff --git a/Engine/Plugins/Experimental/Mutable/Source/MutableTools/Private/MuT/NodeRange.cpp b/Engine/Plugins/Experimental/Mutable/Source/MutableTools/Private/MuT/NodeRange.cpp
new file mode 100644
--- /dev/null
+++ b/Engine/Plugins/Experimental/Mutable/Source/MutableTools/Private/MuT/NodeRange.cpp
@@ -0,0 +1,74 @@
+// Copyright Epic Games, Inc. All Rights Reserved.
+
+
+#include "MuT/NodeRange.h"
+#include "MuT/Node.h"
+#include "MuT/NodePrivate.h"
+#include "MuT/NodeRangeFromScalar.h"
+
+
+namespace mu
+{
+
+	//---------------------------------------------------------------------------------------------
+	// Static initialisation
+	//---------------------------------------------------------------------------------------------
+    static NODE_TYPE s_nodeRangeType = 	NODE_TYPE( "NodeRange", Node::GetStaticType() );
+
+
+	//---------------------------------------------------------------------------------------------
+    void NodeRange::Serialise( const NodeRange* p, OutputArchive& arch )
+	{
+        uint32_t ver = 0;
+		arch << ver;
+
+#define SERIALISE_CHILDREN( C, ID ) \
+        ( const C* pTyped##ID = dynamic_cast<const C*>(p) )			\
+        {                                                           \
+            arch << (uint32_t)ID;                                   \
+            C::Serialise( pTyped##ID, arch );						\
+        }                                                           \
+
+        if SERIALISE_CHILDREN( NodeRangeFromScalar          , 0  )
+        else check(false);
+
+#undef SERIALISE_CHILDREN
+    }
+
+    //---------------------------------------------------------------------------------------------
+    NodeRangePtr NodeRange::StaticUnserialise( InputArchive& arch )
+	{
+        uint32_t ver;
+		arch >> ver;
+		check( ver == 0 );
+
+        uint32_t id;
+		arch >> id;
+
+		switch (id)
+		{
+        case 0 :  return NodeRangeFromScalar::StaticUnserialise( arch ); break;
+        default : check(false);
+		}
+
+        return nullptr;
+	}
+
+
+	//---------------------------------------------------------------------------------------------
+    const NODE_TYPE* NodeRange::GetType() const
+	{
+		return GetStaticType();
+	}
+
+
+	//---------------------------------------------------------------------------------------------
+    const NODE_TYPE* NodeRange::GetStaticType()
+	{
+        return &s_nodeRangeType;
+	}
+
+
+}
+
+
diff --git a/Engine/Plugins/Experimental/Mutable/Source/MutableTools/Private/MuT/NodeRangeFromScalar.cpp b/Engine/Plugins/Experimental/Mutable/Source/MutableTools/Private/MuT/NodeRangeFromScalar.cpp
new file mode 100644
--- /dev/null
+++ b/Engine/Plugins/Experimental/Mutable/Source/MutableTools/Private/MuT/NodeRangeFromScalar.cpp
@@ -0,0 +1,108 @@
+// Copyright Epic Games, Inc. All Rights Reserved.
+
+
+#include "MuT/NodeRangeFromScalar.h"
+#include "NodeRangeFromScalarPrivate.h"
+
+#include "MuT/NodeScalar.h"
+
+#include "MuR/ImagePrivate.h"
+
+
+#define NODE_INPUT_COUNT 	1
+
+
+namespace mu
+{
+
+	//---------------------------------------------------------------------------------------------
+	// Static initialisation
+	//---------------------------------------------------------------------------------------------
+    NODE_TYPE NodeRangeFromScalar::Private::s_type =
+            NODE_TYPE( "RangeFromScalar", NodeRange::GetStaticType() );
+
+
+	//---------------------------------------------------------------------------------------------
+	//!
+	//---------------------------------------------------------------------------------------------
+
+    MUTABLE_IMPLEMENT_NODE( NodeRangeFromScalar, EType::FromScalar, Node, Node::EType::Range)
+
+
+	//---------------------------------------------------------------------------------------------
+	// Node Interface
+	//---------------------------------------------------------------------------------------------
+    int NodeRangeFromScalar::GetInputCount() const
+	{
+		return NODE_INPUT_COUNT;
+	}
+
+
+	//---------------------------------------------------------------------------------------------
+    Node* NodeRangeFromScalar::GetInputNode( int i ) const
+	{
+		check( i>=0 && i<NODE_INPUT_COUNT );
+
+        Node* pResult = nullptr;
+
+		switch (i)
+		{
+        case 0: pResult = m_pD->m_pSize.get(); break;
+		}
+
+		return pResult;
+	}
+
+
+	//---------------------------------------------------------------------------------------------
+    void NodeRangeFromScalar::SetInputNode( int i, NodePtr pNode )
+	{
+		check( i>=0 && i<NODE_INPUT_COUNT );
+
+		switch (i)
+		{
+        case 0: m_pD->m_pSize = dynamic_cast<NodeScalar*>(pNode.get()); break;
+		}
+	}
+
+
+	//---------------------------------------------------------------------------------------------
+	// Own Interface
+	//---------------------------------------------------------------------------------------------
+	Ptr<NodeScalar> NodeRangeFromScalar::GetSize() const
+	{
+        return m_pD->m_pSize.get();
+	}
+
+
+	//---------------------------------------------------------------------------------------------
+    void NodeRangeFromScalar::SetSize( const Ptr<NodeScalar>& pNode )
+	{
+        m_pD->m_pSize = pNode;
+	}
+
+
+	//---------------------------------------------------------------------------------------------
+    const char* NodeRangeFromScalar::GetName() const
+	{
+        return m_pD->m_name.c_str();
+	}
+
+
+	//---------------------------------------------------------------------------------------------
+    void NodeRangeFromScalar::SetName( const char* strName )
+	{
+        if (!strName)
+        {
+            m_pD->m_name.clear();
+        }
+        else
+        {
+            m_pD->m_name = strName;
+        }
+	}
+
+}
+
+#undef NODE_INPUT_COUNT
+
diff --git a/Engine/Plugins/Experimental/Mutable/Source/MutableTools/Private/MuT/NodeRangeFromScalarPrivate.h b/Engine/Plugins/Experimental/Mutable/Source/MutableTools/Private/MuT/NodeRangeFromScalarPrivate.h
new file mode 100644
--- /dev/null
+++ b/Engine/Plugins/Experimental/Mutable/Source/MutableTools/Private/MuT/NodeRangeFromScalarPrivate.h
@@ -0,0 +1,51 @@
+// Copyright Epic Games, Inc. All Rights Reserved.
+
+#pragma once
+
+#include "NodeRangePrivate.h"
+#include "MuT/NodeRangeFromScalar.h"
+#include "MuT/NodeRange.h"
+#include "MuT/NodeScalar.h"
+
+
+namespace mu
+{
+
+
+    class NodeRangeFromScalar::Private : public NodeRange::Private
+	{
+	public:
+
+		MUTABLE_DEFINE_CONST_VISITABLE()
+
+	public:
+
+		static NODE_TYPE s_type;
+
+        NodeScalarPtr m_pSize;
+        string m_name;
+
+		//!
+		void Serialise( OutputArchive& arch ) const
+		{
+            uint32_t ver = 0;
+			arch << ver;
+
+            arch << m_name;
+            arch << m_pSize;
+		}
+
+		//!
+		void Unserialise( InputArchive& arch )
+		{
+            uint32_t ver;
+			arch >> ver;
+			check(ver==0);
+
+            arch >> m_name;
+            arch >> m_pSize;
+		}
+	};
+
+
+}
diff --git a/Engine/Plugins/Experimental/Mutable/Source/MutableTools/Private/MuT/NodeRangePrivate.h b/Engine/Plugins/Experimental/Mutable/Source/MutableTools/Private/MuT/NodeRangePrivate.h
new file mode 100644
--- /dev/null
+++ b/Engine/Plugins/Experimental/Mutable/Source/MutableTools/Private/MuT/NodeRangePrivate.h
@@ -0,0 +1,19 @@
+// Copyright Epic Games, Inc. All Rights Reserved.
+
+#pragma once
+
+#include "MuT/NodePrivate.h"
+#include "MuT/NodeRange.h"
+
+#include "MuR/MemoryPrivate.h"
+
+
+namespace mu
+{
+
+
+    class NodeRange::Private : public Node::Private
+	{
+	};
+
+}
diff --git a/Engine/Plugins/Experimental/Mutable/Source/MutableTools/Private/MuT/NodeScalar.cpp b/Engine/Plugins/Experimental/Mutable/Source/MutableTools/Private/MuT/NodeScalar.cpp
new file mode 100644
--- /dev/null
+++ b/Engine/Plugins/Experimental/Mutable/Source/MutableTools/Private/MuT/NodeScalar.cpp
@@ -0,0 +1,82 @@
+// Copyright Epic Games, Inc. All Rights Reserved.
+
+
+#include "MuT/NodeScalar.h"
+
+#include "MuT/Node.h"
+#include "MuT/NodePrivate.h"
+
+#include "MuT/NodeScalarConstant.h"
+#include "MuT/NodeScalarParameter.h"
+#include "MuT/NodeScalarEnumParameter.h"
+#include "MuT/NodeScalarCurve.h"
+#include "MuT/NodeScalarSwitch.h"
+#include "MuT/NodeScalarArithmeticOperation.h"
+#include "MuT/NodeScalarVariation.h"
+#include "MuT/NodeScalarTable.h"
+
+
+namespace mu
+{
+
+	//---------------------------------------------------------------------------------------------
+	// Static initialisation
+	//---------------------------------------------------------------------------------------------
+	static NODE_TYPE s_nodeScalarType = NODE_TYPE( "NodeScalar", Node::GetStaticType() );
+
+
+	//---------------------------------------------------------------------------------------------
+	void NodeScalar::Serialise( const NodeScalar* p, OutputArchive& arch )
+	{
+        uint32_t ver = 0;
+		arch << ver;
+
+		arch << uint32_t(p->Type);
+		p->SerialiseWrapper(arch);
+    }
+
+
+    //---------------------------------------------------------------------------------------------
+	NodeScalarPtr NodeScalar::StaticUnserialise( InputArchive& arch )
+	{
+        uint32_t ver;
+		arch >> ver;
+		check( ver == 0 );
+
+        uint32_t id;
+		arch >> id;
+
+		switch (id)
+		{
+		case 0 :  return NodeScalarConstant::StaticUnserialise( arch ); break;
+		case 1 :  return NodeScalarParameter::StaticUnserialise( arch ); break;
+		case 2 :  return NodeScalarEnumParameter::StaticUnserialise( arch ); break;
+        case 3 :  return NodeScalarCurve::StaticUnserialise( arch ); break;
+        case 4 :  return NodeScalarSwitch::StaticUnserialise( arch ); break;
+        case 5 :  return NodeScalarArithmeticOperation::StaticUnserialise( arch ); break;
+        case 6 :  return NodeScalarVariation::StaticUnserialise( arch ); break;
+		case 7 :  return NodeScalarTable::StaticUnserialise( arch ); break;
+        default : check(false);
+		}
+
+		return 0;
+	}
+
+
+	//---------------------------------------------------------------------------------------------
+	const NODE_TYPE* NodeScalar::GetType() const
+	{
+		return GetStaticType();
+	}
+
+
+	//---------------------------------------------------------------------------------------------
+	const NODE_TYPE* NodeScalar::GetStaticType()
+	{
+		return &s_nodeScalarType;
+	}
+
+
+}
+
+
diff --git a/Engine/Plugins/Experimental/Mutable/Source/MutableTools/Private/MuT/NodeScalarArithmeticOperation.cpp b/Engine/Plugins/Experimental/Mutable/Source/MutableTools/Private/MuT/NodeScalarArithmeticOperation.cpp
new file mode 100644
--- /dev/null
+++ b/Engine/Plugins/Experimental/Mutable/Source/MutableTools/Private/MuT/NodeScalarArithmeticOperation.cpp
@@ -0,0 +1,126 @@
+// Copyright Epic Games, Inc. All Rights Reserved.
+
+
+#include "MuT/NodeScalarArithmeticOperation.h"
+#include "NodeScalarArithmeticOperationPrivate.h"
+
+#include "MuT/NodeScalar.h"
+
+#include "MuR/ImagePrivate.h"
+
+
+#define NODE_INPUT_COUNT 	2
+
+
+namespace mu
+{
+
+	//---------------------------------------------------------------------------------------------
+	// Static initialisation
+	//---------------------------------------------------------------------------------------------
+    NODE_TYPE NodeScalarArithmeticOperation::Private::s_type =
+            NODE_TYPE( "ScalarArithmenticOperation", NodeScalar::GetStaticType() );
+
+
+	//---------------------------------------------------------------------------------------------
+	//!
+	//---------------------------------------------------------------------------------------------
+
+    MUTABLE_IMPLEMENT_NODE( NodeScalarArithmeticOperation, EType::ArithmeticOperation, Node, Node::EType::Scalar)
+
+
+    //---------------------------------------------------------------------------------------------
+    const char* NodeScalarArithmeticOperation::s_opTypeName[] =
+    {
+        "Add",
+        "Subtract",
+        "Multiply",
+        "Divide"
+    };
+
+
+	//---------------------------------------------------------------------------------------------
+	// Node Interface
+	//---------------------------------------------------------------------------------------------
+    int NodeScalarArithmeticOperation::GetInputCount() const
+	{
+		return NODE_INPUT_COUNT;
+	}
+
+
+	//---------------------------------------------------------------------------------------------
+    Node* NodeScalarArithmeticOperation::GetInputNode( int i ) const
+	{
+		check( i>=0 && i<NODE_INPUT_COUNT );
+
+		Node* pResult = 0;
+
+		switch (i)
+		{
+		case 0: pResult = m_pD->m_pA.get(); break;
+		case 1: pResult = m_pD->m_pB.get(); break;
+		}
+
+		return pResult;
+	}
+
+
+	//---------------------------------------------------------------------------------------------
+    void NodeScalarArithmeticOperation::SetInputNode( int i, NodePtr pNode )
+	{
+		check( i>=0 && i<NODE_INPUT_COUNT );
+
+		switch (i)
+		{
+        case 0: m_pD->m_pA = dynamic_cast<NodeScalar*>(pNode.get()); break;
+        case 1: m_pD->m_pB = dynamic_cast<NodeScalar*>(pNode.get()); break;
+		}
+	}
+
+
+	//---------------------------------------------------------------------------------------------
+	// Own Interface
+	//---------------------------------------------------------------------------------------------
+    NodeScalarArithmeticOperation::OPERATION NodeScalarArithmeticOperation::GetOperation() const
+	{
+		return m_pD->m_operation;
+	}
+
+
+	//---------------------------------------------------------------------------------------------
+    void NodeScalarArithmeticOperation::SetOperation(NodeScalarArithmeticOperation::OPERATION o)
+	{
+		m_pD->m_operation = o;
+	}
+
+
+    NodeScalarPtr NodeScalarArithmeticOperation::GetA() const
+	{
+		return m_pD->m_pA.get();
+	}
+
+
+	//---------------------------------------------------------------------------------------------
+    void NodeScalarArithmeticOperation::SetA(NodeScalarPtr pNode)
+	{
+		m_pD->m_pA = pNode;
+	}
+
+
+	//---------------------------------------------------------------------------------------------
+    NodeScalarPtr NodeScalarArithmeticOperation::GetB() const
+	{
+		return m_pD->m_pB.get();
+	}
+
+
+	//---------------------------------------------------------------------------------------------
+    void NodeScalarArithmeticOperation::SetB( NodeScalarPtr pNode )
+	{
+		m_pD->m_pB = pNode;
+	}
+
+}
+
+#undef NODE_INPUT_COUNT
+
diff --git a/Engine/Plugins/Experimental/Mutable/Source/MutableTools/Private/MuT/NodeScalarArithmeticOperationPrivate.h b/Engine/Plugins/Experimental/Mutable/Source/MutableTools/Private/MuT/NodeScalarArithmeticOperationPrivate.h
new file mode 100644
--- /dev/null
+++ b/Engine/Plugins/Experimental/Mutable/Source/MutableTools/Private/MuT/NodeScalarArithmeticOperationPrivate.h
@@ -0,0 +1,55 @@
+// Copyright Epic Games, Inc. All Rights Reserved.
+
+#pragma once
+
+#include "NodeScalarPrivate.h"
+#include "MuT/NodeScalarArithmeticOperation.h"
+
+#include "MuR/MemoryPrivate.h"
+
+
+namespace mu
+{
+
+    MUTABLE_DEFINE_ENUM_SERIALISABLE(NodeScalarArithmeticOperation::OPERATION)
+
+    class NodeScalarArithmeticOperation::Private : public NodeScalar::Private
+	{
+	public:
+
+		MUTABLE_DEFINE_CONST_VISITABLE()
+
+	public:
+
+		static NODE_TYPE s_type;
+
+        NodeScalarArithmeticOperation::OPERATION m_operation;
+        NodeScalarPtr m_pA;
+        NodeScalarPtr m_pB;
+
+		//!
+		void Serialise( OutputArchive& arch ) const
+		{
+            uint32_t ver = 0;
+			arch << ver;
+
+			arch << m_operation;
+			arch << m_pA;
+			arch << m_pB;
+		}
+
+		//!
+		void Unserialise( InputArchive& arch )
+		{
+            uint32_t ver;
+			arch >> ver;
+			check(ver==0);
+
+			arch >> m_operation;
+			arch >> m_pA;
+			arch >> m_pB;
+		}
+	};
+
+
+}
diff --git a/Engine/Plugins/Experimental/Mutable/Source/MutableTools/Private/MuT/NodeScalarConstant.cpp b/Engine/Plugins/Experimental/Mutable/Source/MutableTools/Private/MuT/NodeScalarConstant.cpp
new file mode 100644
--- /dev/null
+++ b/Engine/Plugins/Experimental/Mutable/Source/MutableTools/Private/MuT/NodeScalarConstant.cpp
@@ -0,0 +1,72 @@
+// Copyright Epic Games, Inc. All Rights Reserved.
+
+
+#include "MuT/NodeScalarConstant.h"
+#include "NodeScalarConstantPrivate.h"
+
+#include "MuR/MutableMath.h"
+
+
+#define NODE_INPUT_COUNT 	0
+
+
+namespace mu
+{
+
+	//---------------------------------------------------------------------------------------------
+	// Static initialisation
+	//---------------------------------------------------------------------------------------------
+	NODE_TYPE NodeScalarConstant::Private::s_type =
+			NODE_TYPE( "ScalarConstant", NodeScalar::GetStaticType() );
+
+
+	//---------------------------------------------------------------------------------------------
+	//!
+	//---------------------------------------------------------------------------------------------
+
+	MUTABLE_IMPLEMENT_NODE( NodeScalarConstant, EType::Constant, Node, Node::EType::Scalar)
+
+
+	//---------------------------------------------------------------------------------------------
+	// Node Interface
+	//---------------------------------------------------------------------------------------------
+	int NodeScalarConstant::GetInputCount() const
+	{
+		return NODE_INPUT_COUNT;
+	}
+
+
+	//---------------------------------------------------------------------------------------------
+    Node* NodeScalarConstant::GetInputNode( int ) const
+	{
+		check( false );
+		return 0;
+	}
+
+
+	//---------------------------------------------------------------------------------------------
+    void NodeScalarConstant::SetInputNode( int, NodePtr )
+	{
+		check( false );
+	}
+
+
+	//---------------------------------------------------------------------------------------------
+	// Own Interface
+	//---------------------------------------------------------------------------------------------
+	float NodeScalarConstant::GetValue() const
+	{
+		return m_pD->m_value;
+	}
+
+	//---------------------------------------------------------------------------------------------
+	void NodeScalarConstant::SetValue( float v )
+	{
+		m_pD->m_value = v;
+	}
+
+
+}
+
+#undef NODE_INPUT_COUNT
+
diff --git a/Engine/Plugins/Experimental/Mutable/Source/MutableTools/Private/MuT/NodeScalarConstantPrivate.h b/Engine/Plugins/Experimental/Mutable/Source/MutableTools/Private/MuT/NodeScalarConstantPrivate.h
new file mode 100644
--- /dev/null
+++ b/Engine/Plugins/Experimental/Mutable/Source/MutableTools/Private/MuT/NodeScalarConstantPrivate.h
@@ -0,0 +1,51 @@
+// Copyright Epic Games, Inc. All Rights Reserved.
+
+#pragma once
+
+#include "NodeScalarPrivate.h"
+
+#include "MuT/NodeScalarConstant.h"
+#include "MuT/NodeImage.h"
+#include "MuT/AST.h"
+#include "MuR/MemoryPrivate.h"
+#include "MuR/ParametersPrivate.h"
+
+
+namespace mu
+{
+
+
+	class NodeScalarConstant::Private : public NodeScalar::Private
+	{
+	public:
+
+		MUTABLE_DEFINE_CONST_VISITABLE()
+
+	public:
+
+		static NODE_TYPE s_type;
+
+		float m_value;
+
+		//!
+		void Serialise( OutputArchive& arch ) const
+		{
+            uint32_t ver = 0;
+			arch << ver;
+
+			arch << m_value;
+		}
+
+		//!
+		void Unserialise( InputArchive& arch )
+		{
+            uint32_t ver;
+			arch >> ver;
+			check(ver==0);
+
+			arch >> m_value;
+		}
+	};
+
+
+}
diff --git a/Engine/Plugins/Experimental/Mutable/Source/MutableTools/Private/MuT/NodeScalarCurve.cpp b/Engine/Plugins/Experimental/Mutable/Source/MutableTools/Private/MuT/NodeScalarCurve.cpp
new file mode 100644
--- /dev/null
+++ b/Engine/Plugins/Experimental/Mutable/Source/MutableTools/Private/MuT/NodeScalarCurve.cpp
@@ -0,0 +1,123 @@
+// Copyright Epic Games, Inc. All Rights Reserved.
+
+
+#include "MuT/NodeScalarCurve.h"
+#include "NodeScalarCurvePrivate.h"
+
+#include "MuT/NodeScalar.h"
+
+
+#define NODE_INPUT_COUNT 	1
+
+
+namespace mu
+{
+
+
+	//---------------------------------------------------------------------------------------------
+	// Static initialisation
+	//---------------------------------------------------------------------------------------------
+    NODE_TYPE NodeScalarCurve::Private::s_type =
+            NODE_TYPE( "Curve", NodeScalarCurve::GetStaticType() );
+
+
+	//---------------------------------------------------------------------------------------------
+	//---------------------------------------------------------------------------------------------
+	//---------------------------------------------------------------------------------------------
+
+    MUTABLE_IMPLEMENT_NODE( NodeScalarCurve, EType::Curve, Node, Node::EType::Scalar)
+
+
+	//---------------------------------------------------------------------------------------------
+	// Node Interface
+	//---------------------------------------------------------------------------------------------
+    int NodeScalarCurve::GetInputCount() const
+	{
+		return NODE_INPUT_COUNT;
+	}
+
+
+	//---------------------------------------------------------------------------------------------
+    Node* NodeScalarCurve::GetInputNode( int i ) const
+	{
+		check( i >=0 && i < GetInputCount() );
+        (void)i;
+
+		return m_pD->m_input_scalar.get();
+	}
+
+
+	//---------------------------------------------------------------------------------------------
+    void NodeScalarCurve::SetInputNode( int i, NodePtr pNode )
+	{
+		check( i >=0 && i < GetInputCount());
+        (void)i;
+
+		m_pD->m_input_scalar = dynamic_cast<NodeScalar*>(pNode.get());
+	}
+
+
+	//---------------------------------------------------------------------------------------------
+	// Own Interface
+	//---------------------------------------------------------------------------------------------
+    int NodeScalarCurve::GetKeyFrameCount() const
+	{
+		return (int)m_pD->m_curve.keyFrames.size();
+	}
+
+
+	//---------------------------------------------------------------------------------------------
+    void NodeScalarCurve::SetKeyFrameCount( int num )
+	{
+		check( num >=0 );
+		m_pD->m_curve.keyFrames.resize( num );
+	}
+
+	//---------------------------------------------------------------------------------------------
+    void NodeScalarCurve::SetDefaultValue(float defaultValue)
+	{
+		m_pD->m_curve.defaultValue = defaultValue;
+	}
+
+
+	//---------------------------------------------------------------------------------------------
+    void NodeScalarCurve::SetKeyFrame(
+            int index, float time, float value,
+            float in_tangent, float in_tangent_weight,
+            float out_tangent, float out_tangent_weight,
+            uint8_t interp_mode, uint8_t tangent_mode, uint8_t tangent_weight_mode )
+	{
+		check( index >=0 && index < (int)m_pD->m_curve.keyFrames.size() );
+
+		CurveKeyFrame& key_frame = m_pD->m_curve.keyFrames[index];
+
+		key_frame.time = time;
+		key_frame.value = value;
+		key_frame.in_tangent = in_tangent;
+		key_frame.in_tangent_weight = in_tangent_weight;
+		key_frame.out_tangent = out_tangent;
+		key_frame.out_tangent_weight = out_tangent_weight;
+		
+		key_frame.interp_mode = interp_mode;
+		key_frame.tangent_mode = tangent_mode;
+		key_frame.tangent_weight_mode = tangent_weight_mode;
+	}
+
+
+	//---------------------------------------------------------------------------------------------
+    NodeScalarPtr NodeScalarCurve::GetT() const
+	{
+		return m_pD->m_input_scalar.get();
+	}
+
+
+	//---------------------------------------------------------------------------------------------
+    void NodeScalarCurve::SetT(NodeScalarPtr pNode)
+	{
+		m_pD->m_input_scalar = pNode;
+	}
+
+}
+
+#undef NODE_INPUT_COUNT
+
diff --git a/Engine/Plugins/Experimental/Mutable/Source/MutableTools/Private/MuT/NodeScalarEnumParameter.cpp b/Engine/Plugins/Experimental/Mutable/Source/MutableTools/Private/MuT/NodeScalarEnumParameter.cpp
new file mode 100644
--- /dev/null
+++ b/Engine/Plugins/Experimental/Mutable/Source/MutableTools/Private/MuT/NodeScalarEnumParameter.cpp
@@ -0,0 +1,156 @@
+// Copyright Epic Games, Inc. All Rights Reserved.
+
+
+#include "MuT/NodeScalarEnumParameter.h"
+#include "NodeScalarEnumParameterPrivate.h"
+
+#include "MuR/MutableMath.h"
+#include "MuR/MemoryPrivate.h"
+
+
+namespace mu
+{
+
+	//---------------------------------------------------------------------------------------------
+	// Static initialisation
+	//---------------------------------------------------------------------------------------------
+	NODE_TYPE NodeScalarEnumParameter::Private::s_type =
+			NODE_TYPE( "ScalarEnumParameter", NodeScalar::GetStaticType() );
+
+
+	//---------------------------------------------------------------------------------------------
+	//!
+	//---------------------------------------------------------------------------------------------
+
+	MUTABLE_IMPLEMENT_NODE( NodeScalarEnumParameter, EType::EnumParameter, Node, Node::EType::Scalar)
+
+
+	//---------------------------------------------------------------------------------------------
+	int NodeScalarEnumParameter::GetInputCount() const
+	{
+        return int( m_pD->m_ranges.size() );
+	}
+
+
+	//---------------------------------------------------------------------------------------------
+    Node* NodeScalarEnumParameter::GetInputNode( int i ) const
+	{
+        check( i<GetInputCount() );
+        if (i<GetInputCount())
+        {
+            return m_pD->m_ranges[i].get();
+        }
+        return nullptr;
+	}
+
+
+	//---------------------------------------------------------------------------------------------
+    void NodeScalarEnumParameter::SetInputNode( int i, NodePtr n )
+	{
+        check( i<GetInputCount() );
+        if (i<GetInputCount())
+        {
+            m_pD->m_ranges[i] = dynamic_cast<NodeRange*>(n.get());
+        }
+	}
+
+
+	//---------------------------------------------------------------------------------------------
+	const char* NodeScalarEnumParameter::GetName() const
+	{
+		return m_pD->m_name.c_str();
+	}
+
+
+	//---------------------------------------------------------------------------------------------
+	void NodeScalarEnumParameter::SetName( const char* strName )
+	{
+		if ( strName )
+		{
+			m_pD->m_name = strName;
+		}
+		else
+		{
+			m_pD->m_name = "";
+		}
+	}
+
+
+	const char* NodeScalarEnumParameter::GetUid() const
+	{
+		return m_pD->m_uid.c_str();
+	}
+
+
+	//---------------------------------------------------------------------------------------------
+	void NodeScalarEnumParameter::SetUid( const char* strUid )
+	{
+		if ( strUid )
+		{
+			m_pD->m_uid = strUid;
+		}
+		else
+		{
+			m_pD->m_uid = "";
+		}
+	}
+
+
+	//---------------------------------------------------------------------------------------------
+	int NodeScalarEnumParameter::GetDefaultValueIndex() const
+	{
+		return m_pD->m_defaultValue;
+	}
+
+
+	//---------------------------------------------------------------------------------------------
+	void NodeScalarEnumParameter::SetDefaultValueIndex( int i )
+	{
+		m_pD->m_defaultValue = i;
+	}
+
+
+	//---------------------------------------------------------------------------------------------
+	void NodeScalarEnumParameter::SetValueCount( int i )
+	{
+		m_pD->m_options.resize(i);
+	}
+
+
+	//---------------------------------------------------------------------------------------------
+	int NodeScalarEnumParameter::GetValueCount() const
+	{
+		return (int)m_pD->m_options.size();
+	}
+
+
+	//---------------------------------------------------------------------------------------------
+	void NodeScalarEnumParameter::SetValue( int i, float value, const char* strName )
+	{
+		check( i>=0 && i<(int)m_pD->m_options.size() );
+		m_pD->m_options[i].name = strName;
+		m_pD->m_options[i].value = value;
+	}
+
+
+    //---------------------------------------------------------------------------------------------
+    void NodeScalarEnumParameter::SetRangeCount( int i )
+    {
+        check(i>=0);
+        m_pD->m_ranges.resize(i);
+    }
+
+
+    //---------------------------------------------------------------------------------------------
+    void NodeScalarEnumParameter::SetRange( int i, NodeRangePtr pRange )
+    {
+        check( i>=0 && i<int(m_pD->m_ranges.size()) );
+        if ( i>=0 && i<int(m_pD->m_ranges.size()) )
+        {
+            m_pD->m_ranges[i] = pRange;
+        }
+    }
+
+}
+
+
diff --git a/Engine/Plugins/Experimental/Mutable/Source/MutableTools/Private/MuT/NodeScalarEnumParameterPrivate.h b/Engine/Plugins/Experimental/Mutable/Source/MutableTools/Private/MuT/NodeScalarEnumParameterPrivate.h
new file mode 100644
--- /dev/null
+++ b/Engine/Plugins/Experimental/Mutable/Source/MutableTools/Private/MuT/NodeScalarEnumParameterPrivate.h
@@ -0,0 +1,94 @@
+// Copyright Epic Games, Inc. All Rights Reserved.
+
+#pragma once
+
+#include "NodeScalarPrivate.h"
+
+#include "MuT/NodeScalarEnumParameter.h"
+#include "MuT/NodeImage.h"
+#include "MuR/MemoryPrivate.h"
+#include "MuR/ParametersPrivate.h"
+
+
+namespace mu
+{
+
+
+	class NodeScalarEnumParameter::Private : public NodeScalar::Private
+	{
+	public:
+
+		MUTABLE_DEFINE_CONST_VISITABLE()
+
+	public:
+
+		static NODE_TYPE s_type;
+
+		string m_name;
+		string m_uid;
+		int m_defaultValue = 0;
+		PARAMETER_DETAILED_TYPE m_detailedType = PARAMETER_DETAILED_TYPE::UNKNOWN;
+
+		struct OPTION
+		{
+			string name;
+			float value;
+
+			//!
+			void Serialise( OutputArchive& arch ) const
+			{
+                uint32_t ver = 0;
+				arch << ver;
+
+				arch << name;
+				arch << value;
+			}
+
+			//!
+			void Unserialise( InputArchive& arch )
+			{
+                uint32_t ver;
+				arch >> ver;
+                check(ver<=0);
+
+				arch >> name;
+				arch >> value;
+			}
+		};
+
+		vector<OPTION> m_options;
+
+        vector<Ptr<NodeRange>> m_ranges;
+
+		//!
+		void Serialise( OutputArchive& arch ) const
+		{
+            uint32_t ver = 2;
+			arch << ver;
+
+			arch << m_name;
+			arch << m_uid;
+			arch << m_defaultValue;
+			arch << m_detailedType;
+            arch << m_options;
+            arch << m_ranges;
+        }
+
+		//!
+		void Unserialise( InputArchive& arch )
+		{
+            uint32_t ver;
+			arch >> ver;
+            check(ver==2);
+
+			arch >> m_name;
+            arch >> m_uid;
+			arch >> m_defaultValue;
+			arch >> m_detailedType;
+			arch >> m_options;
+            arch >> m_ranges;
+        }
+	};
+
+}
+
diff --git a/Engine/Plugins/Experimental/Mutable/Source/MutableTools/Private/MuT/NodeScalarParameter.cpp b/Engine/Plugins/Experimental/Mutable/Source/MutableTools/Private/MuT/NodeScalarParameter.cpp
new file mode 100644
--- /dev/null
+++ b/Engine/Plugins/Experimental/Mutable/Source/MutableTools/Private/MuT/NodeScalarParameter.cpp
@@ -0,0 +1,181 @@
+// Copyright Epic Games, Inc. All Rights Reserved.
+
+
+#include "MuT/NodeScalarParameter.h"
+#include "NodeScalarParameterPrivate.h"
+
+#include "MuR/MutableMath.h"
+#include "MuR/MemoryPrivate.h"
+
+
+namespace mu
+{
+
+	//---------------------------------------------------------------------------------------------
+	// Static initialisation
+	//---------------------------------------------------------------------------------------------
+	NODE_TYPE NodeScalarParameter::Private::s_type =
+			NODE_TYPE( "ScalarParameter", NodeScalar::GetStaticType() );
+
+
+	//---------------------------------------------------------------------------------------------
+	//!
+	//---------------------------------------------------------------------------------------------
+
+	MUTABLE_IMPLEMENT_NODE( NodeScalarParameter, EType::Parameter, Node, Node::EType::Scalar)
+
+
+	//---------------------------------------------------------------------------------------------
+	int NodeScalarParameter::GetInputCount() const
+	{
+        return int( m_pD->m_additionalImages.size()
+                    +
+                    m_pD->m_ranges.size() );
+	}
+
+
+	//---------------------------------------------------------------------------------------------
+    Node* NodeScalarParameter::GetInputNode( int i ) const
+	{
+        check( i<GetInputCount() );
+
+        int imageCount = int(m_pD->m_additionalImages.size());
+        int rangeCount = int(m_pD->m_ranges.size());
+        if (i<imageCount)
+        {
+            return m_pD->m_additionalImages[i].get();
+        }
+        else if ( i < imageCount + rangeCount )
+        {
+            int r = i - imageCount;
+            return m_pD->m_ranges[r].get();
+        }
+        return nullptr;
+	}
+
+
+	//---------------------------------------------------------------------------------------------
+    void NodeScalarParameter::SetInputNode( int i, NodePtr n )
+	{
+        check( i<GetInputCount() );
+        int imageCount = int(m_pD->m_additionalImages.size());
+        int rangeCount = int(m_pD->m_ranges.size());
+        if (i<imageCount)
+        {
+            m_pD->m_additionalImages[i] = dynamic_cast<NodeImage*>(n.get());
+        }
+        else if ( i < imageCount + rangeCount )
+        {
+            int r = i - imageCount;
+            m_pD->m_ranges[r] = dynamic_cast<NodeRange*>(n.get());
+        }
+    }
+
+
+	//---------------------------------------------------------------------------------------------
+	const char* NodeScalarParameter::GetName() const
+	{
+		return m_pD->m_name.c_str();
+	}
+
+
+	//---------------------------------------------------------------------------------------------
+	void NodeScalarParameter::SetName( const char* strName )
+	{
+		if ( strName )
+		{
+			m_pD->m_name = strName;
+		}
+		else
+		{
+			m_pD->m_name = "";
+		}
+	}
+
+
+	const char* NodeScalarParameter::GetUid() const
+	{
+		return m_pD->m_uid.c_str();
+	}
+
+
+	//---------------------------------------------------------------------------------------------
+	void NodeScalarParameter::SetUid( const char* strUid )
+	{
+		if ( strUid )
+		{
+			m_pD->m_uid = strUid;
+		}
+		else
+		{
+			m_pD->m_uid = "";
+		}
+	}
+
+
+	//---------------------------------------------------------------------------------------------
+	float NodeScalarParameter::GetDefaultValue() const
+	{
+		return m_pD->m_defaultValue;
+	}
+
+
+	//---------------------------------------------------------------------------------------------
+	void NodeScalarParameter::SetDefaultValue( float v )
+	{
+		m_pD->m_defaultValue = v;
+	}
+
+
+	//---------------------------------------------------------------------------------------------
+	PARAMETER_DETAILED_TYPE NodeScalarParameter::GetDetailedType() const
+	{
+		return m_pD->m_detailedType;
+	}
+
+
+	//---------------------------------------------------------------------------------------------
+	void NodeScalarParameter::SetDetailedType( PARAMETER_DETAILED_TYPE t )
+	{
+		m_pD->m_detailedType = t;
+	}
+
+
+	//---------------------------------------------------------------------------------------------
+	void NodeScalarParameter::SetDescImageCount( int i )
+	{
+        check(i>=0);
+        m_pD->m_additionalImages.resize(i);
+	}
+
+
+	//---------------------------------------------------------------------------------------------
+	void NodeScalarParameter::SetDescImage( int i, NodeImagePtr pImage )
+	{
+		check( i>=0 && i<(int)m_pD->m_additionalImages.size() );
+		m_pD->m_additionalImages[i] = pImage;
+	}
+
+
+    //---------------------------------------------------------------------------------------------
+    void NodeScalarParameter::SetRangeCount( int i )
+    {
+        check(i>=0);
+        m_pD->m_ranges.resize(i);
+    }
+
+
+    //---------------------------------------------------------------------------------------------
+    void NodeScalarParameter::SetRange( int i, NodeRangePtr pRange )
+    {
+        check( i>=0 && i<int(m_pD->m_ranges.size()) );
+        if ( i>=0 && i<int(m_pD->m_ranges.size()) )
+        {
+            m_pD->m_ranges[i] = pRange;
+        }
+    }
+
+
+}
+
+
diff --git a/Engine/Plugins/Experimental/Mutable/Source/MutableTools/Private/MuT/NodeScalarParameterPrivate.h b/Engine/Plugins/Experimental/Mutable/Source/MutableTools/Private/MuT/NodeScalarParameterPrivate.h
new file mode 100644
--- /dev/null
+++ b/Engine/Plugins/Experimental/Mutable/Source/MutableTools/Private/MuT/NodeScalarParameterPrivate.h
@@ -0,0 +1,67 @@
+// Copyright Epic Games, Inc. All Rights Reserved.
+
+#pragma once
+
+#include "NodeScalarPrivate.h"
+
+#include "MuT/NodeScalarParameter.h"
+#include "MuT/NodeImage.h"
+#include "MuT/NodeRange.h"
+#include "MuR/MemoryPrivate.h"
+#include "MuR/ParametersPrivate.h"
+
+
+namespace mu
+{
+
+
+	class NodeScalarParameter::Private : public NodeScalar::Private
+	{
+	public:
+
+		MUTABLE_DEFINE_CONST_VISITABLE()
+
+	public:
+
+		static NODE_TYPE s_type;
+
+		float m_defaultValue = 0.0f;
+		string m_name;
+		string m_uid;
+		PARAMETER_DETAILED_TYPE m_detailedType = PARAMETER_DETAILED_TYPE::UNKNOWN;
+
+        vector<Ptr<NodeImage>> m_additionalImages;
+
+        vector<Ptr<NodeRange>> m_ranges;
+
+		//!
+		void Serialise( OutputArchive& arch ) const
+		{
+            uint32_t ver = 4;
+			arch << ver;
+
+			arch << m_defaultValue;
+			arch << m_name;
+			arch << m_uid;
+			arch << m_additionalImages;
+			arch << m_detailedType;
+            arch << m_ranges;
+        }
+
+		//!
+		void Unserialise( InputArchive& arch )
+		{
+            uint32_t ver;
+			arch >> ver;
+            check(ver==4);
+
+			arch >> m_defaultValue;
+			arch >> m_name;
+            arch >> m_uid;
+			arch >> m_additionalImages;
+            arch >> m_detailedType;
+            arch >> m_ranges;
+        }
+	};
+
+}
diff --git a/Engine/Plugins/Experimental/Mutable/Source/MutableTools/Private/MuT/NodeScalarSwitch.cpp b/Engine/Plugins/Experimental/Mutable/Source/MutableTools/Private/MuT/NodeScalarSwitch.cpp
new file mode 100644
--- /dev/null
+++ b/Engine/Plugins/Experimental/Mutable/Source/MutableTools/Private/MuT/NodeScalarSwitch.cpp
@@ -0,0 +1,118 @@
+// Copyright Epic Games, Inc. All Rights Reserved.
+
+
+#include "MuT/NodeScalarSwitch.h"
+#include "NodeScalarSwitchPrivate.h"
+
+#include "MuT/NodeScalar.h"
+
+
+namespace mu
+{
+
+	//---------------------------------------------------------------------------------------------
+	// Static initialisation
+	//---------------------------------------------------------------------------------------------
+    NODE_TYPE NodeScalarSwitch::Private::s_type =
+            NODE_TYPE( "ScalarSwitch", NodeScalar::GetStaticType() );
+
+
+	//---------------------------------------------------------------------------------------------
+	//!
+	//---------------------------------------------------------------------------------------------
+
+    MUTABLE_IMPLEMENT_NODE( NodeScalarSwitch, EType::Switch, Node, Node::EType::Scalar)
+
+
+	//---------------------------------------------------------------------------------------------
+	// Node Interface
+	//---------------------------------------------------------------------------------------------
+    int NodeScalarSwitch::GetInputCount() const
+	{
+		return 1 + (int)m_pD->m_options.size();
+	}
+
+
+	//---------------------------------------------------------------------------------------------
+    Node* NodeScalarSwitch::GetInputNode( int i ) const
+	{
+		check( i>=0 && i<GetInputCount() );
+
+		Node* pResult = 0;
+
+		switch (i)
+		{
+		case 0:
+			pResult = m_pD->m_pParameter.get();
+			break;
+
+		default:
+			pResult = m_pD->m_options[i-1].get();
+			break;
+		}
+
+		return pResult;
+	}
+
+
+	//---------------------------------------------------------------------------------------------
+    void NodeScalarSwitch::SetInputNode( int i, NodePtr pNode )
+	{
+		check( i>=0 && i<GetInputCount() );
+
+		switch (i)
+		{
+		case 0:
+			m_pD->m_pParameter = dynamic_cast<NodeScalar*>(pNode.get());
+			break;
+
+		default:
+            m_pD->m_options[i-1] = dynamic_cast<NodeScalar*>(pNode.get());
+			break;
+		}
+	}
+
+
+	//---------------------------------------------------------------------------------------------
+	// Own Interface
+	//---------------------------------------------------------------------------------------------
+    NodeScalarPtr NodeScalarSwitch::GetParameter() const
+	{
+		return m_pD->m_pParameter.get();
+	}
+
+
+	//---------------------------------------------------------------------------------------------
+    void NodeScalarSwitch::SetParameter( NodeScalarPtr pNode )
+	{
+		m_pD->m_pParameter = pNode;
+	}
+
+
+	//---------------------------------------------------------------------------------------------
+    void NodeScalarSwitch::SetOptionCount( int t )
+	{
+		m_pD->m_options.resize(t);
+	}
+
+
+	//---------------------------------------------------------------------------------------------
+    NodeScalarPtr NodeScalarSwitch::GetOption( int t ) const
+	{
+		check( t>=0 && t<(int)m_pD->m_options.size() );
+		return m_pD->m_options[t].get();
+	}
+
+
+	//---------------------------------------------------------------------------------------------
+    void NodeScalarSwitch::SetOption( int t, NodeScalarPtr pNode )
+	{
+		check( t>=0 && t<(int)m_pD->m_options.size() );
+		m_pD->m_options[t] = pNode;
+	}
+
+
+
+}
+
+
diff --git a/Engine/Plugins/Experimental/Mutable/Source/MutableTools/Private/MuT/NodeScalarSwitchPrivate.h b/Engine/Plugins/Experimental/Mutable/Source/MutableTools/Private/MuT/NodeScalarSwitchPrivate.h
new file mode 100644
--- /dev/null
+++ b/Engine/Plugins/Experimental/Mutable/Source/MutableTools/Private/MuT/NodeScalarSwitchPrivate.h
@@ -0,0 +1,51 @@
+// Copyright Epic Games, Inc. All Rights Reserved.
+
+#pragma once
+
+#include "MuT/NodeScalarSwitch.h"
+#include "NodeScalarPrivate.h"
+#include "MuT/NodeScalar.h"
+#include "MuT/AST.h"
+
+
+namespace mu
+{
+
+
+    class NodeScalarSwitch::Private : public NodeScalar::Private
+	{
+	public:
+
+		MUTABLE_DEFINE_CONST_VISITABLE()
+
+	public:
+
+		static NODE_TYPE s_type;
+
+		NodeScalarPtr m_pParameter;
+        vector<NodeScalarPtr> m_options;
+
+		//!
+		void Serialise( OutputArchive& arch ) const
+		{
+            uint32_t ver = 0;
+			arch << ver;
+
+			arch << m_pParameter;
+			arch << m_options;
+		}
+
+		//!
+		void Unserialise( InputArchive& arch )
+		{
+            uint32_t ver;
+			arch >> ver;
+			check(ver==0);
+
+			arch >> m_pParameter;
+			arch >> m_options;
+		}
+	};
+
+
+}
diff --git a/Engine/Plugins/Experimental/Mutable/Source/MutableTools/Private/MuT/NodeScalarTable.cpp b/Engine/Plugins/Experimental/Mutable/Source/MutableTools/Private/MuT/NodeScalarTable.cpp
new file mode 100644
--- /dev/null
+++ b/Engine/Plugins/Experimental/Mutable/Source/MutableTools/Private/MuT/NodeScalarTable.cpp
@@ -0,0 +1,100 @@
+// Copyright Epic Games, Inc. All Rights Reserved.
+
+
+#include "MuT/NodeScalarTable.h"
+#include "MuT/NodeScalarTablePrivate.h"
+
+
+
+namespace mu
+{
+
+
+	//---------------------------------------------------------------------------------------------
+	// Static initialisation
+	//---------------------------------------------------------------------------------------------
+	NODE_TYPE NodeScalarTable::Private::s_type =
+			NODE_TYPE( "TableScalar", NodeScalar::GetStaticType() );
+
+
+	//---------------------------------------------------------------------------------------------
+	//!
+	//---------------------------------------------------------------------------------------------
+
+	MUTABLE_IMPLEMENT_NODE(NodeScalarTable, EType::Table, Node, Node::EType::Scalar)
+
+
+	//---------------------------------------------------------------------------------------------
+	// Node Interface
+	//---------------------------------------------------------------------------------------------
+	int NodeScalarTable::GetInputCount() const
+	{
+		return 0;
+	}
+
+
+	//---------------------------------------------------------------------------------------------
+	Node* NodeScalarTable::GetInputNode( int i ) const
+	{
+		check( i >=0 && i < 0 );
+        (void)i;
+        return 0;
+	}
+
+
+	//---------------------------------------------------------------------------------------------
+    void NodeScalarTable::SetInputNode( int i, NodePtr )
+	{
+		check( i >=0 && i < 0 );
+        (void)i;
+        //m_pD->m_pObject = dynamic_cast<NodeObject*>( pNode.get() );
+	}
+
+
+	//---------------------------------------------------------------------------------------------
+	// Own Interface
+	//---------------------------------------------------------------------------------------------
+	const char* NodeScalarTable::GetColumn() const
+	{
+		return m_pD->m_columnName.c_str();
+	}
+
+
+	//---------------------------------------------------------------------------------------------
+	void NodeScalarTable::SetColumn( const char* strName )
+	{
+		if (strName)
+		{
+			m_pD->m_columnName = strName;
+		}
+		else
+		{
+			m_pD->m_columnName = "";
+		}
+	}
+
+
+	//---------------------------------------------------------------------------------------------
+	void NodeScalarTable::SetParameterName( const char* strName )
+	{
+		m_pD->m_parameterName = strName;
+	}
+
+
+	//---------------------------------------------------------------------------------------------
+	void NodeScalarTable::SetTable( TablePtr pTable )
+	{
+		m_pD->m_pTable = pTable;
+	}
+
+
+	//---------------------------------------------------------------------------------------------
+	TablePtr NodeScalarTable::GetTable() const
+	{
+		return m_pD->m_pTable;
+	}
+
+
+}
+
+
diff --git a/Engine/Plugins/Experimental/Mutable/Source/MutableTools/Private/MuT/NodeScalarTablePrivate.h b/Engine/Plugins/Experimental/Mutable/Source/MutableTools/Private/MuT/NodeScalarTablePrivate.h
new file mode 100644
--- /dev/null
+++ b/Engine/Plugins/Experimental/Mutable/Source/MutableTools/Private/MuT/NodeScalarTablePrivate.h
@@ -0,0 +1,56 @@
+// Copyright Epic Games, Inc. All Rights Reserved.
+
+#pragma once
+
+#include "MuT/NodePrivate.h"
+#include "MuT/NodeScalarTable.h"
+#include "MuT/Table.h"
+#include "MuT/AST.h"
+
+#include "../Private/MuR/MemoryPrivate.h"
+
+
+namespace mu
+{
+
+
+	class NodeScalarTable::Private : public Node::Private
+	{
+	public:
+
+		MUTABLE_DEFINE_CONST_VISITABLE()
+
+	public:
+
+		static NODE_TYPE s_type;
+
+		string m_parameterName;
+		TablePtr m_pTable;
+		string m_columnName;
+
+		//!
+		void Serialise( OutputArchive& arch ) const
+		{
+            uint32_t ver = 0;
+			arch << ver;
+
+			arch << m_parameterName;
+			arch << m_pTable;
+			arch << m_columnName;
+		}
+
+		//!
+		void Unserialise( InputArchive& arch )
+		{
+            uint32_t ver;
+			arch >> ver;
+            check(ver==0);
+
+			arch >> m_parameterName;
+			arch >> m_pTable;
+			arch >> m_columnName;
+		}
+
+	};
+
+}
diff --git a/Engine/Plugins/Experimental/Mutable/Source/MutableTools/Private/MuT/NodeScalarVariation.cpp b/Engine/Plugins/Experimental/Mutable/Source/MutableTools/Private/MuT/NodeScalarVariation.cpp
new file mode 100644
--- /dev/null
+++ b/Engine/Plugins/Experimental/Mutable/Source/MutableTools/Private/MuT/NodeScalarVariation.cpp
@@ -0,0 +1,115 @@
+// Copyright Epic Games, Inc. All Rights Reserved.
+
+
+#include "MuT/NodeScalarVariation.h"
+#include "MuT/NodeScalarVariationPrivate.h"
+
+
+namespace mu
+{
+
+    //---------------------------------------------------------------------------------------------
+    // Static initialisation
+    //---------------------------------------------------------------------------------------------
+    NODE_TYPE NodeScalarVariation::Private::s_type =
+        NODE_TYPE( "ScalarVariation", NodeScalar::GetStaticType() );
+
+
+    //---------------------------------------------------------------------------------------------
+    //---------------------------------------------------------------------------------------------
+    //---------------------------------------------------------------------------------------------
+
+    MUTABLE_IMPLEMENT_NODE( NodeScalarVariation, EType::Variation, Node, Node::EType::Scalar)
+
+
+    //---------------------------------------------------------------------------------------------
+    int NodeScalarVariation::GetInputCount() const { return 1 + int( m_pD->m_variations.size() ); }
+
+
+    //---------------------------------------------------------------------------------------------
+    Node* NodeScalarVariation::GetInputNode( int i ) const
+    {
+        check( i >= 0 && i < GetInputCount() );
+
+        if ( i == 0 )
+        {
+            return m_pD->m_defaultScalar.get();
+        }
+        i -= 1;
+
+        if ( i < int( m_pD->m_variations.size() ) )
+        {
+            return m_pD->m_variations[i].m_scalar.get();
+        }
+        i -= int( m_pD->m_variations.size() );
+
+        return nullptr;
+    }
+
+
+    //---------------------------------------------------------------------------------------------
+    void NodeScalarVariation::SetInputNode( int i, NodePtr pNode )
+    {
+        check( i >= 0 && i < GetInputCount() );
+
+        if ( i == 0 )
+        {
+            m_pD->m_defaultScalar = dynamic_cast<NodeScalar*>( pNode.get() );
+            return;
+        }
+
+        i -= 1;
+        if ( i < int( m_pD->m_variations.size() ) )
+        {
+
+            m_pD->m_variations[i].m_scalar = dynamic_cast<NodeScalar*>( pNode.get() );
+            return;
+        }
+        i -= (int)m_pD->m_variations.size();
+    }
+
+
+    //---------------------------------------------------------------------------------------------
+    // Own Interface
+    //---------------------------------------------------------------------------------------------
+    void NodeScalarVariation::SetDefaultScalar( NodeScalar* p ) { m_pD->m_defaultScalar = p; }
+
+
+    //---------------------------------------------------------------------------------------------
+    int NodeScalarVariation::GetVariationCount() const { return int( m_pD->m_variations.size() ); }
+
+
+    //---------------------------------------------------------------------------------------------
+    void NodeScalarVariation::SetVariationCount( int num )
+    {
+        check( num >= 0 );
+        m_pD->m_variations.resize( num );
+    }
+
+    //---------------------------------------------------------------------------------------------
+    void NodeScalarVariation::SetVariationTag( int index, const char* strTag )
+    {
+        check( index >= 0 && index < (int)m_pD->m_variations.size() );
+        check( strTag );
+
+        if ( strTag )
+        {
+            m_pD->m_variations[index].m_tag = strTag;
+        }
+        else
+        {
+            m_pD->m_variations[index].m_tag = "";
+        }
+    }
+
+
+    //---------------------------------------------------------------------------------------------
+    void NodeScalarVariation::SetVariationScalar( int index, NodeScalar* pNode )
+    {
+        check( index >= 0 && index < (int)m_pD->m_variations.size() );
+
+        m_pD->m_variations[index].m_scalar = pNode;
+    }
+
+
+} // namespace mu
diff --git a/Engine/Plugins/Experimental/Mutable/Source/MutableTools/Private/MuT/NodeScalarVariationPrivate.h b/Engine/Plugins/Experimental/Mutable/Source/MutableTools/Private/MuT/NodeScalarVariationPrivate.h
new file mode 100644
--- /dev/null
+++ b/Engine/Plugins/Experimental/Mutable/Source/MutableTools/Private/MuT/NodeScalarVariationPrivate.h
@@ -0,0 +1,78 @@
+// Copyright Epic Games, Inc. All Rights Reserved.
+
+#pragma once
+
+#include "MuT/NodePrivate.h"
+#include "MuT/NodeScalarVariation.h"
+
+#include "MuR/MemoryPrivate.h"
+
+
+namespace mu
+{
+
+    //---------------------------------------------------------------------------------------------
+    //!
+    //---------------------------------------------------------------------------------------------
+    class NodeScalarVariation::Private : public Node::Private
+    {
+    public:
+        MUTABLE_DEFINE_CONST_VISITABLE()
+
+    public:
+
+        static NODE_TYPE s_type;
+
+        NodeScalarPtr m_defaultScalar;
+
+        struct VARIATION
+        {
+            NodeScalarPtr m_scalar;
+            string m_tag;
+
+            //!
+            void Serialise( OutputArchive& arch ) const
+            {
+                uint32_t ver = 1;
+                arch << ver;
+
+                arch << m_tag;
+                arch << m_scalar;
+            }
+
+            void Unserialise( InputArchive& arch )
+            {
+                uint32_t ver = 0;
+                arch >> ver;
+                check( ver == 1 );
+
+                arch >> m_tag;
+                arch >> m_scalar;
+            }
+        };
+
+        vector<VARIATION> m_variations;
+
+        //!
+        void Serialise( OutputArchive& arch ) const
+        {
+            uint32_t ver = 1;
+            arch << ver;
+
+            arch << m_defaultScalar;
+            arch << m_variations;
+        }
+
+        //!
+        void Unserialise( InputArchive& arch )
+        {
+            uint32_t ver;
+            arch >> ver;
+			check(ver == 1);
+
+            arch >> m_defaultScalar;
+            arch >> m_variations;
+        }
+    };
+
+} // namespace mu
diff --git a/Engine/Plugins/Experimental/Mutable/Source/MutableTools/Private/MuT/NodeString.cpp b/Engine/Plugins/Experimental/Mutable/Source/MutableTools/Private/MuT/NodeString.cpp
new file mode 100644
--- /dev/null
+++ b/Engine/Plugins/Experimental/Mutable/Source/MutableTools/Private/MuT/NodeString.cpp
@@ -0,0 +1,80 @@
+// Copyright Epic Games, Inc. All Rights Reserved.
+
+
+#include "MuT/NodeString.h"
+
+#include "MuT/Node.h"
+#include "MuT/NodePrivate.h"
+
+#include "MuT/NodeStringConstant.h"
+#include "MuT/NodeStringParameter.h"
+
+
+namespace mu
+{
+
+	//---------------------------------------------------------------------------------------------
+	// Static initialisation
+	//---------------------------------------------------------------------------------------------
+	static NODE_TYPE s_nodeStringType = NODE_TYPE( "NodeString", Node::GetStaticType() );
+
+
+	//---------------------------------------------------------------------------------------------
+	void NodeString::Serialise( const NodeString* p, OutputArchive& arch )
+	{
+        uint32_t ver = 0;
+		arch << ver;
+
+
+#define SERIALISE_CHILDREN( C, ID ) \
+        ( const C* pTyped##ID = dynamic_cast<const C*>(p) )			\
+        {                                                           \
+            arch << (uint32_t)ID;                                   \
+            C::Serialise( pTyped##ID, arch );						\
+        }                                                           \
+
+        if SERIALISE_CHILDREN( NodeStringConstant                   , 0 )
+        else if SERIALISE_CHILDREN( NodeStringParameter             , 1 )
+        else check(false);
+
+#undef SERIALISE_CHILDREN
+    }
+
+    //---------------------------------------------------------------------------------------------
+	NodeStringPtr NodeString::StaticUnserialise( InputArchive& arch )
+	{
+        uint32_t ver;
+		arch >> ver;
+		check( ver == 0 );
+
+        uint32_t id;
+		arch >> id;
+
+		switch (id)
+		{
+		case 0 :  return NodeStringConstant::StaticUnserialise( arch ); break;
+		case 1 :  return NodeStringParameter::StaticUnserialise( arch ); break;
+        default : check(false);
+		}
+
+		return 0;
+	}
+
+
+	//---------------------------------------------------------------------------------------------
+	const NODE_TYPE* NodeString::GetType() const
+	{
+		return GetStaticType();
+	}
+
+
+	//---------------------------------------------------------------------------------------------
+	const NODE_TYPE* NodeString::GetStaticType()
+	{
+		return &s_nodeStringType;
+	}
+
+
+}
+
+
diff --git a/Engine/Plugins/Experimental/Mutable/Source/MutableTools/Private/MuT/NodeStringConstant.cpp b/Engine/Plugins/Experimental/Mutable/Source/MutableTools/Private/MuT/NodeStringConstant.cpp
new file mode 100644
--- /dev/null
+++ b/Engine/Plugins/Experimental/Mutable/Source/MutableTools/Private/MuT/NodeStringConstant.cpp
@@ -0,0 +1,72 @@
+// Copyright Epic Games, Inc. All Rights Reserved.
+
+
+#include "MuT/NodeStringConstant.h"
+#include "MuT/NodeStringConstantPrivate.h"
+
+#include "MuR/MutableMath.h"
+
+
+#define NODE_INPUT_COUNT 	0
+
+
+namespace mu
+{
+
+	//---------------------------------------------------------------------------------------------
+	// Static initialisation
+	//---------------------------------------------------------------------------------------------
+	NODE_TYPE NodeStringConstant::Private::s_type =
+			NODE_TYPE( "StringConstant", NodeString::GetStaticType() );
+
+
+	//---------------------------------------------------------------------------------------------
+	//!
+	//---------------------------------------------------------------------------------------------
+
+	MUTABLE_IMPLEMENT_NODE( NodeStringConstant, EType::Constant, Node, Node::EType::String)
+
+
+	//---------------------------------------------------------------------------------------------
+	// Node Interface
+	//---------------------------------------------------------------------------------------------
+	int NodeStringConstant::GetInputCount() const
+	{
+		return NODE_INPUT_COUNT;
+	}
+
+
+	//---------------------------------------------------------------------------------------------
+    Node* NodeStringConstant::GetInputNode( int ) const
+	{
+		check( false );
+		return 0;
+	}
+
+
+	//---------------------------------------------------------------------------------------------
+    void NodeStringConstant::SetInputNode( int, NodePtr )
+	{
+		check( false );
+	}
+
+
+	//---------------------------------------------------------------------------------------------
+	// Own Interface
+	//---------------------------------------------------------------------------------------------
+	const char* NodeStringConstant::GetValue() const
+	{
+		return m_pD->m_value.c_str();
+	}
+
+	//---------------------------------------------------------------------------------------------
+	void NodeStringConstant::SetValue( const char* v )
+	{
+		m_pD->m_value = v?v:"";
+	}
+
+
+}
+
+#undef NODE_INPUT_COUNT
+
diff --git a/Engine/Plugins/Experimental/Mutable/Source/MutableTools/Private/MuT/NodeStringConstantPrivate.h b/Engine/Plugins/Experimental/Mutable/Source/MutableTools/Private/MuT/NodeStringConstantPrivate.h
new file mode 100644
--- /dev/null
+++ b/Engine/Plugins/Experimental/Mutable/Source/MutableTools/Private/MuT/NodeStringConstantPrivate.h
@@ -0,0 +1,51 @@
+// Copyright Epic Games, Inc. All Rights Reserved.
+
+#pragma once
+
+#include "NodeStringPrivate.h"
+
+#include "MuT/NodeStringConstant.h"
+#include "MuT/NodeImage.h"
+#include "MuT/AST.h"
+#include "MuR/MemoryPrivate.h"
+#include "MuR/ParametersPrivate.h"
+
+
+namespace mu
+{
+
+
+	class NodeStringConstant::Private : public NodeString::Private
+	{
+	public:
+
+		MUTABLE_DEFINE_CONST_VISITABLE()
+
+	public:
+
+		static NODE_TYPE s_type;
+
+		string m_value;
+
+		//!
+		void Serialise( OutputArchive& arch ) const
+		{
+            uint32_t ver = 0;
+			arch << ver;
+
+			arch << m_value;
+		}
+
+		//!
+		void Unserialise( InputArchive& arch )
+		{
+            uint32_t ver;
+			arch >> ver;
+			check(ver==0);
+
+			arch >> m_value;
+		}
+	};
+
+
+}
diff --git a/Engine/Plugins/Experimental/Mutable/Source/MutableTools/Private/MuT/NodeStringParameter.cpp b/Engine/Plugins/Experimental/Mutable/Source/MutableTools/Private/MuT/NodeStringParameter.cpp
new file mode 100644
--- /dev/null
+++ b/Engine/Plugins/Experimental/Mutable/Source/MutableTools/Private/MuT/NodeStringParameter.cpp
@@ -0,0 +1,165 @@
+// Copyright Epic Games, Inc. All Rights Reserved.
+
+
+#include "MuT/NodeStringParameter.h"
+#include "MuT/NodeStringParameterPrivate.h"
+
+#include "MuR/MutableMath.h"
+#include "MuR/MemoryPrivate.h"
+
+
+namespace mu
+{
+
+	//---------------------------------------------------------------------------------------------
+	// Static initialisation
+	//---------------------------------------------------------------------------------------------
+	NODE_TYPE NodeStringParameter::Private::s_type =
+			NODE_TYPE( "StringParameter", NodeString::GetStaticType() );
+
+
+	//---------------------------------------------------------------------------------------------
+	//!
+	//---------------------------------------------------------------------------------------------
+
+	MUTABLE_IMPLEMENT_NODE( NodeStringParameter, EType::Parameter, Node, Node::EType::String)
+
+
+	//---------------------------------------------------------------------------------------------
+	int NodeStringParameter::GetInputCount() const
+	{
+        return int( m_pD->m_additionalImages.size()
+                    +
+                    m_pD->m_ranges.size() );
+	}
+
+
+	//---------------------------------------------------------------------------------------------
+    Node* NodeStringParameter::GetInputNode( int i ) const
+	{
+        check( i<GetInputCount() );
+
+        int imageCount = int(m_pD->m_additionalImages.size());
+        int rangeCount = int(m_pD->m_ranges.size());
+        if (i<imageCount)
+        {
+            return m_pD->m_additionalImages[i].get();
+        }
+        else if ( i < imageCount + rangeCount )
+        {
+            int r = i - imageCount;
+            return m_pD->m_ranges[r].get();
+        }
+        return nullptr;
+	}
+
+
+	//---------------------------------------------------------------------------------------------
+    void NodeStringParameter::SetInputNode( int i, NodePtr n )
+	{
+        check( i<GetInputCount() );
+        int imageCount = int(m_pD->m_additionalImages.size());
+        int rangeCount = int(m_pD->m_ranges.size());
+        if (i<imageCount)
+        {
+            m_pD->m_additionalImages[i] = dynamic_cast<NodeImage*>(n.get());
+        }
+        else if ( i < imageCount + rangeCount )
+        {
+            int r = i - imageCount;
+            m_pD->m_ranges[r] = dynamic_cast<NodeRange*>(n.get());
+        }
+    }
+
+
+	//---------------------------------------------------------------------------------------------
+	const char* NodeStringParameter::GetName() const
+	{
+		return m_pD->m_name.c_str();
+	}
+
+
+	//---------------------------------------------------------------------------------------------
+	void NodeStringParameter::SetName( const char* strName )
+	{
+		if ( strName )
+		{
+			m_pD->m_name = strName;
+		}
+		else
+		{
+			m_pD->m_name = "";
+		}
+	}
+
+
+	const char* NodeStringParameter::GetUid() const
+	{
+		return m_pD->m_uid.c_str();
+	}
+
+
+	//---------------------------------------------------------------------------------------------
+	void NodeStringParameter::SetUid( const char* strUid )
+	{
+		if ( strUid )
+		{
+			m_pD->m_uid = strUid;
+		}
+		else
+		{
+			m_pD->m_uid = "";
+		}
+	}
+
+
+	//---------------------------------------------------------------------------------------------
+	const char* NodeStringParameter::GetDefaultValue() const
+	{
+		return m_pD->m_defaultValue.c_str();
+	}
+
+
+	//---------------------------------------------------------------------------------------------
+	void NodeStringParameter::SetDefaultValue( const char* v )
+	{
+		m_pD->m_defaultValue = v?v:"";
+	}
+
+
+	//---------------------------------------------------------------------------------------------
+	PARAMETER_DETAILED_TYPE NodeStringParameter::GetDetailedType() const
+	{
+		return m_pD->m_detailedType;
+	}
+
+
+	//---------------------------------------------------------------------------------------------
+	void NodeStringParameter::SetDetailedType( PARAMETER_DETAILED_TYPE t )
+	{
+		m_pD->m_detailedType = t;
+	}
+
+
+    //---------------------------------------------------------------------------------------------
+    void NodeStringParameter::SetRangeCount( int i )
+    {
+        check(i>=0);
+        m_pD->m_ranges.resize(i);
+    }
+
+
+    //---------------------------------------------------------------------------------------------
+    void NodeStringParameter::SetRange( int i, NodeRangePtr pRange )
+    {
+        check( i>=0 && i<int(m_pD->m_ranges.size()) );
+        if ( i>=0 && i<int(m_pD->m_ranges.size()) )
+        {
+            m_pD->m_ranges[i] = pRange;
+        }
+    }
+
+
+}
+
+
diff --git a/Engine/Plugins/Experimental/Mutable/Source/MutableTools/Private/MuT/NodeStringParameterPrivate.h b/Engine/Plugins/Experimental/Mutable/Source/MutableTools/Private/MuT/NodeStringParameterPrivate.h
new file mode 100644
--- /dev/null
+++ b/Engine/Plugins/Experimental/Mutable/Source/MutableTools/Private/MuT/NodeStringParameterPrivate.h
@@ -0,0 +1,67 @@
+// Copyright Epic Games, Inc. All Rights Reserved.
+
+#pragma once
+
+#include "MuT/NodeStringPrivate.h"
+
+#include "MuT/NodeStringParameter.h"
+#include "MuT/NodeImage.h"
+#include "MuT/NodeRange.h"
+#include "MuR/MemoryPrivate.h"
+#include "MuR/ParametersPrivate.h"
+
+
+namespace mu
+{
+
+
+	class NodeStringParameter::Private : public NodeString::Private
+	{
+	public:
+
+		MUTABLE_DEFINE_CONST_VISITABLE()
+
+	public:
+
+		static NODE_TYPE s_type;
+
+		string m_defaultValue;
+		string m_name;
+		string m_uid;
+		PARAMETER_DETAILED_TYPE m_detailedType = PARAMETER_DETAILED_TYPE::UNKNOWN;
+
+        vector<Ptr<NodeImage>> m_additionalImages;
+
+        vector<Ptr<NodeRange>> m_ranges;
+
+		//!
+		void Serialise( OutputArchive& arch ) const
+		{
+            uint32_t ver = 4;
+			arch << ver;
+
+			arch << m_defaultValue;
+			arch << m_name;
+			arch << m_uid;
+			arch << m_additionalImages;
+			arch << m_detailedType;
+            arch << m_ranges;
+        }
+
+		//!
+		void Unserialise( InputArchive& arch )
+		{
+            uint32_t ver;
+			arch >> ver;
+            check(ver==4);
+
+			arch >> m_defaultValue;
+			arch >> m_name;
+            arch >> m_uid;
+			arch >> m_additionalImages;
+            arch >> m_detailedType;
+            arch >> m_ranges;
+        }
+	};
+
+}
diff --git a/Engine/Plugins/Experimental/Mutable/Source/MutableTools/Private/MuT/NodeStringPrivate.h b/Engine/Plugins/Experimental/Mutable/Source/MutableTools/Private/MuT/NodeStringPrivate.h
new file mode 100644
--- /dev/null
+++ b/Engine/Plugins/Experimental/Mutable/Source/MutableTools/Private/MuT/NodeStringPrivate.h
@@ -0,0 +1,21 @@
+// Copyright Epic Games, Inc. All Rights Reserved.
+
+#pragma once
+
+#include "MuT/NodePrivate.h"
+
+#include "MuT/NodeString.h"
+#include "MuR/MemoryPrivate.h"
+#include "MuR/ParametersPrivate.h"
+
+
+namespace mu
+{
+
+
+	class NodeString::Private : public Node::Private
+	{
+	};
+
+
+}
diff --git a/Engine/Plugins/Experimental/Mutable/Source/MutableTools/Private/MuT/NodeSurface.cpp b/Engine/Plugins/Experimental/Mutable/Source/MutableTools/Private/MuT/NodeSurface.cpp
new file mode 100644
--- /dev/null
+++ b/Engine/Plugins/Experimental/Mutable/Source/MutableTools/Private/MuT/NodeSurface.cpp
@@ -0,0 +1,78 @@
+// Copyright Epic Games, Inc. All Rights Reserved.
+
+
+#include "MuT/NodeSurface.h"
+
+#include "MuT/NodeMesh.h"
+#include "MuT/NodeImage.h"
+
+#include "MuT/NodeSurfaceNew.h"
+#include "MuT/NodeSurfaceEdit.h"
+#include "MuT/NodeSurfaceVariation.h"
+
+#include "MuR/SerialisationPrivate.h"
+
+
+namespace mu
+{
+
+
+	//---------------------------------------------------------------------------------------------
+	// Static initialisation
+	//---------------------------------------------------------------------------------------------
+    static NODE_TYPE s_nodeSurfaceType =
+            NODE_TYPE( "NodeSurface", Node::GetStaticType() );
+
+
+	//---------------------------------------------------------------------------------------------
+	//---------------------------------------------------------------------------------------------
+	//---------------------------------------------------------------------------------------------
+    const NODE_TYPE* NodeSurface::GetType() const
+	{
+		return GetStaticType();
+	}
+
+
+	//---------------------------------------------------------------------------------------------
+    const NODE_TYPE* NodeSurface::GetStaticType()
+	{
+        return &s_nodeSurfaceType;
+	}
+
+
+	//---------------------------------------------------------------------------------------------
+    void NodeSurface::Serialise( const NodeSurface* p, OutputArchive& arch )
+	{
+        uint32_t ver = 0;
+		arch << ver;
+
+		arch << uint32_t(p->Type);
+		p->SerialiseWrapper(arch);
+	}
+
+
+    //---------------------------------------------------------------------------------------------
+    NodeSurfacePtr NodeSurface::StaticUnserialise( InputArchive& arch )
+	{
+        uint32_t ver;
+		arch >> ver;
+		check( ver == 0 );
+
+        uint32_t id;
+		arch >> id;
+
+		switch (id)
+		{
+        case 0: return NodeSurfaceNew::StaticUnserialise( arch ); break;
+        case 1: return NodeSurfaceEdit::StaticUnserialise( arch ); break;
+		//case 2: return NodeSelectSurface::StaticUnserialise(arch); break;
+		case 3: return NodeSurfaceVariation::StaticUnserialise(arch); break;
+		default : check(false);
+		}
+
+		return 0;
+	}
+
+}
+
+
