Moved the Mutable plugin files to live under a directory named after the module. This allows us to get rid of long include paths that contain public and private folder names in the path and fixes head

diff --git a/Engine/Plugins/Experimental/Mutable/Source/MutableTools/Public/MuT/NodeBool.h b/Engine/Plugins/Experimental/Mutable/Source/MutableTools/Public/MuT/NodeBool.h
new file mode 100644
--- /dev/null
+++ b/Engine/Plugins/Experimental/Mutable/Source/MutableTools/Public/MuT/NodeBool.h
@@ -0,0 +1,403 @@
+// Copyright Epic Games, Inc. All Rights Reserved.
+
+#pragma once
+
+#include "MuT/Node.h"
+
+
+namespace mu
+{
+
+	// Forward definitions
+	class NodeBool;
+    using NodeBoolPtr = Ptr<NodeBool>;
+    using NodeBoolPtrConst = Ptr<const NodeBool>;
+
+	class NodeBoolConstant;
+    using NodeBoolConstantPtr = Ptr<NodeBoolConstant>;
+    using NodeBoolConstantPtrConst = Ptr<const NodeBoolConstant>;
+
+	class NodeBoolParameter;
+    using NodeBoolParameterPtr = Ptr<NodeBoolParameter>;
+    using NodeBoolParameterPtrConst = Ptr<const NodeBoolParameter>;
+
+	class NodeBoolIsNull;
+    using NodeBoolIsNullPtr = Ptr<NodeBoolIsNull>;
+    using NodeBoolIsNullPtrConst = Ptr<const NodeBoolIsNull>;
+
+	class NodeBoolNot;
+    using NodeBoolNotPtr = Ptr<NodeBoolNot>;
+    using NodeBoolNotPtrConst = Ptr<const NodeBoolNot>;
+
+	class NodeBoolAnd;
+    using NodeBoolAndPtr = Ptr<NodeBoolAnd>;
+    using NodeBoolAndPtrConst = Ptr<const NodeBoolAnd>;
+
+    class NodeRange;
+    using NodeRangePtr = Ptr<NodeRange>;
+    using NodeRangePtrConst = Ptr<const NodeRange>;
+
+
+	//---------------------------------------------------------------------------------------------
+    //! %Base class of any node that outputs a Bool value.
+	//! \ingroup model
+	//---------------------------------------------------------------------------------------------
+	class MUTABLETOOLS_API NodeBool : public Node
+	{
+	public:
+
+		// Possible subclasses
+		enum class EType : uint8
+		{
+			Constant = 0,
+			Parameter = 1,
+			IsNull = 2,
+			Not = 3,
+			And = 4,
+
+			None
+		};
+
+		//-----------------------------------------------------------------------------------------
+		// Life cycle
+		//-----------------------------------------------------------------------------------------
+
+		static void Serialise( const NodeBool* pNode, OutputArchive& arch );
+		static NodeBoolPtr StaticUnserialise( InputArchive& arch );
+
+
+		//-----------------------------------------------------------------------------------------
+		// Node Interface
+		//-----------------------------------------------------------------------------------------
+
+        const NODE_TYPE* GetType() const override;
+		static const NODE_TYPE* GetStaticType();
+
+
+		//-----------------------------------------------------------------------------------------
+		// Interface pattern
+		//-----------------------------------------------------------------------------------------
+		class Private;
+
+	protected:
+
+		//! Forbidden. Manage with the Ptr<> template.
+		inline ~NodeBool() {}
+
+		//!
+		EType Type = EType::None;
+
+	};
+
+
+	//---------------------------------------------------------------------------------------------
+	//! Node returning a Bool constant value.
+	//! \ingroup model
+	//---------------------------------------------------------------------------------------------
+	class MUTABLETOOLS_API NodeBoolConstant : public NodeBool
+	{
+	public:
+
+		//-----------------------------------------------------------------------------------------
+		// Life cycle
+		//-----------------------------------------------------------------------------------------
+
+		NodeBoolConstant();
+
+		static void Serialise( const NodeBoolConstant* pNode, OutputArchive& arch );
+		void SerialiseWrapper(OutputArchive& arch) const override;
+		static NodeBoolConstantPtr StaticUnserialise( InputArchive& arch );
+
+
+		//-----------------------------------------------------------------------------------------
+		// Node Interface
+		//-----------------------------------------------------------------------------------------
+
+        NodePtr Clone( NodeMapPtr pMap = 0 ) const override;
+
+        const NODE_TYPE* GetType() const override;
+		static const NODE_TYPE* GetStaticType();
+
+        virtual int GetInputCount() const override;
+        virtual Node* GetInputNode( int i ) const override;
+        void SetInputNode( int i, NodePtr pNode ) override;
+
+		//-----------------------------------------------------------------------------------------
+		// Own Interface
+		//-----------------------------------------------------------------------------------------
+
+		//! Get the value to be returned by the node.
+		bool GetValue() const;
+
+		//! Set the value to be returned by the node.
+		void SetValue( bool v );
+
+		//-----------------------------------------------------------------------------------------
+		// Interface pattern
+		//-----------------------------------------------------------------------------------------
+		class Private;
+		Private* GetPrivate() const;
+        Node::Private* GetBasePrivate() const override;
+
+	protected:
+
+		//! Forbidden.
+		~NodeBoolConstant();
+
+	private:
+
+		Private* m_pD;
+
+	};
+
+
+	//---------------------------------------------------------------------------------------------
+	//! Node that defines a Bool model parameter.
+	//! \ingroup model
+	//---------------------------------------------------------------------------------------------
+	class MUTABLETOOLS_API NodeBoolParameter : public NodeBool
+	{
+	public:
+
+		//-----------------------------------------------------------------------------------------
+		// Life cycle
+		//-----------------------------------------------------------------------------------------
+
+		NodeBoolParameter();
+
+		void SerialiseWrapper(OutputArchive& arch) const override;
+		static void Serialise( const NodeBoolParameter* pNode, OutputArchive& arch );
+		static NodeBoolParameterPtr StaticUnserialise( InputArchive& arch );
+
+
+		//-----------------------------------------------------------------------------------------
+		// Node Interface
+		//-----------------------------------------------------------------------------------------
+
+        NodePtr Clone( NodeMapPtr pMap = 0 ) const override;
+
+        const NODE_TYPE* GetType() const override;
+		static const NODE_TYPE* GetStaticType();
+
+        int GetInputCount() const override;
+        Node* GetInputNode( int i ) const override;
+        void SetInputNode( int i, NodePtr pNode ) override;
+
+		//-----------------------------------------------------------------------------------------
+		// Own Interface
+		//-----------------------------------------------------------------------------------------
+
+		//! Get the name of the parameter. It will be exposed in the final compiled data.
+		const char* GetName() const;
+		void SetName( const char* );
+
+		//! Get the uid of the parameter. It will be exposed in the final compiled data.
+		const char* GetUid() const;
+		void SetUid( const char* );
+
+		//! Get the default value of the parameter.
+		bool GetDefaultValue() const;
+		void SetDefaultValue( bool v );
+
+        //! Set the number of ranges (dimensions) for this parameter.
+        //! By default a parameter has 0 ranges, meaning it only has one value.
+        void SetRangeCount( int i );
+        void SetRange( int i, NodeRangePtr pRange );
+
+		//-----------------------------------------------------------------------------------------
+		// Interface pattern
+		//-----------------------------------------------------------------------------------------
+		class Private;
+		Private* GetPrivate() const;
+        Node::Private* GetBasePrivate() const override;
+
+	protected:
+
+		//! Forbidden. 
+		~NodeBoolParameter();
+
+	private:
+
+		Private* m_pD;
+
+	};
+
+
+	//---------------------------------------------------------------------------------------------
+	//! Node that returns true if there is an input or false if there is nothing connected.
+	//! This node is mostly useful when used in the model Transform's conditions.
+	//! \ingroup model
+	//---------------------------------------------------------------------------------------------
+	class MUTABLETOOLS_API NodeBoolIsNull : public NodeBool
+	{
+	public:
+
+		//-----------------------------------------------------------------------------------------
+		// Life cycle
+		//-----------------------------------------------------------------------------------------
+
+		NodeBoolIsNull();
+
+		void SerialiseWrapper(OutputArchive& arch) const override;
+		static void Serialise( const NodeBoolIsNull* pNode, OutputArchive& arch );
+		static NodeBoolIsNullPtr StaticUnserialise( InputArchive& arch );
+
+
+		//-----------------------------------------------------------------------------------------
+		// Node Interface
+		//-----------------------------------------------------------------------------------------
+
+        NodePtr Clone( NodeMapPtr pMap = 0 ) const override;
+
+        const NODE_TYPE* GetType() const override;
+		static const NODE_TYPE* GetStaticType();
+
+        int GetInputCount() const override;
+        Node* GetInputNode( int i ) const override;
+        void SetInputNode( int i, NodePtr pNode ) override;
+
+		//-----------------------------------------------------------------------------------------
+		// Own Interface
+		//-----------------------------------------------------------------------------------------
+
+		//-----------------------------------------------------------------------------------------
+		// Interface pattern
+		//-----------------------------------------------------------------------------------------
+		class Private;
+		Private* GetPrivate() const;
+        Node::Private* GetBasePrivate() const override;
+
+	protected:
+
+		//! Forbidden. 
+		~NodeBoolIsNull();
+
+	private:
+
+		Private* m_pD;
+
+	};
+
+
+	//---------------------------------------------------------------------------------------------
+	//! Node that returns the oposite of the input value.
+	//! \ingroup model
+	//---------------------------------------------------------------------------------------------
+	class MUTABLETOOLS_API NodeBoolNot : public NodeBool
+	{
+	public:
+
+		//-----------------------------------------------------------------------------------------
+		// Life cycle
+		//-----------------------------------------------------------------------------------------
+
+		NodeBoolNot();
+
+		void SerialiseWrapper(OutputArchive& arch) const override;
+		static void Serialise( const NodeBoolNot* pNode, OutputArchive& arch );
+		static NodeBoolNotPtr StaticUnserialise( InputArchive& arch );
+
+
+		//-----------------------------------------------------------------------------------------
+		// Node Interface
+		//-----------------------------------------------------------------------------------------
+
+        NodePtr Clone( NodeMapPtr pMap = 0 ) const override;
+
+        const NODE_TYPE* GetType() const override;
+		static const NODE_TYPE* GetStaticType();
+
+        int GetInputCount() const override;
+        Node* GetInputNode( int i ) const override;
+        void SetInputNode( int i, NodePtr pNode ) override;
+
+		//-----------------------------------------------------------------------------------------
+		// Own Interface
+		//-----------------------------------------------------------------------------------------
+		
+		//! Input
+		NodeBoolPtr GetInput() const;
+		void SetInput( NodeBoolPtr );
+
+		//-----------------------------------------------------------------------------------------
+		// Interface pattern
+		//-----------------------------------------------------------------------------------------
+		class Private;
+		Private* GetPrivate() const;
+        Node::Private* GetBasePrivate() const override;
+
+	protected:
+
+		//! Forbidden. 
+		~NodeBoolNot();
+
+	private:
+
+		Private* m_pD;
+
+	};
+
+
+	//---------------------------------------------------------------------------------------------
+	//! 
+	//! \ingroup model
+	//---------------------------------------------------------------------------------------------
+	class MUTABLETOOLS_API NodeBoolAnd : public NodeBool
+	{
+	public:
+
+		//-----------------------------------------------------------------------------------------
+		// Life cycle
+		//-----------------------------------------------------------------------------------------
+
+		NodeBoolAnd();
+
+		void SerialiseWrapper(OutputArchive& arch) const override;
+		static void Serialise( const NodeBoolAnd* pNode, OutputArchive& arch );
+		static NodeBoolAndPtr StaticUnserialise( InputArchive& arch );
+
+
+		//-----------------------------------------------------------------------------------------
+		// Node Interface
+		//-----------------------------------------------------------------------------------------
+
+        NodePtr Clone( NodeMapPtr pMap = 0 ) const override;
+
+        const NODE_TYPE* GetType() const override;
+		static const NODE_TYPE* GetStaticType();
+
+        int GetInputCount() const override;
+        Node* GetInputNode( int i ) const override;
+        void SetInputNode( int i, NodePtr pNode ) override;
+
+		//-----------------------------------------------------------------------------------------
+		// Own Interface
+		//-----------------------------------------------------------------------------------------
+		
+		//! Inputs
+		NodeBoolPtr GetA() const;
+		void SetA( NodeBoolPtr );
+
+		NodeBoolPtr GetB() const;
+		void SetB( NodeBoolPtr );
+
+		//-----------------------------------------------------------------------------------------
+		// Interface pattern
+		//-----------------------------------------------------------------------------------------
+		class Private;
+		Private* GetPrivate() const;
+        Node::Private* GetBasePrivate() const override;
+
+	protected:
+
+		//! Forbidden. 
+		~NodeBoolAnd();
+
+	private:
+
+		Private* m_pD;
+
+	};
+
+
+}
diff --git a/Engine/Plugins/Experimental/Mutable/Source/MutableTools/Public/MuT/NodeColour.h b/Engine/Plugins/Experimental/Mutable/Source/MutableTools/Public/MuT/NodeColour.h
new file mode 100644
--- /dev/null
+++ b/Engine/Plugins/Experimental/Mutable/Source/MutableTools/Public/MuT/NodeColour.h
@@ -0,0 +1,70 @@
+// Copyright Epic Games, Inc. All Rights Reserved.
+
+#pragma once
+
+#include "MuT/Node.h"
+
+
+namespace mu
+{
+
+	class NodeColour;
+	typedef Ptr<NodeColour> NodeColourPtr;
+	typedef Ptr<const NodeColour> NodeColourPtrConst;
+
+
+    //! %Base class of any node that outputs a colour.
+    //! \ingroup model
+    class MUTABLETOOLS_API NodeColour : public Node
+    {
+    public:
+
+		// Possible subclasses
+		enum class EType : uint8
+		{
+			Constant = 0,
+			Parameter = 1,
+			SampleImage = 2,
+			Table = 3,
+			ImageSize = 4,
+			FromScalars = 5,
+			ArithmeticOperation = 6,
+			Switch = 7,
+			Variation = 8,
+
+			None
+		};
+
+        //-----------------------------------------------------------------------------------------
+        // Life cycle
+        //-----------------------------------------------------------------------------------------
+
+        static void Serialise( const NodeColour* pNode, OutputArchive& arch );
+        static NodeColourPtr StaticUnserialise( InputArchive& arch );
+
+
+        //-----------------------------------------------------------------------------------------
+        // Node Interface
+        //-----------------------------------------------------------------------------------------
+
+        const NODE_TYPE* GetType() const override;
+        static const NODE_TYPE* GetStaticType();
+
+
+        //-----------------------------------------------------------------------------------------
+        // Interface pattern
+        //-----------------------------------------------------------------------------------------
+        class Private;
+
+    protected:
+
+        //! Forbidden. Manage with the Ptr<> template.
+        inline ~NodeColour() {}
+
+		//!
+		EType Type = EType::None;
+
+    };
+
+}
+
diff --git a/Engine/Plugins/Experimental/Mutable/Source/MutableTools/Public/MuT/NodeColourArithmeticOperation.h b/Engine/Plugins/Experimental/Mutable/Source/MutableTools/Public/MuT/NodeColourArithmeticOperation.h
new file mode 100644
--- /dev/null
+++ b/Engine/Plugins/Experimental/Mutable/Source/MutableTools/Public/MuT/NodeColourArithmeticOperation.h
@@ -0,0 +1,98 @@
+// Copyright Epic Games, Inc. All Rights Reserved.
+
+#pragma once
+
+#include "MuT/NodeColour.h"
+
+
+namespace mu
+{
+
+	// Forward definitions
+	class NodeColourArithmeticOperation;
+	typedef Ptr<NodeColourArithmeticOperation> NodeColourArithmeticOperationPtr;
+	typedef Ptr<const NodeColourArithmeticOperation> NodeColourArithmeticOperationPtrConst;
+
+
+    //! Perform a per-component arithmetic operation between two colours.
+	//! \ingroup model
+	class MUTABLETOOLS_API NodeColourArithmeticOperation : public NodeColour
+	{
+	public:
+
+		//-----------------------------------------------------------------------------------------
+		// Life cycle
+		//-----------------------------------------------------------------------------------------
+
+		NodeColourArithmeticOperation();
+
+		void SerialiseWrapper(OutputArchive& arch) const override;
+		static void Serialise( const NodeColourArithmeticOperation* pNode, OutputArchive& arch );
+		static NodeColourArithmeticOperationPtr StaticUnserialise( InputArchive& arch );
+
+
+		//-----------------------------------------------------------------------------------------
+		// Node Interface
+		//-----------------------------------------------------------------------------------------
+
+		virtual NodePtr Clone( NodeMapPtr pMap = 0 ) const override;
+
+		virtual const NODE_TYPE* GetType() const override;
+		static const NODE_TYPE* GetStaticType();
+
+		virtual int GetInputCount() const override;
+		virtual Node* GetInputNode( int i ) const override;
+		virtual void SetInputNode( int i, NodePtr pNode ) override;
+
+
+		//-----------------------------------------------------------------------------------------
+		// Own Interface
+		//-----------------------------------------------------------------------------------------
+
+		//! Possible arithmetic operations
+		typedef enum
+		{
+			AO_ADD,
+			AO_SUBTRACT,
+			AO_MULTIPLY,
+            AO_DIVIDE,
+            _AO_COUNT
+		} OPERATION;
+
+        static const char* s_opTypeName[ _AO_COUNT ];
+
+		//! Get child selection type
+		OPERATION GetOperation() const;
+
+		//! Set the child selection type
+		void SetOperation(OPERATION);
+
+
+		//! Get the first operand
+		NodeColourPtr GetA() const;
+		void SetA(NodeColourPtr);
+
+		// Get the second operand
+		NodeColourPtr GetB() const;
+		void SetB(NodeColourPtr);
+
+		//-----------------------------------------------------------------------------------------
+		// Interface pattern
+		//-----------------------------------------------------------------------------------------
+		class Private;
+		Private* GetPrivate() const;
+        Node::Private* GetBasePrivate() const override;
+
+	protected:
+
+		//! Forbidden. Manage with the Ptr<> template.
+		~NodeColourArithmeticOperation();
+
+	private:
+
+		Private* m_pD;
+
+	};
+
+
+}
diff --git a/Engine/Plugins/Experimental/Mutable/Source/MutableTools/Public/MuT/NodeColourConstant.h b/Engine/Plugins/Experimental/Mutable/Source/MutableTools/Public/MuT/NodeColourConstant.h
new file mode 100644
--- /dev/null
+++ b/Engine/Plugins/Experimental/Mutable/Source/MutableTools/Public/MuT/NodeColourConstant.h
@@ -0,0 +1,74 @@
+// Copyright Epic Games, Inc. All Rights Reserved.
+
+#pragma once
+
+#include "MuT/NodeColour.h"
+
+
+namespace mu
+{
+
+	class NodeColourConstant;
+	typedef Ptr<NodeColourConstant> NodeColourConstantPtr;
+	typedef Ptr<const NodeColourConstant> NodeColourConstantPtrConst;
+
+
+	//! This node outputs a predefined colour value.
+	//! \ingroup model
+	class MUTABLETOOLS_API NodeColourConstant : public NodeColour
+	{
+	public:
+
+		//-----------------------------------------------------------------------------------------
+		// Life cycle
+		//-----------------------------------------------------------------------------------------
+
+		NodeColourConstant();
+
+		void SerialiseWrapper(OutputArchive& arch) const override;
+		static void Serialise( const NodeColourConstant* pNode, OutputArchive& arch );
+		static NodeColourConstantPtr StaticUnserialise( InputArchive& arch );
+
+		//-----------------------------------------------------------------------------------------
+		// Node Interface
+		//-----------------------------------------------------------------------------------------
+
+        NodePtr Clone( NodeMapPtr pMap = 0 ) const override;
+
+        const NODE_TYPE* GetType() const override;
+		static const NODE_TYPE* GetStaticType();
+
+        virtual int GetInputCount() const override;
+        virtual Node* GetInputNode( int i ) const override;
+        void SetInputNode( int i, NodePtr pNode ) override;
+
+		//-----------------------------------------------------------------------------------------
+		// Own Interface
+		//-----------------------------------------------------------------------------------------
+
+		//! Get the value that this node returns
+		void GetValue( float* pR, float* pG, float* pB ) const;
+
+		//! Set the value to be returned by this node
+		void SetValue( float r, float g, float b );
+
+		//-----------------------------------------------------------------------------------------
+		// Interface pattern
+		//-----------------------------------------------------------------------------------------
+		class Private;
+		Private* GetPrivate() const;
+        Node::Private* GetBasePrivate() const override;
+
+	protected:
+
+		//! Forbidden. Manage with the Ptr<> template.
+		~NodeColourConstant();
+
+	private:
+
+		Private* m_pD;
+
+	};
+
+
+}
diff --git a/Engine/Plugins/Experimental/Mutable/Source/MutableTools/Public/MuT/NodeColourFromScalars.h b/Engine/Plugins/Experimental/Mutable/Source/MutableTools/Public/MuT/NodeColourFromScalars.h
new file mode 100644
--- /dev/null
+++ b/Engine/Plugins/Experimental/Mutable/Source/MutableTools/Public/MuT/NodeColourFromScalars.h
@@ -0,0 +1,91 @@
+// Copyright Epic Games, Inc. All Rights Reserved.
+
+#pragma once
+
+#include "MuT/NodeColour.h"
+
+#include "MuR/Image.h"
+
+
+namespace mu
+{
+
+	// Forward definitions
+	class NodeScalar;
+	typedef Ptr<NodeScalar> NodeScalarPtr;
+	typedef Ptr<const NodeScalar> NodeScalarPtrConst;
+
+	class NodeColourFromScalars;
+	typedef Ptr<NodeColourFromScalars> NodeColourFromScalarsPtr;
+	typedef Ptr<const NodeColourFromScalars> NodeColourFromScalarsPtrConst;
+
+
+	//! Obtain a colour by sampling an image at specific homogeneous coordinates.
+	//! \ingroup model
+	class MUTABLETOOLS_API NodeColourFromScalars : public NodeColour
+	{
+	public:
+
+		//-----------------------------------------------------------------------------------------
+		// Life cycle
+		//-----------------------------------------------------------------------------------------
+
+		NodeColourFromScalars();
+
+		void SerialiseWrapper(OutputArchive& arch) const override;
+		static void Serialise( const NodeColourFromScalars* pNode, OutputArchive& arch );
+		static NodeColourFromScalarsPtr StaticUnserialise( InputArchive& arch );
+
+
+		//-----------------------------------------------------------------------------------------
+		// Node Interface
+		//-----------------------------------------------------------------------------------------
+
+        NodePtr Clone( NodeMapPtr pMap = nullptr ) const override;
+
+		const NODE_TYPE* GetType() const override;
+		static const NODE_TYPE* GetStaticType();
+
+        int GetInputCount() const override;
+        Node* GetInputNode( int i ) const override;
+		void SetInputNode( int i, NodePtr pNode ) override;
+
+
+		//-----------------------------------------------------------------------------------------
+		// Own Interface
+		//-----------------------------------------------------------------------------------------
+
+		//! Get the node generating the x coordinate to sample, in the range of 0.0 to 1.0.
+		//! If it is not specified, a value of 1 will be used.
+		NodeScalarPtr GetX() const;
+		void SetX( NodeScalarPtr );
+
+		NodeScalarPtr GetY() const;
+		void SetY( NodeScalarPtr );
+
+		NodeScalarPtr GetZ() const;
+		void SetZ( NodeScalarPtr );
+
+		NodeScalarPtr GetW() const;
+		void SetW( NodeScalarPtr );
+
+		//-----------------------------------------------------------------------------------------
+		// Interface pattern
+		//-----------------------------------------------------------------------------------------
+		class Private;
+		Private* GetPrivate() const;
+        Node::Private* GetBasePrivate() const override;
+
+	protected:
+
+		//! Forbidden. Manage with the Ptr<> template.
+		~NodeColourFromScalars();
+
+	private:
+
+		Private* m_pD;
+
+	};
+
+
+}
diff --git a/Engine/Plugins/Experimental/Mutable/Source/MutableTools/Public/MuT/NodeColourParameter.h b/Engine/Plugins/Experimental/Mutable/Source/MutableTools/Public/MuT/NodeColourParameter.h
new file mode 100644
--- /dev/null
+++ b/Engine/Plugins/Experimental/Mutable/Source/MutableTools/Public/MuT/NodeColourParameter.h
@@ -0,0 +1,96 @@
+// Copyright Epic Games, Inc. All Rights Reserved.
+
+#pragma once
+
+#include "MuT/NodeColour.h"
+
+
+namespace mu
+{
+
+	class NodeColourParameter;
+    using NodeColourParameterPtr = Ptr<NodeColourParameter>;
+    using NodeColourParameterPtrConst = Ptr<const NodeColourParameter>;
+
+    class NodeRange;
+    using NodeRangePtr = Ptr<NodeRange>;
+    using NodeRangePtrConst = Ptr<const NodeRange>;
+
+
+	//! Node that defines a colour model parameter.
+	//! \ingroup model
+	class MUTABLETOOLS_API NodeColourParameter : public NodeColour
+	{
+	public:
+
+		//-----------------------------------------------------------------------------------------
+		// Life cycle
+		//-----------------------------------------------------------------------------------------
+
+		NodeColourParameter();
+
+		void SerialiseWrapper(OutputArchive& arch) const override;
+		static void Serialise( const NodeColourParameter* pNode, OutputArchive& arch );
+		static NodeColourParameterPtr StaticUnserialise( InputArchive& arch );
+
+
+		//-----------------------------------------------------------------------------------------
+		// Node Interface
+		//-----------------------------------------------------------------------------------------
+
+		NodePtr Clone( NodeMapPtr pMap = 0 ) const override;
+
+		const NODE_TYPE* GetType() const override;
+		static const NODE_TYPE* GetStaticType();
+
+		int GetInputCount() const override;
+		Node* GetInputNode( int i ) const override;
+		void SetInputNode( int i, NodePtr pNode ) override;
+
+		//-----------------------------------------------------------------------------------------
+		// Own Interface
+		//-----------------------------------------------------------------------------------------
+
+		//! Get the name of the parameter. It will be exposed in the final compiled data.
+		const char* GetName() const;
+
+		//! Set the name of the parameter.
+		void SetName( const char* );
+
+		//! Get the uid of the parameter. It will be exposed in the final compiled data.
+		const char* GetUid() const;
+
+		//! Set the uid of the parameter.
+		void SetUid( const char* );
+
+		//! Get the default value of the parameter.
+		void GetDefaultValue( float* pR, float* pG, float* pB ) const;
+
+		//! Set the default value of the parameter.
+		void SetDefaultValue( float r, float g, float b );
+
+        //! Set the number of ranges (dimensions) for this parameter.
+        //! By default a parameter has 0 ranges, meaning it only has one value.
+        void SetRangeCount( int i );
+        void SetRange( int i, NodeRangePtr pRange );
+
+		//-----------------------------------------------------------------------------------------
+		// Interface pattern
+		//-----------------------------------------------------------------------------------------
+		class Private;
+		Private* GetPrivate() const;
+        Node::Private* GetBasePrivate() const override;
+
+	protected:
+
+		//! Forbidden. Manage with the Ptr<> template.
+		~NodeColourParameter();
+
+	private:
+
+		Private* m_pD;
+
+	};
+
+
+}
diff --git a/Engine/Plugins/Experimental/Mutable/Source/MutableTools/Public/MuT/NodeColourSampleImage.h b/Engine/Plugins/Experimental/Mutable/Source/MutableTools/Public/MuT/NodeColourSampleImage.h
new file mode 100644
--- /dev/null
+++ b/Engine/Plugins/Experimental/Mutable/Source/MutableTools/Public/MuT/NodeColourSampleImage.h
@@ -0,0 +1,95 @@
+// Copyright Epic Games, Inc. All Rights Reserved.
+
+#pragma once
+
+#include "MuT/NodeColour.h"
+
+#include "MuR/Image.h"
+
+
+namespace mu
+{
+
+	// Forward definitions
+	class NodeScalar;
+	typedef Ptr<NodeScalar> NodeScalarPtr;
+	typedef Ptr<const NodeScalar> NodeScalarPtrConst;
+
+	class NodeImage;
+	typedef Ptr<NodeImage> NodeImagePtr;
+	typedef Ptr<const NodeImage> NodeImagePtrConst;
+
+	class NodeColourSampleImage;
+	typedef Ptr<NodeColourSampleImage> NodeColourSampleImagePtr;
+	typedef Ptr<const NodeColourSampleImage> NodeColourSampleImagePtrConst;
+
+
+	//! Obtain a colour by sampling an image at specific homogeneous coordinates.
+	//! \ingroup model
+	class MUTABLETOOLS_API NodeColourSampleImage : public NodeColour
+	{
+	public:
+
+		//-----------------------------------------------------------------------------------------
+		// Life cycle
+		//-----------------------------------------------------------------------------------------
+
+		NodeColourSampleImage();
+
+		void SerialiseWrapper(OutputArchive& arch) const override;
+		static void Serialise( const NodeColourSampleImage* pNode, OutputArchive& arch );
+		static NodeColourSampleImagePtr StaticUnserialise( InputArchive& arch );
+
+
+		//-----------------------------------------------------------------------------------------
+		// Node Interface
+		//-----------------------------------------------------------------------------------------
+
+		NodePtr Clone( NodeMapPtr pMap = 0 ) const override;
+
+		const NODE_TYPE* GetType() const override;
+		static const NODE_TYPE* GetStaticType();
+
+		int GetInputCount() const override;
+		Node* GetInputNode( int i ) const override;
+		void SetInputNode( int i, NodePtr pNode ) override;
+
+
+		//-----------------------------------------------------------------------------------------
+		// Own Interface
+		//-----------------------------------------------------------------------------------------
+
+		//! Get the node generating the x coordinate to sample, in the range of 0.0 to 1.0.
+		//! If it is not specified, a value of x =0.5 will be used.
+		NodeScalarPtr GetX() const;
+		void SetX( NodeScalarPtr );
+
+		//! Get the node generating the y coordinate to sample, in the range of 0.0 to 1.0.
+		//! If it is not specified, a value of y =0.5 will be used.
+		NodeScalarPtr GetY() const;
+		void SetY( NodeScalarPtr );
+
+		//! Get the node generating the source image to be sampled.
+		NodeImagePtr GetImage() const;
+		void SetImage( NodeImagePtr );
+
+		//-----------------------------------------------------------------------------------------
+		// Interface pattern
+		//-----------------------------------------------------------------------------------------
+		class Private;
+		Private* GetPrivate() const;
+        Node::Private* GetBasePrivate() const override;
+
+	protected:
+
+		//! Forbidden. Manage with the Ptr<> template.
+		~NodeColourSampleImage();
+
+	private:
+
+		Private* m_pD;
+
+	};
+
+
+}
diff --git a/Engine/Plugins/Experimental/Mutable/Source/MutableTools/Public/MuT/NodeColourSwitch.h b/Engine/Plugins/Experimental/Mutable/Source/MutableTools/Public/MuT/NodeColourSwitch.h
new file mode 100644
--- /dev/null
+++ b/Engine/Plugins/Experimental/Mutable/Source/MutableTools/Public/MuT/NodeColourSwitch.h
@@ -0,0 +1,87 @@
+// Copyright Epic Games, Inc. All Rights Reserved.
+
+#pragma once
+
+#include "MuT/NodeColour.h"
+
+
+namespace mu
+{
+
+	// Forward definitions
+	class NodeScalar;
+	typedef Ptr<NodeScalar> NodeScalarPtr;
+	typedef Ptr<const NodeScalar> NodeScalarPtrConst;
+
+    class NodeColourSwitch;
+    typedef Ptr<NodeColourSwitch> NodeColourSwitchPtr;
+    typedef Ptr<const NodeColourSwitch> NodeColourSwitchPtrConst;
+
+
+    //! This node selects an output Colour from a set of input Colours based on a parameter.
+	//! \ingroup model
+    class MUTABLETOOLS_API NodeColourSwitch : public NodeColour
+	{
+	public:
+
+		//-----------------------------------------------------------------------------------------
+		// Life cycle
+		//-----------------------------------------------------------------------------------------
+
+        NodeColourSwitch();
+
+		void SerialiseWrapper(OutputArchive& arch) const override;
+		static void Serialise( const NodeColourSwitch* pNode, OutputArchive& arch );
+        static NodeColourSwitchPtr StaticUnserialise( InputArchive& arch );
+
+
+		//-----------------------------------------------------------------------------------------
+		// Node Interface
+		//-----------------------------------------------------------------------------------------
+
+        NodePtr Clone( NodeMapPtr pMap = 0 ) const override;
+
+        const NODE_TYPE* GetType() const override;
+		static const NODE_TYPE* GetStaticType();
+
+        virtual int GetInputCount() const override;
+        virtual Node* GetInputNode( int i ) const override;
+        void SetInputNode( int i, NodePtr pNode ) override;
+
+		//-----------------------------------------------------------------------------------------
+		// Own Interface
+		//-----------------------------------------------------------------------------------------
+
+		//! Get the node generating the parameter used to select the option.
+		NodeScalarPtr GetParameter() const;
+		void SetParameter( NodeScalarPtr );
+
+        //! Set the number of option Colours. It will keep the currently set targets and initialise
+		//! the new ones as null.
+		void SetOptionCount( int );
+
+        //! Get the node generating the t-th option Colour.
+        NodeColourPtr GetOption( int t ) const;
+        void SetOption( int t, NodeColourPtr );
+
+
+		//-----------------------------------------------------------------------------------------
+		// Interface pattern
+		//-----------------------------------------------------------------------------------------
+		class Private;
+		Private* GetPrivate() const;
+        Node::Private* GetBasePrivate() const override;
+
+	protected:
+
+		//! Forbidden. Manage with the Ptr<> template.
+        ~NodeColourSwitch();
+
+	private:
+
+		Private* m_pD;
+
+	};
+
+
+}
diff --git a/Engine/Plugins/Experimental/Mutable/Source/MutableTools/Public/MuT/NodeColourTable.h b/Engine/Plugins/Experimental/Mutable/Source/MutableTools/Public/MuT/NodeColourTable.h
new file mode 100644
--- /dev/null
+++ b/Engine/Plugins/Experimental/Mutable/Source/MutableTools/Public/MuT/NodeColourTable.h
@@ -0,0 +1,86 @@
+// Copyright Epic Games, Inc. All Rights Reserved.
+
+#pragma once
+
+#include "MuT/NodeColour.h"
+
+
+namespace mu
+{
+
+	// Forward definitions
+	class NodeColourTable;
+	typedef Ptr<NodeColourTable> NodeColourTablePtr;
+	typedef Ptr<const NodeColourTable> NodeColourTablePtrConst;
+
+	class Table;
+	typedef Ptr<Table> TablePtr;
+	typedef Ptr<const Table> TablePtrConst;
+
+
+	//! This node provides the meshes stored in the column of a table.
+	//! \ingroup transform
+	class MUTABLETOOLS_API NodeColourTable : public NodeColour
+	{
+	public:
+
+		//-----------------------------------------------------------------------------------------
+		// Life cycle
+		//-----------------------------------------------------------------------------------------
+
+		NodeColourTable();
+
+		void SerialiseWrapper(OutputArchive& arch) const override;
+		static void Serialise( const NodeColourTable* pNode, OutputArchive& arch );
+		static NodeColourTablePtr StaticUnserialise( InputArchive& arch );
+
+
+		//-----------------------------------------------------------------------------------------
+		// Node Interface
+		//-----------------------------------------------------------------------------------------
+
+		NodePtr Clone( NodeMapPtr pMap = 0 ) const override;
+
+		const NODE_TYPE* GetType() const override;
+		static const NODE_TYPE* GetStaticType();
+
+		int GetInputCount() const override;
+		Node* GetInputNode( int i ) const override;
+		void SetInputNode( int i, NodePtr pNode ) override;
+
+		//-----------------------------------------------------------------------------------------
+		// Own Interface
+		//-----------------------------------------------------------------------------------------
+
+		//! Set the name of the implicit table parameter.
+		void SetParameterName( const char* strName );
+
+		//!
+		TablePtr GetTable() const;
+		void SetTable( TablePtr );
+
+		//!
+		const char* GetColumn() const;
+		void SetColumn( const char* strName );
+
+		//-----------------------------------------------------------------------------------------
+		// Interface pattern
+		//-----------------------------------------------------------------------------------------
+		class Private;
+		Private* GetPrivate() const;
+        Node::Private* GetBasePrivate() const override;
+
+
+	protected:
+
+		//! Forbidden. Manage with the Ptr<> template.
+		~NodeColourTable();
+
+	private:
+
+		Private* m_pD;
+
+	};
+
+
+}
diff --git a/Engine/Plugins/Experimental/Mutable/Source/MutableTools/Public/MuT/NodeColourVariation.h b/Engine/Plugins/Experimental/Mutable/Source/MutableTools/Public/MuT/NodeColourVariation.h
new file mode 100644
--- /dev/null
+++ b/Engine/Plugins/Experimental/Mutable/Source/MutableTools/Public/MuT/NodeColourVariation.h
@@ -0,0 +1,83 @@
+// Copyright Epic Games, Inc. All Rights Reserved.
+
+#pragma once
+
+#include "MuT/NodeColour.h"
+
+
+namespace mu
+{
+
+    // Forward definitions
+    class NodeColourVariation;
+    typedef Ptr<NodeColourVariation> NodeColourVariationPtr;
+    typedef Ptr<const NodeColourVariation> NodeColourVariationPtrConst;
+
+
+    //!
+    //! \ingroup model
+    class MUTABLETOOLS_API NodeColourVariation : public NodeColour
+    {
+    public:
+        //-----------------------------------------------------------------------------------------
+        // Life cycle
+        //-----------------------------------------------------------------------------------------
+
+        NodeColourVariation();
+
+		void SerialiseWrapper(OutputArchive& arch) const override;
+		static void Serialise( const NodeColourVariation* pNode, OutputArchive& arch );
+        static NodeColourVariationPtr StaticUnserialise( InputArchive& arch );
+
+        //-----------------------------------------------------------------------------------------
+        // Node Interface
+        //-----------------------------------------------------------------------------------------
+
+        NodePtr Clone( NodeMapPtr pMap = 0 ) const override;
+
+        const NODE_TYPE* GetType() const override;
+        static const NODE_TYPE* GetStaticType();
+
+        virtual int GetInputCount() const override;
+        virtual Node* GetInputNode( int i ) const override;
+        void SetInputNode( int i, NodePtr pNode ) override;
+
+        //-----------------------------------------------------------------------------------------
+        // Own Interface
+        //-----------------------------------------------------------------------------------------
+
+        //!
+        void SetDefaultColour( NodeColour* Colour );
+
+        //! Set the number of tags to consider in this variation
+        void SetVariationCount( int count );
+
+        //!
+        int GetVariationCount() const;
+
+        //! Set the tag or state name that will enable a specific vartiation
+        void SetVariationTag( int index, const char* strTag );
+
+        //!
+        void SetVariationColour( int index, NodeColour* Colour );
+
+        //!}
+
+
+        //-----------------------------------------------------------------------------------------
+        // Interface pattern
+        //-----------------------------------------------------------------------------------------
+        class Private;
+        Private* GetPrivate() const;
+        Node::Private* GetBasePrivate() const override;
+
+    protected:
+        //! Forbidden. Manage with the Ptr<> template.
+        ~NodeColourVariation();
+
+    private:
+        Private* m_pD;
+    };
+
+
+} // namespace mu
diff --git a/Engine/Plugins/Experimental/Mutable/Source/MutableTools/Public/MuT/NodeComponent.h b/Engine/Plugins/Experimental/Mutable/Source/MutableTools/Public/MuT/NodeComponent.h
new file mode 100644
--- /dev/null
+++ b/Engine/Plugins/Experimental/Mutable/Source/MutableTools/Public/MuT/NodeComponent.h
@@ -0,0 +1,87 @@
+// Copyright Epic Games, Inc. All Rights Reserved.
+
+#pragma once
+
+#include "MuT/Node.h"
+
+
+namespace mu
+{
+
+	// Forward definitions
+	class NodeComponent;
+	typedef Ptr<NodeComponent> NodeComponentPtr;
+	typedef Ptr<const NodeComponent> NodeComponentPtrConst;
+
+    class NodeSurface;
+
+
+	//! This class is the parent of all nodes that output a component.
+	//! \ingroup model
+	class MUTABLETOOLS_API NodeComponent : public Node
+	{
+	public:
+
+		// Possible subclasses
+		enum class EType : uint8
+		{
+			New = 0,
+			Edit = 1,
+
+			None
+		};
+
+		//-----------------------------------------------------------------------------------------
+		// Life cycle
+		//-----------------------------------------------------------------------------------------
+
+		static void Serialise( const NodeComponent* pNode, OutputArchive& arch );
+		static NodeComponentPtr StaticUnserialise( InputArchive& arch );
+
+		//-----------------------------------------------------------------------------------------
+        // Node interface
+		//-----------------------------------------------------------------------------------------
+
+        const NODE_TYPE* GetType() const override;
+		static const NODE_TYPE* GetStaticType();
+
+        //-----------------------------------------------------------------------------------------
+        // Own interface
+        //-----------------------------------------------------------------------------------------
+
+        //! \name Surfaces
+        //! \{
+
+        //! Get the number of meshes in the component.
+        virtual int GetSurfaceCount() const = 0;
+
+        //! Set the number of meshes in the component.
+        virtual void SetSurfaceCount( int ) = 0;
+
+        //! Get the node generating one of the meshes in the component.
+        //! \param index index of the mesh, from 0 to GetSurfaceCount()-1
+        virtual NodeSurface* GetSurface( int index ) const = 0;
+
+        //! Set the node generating one of the meshes in the component.
+        //! \param index index of the mesh, from 0 to GetSurfaceCount()-1
+        virtual void SetSurface( int index, NodeSurface* ) = 0;
+
+        //! \}
+
+		//-----------------------------------------------------------------------------------------
+		// Interface pattern
+		//-----------------------------------------------------------------------------------------
+		class Private;
+
+	protected:
+
+		//! Forbidden. Manage with the Ptr<> template.
+		inline ~NodeComponent() {}
+
+		//!
+		EType Type = EType::None;
+
+	};
+
+
+}
diff --git a/Engine/Plugins/Experimental/Mutable/Source/MutableTools/Public/MuT/NodeComponentEdit.h b/Engine/Plugins/Experimental/Mutable/Source/MutableTools/Public/MuT/NodeComponentEdit.h
new file mode 100644
--- /dev/null
+++ b/Engine/Plugins/Experimental/Mutable/Source/MutableTools/Public/MuT/NodeComponentEdit.h
@@ -0,0 +1,90 @@
+// Copyright Epic Games, Inc. All Rights Reserved.
+
+#pragma once
+
+#include "MuT/NodeComponent.h"
+
+
+namespace mu
+{
+
+	// Forward definitions
+	class NodeComponentEdit;
+	typedef Ptr<NodeComponentEdit> NodeComponentEditPtr;
+	typedef Ptr<const NodeComponentEdit> NodeComponentEditPtrConst;
+
+    class NodeSurface;
+    typedef Ptr<NodeSurface> NodeSurfacePtr;
+    typedef Ptr<const NodeSurface> NodeSurfacePtrConst;
+
+
+
+	//! This node modifies a node of the parent object of the object that this node belongs to.
+	//! It allows to extend, cut and morph the parent component's meshes.
+	//! It also allows to patch the parent component's textures.
+	//! \ingroup model
+	class MUTABLETOOLS_API NodeComponentEdit : public NodeComponent
+	{
+	public:
+
+		//-----------------------------------------------------------------------------------------
+		// Life cycle
+		//-----------------------------------------------------------------------------------------
+
+		NodeComponentEdit();
+
+		void SerialiseWrapper(OutputArchive& arch) const override;
+		static void Serialise( const NodeComponentEdit* pNode, OutputArchive& arch );
+		static NodeComponentEditPtr StaticUnserialise( InputArchive& arch );
+
+		//-----------------------------------------------------------------------------------------
+        // Node interface
+		//-----------------------------------------------------------------------------------------
+		NodePtr Clone( NodeMapPtr pMap = 0 ) const override;
+
+		const NODE_TYPE* GetType() const override;
+		static const NODE_TYPE* GetStaticType();
+
+		int GetInputCount() const override;
+		Node* GetInputNode( int i ) const override;
+		void SetInputNode( int i, NodePtr pNode ) override;
+
+        //-----------------------------------------------------------------------------------------
+        // NodeComponent interface
+        //-----------------------------------------------------------------------------------------
+        int GetSurfaceCount() const override;
+        void SetSurfaceCount( int ) override;
+        NodeSurface* GetSurface( int index ) const override;
+        void SetSurface( int index, NodeSurface* ) override;
+
+		//-----------------------------------------------------------------------------------------
+        // Own interface
+		//-----------------------------------------------------------------------------------------
+
+		//! Set the parent component to modify. It should be a node from the same LOD of one of
+		//! the parent objects of this node's object.
+        void SetParent( NodeComponent* );
+        NodeComponent* GetParent() const;
+
+
+		//-----------------------------------------------------------------------------------------
+		// Interface pattern
+		//-----------------------------------------------------------------------------------------
+		class Private;
+		Private* GetPrivate() const;
+        Node::Private* GetBasePrivate() const override;
+
+	protected:
+
+		//! Forbidden. Manage with the Ptr<> template.
+		~NodeComponentEdit();
+
+	private:
+
+		Private* m_pD;
+
+	};
+
+
+
+}
diff --git a/Engine/Plugins/Experimental/Mutable/Source/MutableTools/Public/MuT/NodeComponentNew.h b/Engine/Plugins/Experimental/Mutable/Source/MutableTools/Public/MuT/NodeComponentNew.h
new file mode 100644
--- /dev/null
+++ b/Engine/Plugins/Experimental/Mutable/Source/MutableTools/Public/MuT/NodeComponentNew.h
@@ -0,0 +1,96 @@
+// Copyright Epic Games, Inc. All Rights Reserved.
+
+#pragma once
+
+#include "MuT/NodeComponent.h"
+
+
+namespace mu
+{
+
+	// Forward definitions
+	class NodeComponent;
+	typedef Ptr<NodeComponent> NodeComponentPtr;
+	typedef Ptr<const NodeComponent> NodeComponentPtrConst;
+
+	class NodeComponentNew;
+	typedef Ptr<NodeComponentNew> NodeComponentNewPtr;
+	typedef Ptr<const NodeComponentNew> NodeComponentNewPtrConst;
+
+    class NodeSurface;
+    typedef Ptr<NodeSurface> NodeSurfacePtr;
+    typedef Ptr<const NodeSurface> NodeSurfacePtrConst;
+
+
+	//! This node makes a new component from several meshes and images.
+	//! \ingroup model
+	class MUTABLETOOLS_API NodeComponentNew : public NodeComponent
+	{
+	public:
+
+		//-----------------------------------------------------------------------------------------
+		// Life cycle
+		//-----------------------------------------------------------------------------------------
+
+		NodeComponentNew();
+
+		void SerialiseWrapper(OutputArchive& arch) const override;
+		static void Serialise( const NodeComponentNew* pNode, OutputArchive& arch );
+		static NodeComponentNewPtr StaticUnserialise( InputArchive& arch );
+
+		//-----------------------------------------------------------------------------------------
+        // Node interface
+		//-----------------------------------------------------------------------------------------
+		NodePtr Clone( NodeMapPtr pMap = 0 ) const override;
+
+        const NODE_TYPE* GetType() const override;
+		static const NODE_TYPE* GetStaticType();
+
+		int GetInputCount() const override;
+		Node* GetInputNode( int i ) const override;
+		void SetInputNode( int i, NodePtr pNode ) override;
+
+        //-----------------------------------------------------------------------------------------
+        // NodeComponent interface
+        //-----------------------------------------------------------------------------------------
+        int GetSurfaceCount() const override;
+        void SetSurfaceCount( int ) override;
+        NodeSurface* GetSurface( int index ) const override;
+        void SetSurface( int index, NodeSurface* ) override;
+
+		//-----------------------------------------------------------------------------------------
+        // Own interface
+		//-----------------------------------------------------------------------------------------
+
+		//! Get the name of the component.
+		const char* GetName() const;
+
+		//! Set the name of the component.
+		void SetName( const char* );
+
+		//! Get the id of the component.
+		uint16 GetId() const;
+
+		//! Set the id of the component.
+		void SetId( uint16 );
+
+		//-----------------------------------------------------------------------------------------
+		// Interface pattern
+		//-----------------------------------------------------------------------------------------
+		class Private;
+		Private* GetPrivate() const;
+        Node::Private* GetBasePrivate() const override;
+
+	protected:
+
+		//! Forbidden. Manage with the Ptr<> template.
+		~NodeComponentNew();
+
+	private:
+
+		Private* m_pD;
+
+	};
+
+
+}
diff --git a/Engine/Plugins/Experimental/Mutable/Source/MutableTools/Public/MuT/NodeImage.h b/Engine/Plugins/Experimental/Mutable/Source/MutableTools/Public/MuT/NodeImage.h
new file mode 100644
--- /dev/null
+++ b/Engine/Plugins/Experimental/Mutable/Source/MutableTools/Public/MuT/NodeImage.h
@@ -0,0 +1,91 @@
+// Copyright Epic Games, Inc. All Rights Reserved.
+
+#pragma once
+
+#include "MuT/Node.h"
+
+namespace mu
+{
+
+	// Forward definitions
+	class NodeImage;
+	typedef Ptr<NodeImage> NodeImagePtr;
+	typedef Ptr<const NodeImage> NodeImagePtrConst;
+
+
+    //! %Base class of any node that outputs an image.
+	//! \ingroup model
+	class MUTABLETOOLS_API NodeImage : public Node
+	{
+	public:
+
+		// Possible subclasses
+		enum class EType : uint8
+		{
+			Constant = 1,
+			Difference = 2,
+			Interpolate = 4,
+			Saturate = 6,
+			Table = 9,
+			Swizzle = 10,
+			SelectColour = 11,
+			ColourMap = 12,
+			Gradient = 13,
+			VolumeLayer_Deprecated = 14,
+			Binarise = 17,
+			Luminance = 18,
+			Layer = 19,
+			LayerColour = 20,
+			Resize = 21,
+			PlainColour = 22,
+			Interpolate3 = 23,
+			Project = 24,
+			Mipmap = 25,
+			Switch = 26,
+			Conditional = 27,
+			Format = 28,
+			Parameter = 29,
+			MultiLayer = 30,
+			Invert = 31,
+			Variation = 32,
+			NormalComposite = 33,
+			Transform = 34,
+			None
+		};
+
+		//-----------------------------------------------------------------------------------------
+		// Life cycle
+		//-----------------------------------------------------------------------------------------
+
+		static void Serialise( const NodeImage* pNode, OutputArchive& arch );
+		static NodeImagePtr StaticUnserialise( InputArchive& arch );
+
+
+		//-----------------------------------------------------------------------------------------
+		// Node Interface
+		//-----------------------------------------------------------------------------------------
+
+		const NODE_TYPE* GetType() const override;
+		static const NODE_TYPE* GetStaticType();
+
+		inline EType GetImageNodeType() const { return Type; }
+
+
+		//-----------------------------------------------------------------------------------------
+		// Interface pattern
+		//-----------------------------------------------------------------------------------------
+		class Private;
+
+	protected:
+
+		//! Forbidden. Manage with the Ptr<> template.
+		inline ~NodeImage() {}
+
+		//!
+		EType Type = EType::None;
+
+	};
+
+
+
+}
diff --git a/Engine/Plugins/Experimental/Mutable/Source/MutableTools/Public/MuT/NodeImageBinarise.h b/Engine/Plugins/Experimental/Mutable/Source/MutableTools/Public/MuT/NodeImageBinarise.h
new file mode 100644
--- /dev/null
+++ b/Engine/Plugins/Experimental/Mutable/Source/MutableTools/Public/MuT/NodeImageBinarise.h
@@ -0,0 +1,82 @@
+// Copyright Epic Games, Inc. All Rights Reserved.
+
+#pragma once
+
+#include "MuT/NodeImage.h"
+
+
+namespace mu
+{
+
+	// Forward definitions
+	class NodeScalar;
+	typedef Ptr<NodeScalar> NodeScalarPtr;
+	typedef Ptr<const NodeScalar> NodeScalarPtrConst;
+
+	class NodeImageBinarise;
+	typedef Ptr<NodeImageBinarise> NodeImageBinarisePtr;
+	typedef Ptr<const NodeImageBinarise> NodeImageBinarisePtrConst;
+
+
+	//! Node that multiplies the colors of an image, channel by channel.
+	//! \ingroup model
+	class MUTABLETOOLS_API NodeImageBinarise : public NodeImage
+	{
+	public:
+
+		//-----------------------------------------------------------------------------------------
+		// Life cycle
+		//-----------------------------------------------------------------------------------------
+
+		NodeImageBinarise();
+
+		void SerialiseWrapper(OutputArchive& arch) const override;
+		static void Serialise( const NodeImageBinarise* pNode, OutputArchive& arch );
+		static NodeImageBinarisePtr StaticUnserialise( InputArchive& arch );
+
+
+		//-----------------------------------------------------------------------------------------
+		// Node Interface
+		//-----------------------------------------------------------------------------------------
+
+        NodePtr Clone( NodeMapPtr pMap = 0 ) const override;
+
+        const NODE_TYPE* GetType() const override;
+		static const NODE_TYPE* GetStaticType();
+
+        virtual int GetInputCount() const override;
+        virtual Node* GetInputNode( int i ) const override;
+        void SetInputNode( int i, NodePtr pNode ) override;
+
+		//-----------------------------------------------------------------------------------------
+		// Own Interface
+		//-----------------------------------------------------------------------------------------
+
+		//! Base image to multiply.
+		NodeImagePtr GetBase() const;
+		void SetBase( NodeImagePtr );
+
+		//! Colour to multiply by the image.
+		NodeScalarPtr GetThreshold() const;
+		void SetThreshold( NodeScalarPtr );
+
+		//-----------------------------------------------------------------------------------------
+		// Interface pattern
+		//-----------------------------------------------------------------------------------------
+		class Private;
+		Private* GetPrivate() const;
+        Node::Private* GetBasePrivate() const override;
+
+	protected:
+
+		//! Forbidden. Manage with the Ptr<> template.
+		~NodeImageBinarise();
+
+	private:
+
+		Private* m_pD;
+
+	};
+
+
+}
diff --git a/Engine/Plugins/Experimental/Mutable/Source/MutableTools/Public/MuT/NodeImageColourMap.h b/Engine/Plugins/Experimental/Mutable/Source/MutableTools/Public/MuT/NodeImageColourMap.h
new file mode 100644
--- /dev/null
+++ b/Engine/Plugins/Experimental/Mutable/Source/MutableTools/Public/MuT/NodeImageColourMap.h
@@ -0,0 +1,84 @@
+// Copyright Epic Games, Inc. All Rights Reserved.
+
+#pragma once
+
+#include "MuT/NodeImage.h"
+
+
+namespace mu
+{
+
+	// Forward definitions
+	class NodeImageColourMap;
+	typedef Ptr<NodeImageColourMap> NodeImageColourMapPtr;
+	typedef Ptr<const NodeImageColourMap> NodeImageColourMapPtrConst;
+
+
+	//! This node changes the colours of a selectd part of the image, applying a colour map from
+	//! conteined in another image.
+	//! \ingroup model
+	class MUTABLETOOLS_API NodeImageColourMap : public NodeImage
+	{
+	public:
+
+		//-----------------------------------------------------------------------------------------
+		// Life cycle
+		//-----------------------------------------------------------------------------------------
+
+		NodeImageColourMap();
+
+		void SerialiseWrapper(OutputArchive& arch) const override;
+		static void Serialise( const NodeImageColourMap* pNode, OutputArchive& arch );
+		static NodeImageColourMapPtr StaticUnserialise( InputArchive& arch );
+
+
+		//-----------------------------------------------------------------------------------------
+		// Node Interface
+		//-----------------------------------------------------------------------------------------
+
+        NodePtr Clone( NodeMapPtr pMap = 0 ) const override;
+
+        const NODE_TYPE* GetType() const override;
+		static const NODE_TYPE* GetStaticType();
+
+        virtual int GetInputCount() const override;
+        virtual Node* GetInputNode( int i ) const override;
+        void SetInputNode( int i, NodePtr pNode ) override;
+
+		//-----------------------------------------------------------------------------------------
+		// Own Interface
+		//-----------------------------------------------------------------------------------------
+
+		//! Get the node generating the base image that will have the other one blended on top.
+		NodeImagePtr GetBase() const;
+		void SetBase( NodeImagePtr );
+
+		//! Get the node generating the mask image controlling the weight of the blend effect.
+		//! \todo: make it optional
+		NodeImagePtr GetMask() const;
+		void SetMask( NodeImagePtr );
+
+		//! Get the node generating the image to blend on the base.
+		NodeImagePtr GetMap() const;
+		void SetMap( NodeImagePtr );
+
+		//-----------------------------------------------------------------------------------------
+		// Interface pattern
+		//-----------------------------------------------------------------------------------------
+		class Private;
+		Private* GetPrivate() const;
+        Node::Private* GetBasePrivate() const override;
+
+	protected:
+
+		//! Forbidden. Manage with the Ptr<> template.
+		~NodeImageColourMap();
+
+	private:
+
+		Private* m_pD;
+
+	};
+
+
+}
diff --git a/Engine/Plugins/Experimental/Mutable/Source/MutableTools/Public/MuT/NodeImageConditional.h b/Engine/Plugins/Experimental/Mutable/Source/MutableTools/Public/MuT/NodeImageConditional.h
new file mode 100644
--- /dev/null
+++ b/Engine/Plugins/Experimental/Mutable/Source/MutableTools/Public/MuT/NodeImageConditional.h
@@ -0,0 +1,88 @@
+// Copyright Epic Games, Inc. All Rights Reserved.
+
+#pragma once
+
+#include "MuT/NodeImage.h"
+
+#include "MuR/Image.h"
+
+
+namespace mu
+{
+
+	// Forward definitions
+    class NodeBool;
+    typedef Ptr<NodeBool> NodeBoolPtr;
+    typedef Ptr<const NodeBool> NodeBoolPtrConst;
+
+    class NodeImageConditional;
+    typedef Ptr<NodeImageConditional> NodeImageConditionalPtr;
+    typedef Ptr<const NodeImageConditional> NodeImageConditionalPtrConst;
+
+
+	//! This node selects an output image from a set of input images based on a parameter.
+	//! \ingroup model
+    class MUTABLETOOLS_API NodeImageConditional : public NodeImage
+	{
+	public:
+
+		//-----------------------------------------------------------------------------------------
+		// Life cycle
+		//-----------------------------------------------------------------------------------------
+
+        NodeImageConditional();
+
+		void SerialiseWrapper(OutputArchive& arch) const override;
+		static void Serialise( const NodeImageConditional* pNode, OutputArchive& arch );
+        static NodeImageConditionalPtr StaticUnserialise( InputArchive& arch );
+
+
+		//-----------------------------------------------------------------------------------------
+		// Node Interface
+		//-----------------------------------------------------------------------------------------
+
+        NodePtr Clone( NodeMapPtr pMap = 0 ) const override;
+
+        const NODE_TYPE* GetType() const override;
+		static const NODE_TYPE* GetStaticType();
+
+        virtual int GetInputCount() const override;
+        virtual Node* GetInputNode( int i ) const override;
+        void SetInputNode( int i, NodePtr pNode ) override;
+
+		//-----------------------------------------------------------------------------------------
+		// Own Interface
+		//-----------------------------------------------------------------------------------------
+
+		//! Get the node generating the parameter used to select the option.
+        NodeBool* GetParameter() const;
+        void SetParameter( NodeBool* );
+
+        //! Get the node generating the option image when the parameters is true.
+        NodeImage* GetOptionTrue() const;
+        void SetOptionTrue( NodeImage* );
+
+        NodeImage* GetOptionFalse() const;
+        void SetOptionFalse( NodeImage* );
+
+
+		//-----------------------------------------------------------------------------------------
+		// Interface pattern
+		//-----------------------------------------------------------------------------------------
+		class Private;
+		Private* GetPrivate() const;
+        Node::Private* GetBasePrivate() const override;
+
+	protected:
+
+		//! Forbidden. Manage with the Ptr<> template.
+        ~NodeImageConditional();
+
+	private:
+
+		Private* m_pD;
+
+	};
+
+
+}
diff --git a/Engine/Plugins/Experimental/Mutable/Source/MutableTools/Public/MuT/NodeImageConstant.h b/Engine/Plugins/Experimental/Mutable/Source/MutableTools/Public/MuT/NodeImageConstant.h
new file mode 100644
--- /dev/null
+++ b/Engine/Plugins/Experimental/Mutable/Source/MutableTools/Public/MuT/NodeImageConstant.h
@@ -0,0 +1,80 @@
+// Copyright Epic Games, Inc. All Rights Reserved.
+
+#pragma once
+
+#include "MuT/NodeImage.h"
+
+#include "MuR/Image.h"
+
+
+namespace mu
+{
+
+	// Forward definitions
+    class NodeImageConstant;
+    typedef Ptr<NodeImageConstant> NodeImageConstantPtr;
+    typedef Ptr<const NodeImageConstant> NodeImageConstantPtrConst;
+
+
+	//! Node that outputs a constant image.
+	//! \ingroup model
+	class MUTABLETOOLS_API NodeImageConstant : public NodeImage
+	{
+	public:
+
+		//-----------------------------------------------------------------------------------------
+		// Life cycle
+		//-----------------------------------------------------------------------------------------
+
+		NodeImageConstant();
+
+		void SerialiseWrapper(OutputArchive& arch) const override;
+		static void Serialise( const NodeImageConstant* pNode, OutputArchive& arch );
+		static NodeImageConstantPtr StaticUnserialise( InputArchive& arch );
+
+
+		//-----------------------------------------------------------------------------------------
+		// Node Interface
+		//-----------------------------------------------------------------------------------------
+
+        NodePtr Clone( NodeMapPtr pMap = 0 ) const override;
+
+        const NODE_TYPE* GetType() const override;
+		static const NODE_TYPE* GetStaticType();
+
+        virtual int GetInputCount() const override;
+        virtual Node* GetInputNode( int i ) const override;
+        void SetInputNode( int i, NodePtr pNode ) override;
+
+		//-----------------------------------------------------------------------------------------
+		// Own Interface
+		//-----------------------------------------------------------------------------------------
+
+		//! Get the image that will be output by this node
+        Ptr<const Image> GetValue() const;
+
+        //! Set the image to be output by this node
+        void SetValue( Ptr<const Image> pImage );
+
+        //! Set the image proxy that will provide the image for this node when necessary
+        void SetValue( Ptr<ResourceProxy<Image>> pImage );
+
+		//-----------------------------------------------------------------------------------------
+		// Interface pattern
+		//-----------------------------------------------------------------------------------------
+		class Private;
+		Private* GetPrivate() const;
+        Node::Private* GetBasePrivate() const override;
+
+	protected:
+
+		//! Forbidden. Manage with the Ptr<> template.
+		~NodeImageConstant();
+
+	private:
+
+		Private* m_pD;
+
+	};
+
+}
diff --git a/Engine/Plugins/Experimental/Mutable/Source/MutableTools/Public/MuT/NodeImageDifference.h b/Engine/Plugins/Experimental/Mutable/Source/MutableTools/Public/MuT/NodeImageDifference.h
new file mode 100644
--- /dev/null
+++ b/Engine/Plugins/Experimental/Mutable/Source/MutableTools/Public/MuT/NodeImageDifference.h
@@ -0,0 +1,85 @@
+// Copyright Epic Games, Inc. All Rights Reserved.
+
+#pragma once
+
+#include "MuT/NodeImage.h"
+
+
+namespace mu
+{
+
+	// Forward definitions
+	class NodeImageDifference;
+	typedef Ptr<NodeImageDifference> NodeImageDifferencePtr;
+	typedef Ptr<const NodeImageDifference> NodeImageDifferencePtrConst;
+
+
+	//! Node that compares two images and returns a black-and-white image with pixels that are
+	//! "significally different", meaning around 0.03 different. This node is for testing
+	//! purposes.
+	//! \todo rename to NodeImageCompare
+	//! \ingroup model
+	class MUTABLETOOLS_API NodeImageDifference : public NodeImage
+	{
+	public:
+
+		//-----------------------------------------------------------------------------------------
+		// Life cycle
+		//-----------------------------------------------------------------------------------------
+
+		NodeImageDifference();
+
+		void SerialiseWrapper(OutputArchive& arch) const override;
+		static void Serialise( const NodeImageDifference* pNode, OutputArchive& arch );
+		static NodeImageDifferencePtr StaticUnserialise( InputArchive& arch );
+
+
+		//-----------------------------------------------------------------------------------------
+		// Node Interface
+		//-----------------------------------------------------------------------------------------
+
+        NodePtr Clone( NodeMapPtr pMap = 0 ) const override;
+
+        const NODE_TYPE* GetType() const override;
+		static const NODE_TYPE* GetStaticType();
+
+        virtual int GetInputCount() const override;
+        virtual Node* GetInputNode( int i ) const override;
+        void SetInputNode( int i, NodePtr pNode ) override;
+
+		//-----------------------------------------------------------------------------------------
+		// Own Interface
+		//-----------------------------------------------------------------------------------------
+
+		//! Return the first image to compare.
+		NodeImagePtr GetA() const;
+
+		//! Set the first image to compare.
+		void SetA( NodeImagePtr );
+
+		//! Return the second image to compare.
+		NodeImagePtr GetB() const;
+
+		//! Set the second image to compare.
+		void SetB( NodeImagePtr );
+
+		//-----------------------------------------------------------------------------------------
+		// Interface pattern
+		//-----------------------------------------------------------------------------------------
+		class Private;
+		Private* GetPrivate() const;
+        Node::Private* GetBasePrivate() const override;
+
+	protected:
+
+		//! Forbidden. Manage with the Ptr<> template.
+		~NodeImageDifference();
+
+	private:
+
+		Private* m_pD;
+
+	};
+
+
+}
diff --git a/Engine/Plugins/Experimental/Mutable/Source/MutableTools/Public/MuT/NodeImageFormat.h b/Engine/Plugins/Experimental/Mutable/Source/MutableTools/Public/MuT/NodeImageFormat.h
new file mode 100644
--- /dev/null
+++ b/Engine/Plugins/Experimental/Mutable/Source/MutableTools/Public/MuT/NodeImageFormat.h
@@ -0,0 +1,83 @@
+// Copyright Epic Games, Inc. All Rights Reserved.
+
+#pragma once
+
+#include "MuT/NodeImage.h"
+
+#include "MuR/Image.h"
+
+
+namespace mu
+{
+
+	// Forward definitions
+	class NodeImageFormat;
+	typedef Ptr<NodeImageFormat> NodeImageFormatPtr;
+	typedef Ptr<const NodeImageFormat> NodeImageFormatPtrConst;
+
+
+	//! Node that composes a new image by gathering pixel data from channels in other images.
+	//! \ingroup model
+	class MUTABLETOOLS_API NodeImageFormat : public NodeImage
+	{
+	public:
+
+		//-----------------------------------------------------------------------------------------
+		// Life cycle
+		//-----------------------------------------------------------------------------------------
+
+		NodeImageFormat();
+
+		void SerialiseWrapper(OutputArchive& arch) const override;
+		static void Serialise( const NodeImageFormat* pNode, OutputArchive& arch );
+		static NodeImageFormatPtr StaticUnserialise( InputArchive& arch );
+
+
+		//-----------------------------------------------------------------------------------------
+		// Node Interface
+		//-----------------------------------------------------------------------------------------
+
+		NodePtr Clone( NodeMapPtr pMap = 0 ) const override;
+
+		const NODE_TYPE* GetType() const override;
+		static const NODE_TYPE* GetStaticType();
+
+		int GetInputCount() const override;
+		Node* GetInputNode( int i ) const override;
+		void SetInputNode( int i, NodePtr pNode ) override;
+
+		//-----------------------------------------------------------------------------------------
+		// Own Interface
+		//-----------------------------------------------------------------------------------------
+
+		//!
+		EImageFormat GetFormat() const;
+        void SetFormat(EImageFormat format, EImageFormat formatIfAlpha = EImageFormat::IF_NONE );
+
+		//!
+		NodeImagePtr GetSource() const;
+
+		void SetSource( NodeImagePtr );
+
+
+		//-----------------------------------------------------------------------------------------
+		// Interface pattern
+		//-----------------------------------------------------------------------------------------
+		class Private;
+		Private* GetPrivate() const;
+        Node::Private* GetBasePrivate() const override;
+
+	protected:
+
+		//! Forbidden. Manage with the Ptr<> template.
+		~NodeImageFormat();
+
+	private:
+
+		Private* m_pD;
+
+	};
+
+
+}
+
diff --git a/Engine/Plugins/Experimental/Mutable/Source/MutableTools/Public/MuT/NodeImageGradient.h b/Engine/Plugins/Experimental/Mutable/Source/MutableTools/Public/MuT/NodeImageGradient.h
new file mode 100644
--- /dev/null
+++ b/Engine/Plugins/Experimental/Mutable/Source/MutableTools/Public/MuT/NodeImageGradient.h
@@ -0,0 +1,87 @@
+// Copyright Epic Games, Inc. All Rights Reserved.
+
+#pragma once
+
+#include "MuT/NodeImage.h"
+
+
+namespace mu
+{
+
+	// Forward definitions
+	class NodeImageGradient;
+	typedef Ptr<NodeImageGradient> NodeImageGradientPtr;
+	typedef Ptr<const NodeImageGradient> NodeImageGradientPtrConst;
+
+	class NodeColour;
+	typedef Ptr<NodeColour> NodeColourPtr;
+	typedef Ptr<const NodeColour> NodeColourPtrConst;
+
+
+	//! This node generats an horizontal linear gradient between two colours in a new image.
+	//! \ingroup model
+	class MUTABLETOOLS_API NodeImageGradient : public NodeImage
+	{
+	public:
+
+		//-----------------------------------------------------------------------------------------
+		// Life cycle
+		//-----------------------------------------------------------------------------------------
+
+		NodeImageGradient();
+
+		void SerialiseWrapper(OutputArchive& arch) const override;
+		static void Serialise( const NodeImageGradient* pNode, OutputArchive& arch );
+		static NodeImageGradientPtr StaticUnserialise( InputArchive& arch );
+
+
+		//-----------------------------------------------------------------------------------------
+		// Node Interface
+		//-----------------------------------------------------------------------------------------
+
+        NodePtr Clone( NodeMapPtr pMap = 0 ) const override;
+
+        const NODE_TYPE* GetType() const override;
+		static const NODE_TYPE* GetStaticType();
+
+        virtual int GetInputCount() const override;
+        virtual Node* GetInputNode( int i ) const override;
+        void SetInputNode( int i, NodePtr pNode ) override;
+
+		//-----------------------------------------------------------------------------------------
+		// Own Interface
+		//-----------------------------------------------------------------------------------------
+
+		//! First colour of the gradient.
+		NodeColourPtr GetColour0() const;
+		void SetColour0( NodeColourPtr );
+
+		//! Second colour of the gradient.
+		NodeColourPtr GetColour1() const;
+		void SetColour1( NodeColourPtr );
+
+		//! Generated image size.
+		int GetSizeX() const;
+		int GetSizeY() const;
+		void SetSize( int x, int y );
+
+		//-----------------------------------------------------------------------------------------
+		// Interface pattern
+		//-----------------------------------------------------------------------------------------
+		class Private;
+		Private* GetPrivate() const;
+        Node::Private* GetBasePrivate() const override;
+
+	protected:
+
+		//! Forbidden. Manage with the Ptr<> template.
+		~NodeImageGradient();
+
+	private:
+
+		Private* m_pD;
+
+	};
+
+
+}
diff --git a/Engine/Plugins/Experimental/Mutable/Source/MutableTools/Public/MuT/NodeImageInterpolate.h b/Engine/Plugins/Experimental/Mutable/Source/MutableTools/Public/MuT/NodeImageInterpolate.h
new file mode 100644
--- /dev/null
+++ b/Engine/Plugins/Experimental/Mutable/Source/MutableTools/Public/MuT/NodeImageInterpolate.h
@@ -0,0 +1,95 @@
+// Copyright Epic Games, Inc. All Rights Reserved.
+
+#pragma once
+
+#include "MuT/NodeImage.h"
+
+#include "MuR/Image.h"
+
+
+namespace mu
+{
+
+	// Forward definitions
+	class NodeScalar;
+	typedef Ptr<NodeScalar> NodeScalarPtr;
+	typedef Ptr<const NodeScalar> NodeScalarPtrConst;
+
+	class NodeImageInterpolate;
+	typedef Ptr<NodeImageInterpolate> NodeImageInterpolatePtr;
+	typedef Ptr<const NodeImageInterpolate> NodeImageInterpolatePtrConst;
+
+
+	//! Node that interpolates linearly between several images based on a weight.
+	//! The different input images are uniformly distributed in the 0 to 1 range:
+	//! - if two images are set, the first one is 0.0 and the second one is 1.0
+	//! - if three images are set, the first one is 0.0 the second one is 0.5 and the third one 1.0
+	//! regardless of the input slots used: If B and C are set, but not A, B will be at weight 0
+	//! and C will be at weight 1
+	//! \ingroup model
+	class MUTABLETOOLS_API NodeImageInterpolate : public NodeImage
+	{
+	public:
+
+		//-----------------------------------------------------------------------------------------
+		// Life cycle
+		//-----------------------------------------------------------------------------------------
+
+		NodeImageInterpolate();
+
+		void SerialiseWrapper(OutputArchive& arch) const override;
+		static void Serialise( const NodeImageInterpolate* pNode, OutputArchive& arch );
+		static NodeImageInterpolatePtr StaticUnserialise( InputArchive& arch );
+
+
+		//-----------------------------------------------------------------------------------------
+		// Node Interface
+		//-----------------------------------------------------------------------------------------
+
+        NodePtr Clone( NodeMapPtr pMap = 0 ) const override;
+
+        const NODE_TYPE* GetType() const override;
+		static const NODE_TYPE* GetStaticType();
+
+        virtual int GetInputCount() const override;
+        virtual Node* GetInputNode( int i ) const override;
+        void SetInputNode( int i, NodePtr pNode ) override;
+
+		//-----------------------------------------------------------------------------------------
+		// Own Interface
+		//-----------------------------------------------------------------------------------------
+
+		//! Get the node generating the weight of the interpolation.
+		NodeScalarPtr GetFactor() const;
+		void SetFactor( NodeScalarPtr );
+
+		//! Set the number of target images. It will keep the currently set targets and initialise
+		//! the new ones as null.
+		void SetTargetCount( int );
+		int GetTargetCount() const;
+
+		//! Get the node generating the t-th target image of the interpolation.
+		NodeImagePtr GetTarget( int t ) const;
+		void SetTarget( int t, NodeImagePtr );
+
+
+		//-----------------------------------------------------------------------------------------
+		// Interface pattern
+		//-----------------------------------------------------------------------------------------
+		class Private;
+		Private* GetPrivate() const;
+        Node::Private* GetBasePrivate() const override;
+
+	protected:
+
+		//! Forbidden. Manage with the Ptr<> template.
+		~NodeImageInterpolate();
+
+	private:
+
+		Private* m_pD;
+
+	};
+
+
+}
diff --git a/Engine/Plugins/Experimental/Mutable/Source/MutableTools/Public/MuT/NodeImageInterpolate3.h b/Engine/Plugins/Experimental/Mutable/Source/MutableTools/Public/MuT/NodeImageInterpolate3.h
new file mode 100644
--- /dev/null
+++ b/Engine/Plugins/Experimental/Mutable/Source/MutableTools/Public/MuT/NodeImageInterpolate3.h
@@ -0,0 +1,100 @@
+// Copyright Epic Games, Inc. All Rights Reserved.
+
+#pragma once
+
+#include "MuT/NodeImage.h"
+
+#include "MuR/Image.h"
+
+
+namespace mu
+{
+
+	// Forward definitions
+	class NodeScalar;
+	typedef Ptr<NodeScalar> NodeScalarPtr;
+	typedef Ptr<const NodeScalar> NodeScalarPtrConst;
+
+	class NodeImageInterpolate3;
+	typedef Ptr<NodeImageInterpolate3> NodeImageInterpolate3Ptr;
+	typedef Ptr<const NodeImageInterpolate3> NodeImageInterpolate3PtrConst;
+
+
+	//! Node that Interpolates linearly between 3 images. It has two weights and the third one is
+	//! deduced as 1-w0-w1.
+	//! \ingroup model
+	class MUTABLETOOLS_API NodeImageInterpolate3 : public NodeImage
+	{
+	public:
+
+		//-----------------------------------------------------------------------------------------
+		// Life cycle
+		//-----------------------------------------------------------------------------------------
+
+		NodeImageInterpolate3();
+
+		void SerialiseWrapper(OutputArchive& arch) const override;
+		static void Serialise( const NodeImageInterpolate3* pNode, OutputArchive& arch );
+		static NodeImageInterpolate3Ptr StaticUnserialise( InputArchive& arch );
+
+
+		//-----------------------------------------------------------------------------------------
+		// Node Interface
+		//-----------------------------------------------------------------------------------------
+
+        NodePtr Clone( NodeMapPtr pMap = 0 ) const override;
+
+        const NODE_TYPE* GetType() const override;
+		static const NODE_TYPE* GetStaticType();
+
+        virtual int GetInputCount() const override;
+        virtual Node* GetInputNode( int i ) const override;
+        void SetInputNode( int i, NodePtr pNode ) override;
+
+		//-----------------------------------------------------------------------------------------
+		// Own Interface
+		//-----------------------------------------------------------------------------------------
+
+		//! Get the weight of the first target. If this is 1 (so the other is 0) the resulting
+		//! image is the target 1. If both weights are 0, the result is target 0.
+		NodeScalarPtr GetFactor1() const;
+		void SetFactor1( NodeScalarPtr );
+
+		//! Get the weight of the second target. If this is 1 (so the other is 1) the resulting
+		//! image is the target 2. If both weights are 0, the result is target 0.
+		NodeScalarPtr GetFactor2() const;
+		void SetFactor2( NodeScalarPtr );
+
+		//! Get the node generating the first target image of the interpolation.
+		NodeImagePtr GetTarget0() const;
+		void SetTarget0( NodeImagePtr );
+
+		//! Get the node generating the second target image of the interpolation.
+		NodeImagePtr GetTarget1() const;
+		void SetTarget1( NodeImagePtr );
+
+		//! Get the node generating the third target image of the interpolation.
+		NodeImagePtr GetTarget2() const;
+		void SetTarget2( NodeImagePtr );
+
+
+		//-----------------------------------------------------------------------------------------
+		// Interface pattern
+		//-----------------------------------------------------------------------------------------
+		class Private;
+		Private* GetPrivate() const;
+        Node::Private* GetBasePrivate() const override;
+
+	protected:
+
+		//! Forbidden. Manage with the Ptr<> template.
+		~NodeImageInterpolate3();
+
+	private:
+
+		Private* m_pD;
+
+	};
+
+
+}
diff --git a/Engine/Plugins/Experimental/Mutable/Source/MutableTools/Public/MuT/NodeImageInvert.h b/Engine/Plugins/Experimental/Mutable/Source/MutableTools/Public/MuT/NodeImageInvert.h
new file mode 100644
--- /dev/null
+++ b/Engine/Plugins/Experimental/Mutable/Source/MutableTools/Public/MuT/NodeImageInvert.h
@@ -0,0 +1,71 @@
+// Copyright Epic Games, Inc. All Rights Reserved.
+
+#pragma once
+
+#include "MuT/NodeImage.h"
+
+
+namespace mu
+{
+
+	// Forward definitions
+	class NodeImageInvert;
+	typedef Ptr<NodeImageInvert> NodeImageInvertPtr;
+	typedef Ptr<const NodeImageInvert> NodeImageInvertPtrConst;
+
+	//! Node that inverts the colors of an image, channel by channel
+	//! \ingroup model
+
+	class MUTABLETOOLS_API NodeImageInvert : public NodeImage
+	{
+	public:
+
+		//-----------------------------------------------------------------------------------------
+		// Life cycle
+		//-----------------------------------------------------------------------------------------
+
+		NodeImageInvert();
+
+		static void Serialise(const NodeImageInvert* pNode, OutputArchive& arch);
+		void SerialiseWrapper(OutputArchive& arch) const override;
+		static NodeImageInvertPtr StaticUnserialise(InputArchive& arch);
+
+
+		//-----------------------------------------------------------------------------------------
+		// Node Interface
+		//-----------------------------------------------------------------------------------------
+
+		NodePtr Clone(NodeMapPtr pMap = 0) const override;
+
+		const NODE_TYPE* GetType() const override;
+		static const NODE_TYPE* GetStaticType();
+
+		virtual int GetInputCount() const override;
+		virtual Node* GetInputNode(int i)const override;
+		void SetInputNode(int i, NodePtr pNode) override;
+
+		//-----------------------------------------------------------------------------------------
+		// Own Interface
+		//-----------------------------------------------------------------------------------------
+
+		//! Base image to invert.
+		NodeImagePtr GetBase() const;
+		void SetBase(NodeImagePtr);
+
+		//-----------------------------------------------------------------------------------------
+		// Interface pattern
+		//-----------------------------------------------------------------------------------------
+		class Private;
+		Private* GetPrivate() const;
+		Node::Private*GetBasePrivate()const override;
+
+	protected:
+
+		//! Forbidden. Manage with the Ptr<> template
+		~NodeImageInvert();
+
+	private:
+
+		Private* m_pD;
+	};
+}
diff --git a/Engine/Plugins/Experimental/Mutable/Source/MutableTools/Public/MuT/NodeImageLayer.h b/Engine/Plugins/Experimental/Mutable/Source/MutableTools/Public/MuT/NodeImageLayer.h
new file mode 100644
--- /dev/null
+++ b/Engine/Plugins/Experimental/Mutable/Source/MutableTools/Public/MuT/NodeImageLayer.h
@@ -0,0 +1,89 @@
+// Copyright Epic Games, Inc. All Rights Reserved.
+
+#pragma once
+
+#include "MuT/NodeImage.h"
+#include "MuR/Image.h"
+
+
+namespace mu
+{
+
+	// Forward definitions
+	class NodeImageLayer;
+	typedef Ptr<NodeImageLayer> NodeImageLayerPtr;
+	typedef Ptr<const NodeImageLayer> NodeImageLayerPtrConst;
+
+
+	//! This node applies a layer blending effect on a base image using a mask and a blended image.
+	//! \ingroup model
+	class MUTABLETOOLS_API NodeImageLayer : public NodeImage
+	{
+	public:
+
+		//-----------------------------------------------------------------------------------------
+		// Life cycle
+		//-----------------------------------------------------------------------------------------
+
+		NodeImageLayer();
+
+		void SerialiseWrapper(OutputArchive& arch) const override;
+		static void Serialise( const NodeImageLayer* pNode, OutputArchive& arch );
+		static NodeImageLayerPtr StaticUnserialise( InputArchive& arch );
+
+
+		//-----------------------------------------------------------------------------------------
+		// Node Interface
+		//-----------------------------------------------------------------------------------------
+
+        NodePtr Clone( NodeMapPtr pMap = 0 ) const override;
+
+        const NODE_TYPE* GetType() const override;
+		static const NODE_TYPE* GetStaticType();
+
+        virtual int GetInputCount() const override;
+        virtual Node* GetInputNode( int i ) const override;
+        void SetInputNode( int i, NodePtr pNode ) override;
+
+		//-----------------------------------------------------------------------------------------
+		// Own Interface
+		//-----------------------------------------------------------------------------------------
+
+		//! Get the node generating the base image that will have the blendeing effect applied.
+		NodeImagePtr GetBase() const;
+		void SetBase( NodeImagePtr );
+
+		//! Get the node generating the mask image controlling the weight of the effect.
+		NodeImagePtr GetMask() const;
+		void SetMask( NodeImagePtr );
+
+		//! Get the image blended on top of the base.
+		NodeImagePtr GetBlended() const;
+		void SetBlended( NodeImagePtr );
+
+		static const char* s_blendTypeName[int32(EBlendType::_BT_COUNT)];
+
+		//!
+		EBlendType GetBlendType() const;
+		void SetBlendType(EBlendType);
+
+		//-----------------------------------------------------------------------------------------
+		// Interface pattern
+		//-----------------------------------------------------------------------------------------
+		class Private;
+		Private* GetPrivate() const;
+        Node::Private* GetBasePrivate() const override;
+
+	protected:
+
+		//! Forbidden. Manage with the Ptr<> template.
+		~NodeImageLayer();
+
+	private:
+
+		Private* m_pD;
+
+	};
+
+
+}
diff --git a/Engine/Plugins/Experimental/Mutable/Source/MutableTools/Public/MuT/NodeImageLayerColour.h b/Engine/Plugins/Experimental/Mutable/Source/MutableTools/Public/MuT/NodeImageLayerColour.h
new file mode 100644
--- /dev/null
+++ b/Engine/Plugins/Experimental/Mutable/Source/MutableTools/Public/MuT/NodeImageLayerColour.h
@@ -0,0 +1,92 @@
+// Copyright Epic Games, Inc. All Rights Reserved.
+
+#pragma once
+
+#include "MuT/NodeImageLayer.h"
+#include "MuR/Image.h"
+
+namespace mu
+{
+
+	// Forward definitions
+	class NodeColour;
+	typedef Ptr<NodeColour> NodeColourPtr;
+	typedef Ptr<const NodeColour> NodeColourPtrConst;
+
+	class NodeImageLayerColour;
+	typedef Ptr<NodeImageLayerColour> NodeImageLayerColourPtr;
+	typedef Ptr<const NodeImageLayerColour> NodeImageLayerColourPtrConst;
+
+
+	//! This node applies a layer blending effect on a base image using a mask and a colour.
+	//! \ingroup model
+	class MUTABLETOOLS_API NodeImageLayerColour : public NodeImage
+	{
+	public:
+
+		//-----------------------------------------------------------------------------------------
+		// Life cycle
+		//-----------------------------------------------------------------------------------------
+
+		NodeImageLayerColour();
+
+		void SerialiseWrapper(OutputArchive& arch) const override;
+		static void Serialise( const NodeImageLayerColour* pNode, OutputArchive& arch );
+		static NodeImageLayerColourPtr StaticUnserialise( InputArchive& arch );
+
+
+		//-----------------------------------------------------------------------------------------
+		// Node Interface
+		//-----------------------------------------------------------------------------------------
+
+        NodePtr Clone( NodeMapPtr pMap = 0 ) const override;
+
+        const NODE_TYPE* GetType() const override;
+		static const NODE_TYPE* GetStaticType();
+
+        virtual int GetInputCount() const override;
+        virtual Node* GetInputNode( int i ) const override;
+        void SetInputNode( int i, NodePtr pNode ) override;
+
+		//-----------------------------------------------------------------------------------------
+		// Own Interface
+		//-----------------------------------------------------------------------------------------
+
+		//! Get the node generating the base image that will have the blending effect applied.
+		NodeImagePtr GetBase() const;
+		void SetBase( NodeImagePtr );
+
+		//! Get the node generating the mask image controlling the weight of the effect.
+		NodeImagePtr GetMask() const;
+		void SetMask( NodeImagePtr );
+
+		//! Get the node generating the color to blend on the base.
+		NodeColourPtr GetColour() const;
+		void SetColour( NodeColourPtr );
+
+		//!
+		EBlendType GetBlendType() const;
+
+		//!
+		void SetBlendType(EBlendType);
+
+		//-----------------------------------------------------------------------------------------
+		// Interface pattern
+		//-----------------------------------------------------------------------------------------
+		class Private;
+		Private* GetPrivate() const;
+        Node::Private* GetBasePrivate() const override;
+
+	protected:
+
+		//! Forbidden. Manage with the Ptr<> template.
+		~NodeImageLayerColour();
+
+	private:
+
+		Private* m_pD;
+
+	};
+
+
+}
diff --git a/Engine/Plugins/Experimental/Mutable/Source/MutableTools/Public/MuT/NodeImageLuminance.h b/Engine/Plugins/Experimental/Mutable/Source/MutableTools/Public/MuT/NodeImageLuminance.h
new file mode 100644
--- /dev/null
+++ b/Engine/Plugins/Experimental/Mutable/Source/MutableTools/Public/MuT/NodeImageLuminance.h
@@ -0,0 +1,76 @@
+// Copyright Epic Games, Inc. All Rights Reserved.
+
+#pragma once
+
+#include "MuT/NodeImage.h"
+
+#include "MuR/Image.h"
+
+
+namespace mu
+{
+
+	// Forward definitions
+	class NodeImageLuminance;
+	typedef Ptr<NodeImageLuminance> NodeImageLuminancePtr;
+	typedef Ptr<const NodeImageLuminance> NodeImageLuminancePtrConst;
+
+
+	//! Calculate the luminance of an image into a new single-channel image..
+	//! \ingroup model
+	class MUTABLETOOLS_API NodeImageLuminance : public NodeImage
+	{
+	public:
+
+		//-----------------------------------------------------------------------------------------
+		// Life cycle
+		//-----------------------------------------------------------------------------------------
+
+		NodeImageLuminance();
+
+		void SerialiseWrapper(OutputArchive& arch) const override;
+		static void Serialise( const NodeImageLuminance* pNode, OutputArchive& arch );
+		static NodeImageLuminancePtr StaticUnserialise( InputArchive& arch );
+
+
+		//-----------------------------------------------------------------------------------------
+		// Node Interface
+		//-----------------------------------------------------------------------------------------
+
+        NodePtr Clone( NodeMapPtr pMap = 0 ) const override;
+
+        const NODE_TYPE* GetType() const override;
+		static const NODE_TYPE* GetStaticType();
+
+        virtual int GetInputCount() const override;
+        virtual Node* GetInputNode( int i ) const override;
+        void SetInputNode( int i, NodePtr pNode ) override;
+
+		//-----------------------------------------------------------------------------------------
+		// Own Interface
+		//-----------------------------------------------------------------------------------------
+
+		//! Get the node generating the source image calculate the luminance from.
+		NodeImagePtr GetSource() const;
+		void SetSource( NodeImagePtr );
+
+		//-----------------------------------------------------------------------------------------
+		// Interface pattern
+		//-----------------------------------------------------------------------------------------
+		class Private;
+		Private* GetPrivate() const;
+        Node::Private* GetBasePrivate() const override;
+
+	protected:
+
+		//! Forbidden. Manage with the Ptr<> template.
+		~NodeImageLuminance();
+
+	private:
+
+		Private* m_pD;
+
+	};
+
+
+}
diff --git a/Engine/Plugins/Experimental/Mutable/Source/MutableTools/Public/MuT/NodeImageMipmap.h b/Engine/Plugins/Experimental/Mutable/Source/MutableTools/Public/MuT/NodeImageMipmap.h
new file mode 100644
--- /dev/null
+++ b/Engine/Plugins/Experimental/Mutable/Source/MutableTools/Public/MuT/NodeImageMipmap.h
@@ -0,0 +1,85 @@
+// Copyright Epic Games, Inc. All Rights Reserved.
+
+#pragma once
+
+#include "MuT/NodeImage.h"
+
+#include "MuR/Image.h"
+
+
+namespace mu
+{
+
+	// Forward definitions
+	class NodeScalar;
+	typedef Ptr<NodeScalar> NodeScalarPtr;
+	typedef Ptr<const NodeScalar> NodeScalarPtrConst;
+
+    class NodeImageMipmap;
+    typedef Ptr<NodeImageMipmap> NodeImageMipmapPtr;
+    typedef Ptr<const NodeImageMipmap> NodeImageMipmapPtrConst;
+
+
+	//! Generate mimaps for an image.
+	//! \ingroup model
+    class MUTABLETOOLS_API NodeImageMipmap : public NodeImage
+	{
+	public:
+
+		//-----------------------------------------------------------------------------------------
+		// Life cycle
+		//-----------------------------------------------------------------------------------------
+
+        NodeImageMipmap();
+
+		void SerialiseWrapper(OutputArchive& arch) const override;
+		static void Serialise( const NodeImageMipmap* pNode, OutputArchive& arch );
+        static NodeImageMipmapPtr StaticUnserialise( InputArchive& arch );
+
+
+		//-----------------------------------------------------------------------------------------
+		// Node Interface
+		//-----------------------------------------------------------------------------------------
+
+        NodePtr Clone( NodeMapPtr pMap = 0 ) const override;
+
+        const NODE_TYPE* GetType() const override;
+		static const NODE_TYPE* GetStaticType();
+
+        virtual int GetInputCount() const override;
+        virtual Node* GetInputNode( int i ) const override;
+        void SetInputNode( int i, NodePtr pNode ) override;
+
+		//-----------------------------------------------------------------------------------------
+		// Own Interface
+		//-----------------------------------------------------------------------------------------
+
+        //! Get the node generating the source image to be mipmapped.
+		NodeImagePtr GetSource() const;
+		void SetSource( NodeImagePtr );
+
+		void SetMipmapGenerationSettings( EMipmapFilterType filterType,
+										  EAddressMode addressMode,
+										  float sharpenFactor,  
+										  bool bDitherMipAlpha );
+
+		//-----------------------------------------------------------------------------------------
+		// Interface pattern
+		//-----------------------------------------------------------------------------------------
+		class Private;
+		Private* GetPrivate() const;
+        Node::Private* GetBasePrivate() const override;
+
+	protected:
+
+		//! Forbidden. Manage with the Ptr<> template.
+        ~NodeImageMipmap();
+
+	private:
+
+		Private* m_pD;
+
+	};
+
+
+}
diff --git a/Engine/Plugins/Experimental/Mutable/Source/MutableTools/Public/MuT/NodeImageMultiLayer.h b/Engine/Plugins/Experimental/Mutable/Source/MutableTools/Public/MuT/NodeImageMultiLayer.h
new file mode 100644
--- /dev/null
+++ b/Engine/Plugins/Experimental/Mutable/Source/MutableTools/Public/MuT/NodeImageMultiLayer.h
@@ -0,0 +1,94 @@
+// Copyright Epic Games, Inc. All Rights Reserved.
+
+#pragma once
+
+#include "MuT/NodeImage.h"
+#include "MuT/NodeImageLayer.h"
+#include "MuT/NodeRange.h"
+#include "MuR/Image.h"
+
+
+namespace mu
+{
+
+	// Forward definitions
+    class NodeImageMultiLayer;
+    typedef Ptr<NodeImageMultiLayer> NodeImageMultiLayerPtr;
+    typedef Ptr<const NodeImageMultiLayer> NodeImageMultiLayerPtrConst;
+
+
+    //! This node applies any numbre of layer blending effect on a base image using a mask and a
+    //! blended image. The number of layers depends on a scalar input of this node.
+	//! \ingroup model
+    class MUTABLETOOLS_API NodeImageMultiLayer : public NodeImage
+	{
+	public:
+
+		//-----------------------------------------------------------------------------------------
+		// Life cycle
+		//-----------------------------------------------------------------------------------------
+
+        NodeImageMultiLayer();
+
+		void SerialiseWrapper(OutputArchive& arch) const override;
+		static void Serialise( const NodeImageMultiLayer* pNode, OutputArchive& arch );
+        static NodeImageMultiLayerPtr StaticUnserialise( InputArchive& arch );
+
+
+		//-----------------------------------------------------------------------------------------
+		// Node Interface
+		//-----------------------------------------------------------------------------------------
+
+        NodePtr Clone( NodeMapPtr pMap = 0 ) const override;
+
+        const NODE_TYPE* GetType() const override;
+		static const NODE_TYPE* GetStaticType();
+
+        virtual int GetInputCount() const override;
+        virtual Node* GetInputNode( int i ) const override;
+        void SetInputNode( int i, NodePtr pNode ) override;
+
+		//-----------------------------------------------------------------------------------------
+		// Own Interface
+		//-----------------------------------------------------------------------------------------
+
+		//! Get the node generating the base image that will have the blendeing effect applied.
+		NodeImagePtr GetBase() const;
+		void SetBase( NodeImagePtr );
+
+		//! Get the node generating the mask image controlling the weight of the effect.
+		NodeImagePtr GetMask() const;
+		void SetMask( NodeImagePtr );
+
+        //! Get the image blended on top of the base.
+        NodeImagePtr GetBlended() const;
+        void SetBlended( NodeImagePtr );
+
+        //! Get the number of layers to apply.
+        NodeRangePtr GetRange() const;
+        void SetRange( NodeRangePtr );
+
+		//!
+		EBlendType GetBlendType() const;
+        void SetBlendType(EBlendType);
+
+		//-----------------------------------------------------------------------------------------
+		// Interface pattern
+		//-----------------------------------------------------------------------------------------
+		class Private;
+		Private* GetPrivate() const;
+        Node::Private* GetBasePrivate() const override;
+
+	protected:
+
+		//! Forbidden. Manage with the Ptr<> template.
+        ~NodeImageMultiLayer();
+
+	private:
+
+		Private* m_pD;
+
+	};
+
+
+}
