Moved the Mutable plugin files to live under a directory named after the module. This allows us to get rid of long include paths that contain public and private folder names in the path and fixes head

diff --git a/Engine/Plugins/Experimental/Mutable/Source/MutableTools/Public/MuT/NodeImageNormalComposite.h b/Engine/Plugins/Experimental/Mutable/Source/MutableTools/Public/MuT/NodeImageNormalComposite.h
new file mode 100644
--- /dev/null
+++ b/Engine/Plugins/Experimental/Mutable/Source/MutableTools/Public/MuT/NodeImageNormalComposite.h
@@ -0,0 +1,82 @@
+// Copyright Epic Games, Inc. All Rights Reserved.
+
+#pragma once
+
+#include "MuT/NodeImage.h"
+#include "MuR/Image.h"
+
+namespace mu
+{
+
+	// Forward definitions
+	class NodeImageNormalComposite;
+	typedef Ptr<NodeImageNormalComposite> NodeImageNormalCompositePtr;
+	typedef Ptr<const NodeImageNormalComposite> NodeImageNormalCompositePtrConst;
+
+	//! Node that inverts the colors of an image, channel by channel
+	//! \ingroup model
+
+	class MUTABLETOOLS_API NodeImageNormalComposite : public NodeImage
+	{
+	public:
+
+		//MUTABLE_DEFINE_CONST_VISITABLE()
+		
+		//-----------------------------------------------------------------------------------------
+		// Life cycle
+		//-----------------------------------------------------------------------------------------
+
+		NodeImageNormalComposite();
+
+		void SerialiseWrapper(OutputArchive& arch) const override;
+		static void Serialise(const NodeImageNormalComposite* pNode, OutputArchive& arch);
+		static NodeImageNormalCompositePtr StaticUnserialise(InputArchive& arch);
+
+
+		//-----------------------------------------------------------------------------------------
+		// Node Interface
+		//-----------------------------------------------------------------------------------------
+
+		NodePtr Clone(NodeMapPtr pMap = 0) const override;
+
+		const NODE_TYPE* GetType() const override;
+		static const NODE_TYPE* GetStaticType();
+
+		virtual int GetInputCount() const override;
+		virtual Node* GetInputNode(int i)const override;
+		void SetInputNode(int i, NodePtr pNode) override;
+
+		//-----------------------------------------------------------------------------------------
+		// Own Interface
+		//-----------------------------------------------------------------------------------------
+
+		NodeImagePtr GetBase() const;
+		void SetBase(NodeImagePtr);
+
+        NodeImagePtr GetNormal() const;
+        void SetNormal(NodeImagePtr);
+
+		float GetPower() const;
+		void SetPower(float); 
+
+        ECompositeImageMode GetMode() const;
+        void SetMode(ECompositeImageMode);
+
+		//-----------------------------------------------------------------------------------------
+		// Interface pattern
+		//-----------------------------------------------------------------------------------------
+		class Private;
+		Private* GetPrivate() const;
+		Node::Private*GetBasePrivate()const override;
+
+	protected:
+
+		//! Forbidden. Manage with the Ptr<> template
+		~NodeImageNormalComposite();
+
+	private:
+
+		Private* m_pD;
+	};
+
+}
diff --git a/Engine/Plugins/Experimental/Mutable/Source/MutableTools/Public/MuT/NodeImageParameter.h b/Engine/Plugins/Experimental/Mutable/Source/MutableTools/Public/MuT/NodeImageParameter.h
new file mode 100644
--- /dev/null
+++ b/Engine/Plugins/Experimental/Mutable/Source/MutableTools/Public/MuT/NodeImageParameter.h
@@ -0,0 +1,88 @@
+// Copyright Epic Games, Inc. All Rights Reserved.
+
+#pragma once
+
+#include "MuT/NodeImage.h"
+#include "MuR/Image.h"
+
+
+namespace mu
+{
+
+    class NodeImageParameter;
+    typedef Ptr<NodeImageParameter> NodeImageParameterPtr;
+    typedef Ptr<const NodeImageParameter> NodeImageParameterPtrConst;
+
+
+    //! Node that defines a Image model parameter.
+	//! \ingroup model
+    class MUTABLETOOLS_API NodeImageParameter : public NodeImage
+	{
+	public:
+
+		//-----------------------------------------------------------------------------------------
+		// Life cycle
+		//-----------------------------------------------------------------------------------------
+
+        NodeImageParameter();
+
+		void SerialiseWrapper(OutputArchive& arch) const override;
+		static void Serialise( const NodeImageParameter* pNode, OutputArchive& arch );
+        static NodeImageParameterPtr StaticUnserialise( InputArchive& arch );
+
+
+		//-----------------------------------------------------------------------------------------
+		// Node Interface
+		//-----------------------------------------------------------------------------------------
+
+		NodePtr Clone( NodeMapPtr pMap = 0 ) const override;
+
+		const NODE_TYPE* GetType() const override;
+		static const NODE_TYPE* GetStaticType();
+
+		int GetInputCount() const override;
+		Node* GetInputNode( int i ) const override;
+		void SetInputNode( int i, NodePtr pNode ) override;
+
+		//-----------------------------------------------------------------------------------------
+		// Own Interface
+		//-----------------------------------------------------------------------------------------
+
+		//! Get the name of the parameter. It will be exposed in the final compiled data.
+		const char* GetName() const;
+
+		//! Set the name of the parameter.
+		void SetName( const char* );
+
+		//! Get the uid of the parameter. It will be exposed in the final compiled data.
+		const char* GetUid() const;
+
+		//! Set the uid of the parameter.
+		void SetUid( const char* );
+
+		//! Get the default value of the parameter.
+        //ImagePtrConst GetDefaultValue() const;
+
+		//! Set the default value of the parameter.
+        //void SetDefaultValue( ImagePtr );
+
+		//-----------------------------------------------------------------------------------------
+		// Interface pattern
+		//-----------------------------------------------------------------------------------------
+		class Private;
+		Private* GetPrivate() const;
+        Node::Private* GetBasePrivate() const override;
+
+	protected:
+
+		//! Forbidden. Manage with the Ptr<> template.
+        ~NodeImageParameter();
+
+	private:
+
+		Private* m_pD;
+
+	};
+
+
+}
diff --git a/Engine/Plugins/Experimental/Mutable/Source/MutableTools/Public/MuT/NodeImagePlainColour.h b/Engine/Plugins/Experimental/Mutable/Source/MutableTools/Public/MuT/NodeImagePlainColour.h
new file mode 100644
--- /dev/null
+++ b/Engine/Plugins/Experimental/Mutable/Source/MutableTools/Public/MuT/NodeImagePlainColour.h
@@ -0,0 +1,81 @@
+// Copyright Epic Games, Inc. All Rights Reserved.
+
+#pragma once
+
+#include "MuT/NodeImage.h"
+
+
+namespace mu
+{
+
+	// Forward definitions
+	class NodeColour;
+	typedef Ptr<NodeColour> NodeColourPtr;
+	typedef Ptr<const NodeColour> NodeColourPtrConst;
+
+	class NodeImagePlainColour;
+	typedef Ptr<NodeImagePlainColour> NodeImagePlainColourPtr;
+	typedef Ptr<const NodeImagePlainColour> NodeImagePlainColourPtrConst;
+
+
+	//! Node that multiplies the colors of an image, channel by channel.
+	//! \ingroup model
+	class MUTABLETOOLS_API NodeImagePlainColour : public NodeImage
+	{
+	public:
+
+		//-----------------------------------------------------------------------------------------
+		// Life cycle
+		//-----------------------------------------------------------------------------------------
+
+		NodeImagePlainColour();
+
+		void SerialiseWrapper(OutputArchive& arch) const override;
+		static void Serialise( const NodeImagePlainColour* pNode, OutputArchive& arch );
+		static NodeImagePlainColourPtr StaticUnserialise( InputArchive& arch );
+
+
+		//-----------------------------------------------------------------------------------------
+		// Node Interface
+		//-----------------------------------------------------------------------------------------
+
+        NodePtr Clone( NodeMapPtr pMap = 0 ) const override;
+
+        const NODE_TYPE* GetType() const override;
+		static const NODE_TYPE* GetStaticType();
+
+        virtual int GetInputCount() const override;
+        virtual Node* GetInputNode( int i ) const override;
+        void SetInputNode( int i, NodePtr pNode ) override;
+
+		//-----------------------------------------------------------------------------------------
+		// Own Interface
+		//-----------------------------------------------------------------------------------------
+
+		//! Colour of the image.
+		NodeColourPtr GetColour() const;
+		void SetColour( NodeColourPtr );
+
+		//! New size or relative factor
+		void SetSize( int x, int y );
+
+		//-----------------------------------------------------------------------------------------
+		// Interface pattern
+		//-----------------------------------------------------------------------------------------
+		class Private;
+		Private* GetPrivate() const;
+        Node::Private* GetBasePrivate() const override;
+
+	protected:
+
+		//! Forbidden. Manage with the Ptr<> template.
+		~NodeImagePlainColour();
+
+	private:
+
+		Private* m_pD;
+
+	};
+
+
+}
diff --git a/Engine/Plugins/Experimental/Mutable/Source/MutableTools/Public/MuT/NodeImageProject.h b/Engine/Plugins/Experimental/Mutable/Source/MutableTools/Public/MuT/NodeImageProject.h
new file mode 100644
--- /dev/null
+++ b/Engine/Plugins/Experimental/Mutable/Source/MutableTools/Public/MuT/NodeImageProject.h
@@ -0,0 +1,121 @@
+// Copyright Epic Games, Inc. All Rights Reserved.
+
+
+#pragma once
+
+#include "MuT/NodeImage.h"
+#include "MuR/Image.h"
+#include "Math/IntVector.h"
+
+namespace mu
+{
+
+	// Forward definitions
+	class NodeScalar;
+	typedef Ptr<NodeScalar> NodeScalarPtr;
+	typedef Ptr<const NodeScalar> NodeScalarPtrConst;
+
+	class NodeColour;
+	typedef Ptr<NodeColour> NodeColourPtr;
+	typedef Ptr<const NodeColour> NodeColourPtrConst;
+
+	class NodeMesh;
+	typedef Ptr<NodeMesh> NodeMeshPtr;
+	typedef Ptr<const NodeMesh> NodeMeshPtrConst;
+
+	class NodeProjector;
+	typedef Ptr<NodeProjector> NodeProjectorPtr;
+	typedef Ptr<const NodeProjector> NodeProjectorPtrConst;
+
+	class NodeImageProject;
+	typedef Ptr<NodeImageProject> NodeImageProjectPtr;
+	typedef Ptr<const NodeImageProject> NodeImageProjectPtrConst;
+
+
+	//!
+	//! \ingroup model
+	class MUTABLETOOLS_API NodeImageProject : public NodeImage
+	{
+	public:
+
+		//-----------------------------------------------------------------------------------------
+		// Life cycle
+		//-----------------------------------------------------------------------------------------
+
+		NodeImageProject();
+
+		void SerialiseWrapper(OutputArchive& arch) const override;
+		static void Serialise( const NodeImageProject* pNode, OutputArchive& arch );
+		static NodeImageProjectPtr StaticUnserialise( InputArchive& arch );
+
+
+		//-----------------------------------------------------------------------------------------
+		// Node Interface
+		//-----------------------------------------------------------------------------------------
+
+        NodePtr Clone( NodeMapPtr pMap = 0 ) const override;
+
+        const NODE_TYPE* GetType() const override;
+		static const NODE_TYPE* GetStaticType();
+
+        virtual int GetInputCount() const override;
+        virtual Node* GetInputNode( int i ) const override;
+        void SetInputNode( int i, NodePtr pNode ) override;
+
+		//-----------------------------------------------------------------------------------------
+		// Own Interface
+		//-----------------------------------------------------------------------------------------
+
+		//! Get the node generating the mesh to be used for the projection.
+		NodeMeshPtr GetMesh() const;
+		void SetMesh( NodeMeshPtr );
+
+		//! Get the node generating the projector to be used for the projection.
+		NodeProjectorPtr GetProjector() const;
+		void SetProjector( NodeProjectorPtr );
+
+		//! Get the node generating a mask to filter out the projected images. It is optional.
+		NodeImagePtr GetTargetMask() const;
+		void SetTargetMask( NodeImagePtr );
+
+        //! Get the node generating the threshold angle of the projection. In degrees.
+		NodeScalarPtr GetAngleFadeStart() const;
+		void SetAngleFadeStart( NodeScalarPtr );
+
+        //! Get the node generating the threshold angle of the projection. In degrees.
+		NodeScalarPtr GetAngleFadeEnd() const;
+		void SetAngleFadeEnd( NodeScalarPtr );
+
+        //! Get the node generating the image to project.
+        NodeImagePtr GetImage() const;
+        void SetImage( NodeImagePtr );
+
+        //! UV layout of the mesh to use for the generated image. Defaults to 0.
+        uint8_t GetLayout() const;
+        void SetLayout( uint8_t  );
+
+		//! Set the size of the image to generate with the projection. If set to 0 (default) a size
+		//! that matches how this node is used will try to be guessed. 
+		const FUintVector2& GetImageSize() const;
+		void SetImageSize( const FUintVector2& size );
+
+		//-----------------------------------------------------------------------------------------
+		// Interface pattern
+		//-----------------------------------------------------------------------------------------
+		class Private;
+		Private* GetPrivate() const;
+        Node::Private* GetBasePrivate() const override;
+
+	protected:
+
+		//! Forbidden. Manage with the Ptr<> template.
+		~NodeImageProject();
+
+	private:
+
+		Private* m_pD;
+
+	};
+
+
+}
diff --git a/Engine/Plugins/Experimental/Mutable/Source/MutableTools/Public/MuT/NodeImageResize.h b/Engine/Plugins/Experimental/Mutable/Source/MutableTools/Public/MuT/NodeImageResize.h
new file mode 100644
--- /dev/null
+++ b/Engine/Plugins/Experimental/Mutable/Source/MutableTools/Public/MuT/NodeImageResize.h
@@ -0,0 +1,84 @@
+// Copyright Epic Games, Inc. All Rights Reserved.
+
+#pragma once
+
+#include "MuT/NodeImage.h"
+
+
+namespace mu
+{
+
+	// Forward definitions
+	class NodeScalar;
+	typedef Ptr<NodeScalar> NodeScalarPtr;
+	typedef Ptr<const NodeScalar> NodeScalarPtrConst;
+
+	class NodeImageResize;
+	typedef Ptr<NodeImageResize> NodeImageResizePtr;
+	typedef Ptr<const NodeImageResize> NodeImageResizePtrConst;
+
+
+	//! Node that multiplies the colors of an image, channel by channel.
+	//! \ingroup model
+	class MUTABLETOOLS_API NodeImageResize : public NodeImage
+	{
+	public:
+
+		//-----------------------------------------------------------------------------------------
+		// Life cycle
+		//-----------------------------------------------------------------------------------------
+
+		NodeImageResize();
+
+		void SerialiseWrapper(OutputArchive& arch) const override;
+		static void Serialise( const NodeImageResize* pNode, OutputArchive& arch );
+		static NodeImageResizePtr StaticUnserialise( InputArchive& arch );
+
+
+		//-----------------------------------------------------------------------------------------
+		// Node Interface
+		//-----------------------------------------------------------------------------------------
+
+        NodePtr Clone( NodeMapPtr pMap = 0 ) const override;
+
+        const NODE_TYPE* GetType() const override;
+		static const NODE_TYPE* GetStaticType();
+
+        virtual int GetInputCount() const override;
+        virtual Node* GetInputNode( int i ) const override;
+        void SetInputNode( int i, NodePtr pNode ) override;
+
+		//-----------------------------------------------------------------------------------------
+		// Own Interface
+		//-----------------------------------------------------------------------------------------
+
+		//! Base image to resize.
+		NodeImagePtr GetBase() const;
+		void SetBase( NodeImagePtr );
+
+		//! Is the size a relative factor or an absolute size?
+		void SetRelative( bool );
+
+		//! New size or relative factor
+		void SetSize( float x, float y );
+
+		//-----------------------------------------------------------------------------------------
+		// Interface pattern
+		//-----------------------------------------------------------------------------------------
+		class Private;
+		Private* GetPrivate() const;
+        Node::Private* GetBasePrivate() const override;
+
+	protected:
+
+		//! Forbidden. Manage with the Ptr<> template.
+		~NodeImageResize();
+
+	private:
+
+		Private* m_pD;
+
+	};
+
+
+}
diff --git a/Engine/Plugins/Experimental/Mutable/Source/MutableTools/Public/MuT/NodeImageSaturate.h b/Engine/Plugins/Experimental/Mutable/Source/MutableTools/Public/MuT/NodeImageSaturate.h
new file mode 100644
--- /dev/null
+++ b/Engine/Plugins/Experimental/Mutable/Source/MutableTools/Public/MuT/NodeImageSaturate.h
@@ -0,0 +1,88 @@
+// Copyright Epic Games, Inc. All Rights Reserved.
+
+#pragma once
+
+#include "MuT/NodeImage.h"
+
+#include "MuR/Image.h"
+
+
+namespace mu
+{
+
+	// Forward definitions
+	class NodeScalar;
+	typedef Ptr<NodeScalar> NodeScalarPtr;
+	typedef Ptr<const NodeScalar> NodeScalarPtrConst;
+
+	class NodeImageSaturate;
+	typedef Ptr<NodeImageSaturate> NodeImageSaturatePtr;
+	typedef Ptr<const NodeImageSaturate> NodeImageSaturatePtrConst;
+
+
+	//! Change the saturation of an image. This node can be used to increase the saturation with a
+	//! factor bigger than 1, or to decrease it or desaturate it completely with a factor smaller
+	//! than 1 or 0.
+	//! \ingroup model
+	class MUTABLETOOLS_API NodeImageSaturate : public NodeImage
+	{
+	public:
+
+		//-----------------------------------------------------------------------------------------
+		// Life cycle
+		//-----------------------------------------------------------------------------------------
+
+		NodeImageSaturate();
+
+		void SerialiseWrapper(OutputArchive& arch) const override;
+		static void Serialise( const NodeImageSaturate* pNode, OutputArchive& arch );
+		static NodeImageSaturatePtr StaticUnserialise( InputArchive& arch );
+
+
+		//-----------------------------------------------------------------------------------------
+		// Node Interface
+		//-----------------------------------------------------------------------------------------
+
+        NodePtr Clone( NodeMapPtr pMap = 0 ) const override;
+
+        const NODE_TYPE* GetType() const override;
+		static const NODE_TYPE* GetStaticType();
+
+        virtual int GetInputCount() const override;
+        virtual Node* GetInputNode( int i ) const override;
+        void SetInputNode( int i, NodePtr pNode ) override;
+
+		//-----------------------------------------------------------------------------------------
+		// Own Interface
+		//-----------------------------------------------------------------------------------------
+
+		//! Get the node generating the saturation factor.
+		//! A value of 0 completely desaturates. 1 will leave the same saturation and bigger than 1
+		//! will increase it.
+		NodeScalarPtr GetFactor() const;
+		void SetFactor( NodeScalarPtr );
+
+		//! Get the node generating the source image to be [de]saturated.
+		NodeImagePtr GetSource() const;
+		void SetSource( NodeImagePtr );
+
+		//-----------------------------------------------------------------------------------------
+		// Interface pattern
+		//-----------------------------------------------------------------------------------------
+		class Private;
+		Private* GetPrivate() const;
+        Node::Private* GetBasePrivate() const override;
+
+	protected:
+
+		//! Forbidden. Manage with the Ptr<> template.
+		~NodeImageSaturate();
+
+	private:
+
+		Private* m_pD;
+
+	};
+
+
+}
diff --git a/Engine/Plugins/Experimental/Mutable/Source/MutableTools/Public/MuT/NodeImageSelectColour.h b/Engine/Plugins/Experimental/Mutable/Source/MutableTools/Public/MuT/NodeImageSelectColour.h
new file mode 100644
--- /dev/null
+++ b/Engine/Plugins/Experimental/Mutable/Source/MutableTools/Public/MuT/NodeImageSelectColour.h
@@ -0,0 +1,85 @@
+// Copyright Epic Games, Inc. All Rights Reserved.
+
+#pragma once
+
+#include "MuT/NodeImage.h"
+
+#include "MuR/Image.h"
+
+
+namespace mu
+{
+
+	// Forward definitions
+	class NodeColour;
+	typedef Ptr<NodeColour> NodeColourPtr;
+	typedef Ptr<const NodeColour> NodeColourPtrConst;
+
+	class NodeImageSelectColour;
+	typedef Ptr<NodeImageSelectColour> NodeImageSelectColourPtr;
+	typedef Ptr<const NodeImageSelectColour> NodeImageSelectColourPtrConst;
+
+
+	//! Create a black and white mask image by selecting the pixels from a source image that match
+	//! a colour.
+	//! \ingroup model
+	class MUTABLETOOLS_API NodeImageSelectColour : public NodeImage
+	{
+	public:
+
+		//-----------------------------------------------------------------------------------------
+		// Life cycle
+		//-----------------------------------------------------------------------------------------
+
+		NodeImageSelectColour();
+
+		void SerialiseWrapper(OutputArchive& arch) const override;
+		static void Serialise( const NodeImageSelectColour* pNode, OutputArchive& arch );
+		static NodeImageSelectColourPtr StaticUnserialise( InputArchive& arch );
+
+
+		//-----------------------------------------------------------------------------------------
+		// Node Interface
+		//-----------------------------------------------------------------------------------------
+
+        NodePtr Clone( NodeMapPtr pMap = 0 ) const override;
+
+        const NODE_TYPE* GetType() const override;
+		static const NODE_TYPE* GetStaticType();
+
+        virtual int GetInputCount() const override;
+        virtual Node* GetInputNode( int i ) const override;
+        void SetInputNode( int i, NodePtr pNode ) override;
+
+		//-----------------------------------------------------------------------------------------
+		// Own Interface
+		//-----------------------------------------------------------------------------------------
+
+		//! Get the node generating the colour to select.
+		NodeColourPtr GetColour() const;
+		void SetColour( NodeColourPtr );
+
+		//! Get the node generating the source image to select pixels from.
+		NodeImagePtr GetSource() const;
+		void SetSource( NodeImagePtr );
+
+		//-----------------------------------------------------------------------------------------
+		// Interface pattern
+		//-----------------------------------------------------------------------------------------
+		class Private;
+		Private* GetPrivate() const;
+        Node::Private* GetBasePrivate() const override;
+
+	protected:
+
+		//! Forbidden. Manage with the Ptr<> template.
+		~NodeImageSelectColour();
+
+	private:
+
+		Private* m_pD;
+
+	};
+
+
+}
diff --git a/Engine/Plugins/Experimental/Mutable/Source/MutableTools/Public/MuT/NodeImageSwitch.h b/Engine/Plugins/Experimental/Mutable/Source/MutableTools/Public/MuT/NodeImageSwitch.h
new file mode 100644
--- /dev/null
+++ b/Engine/Plugins/Experimental/Mutable/Source/MutableTools/Public/MuT/NodeImageSwitch.h
@@ -0,0 +1,90 @@
+// Copyright Epic Games, Inc. All Rights Reserved.
+
+#pragma once
+
+#include "MuT/NodeImage.h"
+
+#include "MuR/Image.h"
+
+
+namespace mu
+{
+
+	// Forward definitions
+	class NodeScalar;
+	typedef Ptr<NodeScalar> NodeScalarPtr;
+	typedef Ptr<const NodeScalar> NodeScalarPtrConst;
+
+	class NodeImageSwitch;
+	typedef Ptr<NodeImageSwitch> NodeImageSwitchPtr;
+	typedef Ptr<const NodeImageSwitch> NodeImageSwitchPtrConst;
+
+
+	//! This node selects an output image from a set of input images based on a parameter.
+	//! \ingroup model
+	class MUTABLETOOLS_API NodeImageSwitch : public NodeImage
+	{
+	public:
+
+		//-----------------------------------------------------------------------------------------
+		// Life cycle
+		//-----------------------------------------------------------------------------------------
+
+		NodeImageSwitch();
+
+		void SerialiseWrapper(OutputArchive& arch) const override;
+		static void Serialise( const NodeImageSwitch* pNode, OutputArchive& arch );
+		static NodeImageSwitchPtr StaticUnserialise( InputArchive& arch );
+
+
+		//-----------------------------------------------------------------------------------------
+		// Node Interface
+		//-----------------------------------------------------------------------------------------
+
+        NodePtr Clone( NodeMapPtr pMap = 0 ) const override;
+
+        const NODE_TYPE* GetType() const override;
+		static const NODE_TYPE* GetStaticType();
+
+        virtual int GetInputCount() const override;
+        virtual Node* GetInputNode( int i ) const override;
+        void SetInputNode( int i, NodePtr pNode ) override;
+
+		//-----------------------------------------------------------------------------------------
+		// Own Interface
+		//-----------------------------------------------------------------------------------------
+
+		//! Get the node generating the parameter used to select the option.
+		NodeScalarPtr GetParameter() const;
+		void SetParameter( NodeScalarPtr );
+
+		//! Set the number of option images. It will keep the currently set targets and initialise
+		//! the new ones as null.
+		void SetOptionCount( int );
+		int GetOptionCount() const;
+
+		//! Get the node generating the t-th option image.
+		NodeImagePtr GetOption( int t ) const;
+		void SetOption( int t, NodeImagePtr );
+
+
+		//-----------------------------------------------------------------------------------------
+		// Interface pattern
+		//-----------------------------------------------------------------------------------------
+		class Private;
+		Private* GetPrivate() const;
+        Node::Private* GetBasePrivate() const override;
+
+	protected:
+
+		//! Forbidden. Manage with the Ptr<> template.
+		~NodeImageSwitch();
+
+	private:
+
+		Private* m_pD;
+
+	};
+
+
+}
diff --git a/Engine/Plugins/Experimental/Mutable/Source/MutableTools/Public/MuT/NodeImageSwizzle.h b/Engine/Plugins/Experimental/Mutable/Source/MutableTools/Public/MuT/NodeImageSwizzle.h
new file mode 100644
--- /dev/null
+++ b/Engine/Plugins/Experimental/Mutable/Source/MutableTools/Public/MuT/NodeImageSwizzle.h
@@ -0,0 +1,84 @@
+// Copyright Epic Games, Inc. All Rights Reserved.
+
+#pragma once
+
+#include "MuT/NodeImage.h"
+
+#include "MuR/Image.h"
+
+
+namespace mu
+{
+
+	// Forward definitions
+	class NodeImageSwizzle;
+	typedef Ptr<NodeImageSwizzle> NodeImageSwizzlePtr;
+	typedef Ptr<const NodeImageSwizzle> NodeImageSwizzlePtrConst;
+
+
+	//! Node that composes a new image by gathering pixel data from channels in other images.
+	//! \ingroup model
+	class MUTABLETOOLS_API NodeImageSwizzle : public NodeImage
+	{
+	public:
+
+		//-----------------------------------------------------------------------------------------
+		// Life cycle
+		//-----------------------------------------------------------------------------------------
+
+		NodeImageSwizzle();
+
+		void SerialiseWrapper(OutputArchive& arch) const override;
+		static void Serialise( const NodeImageSwizzle* pNode, OutputArchive& arch );
+		static NodeImageSwizzlePtr StaticUnserialise( InputArchive& arch );
+
+
+		//-----------------------------------------------------------------------------------------
+		// Node Interface
+		//-----------------------------------------------------------------------------------------
+
+        NodePtr Clone( NodeMapPtr pMap = 0 ) const override;
+
+        const NODE_TYPE* GetType() const override;
+		static const NODE_TYPE* GetStaticType();
+
+        virtual int GetInputCount() const override;
+        virtual Node* GetInputNode( int i ) const override;
+        void SetInputNode( int i, NodePtr pNode ) override;
+
+		//-----------------------------------------------------------------------------------------
+		// Own Interface
+		//-----------------------------------------------------------------------------------------
+
+		//!
+		EImageFormat GetFormat() const;
+		void SetFormat(EImageFormat);
+
+		//!
+		NodeImagePtr GetSource( int ) const;
+		int GetSourceChannel( int ) const;
+
+		void SetSource( int, NodeImagePtr );
+		void SetSourceChannel( int, int );
+
+
+		//-----------------------------------------------------------------------------------------
+		// Interface pattern
+		//-----------------------------------------------------------------------------------------
+		class Private;
+		Private* GetPrivate() const;
+        Node::Private* GetBasePrivate() const override;
+
+	protected:
+
+		//! Forbidden. Manage with the Ptr<> template.
+		~NodeImageSwizzle();
+
+	private:
+
+		Private* m_pD;
+
+	};
+
+
+}
diff --git a/Engine/Plugins/Experimental/Mutable/Source/MutableTools/Public/MuT/NodeImageTable.h b/Engine/Plugins/Experimental/Mutable/Source/MutableTools/Public/MuT/NodeImageTable.h
new file mode 100644
--- /dev/null
+++ b/Engine/Plugins/Experimental/Mutable/Source/MutableTools/Public/MuT/NodeImageTable.h
@@ -0,0 +1,86 @@
+// Copyright Epic Games, Inc. All Rights Reserved.
+
+#pragma once
+
+#include "MuT/NodeImage.h"
+
+
+namespace mu
+{
+
+	// Forward definitions
+	class NodeImageTable;
+	typedef Ptr<NodeImageTable> NodeImageTablePtr;
+	typedef Ptr<const NodeImageTable> NodeImageTablePtrConst;
+
+	class Table;
+	typedef Ptr<Table> TablePtr;
+	typedef Ptr<const Table> TablePtrConst;
+
+
+	//! This node provides the meshes stored in the column of a table.
+	//! \ingroup transform
+	class MUTABLETOOLS_API NodeImageTable : public NodeImage
+	{
+	public:
+
+		//-----------------------------------------------------------------------------------------
+		// Life cycle
+		//-----------------------------------------------------------------------------------------
+
+		NodeImageTable();
+
+		void SerialiseWrapper(OutputArchive& arch) const override;
+		static void Serialise( const NodeImageTable* pNode, OutputArchive& arch );
+		static NodeImageTablePtr StaticUnserialise( InputArchive& arch );
+
+
+		//-----------------------------------------------------------------------------------------
+		// Node Interface
+		//-----------------------------------------------------------------------------------------
+
+        NodePtr Clone( NodeMapPtr pMap = 0 ) const override;
+
+        const NODE_TYPE* GetType() const override;
+		static const NODE_TYPE* GetStaticType();
+
+        virtual int GetInputCount() const override;
+        virtual Node* GetInputNode( int i ) const override;
+        void SetInputNode( int i, NodePtr pNode ) override;
+
+		//-----------------------------------------------------------------------------------------
+		// Own Interface
+		//-----------------------------------------------------------------------------------------
+
+		//! Set the name of the implicit table parameter.
+		void SetParameterName( const char* strName );
+
+		//!
+		TablePtr GetTable() const;
+		void SetTable( TablePtr );
+
+		//!
+		const char* GetColumn() const;
+		void SetColumn( const char* strName );
+
+		//-----------------------------------------------------------------------------------------
+		// Interface pattern
+		//-----------------------------------------------------------------------------------------
+		class Private;
+		Private* GetPrivate() const;
+        Node::Private* GetBasePrivate() const override;
+
+
+	protected:
+
+		//! Forbidden. Manage with the Ptr<> template.
+		~NodeImageTable();
+
+	private:
+
+		Private* m_pD;
+
+	};
+
+
+}
diff --git a/Engine/Plugins/Experimental/Mutable/Source/MutableTools/Public/MuT/NodeImageTransform.h b/Engine/Plugins/Experimental/Mutable/Source/MutableTools/Public/MuT/NodeImageTransform.h
new file mode 100644
--- /dev/null
+++ b/Engine/Plugins/Experimental/Mutable/Source/MutableTools/Public/MuT/NodeImageTransform.h
@@ -0,0 +1,88 @@
+// Copyright Epic Games, Inc. All Rights Reserved.
+
+#pragma once
+
+#include "MuT/NodeImage.h"
+
+
+namespace mu
+{
+
+	// Forward definitions
+	class NodeScalar;
+	typedef Ptr<NodeScalar> NodeScalarPtr;
+	typedef Ptr<const NodeScalar> NodeScalarPtrConst;
+
+	class NodeImageTransform;
+	typedef Ptr<NodeImageTransform> NodeImageTransformPtr;
+	typedef Ptr<const NodeImageTransform> NodeImageTransformPtrConst;
+
+
+	//! Node that multiplies the colors of an image, channel by channel.
+	//! \ingroup model
+	class MUTABLETOOLS_API NodeImageTransform : public NodeImage
+	{
+	public:
+
+		//-----------------------------------------------------------------------------------------
+		// Life cycle
+		//-----------------------------------------------------------------------------------------
+
+		NodeImageTransform();
+
+		void SerialiseWrapper(OutputArchive& arch) const override;
+		static void Serialise( const NodeImageTransform* pNode, OutputArchive& arch );
+		static NodeImageTransformPtr StaticUnserialise( InputArchive& arch );
+
+
+		//-----------------------------------------------------------------------------------------
+		// Node Interface
+		//-----------------------------------------------------------------------------------------
+
+        NodePtr Clone( NodeMapPtr pMap = 0 ) const override;
+
+        const NODE_TYPE* GetType() const override;
+		static const NODE_TYPE* GetStaticType();
+
+        virtual int GetInputCount() const override;
+        virtual Node* GetInputNode( int i ) const override;
+        void SetInputNode( int i, NodePtr pNode ) override;
+
+		//-----------------------------------------------------------------------------------------
+		// Own Interface
+		//-----------------------------------------------------------------------------------------
+
+		//! Base image to multiply.
+		NodeImagePtr GetBase() const;
+		void SetBase( NodeImagePtr );
+	
+		NodeScalarPtr GetOffsetX() const;
+		NodeScalarPtr GetOffsetY() const;
+		NodeScalarPtr GetScaleX() const;
+		NodeScalarPtr GetScaleY() const;
+		NodeScalarPtr GetRotation() const;
+		void SetOffsetX( NodeScalarPtr pNode );
+		void SetOffsetY( NodeScalarPtr pNode );
+		void SetScaleX( NodeScalarPtr pNode );
+		void SetScaleY( NodeScalarPtr pNode );
+		void SetRotation( NodeScalarPtr pNode );
+		//-----------------------------------------------------------------------------------------
+		// Interface pattern
+		//-----------------------------------------------------------------------------------------
+		class Private;
+		Private* GetPrivate() const;
+        Node::Private* GetBasePrivate() const override;
+
+	protected:
+
+		//! Forbidden. Manage with the Ptr<> template.
+		~NodeImageTransform();
+
+	private:
+
+		Private* m_pD;
+
+	};
+
+
+}
diff --git a/Engine/Plugins/Experimental/Mutable/Source/MutableTools/Public/MuT/NodeImageVariation.h b/Engine/Plugins/Experimental/Mutable/Source/MutableTools/Public/MuT/NodeImageVariation.h
new file mode 100644
--- /dev/null
+++ b/Engine/Plugins/Experimental/Mutable/Source/MutableTools/Public/MuT/NodeImageVariation.h
@@ -0,0 +1,83 @@
+// Copyright Epic Games, Inc. All Rights Reserved.
+
+#pragma once
+
+#include "MuT/NodeImage.h"
+
+
+namespace mu
+{
+
+    // Forward definitions
+    class NodeImageVariation;
+    typedef Ptr<NodeImageVariation> NodeImageVariationPtr;
+    typedef Ptr<const NodeImageVariation> NodeImageVariationPtrConst;
+
+
+    //!
+    //! \ingroup model
+    class MUTABLETOOLS_API NodeImageVariation : public NodeImage
+    {
+    public:
+        //-----------------------------------------------------------------------------------------
+        // Life cycle
+        //-----------------------------------------------------------------------------------------
+
+        NodeImageVariation();
+
+		void SerialiseWrapper(OutputArchive& arch) const override;
+		static void Serialise( const NodeImageVariation* pNode, OutputArchive& arch );
+        static NodeImageVariationPtr StaticUnserialise( InputArchive& arch );
+
+        //-----------------------------------------------------------------------------------------
+        // Node Interface
+        //-----------------------------------------------------------------------------------------
+
+        NodePtr Clone( NodeMapPtr pMap = 0 ) const override;
+
+        const NODE_TYPE* GetType() const override;
+        static const NODE_TYPE* GetStaticType();
+
+        virtual int GetInputCount() const override;
+        virtual Node* GetInputNode( int i ) const override;
+        void SetInputNode( int i, NodePtr pNode ) override;
+
+        //-----------------------------------------------------------------------------------------
+        // Own Interface
+        //-----------------------------------------------------------------------------------------
+
+        //!
+        void SetDefaultImage( NodeImage* Image );
+
+        //! Set the number of tags to consider in this variation
+        void SetVariationCount( int count );
+
+        //!
+        int GetVariationCount() const;
+
+        //! Set the tag or state name that will enable a specific vartiation
+        void SetVariationTag( int index, const char* strTag );
+
+        //!
+        void SetVariationImage( int index, NodeImage* Image );
+
+        //!}
+
+
+        //-----------------------------------------------------------------------------------------
+        // Interface pattern
+        //-----------------------------------------------------------------------------------------
+        class Private;
+        Private* GetPrivate() const;
+        Node::Private* GetBasePrivate() const override;
+
+    protected:
+        //! Forbidden. Manage with the Ptr<> template.
+        ~NodeImageVariation();
+
+    private:
+        Private* m_pD;
+    };
+
+
+} // namespace mu
diff --git a/Engine/Plugins/Experimental/Mutable/Source/MutableTools/Public/MuT/NodeLOD.h b/Engine/Plugins/Experimental/Mutable/Source/MutableTools/Public/MuT/NodeLOD.h
new file mode 100644
--- /dev/null
+++ b/Engine/Plugins/Experimental/Mutable/Source/MutableTools/Public/MuT/NodeLOD.h
@@ -0,0 +1,104 @@
+// Copyright Epic Games, Inc. All Rights Reserved.
+
+
+#pragma once
+
+#include "MuT/Node.h"
+
+
+namespace mu
+{
+
+	// Forward definitions
+	class NodeComponent;
+	typedef Ptr<NodeComponent> NodeComponentPtr;
+	typedef Ptr<const NodeComponent> NodeComponentPtrConst;
+
+	class NodeModifier;
+	typedef Ptr<NodeModifier> NodeModifierPtr;
+	typedef Ptr<const NodeModifier> NodeModifierPtrConst;
+
+	class NodeLOD;
+	typedef Ptr<NodeLOD> NodeLODPtr;
+	typedef Ptr<const NodeLOD> NodeLODPtrConst;
+
+
+	//! Node that creates a new level of detail by assembling several object components.
+	//! \ingroup model
+	class MUTABLETOOLS_API NodeLOD : public Node
+	{
+	public:
+
+		//-----------------------------------------------------------------------------------------
+		// Life cycle
+		//-----------------------------------------------------------------------------------------
+
+		NodeLOD();
+
+		void SerialiseWrapper(OutputArchive& arch) const override;
+		static void Serialise( const NodeLOD* pNode, OutputArchive& arch );
+		static NodeLODPtr StaticUnserialise( InputArchive& arch );
+
+
+		//-----------------------------------------------------------------------------------------
+		// Node Interface
+		//-----------------------------------------------------------------------------------------
+
+		NodePtr Clone( NodeMapPtr pMap = 0 ) const override;
+
+		const NODE_TYPE* GetType() const override;
+		static const NODE_TYPE* GetStaticType();
+
+		int GetInputCount() const override;
+		Node* GetInputNode( int i ) const override;
+		void SetInputNode( int i, NodePtr pNode ) override;
+
+		//-----------------------------------------------------------------------------------------
+		// Own Interface
+		//-----------------------------------------------------------------------------------------
+
+		//! Get the number of components in the LOD to create
+		int GetComponentCount() const;
+
+		//! Set the number of components in the LOD to create
+		void SetComponentCount( int );
+
+		//! Get the node generating one of the components of the LOD.
+		//! \param index index of the component from 0 to GetComponentCount()-1
+		NodeComponentPtr GetComponent( int index ) const;
+
+		//! Set the node generating one of the components of the LOD.
+		//! \param index index of the component from 0 to GetComponentCount()-1
+		void SetComponent( int index, NodeComponentPtr );
+
+		//! Set the number of modifiers in the LOD
+		void SetModifierCount(int);
+
+		//! Get the number of modifiers in the LOD
+		int GetModifierCount() const;
+
+		//! Set the node generating one of the modifiers of the LOD.
+		//! \param index index of the component from 0 to GetModifierCount()-1
+		void SetModifier(int index, NodeModifierPtr);
+
+		//-----------------------------------------------------------------------------------------
+		// Interface pattern
+		//-----------------------------------------------------------------------------------------
+		class Private;
+		Private* GetPrivate() const;
+        Node::Private* GetBasePrivate() const override;
+
+	protected:
+
+		//! Forbidden. Manage with the Ptr<> template.
+		~NodeLOD();
+
+	private:
+
+		Private* m_pD;
+
+	};
+
+
+
+}
diff --git a/Engine/Plugins/Experimental/Mutable/Source/MutableTools/Public/MuT/NodeLayout.h b/Engine/Plugins/Experimental/Mutable/Source/MutableTools/Public/MuT/NodeLayout.h
new file mode 100644
--- /dev/null
+++ b/Engine/Plugins/Experimental/Mutable/Source/MutableTools/Public/MuT/NodeLayout.h
@@ -0,0 +1,168 @@
+// Copyright Epic Games, Inc. All Rights Reserved.
+
+#pragma once
+
+#include "MuT/Node.h"
+#include "MuR/Mesh.h"
+#include "MuR/Layout.h"
+
+
+namespace mu
+{
+
+	//
+	class NodeLayout;
+	typedef Ptr<NodeLayout> NodeLayoutPtr;
+	typedef Ptr<const NodeLayout> NodeLayoutPtrConst;
+
+	//
+	class NodeLayoutBlocks;
+	typedef Ptr<NodeLayoutBlocks> NodeLayoutBlocksPtr;
+	typedef Ptr<const NodeLayoutBlocks> NodeLayoutBlocksPtrConst;
+
+
+	//! This node is used to define the texture layout for a texture coordinates channel of a
+	//! constant mesh.
+	//! \ingroup model
+	class MUTABLETOOLS_API NodeLayout : public Node
+	{
+	public:
+
+		// Possible subclasses
+		enum class EType : uint8
+		{
+			Blocks = 0,
+
+			None
+		};
+
+		//-----------------------------------------------------------------------------------------
+		// Life cycle
+		//-----------------------------------------------------------------------------------------
+
+		static void Serialise( const NodeLayout* pNode, OutputArchive& arch );
+		static NodeLayoutPtr StaticUnserialise( InputArchive& arch );
+
+
+		//-----------------------------------------------------------------------------------------
+		// Node Interface
+		//-----------------------------------------------------------------------------------------
+
+		const NODE_TYPE* GetType() const override;
+		static const NODE_TYPE* GetStaticType();
+
+		//-----------------------------------------------------------------------------------------
+		// Own Interface
+		//-----------------------------------------------------------------------------------------
+
+
+		//-----------------------------------------------------------------------------------------
+		// Interface pattern
+		//-----------------------------------------------------------------------------------------
+		class Private;
+
+	protected:
+
+		//! Forbidden. Manage with the Ptr<> template.
+		~NodeLayout() {}
+
+		//!
+		EType Type = EType::None;
+
+	};
+
+
+	//! This node is used to define the texture layout for a texture coordinates channel of a
+	//! constant mesh.
+	//! The blocks defined here will be set to the channel of the mesh were this node is connected.
+	//! \ingroup model
+	class MUTABLETOOLS_API NodeLayoutBlocks : public NodeLayout
+	{
+	public:
+
+		//-----------------------------------------------------------------------------------------
+		// Life cycle
+		//-----------------------------------------------------------------------------------------
+
+		NodeLayoutBlocks();
+
+		void SerialiseWrapper(OutputArchive& arch) const override;
+		static void Serialise( const NodeLayoutBlocks* pNode, OutputArchive& arch );
+		static NodeLayoutBlocksPtr StaticUnserialise( InputArchive& arch );
+
+
+		//-----------------------------------------------------------------------------------------
+		// Node Interface
+		//-----------------------------------------------------------------------------------------
+
+		NodePtr Clone( NodeMapPtr pMap = 0 ) const override;
+
+		const NODE_TYPE* GetType() const override;
+		static const NODE_TYPE* GetStaticType();
+
+		int GetInputCount() const override;
+		Node* GetInputNode( int i ) const override;
+		void SetInputNode( int i, NodePtr pNode ) override;
+
+		//-----------------------------------------------------------------------------------------
+		// Own Interface
+		//-----------------------------------------------------------------------------------------
+
+		//! Set the size of the grid where the blocks will be defined
+		void SetGridSize( int x, int y );
+
+		//! Set the maximum size of the grid where the blocks can be defined
+		void SetMaxGridSize(int x, int y);
+
+		//! Get the size of the grid where the blocks will be defined
+		void GetGridSize( int* pX, int* pY ) const;
+
+		//! Get the maximum size of the grid where the blocks can be defined
+		void GetMaxGridSize(int* pX, int* pY) const;
+
+		//! Set the number of blocks in the layout.
+		//! It keeps the current data as much as possible, and the new data is undefined.
+		void SetBlockCount( int );
+
+		//! Get the number of blocks in the layout
+		int GetBlockCount();
+
+		//! Set a block of the layout.
+		//! minx and miny refer to the lowest left corner of the block.
+        void SetBlock( int index, int minx, int miny, int sizex, int sizey );
+
+		//! Get a block of the layout.
+		//! minx and miny refer to the lowest left corner of the block.
+		void GetBlock(int index, int* minx, int* miny, int* sizex, int* sizey);
+
+		//! Set priority of a block.
+		void SetBlockPriority(int index, int priority);
+
+		//! Set the texture layout packing strategy 
+		void SetLayoutPackingStrategy(EPackStrategy strategy);
+
+		//! Generate the blocks of a layout using the UV of the meshes
+		static NodeLayoutBlocksPtr GenerateLayoutBlocks(const MeshPtr pMesh, int layoutIndex, int gridSizeX, int gridSizeY);
+
+
+		//-----------------------------------------------------------------------------------------
+		// Interface pattern
+		//-----------------------------------------------------------------------------------------
+		class Private;
+		Private* GetPrivate() const;
+        Node::Private* GetBasePrivate() const override;
+
+	protected:
+
+		//! Forbidden. Manage with the Ptr<> template.
+		~NodeLayoutBlocks();
+
+	private:
+
+		Private* m_pD;
+
+	};
+
+
+
+}
diff --git a/Engine/Plugins/Experimental/Mutable/Source/MutableTools/Public/MuT/NodeMesh.h b/Engine/Plugins/Experimental/Mutable/Source/MutableTools/Public/MuT/NodeMesh.h
new file mode 100644
--- /dev/null
+++ b/Engine/Plugins/Experimental/Mutable/Source/MutableTools/Public/MuT/NodeMesh.h
@@ -0,0 +1,81 @@
+// Copyright Epic Games, Inc. All Rights Reserved.
+
+#pragma once
+
+#include "MuT/Node.h"
+
+namespace mu
+{
+
+	// Forward definitions
+	class NodeMesh;
+	typedef Ptr<NodeMesh> NodeMeshPtr;
+	typedef Ptr<const NodeMesh> NodeMeshPtrConst;
+
+
+    //! %Base class of any node that outputs a mesh.
+	//! \ingroup model
+	class MUTABLETOOLS_API NodeMesh : public Node
+	{
+	public:
+
+		// Possible subclasses
+		enum class EType : uint8
+		{
+			Constant = 0,
+			Interpolate = 2,
+			Table = 4,
+			Subtract = 6,
+			Format = 7,
+			Tangents = 8,
+			Morph = 9,
+			MakeMorph = 10,
+			Switch = 11,
+			Fragment = 12,
+			Transform = 13,
+			ClipMorphPlane = 14,
+			ClipWithMesh = 15,
+			ApplyPose = 16,
+			Variation = 17,
+			GeometryOperation = 18,
+			Reshape = 19,
+			ClipDeform = 20,
+
+			None
+		};
+
+		//-----------------------------------------------------------------------------------------
+		// Life cycle
+		//-----------------------------------------------------------------------------------------
+
+		static void Serialise( const NodeMesh* pNode, OutputArchive& arch );
+		static NodeMeshPtr StaticUnserialise( InputArchive& arch );
+
+
+		//-----------------------------------------------------------------------------------------
+		// Node Interface
+		//-----------------------------------------------------------------------------------------
+
+        const NODE_TYPE* GetType() const override;
+		static const NODE_TYPE* GetStaticType();
+
+		inline EType GetMeshNodeType() const { return Type; }
+
+
+		//-----------------------------------------------------------------------------------------
+		// Interface pattern
+		//-----------------------------------------------------------------------------------------
+		class Private;
+
+	protected:
+
+		//! Forbidden. Manage with the Ptr<> template.
+		inline ~NodeMesh() {}
+
+		//!
+		EType Type = EType::None;
+	};
+
+
+}
+
diff --git a/Engine/Plugins/Experimental/Mutable/Source/MutableTools/Public/MuT/NodeMeshApplyPose.h b/Engine/Plugins/Experimental/Mutable/Source/MutableTools/Public/MuT/NodeMeshApplyPose.h
new file mode 100644
--- /dev/null
+++ b/Engine/Plugins/Experimental/Mutable/Source/MutableTools/Public/MuT/NodeMeshApplyPose.h
@@ -0,0 +1,81 @@
+// Copyright Epic Games, Inc. All Rights Reserved.
+
+#pragma once
+
+#include "MuT/NodeMesh.h"
+
+#include "MuR/Mesh.h"
+
+
+namespace mu
+{
+
+	// Forward definitions
+    class NodeMeshApplyPose;
+    using NodeMeshApplyPosePtr = Ptr<NodeMeshApplyPose> ;
+    using NodeMeshApplyPosePtrConst = Ptr<const NodeMeshApplyPose>;
+
+
+    //! Node that applies a pose to a mesh, baking it into the vertex data
+	//! \ingroup model
+    class MUTABLETOOLS_API NodeMeshApplyPose : public NodeMesh
+	{
+	public:
+
+		//-----------------------------------------------------------------------------------------
+		// Life cycle
+		//-----------------------------------------------------------------------------------------
+
+        NodeMeshApplyPose();
+
+		void SerialiseWrapper(OutputArchive& arch) const override;
+		static void Serialise( const NodeMeshApplyPose* pNode, OutputArchive& arch );
+        static NodeMeshApplyPosePtr StaticUnserialise( InputArchive& arch );
+
+
+		//-----------------------------------------------------------------------------------------
+		// Node Interface
+		//-----------------------------------------------------------------------------------------
+
+		NodePtr Clone( NodeMapPtr pMap = 0 ) const override;
+
+		const NODE_TYPE* GetType() const override;
+		static const NODE_TYPE* GetStaticType();
+
+		int GetInputCount() const override;
+		Node* GetInputNode( int i ) const override;
+		void SetInputNode( int i, NodePtr pNode ) override;
+
+		//-----------------------------------------------------------------------------------------
+		// Own Interface
+		//-----------------------------------------------------------------------------------------
+
+		//! Get the nodes generating the base mesh to be morphed
+		NodeMeshPtr GetBase() const;
+		void SetBase( NodeMeshPtr );
+
+        //! Get the nodes generating the pose to apply. A pose is represented with a mesh object
+        //! with no geometry: only bone matrices and names are relevant.
+        NodeMeshPtr GetPose() const;
+        void SetPose( NodeMeshPtr );
+
+		//-----------------------------------------------------------------------------------------
+		// Interface pattern
+		//-----------------------------------------------------------------------------------------
+		class Private;
+		Private* GetPrivate() const;
+        Node::Private* GetBasePrivate() const override;
+
+	protected:
+
+		//! Forbidden. Manage with the Ptr<> template.
+        ~NodeMeshApplyPose();
+
+	private:
+
+		Private* m_pD;
+
+	};
+
+
+}
diff --git a/Engine/Plugins/Experimental/Mutable/Source/MutableTools/Public/MuT/NodeMeshClipDeform.h b/Engine/Plugins/Experimental/Mutable/Source/MutableTools/Public/MuT/NodeMeshClipDeform.h
new file mode 100644
--- /dev/null
+++ b/Engine/Plugins/Experimental/Mutable/Source/MutableTools/Public/MuT/NodeMeshClipDeform.h
@@ -0,0 +1,77 @@
+// Copyright Epic Games, Inc. All Rights Reserved.
+
+#pragma once
+
+#include "MuT/NodeMesh.h"
+#include "MuT/NodeImage.h"
+
+#include "MuR/Mesh.h"
+#include "MuR/Image.h"
+
+
+namespace mu
+{
+
+	//! Node that morphs a base mesh with one or two weighted targets from a sequence.
+	//! \ingroup model
+	class MUTABLETOOLS_API NodeMeshClipDeform : public NodeMesh
+	{
+	public:
+
+		//-----------------------------------------------------------------------------------------
+		// Life cycle
+		//-----------------------------------------------------------------------------------------
+
+		NodeMeshClipDeform();
+
+		void SerialiseWrapper(OutputArchive& arch) const override;
+		static void Serialise( const NodeMeshClipDeform* pNode, OutputArchive& arch );
+		static Ptr<NodeMeshClipDeform> StaticUnserialise( InputArchive& arch );
+
+
+		//-----------------------------------------------------------------------------------------
+		// Node Interface
+		//-----------------------------------------------------------------------------------------
+
+		NodePtr Clone( NodeMapPtr pMap = 0 ) const override;
+
+		const NODE_TYPE* GetType() const override;
+		static const NODE_TYPE* GetStaticType();
+
+		int GetInputCount() const override;
+		Node* GetInputNode( int i ) const override;
+		void SetInputNode( int i, NodePtr pNode ) override;
+
+		//-----------------------------------------------------------------------------------------
+		// Own Interface
+		//-----------------------------------------------------------------------------------------
+
+		const NodeMeshPtr& GetBaseMesh() const;
+		void SetBaseMesh( const NodeMeshPtr& );
+
+		const NodeMeshPtr& GetClipShape() const;
+		void SetClipShape(const NodeMeshPtr&);
+
+		const NodeImagePtr& GetShapeWeights() const;
+		void SetShapeWeights(const NodeImagePtr&);
+
+
+        //-----------------------------------------------------------------------------------------
+		// Interface pattern
+		//-----------------------------------------------------------------------------------------
+		class Private;
+		Private* GetPrivate() const;
+        Node::Private* GetBasePrivate() const override;
+
+	protected:
+
+		//! Forbidden. Manage with the Ptr<> template.
+		~NodeMeshClipDeform();
+
+	private:
+
+		Private* m_pD;
+
+	};
+
+}
diff --git a/Engine/Plugins/Experimental/Mutable/Source/MutableTools/Public/MuT/NodeMeshClipMorphPlane.h b/Engine/Plugins/Experimental/Mutable/Source/MutableTools/Public/MuT/NodeMeshClipMorphPlane.h
new file mode 100644
--- /dev/null
+++ b/Engine/Plugins/Experimental/Mutable/Source/MutableTools/Public/MuT/NodeMeshClipMorphPlane.h
@@ -0,0 +1,94 @@
+// Copyright Epic Games, Inc. All Rights Reserved.
+
+#pragma once
+
+#include "MuT/NodeMesh.h"
+
+#include "MuR/Mesh.h"
+
+
+namespace mu
+{
+
+	// Forward definitions
+    class NodeMeshClipMorphPlane;
+    typedef Ptr<NodeMeshClipMorphPlane> NodeMeshClipMorphPlanePtr;
+    typedef Ptr<const NodeMeshClipMorphPlane> NodeMeshClipMorphPlaneConst;
+
+
+    //! This node applies a geometric transform represented by a 4x4 matrix to a mesh
+	//! \ingroup model
+    class MUTABLETOOLS_API NodeMeshClipMorphPlane : public NodeMesh
+	{
+	public:
+
+		//-----------------------------------------------------------------------------------------
+		// Life cycle
+		//-----------------------------------------------------------------------------------------
+
+		NodeMeshClipMorphPlane();
+
+		void SerialiseWrapper(OutputArchive& arch) const override;
+		static void Serialise( const NodeMeshClipMorphPlane* pNode, OutputArchive& arch );
+        static NodeMeshClipMorphPlanePtr StaticUnserialise( InputArchive& arch );
+
+		//-----------------------------------------------------------------------------------------
+		// Node Interface
+		//-----------------------------------------------------------------------------------------
+
+		virtual NodePtr Clone( NodeMapPtr pMap = 0 ) const;
+
+		virtual const NODE_TYPE* GetType() const;
+		static const NODE_TYPE* GetStaticType();
+
+		virtual int GetInputCount() const;
+		virtual const char* GetInputName( int i ) const;
+		virtual const NODE_TYPE* GetInputType( int i ) const;
+		virtual Node* GetInputNode( int i ) const;
+		virtual void SetInputNode( int i, NodePtr pNode );
+
+		//-----------------------------------------------------------------------------------------
+		// Own Interface
+		//-----------------------------------------------------------------------------------------
+
+		//! Source mesh to be clipped and morphed
+        NodeMeshPtr GetSource() const;
+		void SetSource( NodeMesh* );
+
+        //! \param 
+		void SetPlane(float centerX, float centerY, float centerZ, float normalX, float normalY, float normalZ);
+		void SetParams(float dist, float factor);
+		void SetMorphEllipse(float radius1, float radius2, float rotation);
+
+		//! Define an axis-aligned box that will select the vertices to be morphed.
+		//! If this is not set, or the size of the box is 0, all vertices will be affected.
+		//! Only one of Box or Bone Hierarchy can be used (the last one set)
+		void SetVertexSelectionBox(float centerX, float centerY, float centerZ, float radiusX, float radiusY, float radiusZ);
+
+		//! Define the root bone of the subhierarchy of the mesh that will be affected.
+		//! Only one of Box or Bone Hierarchy can be used (the last one set)
+		void SetVertexSelectionBone(const char* strBoneName, float maxEffectRadius);
+
+		//! Add a tag to the clip morph operation, which will only affect surfaces with the same tag
+		void AddTag(const char* tagName);
+
+		//-----------------------------------------------------------------------------------------
+		// Interface pattern
+		//-----------------------------------------------------------------------------------------
+		class Private;
+		Private* GetPrivate() const;
+		Node::Private* GetBasePrivate() const;
+
+	protected:
+
+		//! Forbidden. Manage with the Ptr<> template.
+        ~NodeMeshClipMorphPlane();
+
+	private:
+
+		Private* m_pD;
+
+	};
+
+
+}
diff --git a/Engine/Plugins/Experimental/Mutable/Source/MutableTools/Public/MuT/NodeMeshClipWithMesh.h b/Engine/Plugins/Experimental/Mutable/Source/MutableTools/Public/MuT/NodeMeshClipWithMesh.h
new file mode 100644
--- /dev/null
+++ b/Engine/Plugins/Experimental/Mutable/Source/MutableTools/Public/MuT/NodeMeshClipWithMesh.h
@@ -0,0 +1,83 @@
+// Copyright Epic Games, Inc. All Rights Reserved.
+
+#pragma once
+
+#include "MuT/NodeMesh.h"
+
+#include "MuR/Mesh.h"
+
+
+namespace mu
+{
+
+	// Forward definitions
+    class NodeMeshClipWithMesh;
+    typedef Ptr<NodeMeshClipWithMesh> NodeMeshClipWithMeshPtr;
+    typedef Ptr<const NodeMeshClipWithMesh> NodeMeshClipWithMeshConst;
+
+
+    //! This node applies a geometric transform represented by a 4x4 matrix to a mesh
+	//! \ingroup model
+    class MUTABLETOOLS_API NodeMeshClipWithMesh : public NodeMesh
+	{
+	public:
+
+		//-----------------------------------------------------------------------------------------
+		// Life cycle
+		//-----------------------------------------------------------------------------------------
+
+		NodeMeshClipWithMesh();
+
+		void SerialiseWrapper(OutputArchive& arch) const override;
+		static void Serialise( const NodeMeshClipWithMesh* pNode, OutputArchive& arch );
+        static NodeMeshClipWithMeshPtr StaticUnserialise( InputArchive& arch );
+
+		//-----------------------------------------------------------------------------------------
+		// Node Interface
+		//-----------------------------------------------------------------------------------------
+
+		virtual NodePtr Clone( NodeMapPtr pMap = 0 ) const;
+
+		virtual const NODE_TYPE* GetType() const;
+		static const NODE_TYPE* GetStaticType();
+
+		virtual int GetInputCount() const;
+		virtual const char* GetInputName( int i ) const;
+		virtual const NODE_TYPE* GetInputType( int i ) const;
+		virtual Node* GetInputNode( int i ) const;
+		virtual void SetInputNode( int i, NodePtr pNode );
+
+		//-----------------------------------------------------------------------------------------
+		// Own Interface
+		//-----------------------------------------------------------------------------------------
+
+		//! Source mesh to be clipped
+        NodeMeshPtr GetSource() const;
+		void SetSource(NodeMesh*);
+
+        //! \param 
+		void SetClipMesh(NodeMesh*);
+
+		//! Add a tag to the clip morph operation, which will only affect surfaces with the same tag
+		void AddTag(const char* tagName);
+
+		//-----------------------------------------------------------------------------------------
+		// Interface pattern
+		//-----------------------------------------------------------------------------------------
+		class Private;
+		Private* GetPrivate() const;
+		Node::Private* GetBasePrivate() const;
+
+	protected:
+
+		//! Forbidden. Manage with the Ptr<> template.
+        ~NodeMeshClipWithMesh();
+
+	private:
+
+		Private* m_pD;
+
+	};
+
+
+}
diff --git a/Engine/Plugins/Experimental/Mutable/Source/MutableTools/Public/MuT/NodeMeshConstant.h b/Engine/Plugins/Experimental/Mutable/Source/MutableTools/Public/MuT/NodeMeshConstant.h
new file mode 100644
--- /dev/null
+++ b/Engine/Plugins/Experimental/Mutable/Source/MutableTools/Public/MuT/NodeMeshConstant.h
@@ -0,0 +1,92 @@
+// Copyright Epic Games, Inc. All Rights Reserved.
+
+#pragma once
+
+#include "MuT/NodeMesh.h"
+
+#include "MuR/Mesh.h"
+#include "MuR/Types.h"
+
+
+namespace mu
+{
+
+	// Forward definitions
+	class NodeMeshConstant;
+	typedef Ptr<NodeMeshConstant> NodeMeshConstantPtr;
+	typedef Ptr<const NodeMeshConstant> NodeMeshConstantPtrConst;
+
+	class NodeLayout;
+	typedef Ptr<NodeLayout> NodeLayoutPtr;
+	typedef Ptr<const NodeLayout> NodeLayoutPtrConst;
+
+
+	//! Node that outputs a constant mesh.
+	//! It allows to define the layouts for the texture channels of the constant mesh
+	//! \ingroup model
+	class MUTABLETOOLS_API NodeMeshConstant : public NodeMesh
+	{
+	public:
+
+		//-----------------------------------------------------------------------------------------
+		// Life cycle
+		//-----------------------------------------------------------------------------------------
+
+		NodeMeshConstant();
+
+		void SerialiseWrapper(OutputArchive& arch) const override;
+		static void Serialise( const NodeMeshConstant* pNode, OutputArchive& arch );
+		static NodeMeshConstantPtr StaticUnserialise( InputArchive& arch );
+
+
+		//-----------------------------------------------------------------------------------------
+		// Node Interface
+		//-----------------------------------------------------------------------------------------
+
+		NodePtr Clone( NodeMapPtr pMap = 0 ) const override;
+
+		const NODE_TYPE* GetType() const override;
+		static const NODE_TYPE* GetStaticType();
+
+		int GetInputCount() const override;
+		Node* GetInputNode( int i ) const override;
+		void SetInputNode( int i, NodePtr pNode ) override;
+
+		//-----------------------------------------------------------------------------------------
+		// Own Interface
+		//-----------------------------------------------------------------------------------------
+
+		//! Get the constant mesh that will be returned.
+		MeshPtr GetValue() const;
+
+		//! Set the constant mesh that will be returned.
+		void SetValue( MeshPtr );
+
+		//! Get the number of layouts defined in this mesh.
+		int GetLayoutCount() const;
+		void SetLayoutCount( int );
+
+		//! Get the node defining a layout of the returned mesh.
+		NodeLayoutPtr GetLayout( int index ) const;
+		void SetLayout( int index, NodeLayoutPtr );
+
+		//-----------------------------------------------------------------------------------------
+		// Interface pattern
+		//-----------------------------------------------------------------------------------------
+		class Private;
+		Private* GetPrivate() const;
+        Node::Private* GetBasePrivate() const override;
+
+	protected:
+
+		//! Forbidden. Manage with the Ptr<> template.
+		~NodeMeshConstant();
+
+	private:
+
+		Private* m_pD;
+
+	};
+
+
+}
diff --git a/Engine/Plugins/Experimental/Mutable/Source/MutableTools/Public/MuT/NodeMeshFormat.h b/Engine/Plugins/Experimental/Mutable/Source/MutableTools/Public/MuT/NodeMeshFormat.h
new file mode 100644
--- /dev/null
+++ b/Engine/Plugins/Experimental/Mutable/Source/MutableTools/Public/MuT/NodeMeshFormat.h
@@ -0,0 +1,94 @@
+// Copyright Epic Games, Inc. All Rights Reserved.
+
+#pragma once
+
+#include "MuT/NodeMesh.h"
+
+#include "MuR/Mesh.h"
+
+
+namespace mu
+{
+
+	// Forward definitions
+	class NodeMeshFormat;
+	typedef Ptr<NodeMeshFormat> NodeMeshFormatPtr;
+	typedef Ptr<const NodeMeshFormat> NodeMeshFormatPtrConst;
+
+
+	//! This node can change the buffer formats of a mesh vertices, indices and faces.
+	//! \ingroup model
+	class MUTABLETOOLS_API NodeMeshFormat : public NodeMesh
+	{
+	public:
+
+		//-----------------------------------------------------------------------------------------
+		// Life cycle
+		//-----------------------------------------------------------------------------------------
+
+		NodeMeshFormat();
+
+		void SerialiseWrapper(OutputArchive& arch) const override;
+		static void Serialise( const NodeMeshFormat* pNode, OutputArchive& arch );
+		static NodeMeshFormatPtr StaticUnserialise( InputArchive& arch );
+
+		//-----------------------------------------------------------------------------------------
+		// Node Interface
+		//-----------------------------------------------------------------------------------------
+
+		NodePtr Clone( NodeMapPtr pMap = 0 ) const override;
+
+		const NODE_TYPE* GetType() const override;
+		static const NODE_TYPE* GetStaticType();
+
+		int GetInputCount() const override;
+		Node* GetInputNode( int i ) const override;
+		void SetInputNode( int i, NodePtr pNode ) override;
+
+		//-----------------------------------------------------------------------------------------
+		// Own Interface
+		//-----------------------------------------------------------------------------------------
+
+		//! Source mesh to be re-formatted
+		NodeMeshPtr GetSource() const;
+		void SetSource( NodeMesh* );
+
+		//! Get and set the MeshBufferSet that defines the new format for the mesh vertices. These
+		//! buffers don't really contain any data (they have 0 elements) but they define the
+		//! structure. If this is null, the vertex buffers will not be changed.
+		FMeshBufferSet& GetVertexBuffers();
+
+		//! Get and set the MeshBufferSet that defines the new format for the mesh indices. These
+		//! buffers don't really contain any data (they have 0 elements) but they define the
+		//! structure. If this is null, the vertex buffers will not be changed.
+		FMeshBufferSet& GetIndexBuffers();
+
+		//! Get and set the MeshBufferSet that defines the new format for the mesh faces. These
+		//! buffers don't really contain any data (they have 0 elements) but they define the
+		//! structure. If this is null, the vertex buffers will not be changed.
+		FMeshBufferSet& GetFaceBuffers();
+
+		//! Enable or disable the rebuilding of the tangent space of the mesh.
+		//! By default it is disabled.
+		void SetRebuildTangents( bool );
+
+		//-----------------------------------------------------------------------------------------
+		// Interface pattern
+		//-----------------------------------------------------------------------------------------
+		class Private;
+		Private* GetPrivate() const;
+        Node::Private* GetBasePrivate() const override;
+
+	protected:
+
+		//! Forbidden. Manage with the Ptr<> template.
+		~NodeMeshFormat();
+
+	private:
+
+		Private* m_pD;
+
+	};
+
+
+}
diff --git a/Engine/Plugins/Experimental/Mutable/Source/MutableTools/Public/MuT/NodeMeshFragment.h b/Engine/Plugins/Experimental/Mutable/Source/MutableTools/Public/MuT/NodeMeshFragment.h
new file mode 100644
--- /dev/null
+++ b/Engine/Plugins/Experimental/Mutable/Source/MutableTools/Public/MuT/NodeMeshFragment.h
@@ -0,0 +1,97 @@
+// Copyright Epic Games, Inc. All Rights Reserved.
+
+#pragma once
+
+#include "MuT/NodeMesh.h"
+
+#include "MuR/Mesh.h"
+
+
+namespace mu
+{
+
+	// Forward definitions
+	class NodeMeshFragment;
+	typedef Ptr<NodeMeshFragment> NodeMeshFragmentPtr;
+	typedef Ptr<const NodeMeshFragment> NodeMeshFragmentPtrConst;
+
+
+	//! This node selects a fragment of a mesh, by selecting some of its layout blocks.
+	//! \ingroup model
+	class MUTABLETOOLS_API NodeMeshFragment : public NodeMesh
+	{
+	public:
+
+		//-----------------------------------------------------------------------------------------
+		// Life cycle
+		//-----------------------------------------------------------------------------------------
+
+		NodeMeshFragment();
+
+		void SerialiseWrapper(OutputArchive& arch) const override;
+		static void Serialise( const NodeMeshFragment* pNode, OutputArchive& arch );
+		static NodeMeshFragmentPtr StaticUnserialise( InputArchive& arch );
+
+
+		//-----------------------------------------------------------------------------------------
+		// Node Interface
+		//-----------------------------------------------------------------------------------------
+
+		NodePtr Clone( NodeMapPtr pMap = 0 ) const override;
+
+		const NODE_TYPE* GetType() const override;
+		static const NODE_TYPE* GetStaticType();
+
+		int GetInputCount() const override;
+		Node* GetInputNode( int i ) const override;
+		void SetInputNode( int i, NodePtr pNode ) override;
+
+		//-----------------------------------------------------------------------------------------
+		// Own Interface
+		//-----------------------------------------------------------------------------------------
+
+        typedef enum {
+            FT_LAYOUT_BLOCKS = 0,
+            FT_FACE_GROUP
+        } FRAGMENT_TYPE;
+
+        //! Get the type of mesh fragment to extract.
+        FRAGMENT_TYPE GetFragmentType() const;
+        void SetFragmentType( FRAGMENT_TYPE type );
+
+		//!
+		NodeMeshPtr GetMesh() const;
+		void SetMesh( NodeMeshPtr );
+
+		//! Get the index of the layout to use to extract blocks.
+        int GetLayoutOrGroup() const;
+        void SetLayoutOrGroup( int layoutIndex );
+
+		//! Set the number of layout blocks to extract.
+		void SetBlockCount( int );
+
+		//! Get one of the layout block indices
+		int GetBlock( int i ) const;
+		void SetBlock( int i, int blockIndex );
+
+
+		//-----------------------------------------------------------------------------------------
+		// Interface pattern
+		//-----------------------------------------------------------------------------------------
+		class Private;
+		Private* GetPrivate() const;
+        Node::Private* GetBasePrivate() const override;
+
+	protected:
+
+		//! Forbidden. Manage with the Ptr<> template.
+		~NodeMeshFragment();
+
+	private:
+
+		Private* m_pD;
+
+	};
+
+
+}
diff --git a/Engine/Plugins/Experimental/Mutable/Source/MutableTools/Public/MuT/NodeMeshGeometryOperation.h b/Engine/Plugins/Experimental/Mutable/Source/MutableTools/Public/MuT/NodeMeshGeometryOperation.h
new file mode 100644
--- /dev/null
+++ b/Engine/Plugins/Experimental/Mutable/Source/MutableTools/Public/MuT/NodeMeshGeometryOperation.h
@@ -0,0 +1,86 @@
+// Copyright Epic Games, Inc. All Rights Reserved.
+
+#pragma once
+
+#include "MuT/NodeMesh.h"
+#include "MuR/Mesh.h"
+
+
+namespace mu
+{
+
+	// Forward definitions
+	class NodeScalar;
+	typedef Ptr<NodeScalar> NodeScalarPtr;
+	typedef Ptr<const NodeScalar> NodeScalarPtrConst;
+
+	class NodeMeshGeometryOperation;
+	typedef Ptr<NodeMeshGeometryOperation> NodeMeshGeometryOperationPtr;
+	typedef Ptr<const NodeMeshGeometryOperation> NodeMeshGeometryOperationPtrConst;
+
+
+	//! Node that morphs a base mesh with one or two weighted targets from a sequence.
+	//! \ingroup model
+	class MUTABLETOOLS_API NodeMeshGeometryOperation : public NodeMesh
+	{
+	public:
+
+		//-----------------------------------------------------------------------------------------
+		// Life cycle
+		//-----------------------------------------------------------------------------------------
+
+		NodeMeshGeometryOperation();
+
+		void SerialiseWrapper(OutputArchive& arch) const override;
+		static void Serialise( const NodeMeshGeometryOperation* pNode, OutputArchive& arch );
+		static NodeMeshGeometryOperationPtr StaticUnserialise( InputArchive& arch );
+
+
+		//-----------------------------------------------------------------------------------------
+		// Node Interface
+		//-----------------------------------------------------------------------------------------
+
+		NodePtr Clone( NodeMapPtr pMap = 0 ) const override;
+
+		const NODE_TYPE* GetType() const override;
+		static const NODE_TYPE* GetStaticType();
+
+		int GetInputCount() const override;
+		Node* GetInputNode( int i ) const override;
+		void SetInputNode( int i, NodePtr pNode ) override;
+
+		//-----------------------------------------------------------------------------------------
+		// Own Interface
+		//-----------------------------------------------------------------------------------------
+
+		const NodeMeshPtr& GetMeshA() const;
+		void SetMeshA( const NodeMeshPtr& );
+
+		const NodeMeshPtr& GetMeshB() const;
+		void SetMeshB(const NodeMeshPtr&);
+
+		const NodeScalarPtr& GetScalarA() const;
+		void SetScalarA(const NodeScalarPtr&);
+
+		const NodeScalarPtr& GetScalarB() const;
+		void SetScalarB(const NodeScalarPtr&);
+
+        //-----------------------------------------------------------------------------------------
+		// Interface pattern
+		//-----------------------------------------------------------------------------------------
+		class Private;
+		Private* GetPrivate() const;
+        Node::Private* GetBasePrivate() const override;
+
+	protected:
+
+		//! Forbidden. Manage with the Ptr<> template.
+		~NodeMeshGeometryOperation();
+
+	private:
+
+		Private* m_pD;
+
+	};
+
+}
diff --git a/Engine/Plugins/Experimental/Mutable/Source/MutableTools/Public/MuT/NodeMeshInterpolate.h b/Engine/Plugins/Experimental/Mutable/Source/MutableTools/Public/MuT/NodeMeshInterpolate.h
new file mode 100644
--- /dev/null
+++ b/Engine/Plugins/Experimental/Mutable/Source/MutableTools/Public/MuT/NodeMeshInterpolate.h
@@ -0,0 +1,102 @@
+// Copyright Epic Games, Inc. All Rights Reserved.
+
+#pragma once
+
+#include "MuT/NodeMesh.h"
+
+#include "MuR/Mesh.h"
+
+
+namespace mu
+{
+
+	// Forward definitions
+	class NodeScalar;
+	typedef Ptr<NodeScalar> NodeScalarPtr;
+	typedef Ptr<const NodeScalar> NodeScalarPtrConst;
+
+	class NodeMeshInterpolate;
+	typedef Ptr<NodeMeshInterpolate> NodeMeshInterpolatePtr;
+	typedef Ptr<const NodeMeshInterpolate> NodeMeshInterpolatePtrConst;
+
+
+	//! Node that morphs (interpolates linearly) between several meshes based on a weight.
+	//! The different input images are uniformly distributed in the 0 to 1 range:
+	//! - if two meshes are set, the first one is 0.0 and the second one is 1.0
+	//! - if three meshes are set, the first one is 0.0 the second one is 0.5 and the third one 1.0
+	//! regardless of the input slots used: If B and C are set, but not A, B will be at weight 0
+	//! and C will be at weight 1
+	//! \ingroup model
+	class MUTABLETOOLS_API NodeMeshInterpolate : public NodeMesh
+	{
+	public:
+
+		//-----------------------------------------------------------------------------------------
+		// Life cycle
+		//-----------------------------------------------------------------------------------------
+
+		NodeMeshInterpolate();
+
+		void SerialiseWrapper(OutputArchive& arch) const override;
+		static void Serialise( const NodeMeshInterpolate* pNode, OutputArchive& arch );
+		static NodeMeshInterpolatePtr StaticUnserialise( InputArchive& arch );
+
+
+		//-----------------------------------------------------------------------------------------
+		// Node Interface
+		//-----------------------------------------------------------------------------------------
+
+		NodePtr Clone( NodeMapPtr pMap = 0 ) const override;
+
+		const NODE_TYPE* GetType() const override;
+		static const NODE_TYPE* GetStaticType();
+
+		int GetInputCount() const override;
+		Node* GetInputNode( int i ) const override;
+		void SetInputNode( int i, NodePtr pNode ) override;
+
+		//-----------------------------------------------------------------------------------------
+		// Own Interface
+		//-----------------------------------------------------------------------------------------
+
+		//! Get the node generating the weight of the interpolation.
+		NodeScalarPtr GetFactor() const;
+		void SetFactor( NodeScalarPtr );
+
+		//! Set the number of target meshes. It will keep the currently set targets and initialise
+		//! the new ones as 0.
+		void SetTargetCount( int );
+		int GetTargetCount() const;
+
+		//! Get the node generating the several mesh targets of the interpolation.
+		NodeMeshPtr GetTarget( int t ) const;
+		void SetTarget( int t, NodeMeshPtr );
+
+		//! Set the number of channels of the mesh that will be interpolated. It will keep the
+		//! currently set targets and initialise the new ones as null. The channels that are not
+		//! interpolated will be copied from the first target.
+		void SetChannelCount( int );
+
+		//! Set a channel to interpolate from the mesh.
+		void SetChannel( int index, MESH_BUFFER_SEMANTIC semantic, int semanticIndex );
+
+		//-----------------------------------------------------------------------------------------
+		// Interface pattern
+		//-----------------------------------------------------------------------------------------
+		class Private;
+		Private* GetPrivate() const;
+        Node::Private* GetBasePrivate() const override;
+
+	protected:
+
+		//! Forbidden. Manage with the Ptr<> template.
+		~NodeMeshInterpolate();
+
+	private:
+
+		Private* m_pD;
+
+	};
+
+
+}
diff --git a/Engine/Plugins/Experimental/Mutable/Source/MutableTools/Public/MuT/NodeMeshMakeMorph.h b/Engine/Plugins/Experimental/Mutable/Source/MutableTools/Public/MuT/NodeMeshMakeMorph.h
new file mode 100644
--- /dev/null
+++ b/Engine/Plugins/Experimental/Mutable/Source/MutableTools/Public/MuT/NodeMeshMakeMorph.h
@@ -0,0 +1,77 @@
+// Copyright Epic Games, Inc. All Rights Reserved.
+
+#pragma once
+
+#include "MuT/NodeMesh.h"
+
+
+namespace mu
+{
+
+	// Forward definitions
+    class NodeMeshMakeMorph;
+    typedef Ptr<NodeMeshMakeMorph> NodeMeshMakeMorphPtr;
+    typedef Ptr<const NodeMeshMakeMorph> NodeMeshMakeMorphPtrConst;
+
+
+	//! This node removes from a mesh A all the faces that are also part of a mesh B.
+	//! \ingroup model
+    class MUTABLETOOLS_API NodeMeshMakeMorph : public NodeMesh
+	{
+	public:
+
+		//-----------------------------------------------------------------------------------------
+		// Life cycle
+		//-----------------------------------------------------------------------------------------
+
+        NodeMeshMakeMorph();
+
+		void SerialiseWrapper(OutputArchive& arch) const override;
+		static void Serialise( const NodeMeshMakeMorph* pNode, OutputArchive& arch );
+        static NodeMeshMakeMorphPtr StaticUnserialise( InputArchive& arch );
+
+		//-----------------------------------------------------------------------------------------
+		// Node Interface
+		//-----------------------------------------------------------------------------------------
+
+		NodePtr Clone( NodeMapPtr pMap = 0 ) const override;
+
+		const NODE_TYPE* GetType() const override;
+		static const NODE_TYPE* GetStaticType();
+
+		int GetInputCount() const override;
+		Node* GetInputNode( int i ) const override;
+		void SetInputNode( int i, NodePtr pNode ) override;
+
+		//-----------------------------------------------------------------------------------------
+		// Own Interface
+		//-----------------------------------------------------------------------------------------
+
+		//!
+		NodeMeshPtr GetBase() const;
+		void SetBase( NodeMesh* p );
+
+		//!
+		NodeMeshPtr GetTarget() const;
+		void SetTarget( NodeMesh* p );
+
+		//-----------------------------------------------------------------------------------------
+		// Interface pattern
+		//-----------------------------------------------------------------------------------------
+		class Private;
+		Private* GetPrivate() const;
+        Node::Private* GetBasePrivate() const override;
+
+	protected:
+
+		//! Forbidden. Manage with the Ptr<> template.
+        ~NodeMeshMakeMorph();
+
+	private:
+
+		Private* m_pD;
+
+	};
+
+
+}
diff --git a/Engine/Plugins/Experimental/Mutable/Source/MutableTools/Public/MuT/NodeMeshMorph.h b/Engine/Plugins/Experimental/Mutable/Source/MutableTools/Public/MuT/NodeMeshMorph.h
new file mode 100644
--- /dev/null
+++ b/Engine/Plugins/Experimental/Mutable/Source/MutableTools/Public/MuT/NodeMeshMorph.h
@@ -0,0 +1,116 @@
+// Copyright Epic Games, Inc. All Rights Reserved.
+
+#pragma once
+
+#include "MuT/NodeMesh.h"
+
+#include "MuR/Mesh.h"
+
+
+namespace mu
+{
+
+	// Forward definitions
+	class NodeScalar;
+	typedef Ptr<NodeScalar> NodeScalarPtr;
+	typedef Ptr<const NodeScalar> NodeScalarPtrConst;
+
+	class NodeMeshMorph;
+	typedef Ptr<NodeMeshMorph> NodeMeshMorphPtr;
+	typedef Ptr<const NodeMeshMorph> NodeMeshMorphPtrConst;
+
+
+	//! Node that morphs a base mesh with one or two weighted targets from a sequence.
+	//! \ingroup model
+	class MUTABLETOOLS_API NodeMeshMorph : public NodeMesh
+	{
+	public:
+
+		//-----------------------------------------------------------------------------------------
+		// Life cycle
+		//-----------------------------------------------------------------------------------------
+
+		NodeMeshMorph();
+
+		void SerialiseWrapper(OutputArchive& arch) const override;
+		static void Serialise( const NodeMeshMorph* pNode, OutputArchive& arch );
+		static NodeMeshMorphPtr StaticUnserialise( InputArchive& arch );
+
+
+		//-----------------------------------------------------------------------------------------
+		// Node Interface
+		//-----------------------------------------------------------------------------------------
+
+		NodePtr Clone( NodeMapPtr pMap = 0 ) const override;
+
+		const NODE_TYPE* GetType() const override;
+		static const NODE_TYPE* GetStaticType();
+
+		int GetInputCount() const override;
+		Node* GetInputNode( int i ) const override;
+		void SetInputNode( int i, NodePtr pNode ) override;
+
+		//-----------------------------------------------------------------------------------------
+		// Own Interface
+		//-----------------------------------------------------------------------------------------
+
+		//! Get the node generating the weight used to select (and combine) the morphs to apply.
+		NodeScalarPtr GetFactor() const;
+		void SetFactor( NodeScalarPtr );
+
+		//! Get the nodes generating the base mesh to be morphed
+		NodeMeshPtr GetBase() const;
+		void SetBase( NodeMeshPtr );
+
+		//! Set the number of target meshes. It will keep the currently set targets and initialise
+		//! the new ones as 0.
+		void SetMorphCount( int );
+		int GetMorphCount() const;
+
+		//! Get the nodes generating the several morph targets
+		NodeMeshPtr GetMorph( int t ) const;
+		void SetMorph( int t, NodeMeshPtr );
+
+		//! Set the domain of the vertex indices in the morph connections. By default they are 
+		//! absolute as generated by the MakeMeshMorph node. But if they are provided externally, 
+		//! they may be relative to the main mesh connected in the Base node.
+        void SetMorphIndicesAreRelative( bool relative );
+
+		/** Also deform the mesh skeleton. Disabled by default. */
+		void SetReshapeSkeleton(bool);	
+
+		/** Deform all the bones of the skeleton */
+		void SetDeformAllBones(bool);
+
+		/** Deform all the physiscs bodies */
+		void SetDeformAllPhysics(bool);
+
+		/** Deform Mesh Physics Volumes */
+		void SetReshapePhysicsVolumes(bool);
+
+		/** Sets the number of bones that will be deform */
+		void AddBoneToDeform(const char* BoneName);	
+
+		/** Add a Physics Body to deform */
+		void AddPhysicsBodyToDeform(const char* BoneName);
+
+        //-----------------------------------------------------------------------------------------
+		// Interface pattern
+		//-----------------------------------------------------------------------------------------
+		class Private;
+		Private* GetPrivate() const;
+        Node::Private* GetBasePrivate() const override;
+
+	protected:
+
+		//! Forbidden. Manage with the Ptr<> template.
+		~NodeMeshMorph();
+
+	private:
+
+		Private* m_pD;
+
+	};
+
+
+}
diff --git a/Engine/Plugins/Experimental/Mutable/Source/MutableTools/Public/MuT/NodeMeshReshape.h b/Engine/Plugins/Experimental/Mutable/Source/MutableTools/Public/MuT/NodeMeshReshape.h
new file mode 100644
--- /dev/null
+++ b/Engine/Plugins/Experimental/Mutable/Source/MutableTools/Public/MuT/NodeMeshReshape.h
@@ -0,0 +1,94 @@
+// Copyright Epic Games, Inc. All Rights Reserved.
+
+#pragma once
+
+#include "MuT/NodeMesh.h"
+#include "MuR/Mesh.h"
+
+
+namespace mu
+{
+
+	//! Node that morphs a base mesh with one or two weighted targets from a sequence.
+	//! \ingroup model
+	class MUTABLETOOLS_API NodeMeshReshape : public NodeMesh
+	{
+	public:
+
+		//-----------------------------------------------------------------------------------------
+		// Life cycle
+		//-----------------------------------------------------------------------------------------
+
+		NodeMeshReshape();
+
+		void SerialiseWrapper(OutputArchive& arch) const override;
+		static void Serialise( const NodeMeshReshape* pNode, OutputArchive& arch );
+		static Ptr<NodeMeshReshape> StaticUnserialise( InputArchive& arch );
+
+
+		//-----------------------------------------------------------------------------------------
+		// Node Interface
+		//-----------------------------------------------------------------------------------------
+
+		NodePtr Clone( NodeMapPtr pMap = 0 ) const override;
+
+		const NODE_TYPE* GetType() const override;
+		static const NODE_TYPE* GetStaticType();
+
+		int GetInputCount() const override;
+		Node* GetInputNode( int i ) const override;
+		void SetInputNode( int i, NodePtr pNode ) override;
+
+		//-----------------------------------------------------------------------------------------
+		// Own Interface
+		//-----------------------------------------------------------------------------------------
+
+		const NodeMeshPtr& GetBaseMesh() const;
+		void SetBaseMesh( const NodeMeshPtr& );
+
+		const NodeMeshPtr& GetBaseShape() const;
+		void SetBaseShape(const NodeMeshPtr&);
+
+		const NodeMeshPtr& GetTargetShape() const;
+		void SetTargetShape(const NodeMeshPtr&);
+
+		/** Also deform the mesh skeleton. Disabled by default. */
+		void SetReshapeSkeleton(bool);
+		
+		/** Search for clusters of rigid parts and not deform them. */
+		void SetEnableRigidParts(bool);
+
+		/** Deform all the bones of the skeleton */
+		void SetDeformAllBones(bool);
+
+		/** Deform all the physiscs bodies */
+		void SetDeformAllPhysics(bool);
+
+		/** Deform Mesh Physics Volumes */
+		void SetReshapePhysicsVolumes(bool);
+
+		/** Sets the number of bones that will be deform */
+		void AddBoneToDeform(const char* BoneName);
+	
+		/** Add a Physics Body to deform */
+		void AddPhysicsBodyToDeform(const char* BoneName);
+        
+		//-----------------------------------------------------------------------------------------
+		// Interface pattern
+		//-----------------------------------------------------------------------------------------
+		class Private;
+		Private* GetPrivate() const;
+        Node::Private* GetBasePrivate() const override;
+
+	protected:
+
+		//! Forbidden. Manage with the Ptr<> template.
+		~NodeMeshReshape();
+
+	private:
+
+		Private* m_pD;
+
+	};
+
+}
diff --git a/Engine/Plugins/Experimental/Mutable/Source/MutableTools/Public/MuT/NodeMeshSubtract.h b/Engine/Plugins/Experimental/Mutable/Source/MutableTools/Public/MuT/NodeMeshSubtract.h
new file mode 100644
--- /dev/null
+++ b/Engine/Plugins/Experimental/Mutable/Source/MutableTools/Public/MuT/NodeMeshSubtract.h
@@ -0,0 +1,77 @@
+// Copyright Epic Games, Inc. All Rights Reserved.
+
+#pragma once
+
+#include "MuT/NodeMesh.h"
+
+
+namespace mu
+{
+
+	// Forward definitions
+	class NodeMeshSubtract;
+	typedef Ptr<NodeMeshSubtract> NodeMeshSubtractPtr;
+	typedef Ptr<const NodeMeshSubtract> NodeMeshSubtractPtrConst;
+
+
+	//! This node removes from a mesh A all the faces that are also part of a mesh B.
+	//! \ingroup model
+	class MUTABLETOOLS_API NodeMeshSubtract : public NodeMesh
+	{
+	public:
+
+		//-----------------------------------------------------------------------------------------
+		// Life cycle
+		//-----------------------------------------------------------------------------------------
+
+		NodeMeshSubtract();
+
+		void SerialiseWrapper(OutputArchive& arch) const override;
+		static void Serialise( const NodeMeshSubtract* pNode, OutputArchive& arch );
+		static NodeMeshSubtractPtr StaticUnserialise( InputArchive& arch );
+
+		//-----------------------------------------------------------------------------------------
+		// Node Interface
+		//-----------------------------------------------------------------------------------------
+
+        NodePtr Clone( NodeMapPtr pMap = 0 ) const override;
+
+		const NODE_TYPE* GetType() const override;
+		static const NODE_TYPE* GetStaticType();
+
+		int GetInputCount() const override;
+		Node* GetInputNode( int i ) const override;
+		void SetInputNode( int i, NodePtr pNode ) override;
+
+		//-----------------------------------------------------------------------------------------
+		// Own Interface
+		//-----------------------------------------------------------------------------------------
+
+		//!
+		NodeMeshPtr GetA() const;
+		void SetA( NodeMesh* p );
+
+		//!
+		NodeMeshPtr GetB() const;
+		void SetB( NodeMesh* p );
+
+		//-----------------------------------------------------------------------------------------
+		// Interface pattern
+		//-----------------------------------------------------------------------------------------
+		class Private;
+		Private* GetPrivate() const;
+        Node::Private* GetBasePrivate() const override;
+
+	protected:
+
+		//! Forbidden. Manage with the Ptr<> template.
+		~NodeMeshSubtract();
+
+	private:
+
+		Private* m_pD;
+
+	};
+
+
+}
diff --git a/Engine/Plugins/Experimental/Mutable/Source/MutableTools/Public/MuT/NodeMeshSwitch.h b/Engine/Plugins/Experimental/Mutable/Source/MutableTools/Public/MuT/NodeMeshSwitch.h
new file mode 100644
--- /dev/null
+++ b/Engine/Plugins/Experimental/Mutable/Source/MutableTools/Public/MuT/NodeMeshSwitch.h
@@ -0,0 +1,89 @@
+// Copyright Epic Games, Inc. All Rights Reserved.
+
+#pragma once
+
+#include "MuT/NodeMesh.h"
+
+#include "MuR/Mesh.h"
+
+
+namespace mu
+{
+
+	// Forward definitions
+	class NodeScalar;
+	typedef Ptr<NodeScalar> NodeScalarPtr;
+	typedef Ptr<const NodeScalar> NodeScalarPtrConst;
+
+	class NodeMeshSwitch;
+	typedef Ptr<NodeMeshSwitch> NodeMeshSwitchPtr;
+	typedef Ptr<const NodeMeshSwitch> NodeMeshSwitchPtrConst;
+
+
+	//! This node selects an output Mesh from a set of input Meshs based on a parameter.
+	//! \ingroup model
+	class MUTABLETOOLS_API NodeMeshSwitch : public NodeMesh
+	{
+	public:
+
+		//-----------------------------------------------------------------------------------------
+		// Life cycle
+		//-----------------------------------------------------------------------------------------
+
+		NodeMeshSwitch();
+
+		void SerialiseWrapper(OutputArchive& arch) const override;
+		static void Serialise( const NodeMeshSwitch* pNode, OutputArchive& arch );
+		static NodeMeshSwitchPtr StaticUnserialise( InputArchive& arch );
+
+
+		//-----------------------------------------------------------------------------------------
+		// Node Interface
+		//-----------------------------------------------------------------------------------------
+
+        NodePtr Clone( NodeMapPtr pMap = 0 ) const override;
+
+		const NODE_TYPE* GetType() const override;
+		static const NODE_TYPE* GetStaticType();
+
+		int GetInputCount() const override;
+		Node* GetInputNode( int i ) const override;
+		void SetInputNode( int i, NodePtr pNode ) override;
+
+		//-----------------------------------------------------------------------------------------
+		// Own Interface
+		//-----------------------------------------------------------------------------------------
+
+		//! Get the node generating the parameter used to select the option.
+		NodeScalarPtr GetParameter() const;
+		void SetParameter( NodeScalarPtr );
+
+		//! Set the number of option Meshs. It will keep the currently set targets and initialise
+		//! the new ones as null.
+		void SetOptionCount( int );
+
+		//! Get the node generating the t-th option Mesh.
+		NodeMeshPtr GetOption( int t ) const;
+		void SetOption( int t, NodeMeshPtr );
+
+
+		//-----------------------------------------------------------------------------------------
+		// Interface pattern
+		//-----------------------------------------------------------------------------------------
+		class Private;
+		Private* GetPrivate() const;
+        Node::Private* GetBasePrivate() const override;
+
+	protected:
+
+		//! Forbidden. Manage with the Ptr<> template.
+		~NodeMeshSwitch();
+
+	private:
+
+		Private* m_pD;
+
+	};
+
+
+}
diff --git a/Engine/Plugins/Experimental/Mutable/Source/MutableTools/Public/MuT/NodeMeshTable.h b/Engine/Plugins/Experimental/Mutable/Source/MutableTools/Public/MuT/NodeMeshTable.h
new file mode 100644
--- /dev/null
+++ b/Engine/Plugins/Experimental/Mutable/Source/MutableTools/Public/MuT/NodeMeshTable.h
@@ -0,0 +1,98 @@
+// Copyright Epic Games, Inc. All Rights Reserved.
+
+#pragma once
+
+#include "MuT/NodeMesh.h"
+
+
+namespace mu
+{
+
+	// Forward definitions
+	class NodeMeshTable;
+	typedef Ptr<NodeMeshTable> NodeMeshTablePtr;
+	typedef Ptr<const NodeMeshTable> NodeMeshTablePtrConst;
+
+	class NodeLayout;
+	typedef Ptr<NodeLayout> NodeLayoutPtr;
+	typedef Ptr<const NodeLayout> NodeLayoutPtrConst;
+
+	class Table;
+	typedef Ptr<Table> TablePtr;
+	typedef Ptr<const Table> TablePtrConst;
+
+
+	//! This node provides the meshes stored in the column of a table.
+	//! \ingroup transform
+	class MUTABLETOOLS_API NodeMeshTable : public NodeMesh
+	{
+	public:
+
+		//-----------------------------------------------------------------------------------------
+		// Life cycle
+		//-----------------------------------------------------------------------------------------
+
+		NodeMeshTable();
+
+		void SerialiseWrapper(OutputArchive& arch) const override;
+		static void Serialise( const NodeMeshTable* pNode, OutputArchive& arch );
+		static NodeMeshTablePtr StaticUnserialise( InputArchive& arch );
+
+
+		//-----------------------------------------------------------------------------------------
+		// Node Interface
+		//-----------------------------------------------------------------------------------------
+
+        NodePtr Clone( NodeMapPtr pMap = 0 ) const override;
+
+        const NODE_TYPE* GetType() const override;
+		static const NODE_TYPE* GetStaticType();
+
+        int GetInputCount() const override;
+        Node* GetInputNode( int i ) const override;
+        void SetInputNode( int i, NodePtr pNode ) override;
+
+		//-----------------------------------------------------------------------------------------
+		// Own Interface
+		//-----------------------------------------------------------------------------------------
+
+		//! Set the name of the implicit table parameter.
+		void SetParameterName( const char* strName );
+
+		//!
+		TablePtr GetTable() const;
+		void SetTable( TablePtr );
+
+		//!
+		const char* GetColumn() const;
+		void SetColumn( const char* strName );
+
+		//! Get the number of layouts defined in the meshes on this column.
+		int GetLayoutCount() const;
+		void SetLayoutCount( int );
+
+		//! Get the node defining a layout of the meshes on this column.
+		NodeLayoutPtr GetLayout( int index ) const;
+		void SetLayout( int index, NodeLayoutPtr );
+
+		//-----------------------------------------------------------------------------------------
+		// Interface pattern
+		//-----------------------------------------------------------------------------------------
+		class Private;
+		Private* GetPrivate() const;
+        Node::Private* GetBasePrivate() const override;
+
+
+	protected:
+
+		//! Forbidden. Manage with the Ptr<> template.
+		~NodeMeshTable();
+
+	private:
+
+		Private* m_pD;
+
+	};
+
+
+}
diff --git a/Engine/Plugins/Experimental/Mutable/Source/MutableTools/Public/MuT/NodeMeshTangents.h b/Engine/Plugins/Experimental/Mutable/Source/MutableTools/Public/MuT/NodeMeshTangents.h
new file mode 100644
--- /dev/null
+++ b/Engine/Plugins/Experimental/Mutable/Source/MutableTools/Public/MuT/NodeMeshTangents.h
@@ -0,0 +1,73 @@
+// Copyright Epic Games, Inc. All Rights Reserved.
+
+#pragma once
+
+#include "MuT/NodeMesh.h"
+
+
+namespace mu
+{
+
+	// Forward definitions
+	class NodeMeshTangents;
+	typedef Ptr<NodeMeshTangents> NodeMeshTangentsPtr;
+	typedef Ptr<const NodeMeshTangents> NodeMeshTangentsPtrConst;
+
+
+	//! This node rebuilds the tangents and binormals of the source mesh.
+	//! \ingroup model
+	class MUTABLETOOLS_API NodeMeshTangents : public NodeMesh
+	{
+	public:
+
+		//-----------------------------------------------------------------------------------------
+		// Life cycle
+		//-----------------------------------------------------------------------------------------
+
+		NodeMeshTangents();
+
+		void SerialiseWrapper(OutputArchive& arch) const override;
+		static void Serialise( const NodeMeshTangents* pNode, OutputArchive& arch );
+		static NodeMeshTangentsPtr StaticUnserialise( InputArchive& arch );
+
+		//-----------------------------------------------------------------------------------------
+		// Node Interface
+		//-----------------------------------------------------------------------------------------
+
+        NodePtr Clone( NodeMapPtr pMap = 0 ) const override;
+
+        const NODE_TYPE* GetType() const override;
+		static const NODE_TYPE* GetStaticType();
+
+        int GetInputCount() const override;
+        Node* GetInputNode( int i ) const override;
+        void SetInputNode( int i, NodePtr pNode ) override;
+
+		//-----------------------------------------------------------------------------------------
+		// Own Interface
+		//-----------------------------------------------------------------------------------------
+
+		//!
+		NodeMeshPtr GetSource() const;
+		void SetSource( NodeMesh* p );
+
+		//-----------------------------------------------------------------------------------------
+		// Interface pattern
+		//-----------------------------------------------------------------------------------------
+		class Private;
+		Private* GetPrivate() const;
+        Node::Private* GetBasePrivate() const override;
+
+	protected:
+
+		//! Forbidden. Manage with the Ptr<> template.
+		~NodeMeshTangents();
+
+	private:
+
+		Private* m_pD;
+
+	};
+
+
+}
