Moved the Mutable plugin files to live under a directory named after the module. This allows us to get rid of long include paths that contain public and private folder names in the path and fixes head

diff --git a/Engine/Plugins/Experimental/Mutable/Source/MutableTools/Public/MuT/NodeMeshTransform.h b/Engine/Plugins/Experimental/Mutable/Source/MutableTools/Public/MuT/NodeMeshTransform.h
new file mode 100644
--- /dev/null
+++ b/Engine/Plugins/Experimental/Mutable/Source/MutableTools/Public/MuT/NodeMeshTransform.h
@@ -0,0 +1,79 @@
+// Copyright Epic Games, Inc. All Rights Reserved.
+
+#pragma once
+
+#include "MuT/NodeMesh.h"
+
+#include "MuR/Mesh.h"
+
+
+namespace mu
+{
+
+	// Forward definitions
+    class NodeMeshTransform;
+    typedef Ptr<NodeMeshTransform> NodeMeshTransformPtr;
+    typedef Ptr<const NodeMeshTransform> NodeMeshTransformPtrConst;
+
+
+    //! This node applies a geometric transform represented by a 4x4 matrix to a mesh
+	//! \ingroup model
+    class MUTABLETOOLS_API NodeMeshTransform : public NodeMesh
+	{
+	public:
+
+		//-----------------------------------------------------------------------------------------
+		// Life cycle
+		//-----------------------------------------------------------------------------------------
+
+        NodeMeshTransform();
+
+		void SerialiseWrapper(OutputArchive& arch) const override;
+		static void Serialise( const NodeMeshTransform* pNode, OutputArchive& arch );
+        static NodeMeshTransformPtr StaticUnserialise( InputArchive& arch );
+
+		//-----------------------------------------------------------------------------------------
+		// Node Interface
+		//-----------------------------------------------------------------------------------------
+
+        NodePtr Clone( NodeMapPtr pMap = 0 ) const override;
+
+        const NODE_TYPE* GetType() const override;
+		static const NODE_TYPE* GetStaticType();
+
+        int GetInputCount() const override;
+        Node* GetInputNode( int i ) const override;
+        void SetInputNode( int i, NodePtr pNode ) override;
+
+		//-----------------------------------------------------------------------------------------
+		// Own Interface
+		//-----------------------------------------------------------------------------------------
+
+		//! Source mesh to be re-formatted
+        NodeMeshPtr GetSource() const;
+		void SetSource( NodeMesh* );
+
+        //! \param pMat4x4 pointer to 16 floats representing the 4x4 matrix to apply
+        void SetTransform( const float* pMat4x4 );
+        void GetTransform( float* pMat4x4 ) const;
+
+		//-----------------------------------------------------------------------------------------
+		// Interface pattern
+		//-----------------------------------------------------------------------------------------
+		class Private;
+		Private* GetPrivate() const;
+        Node::Private* GetBasePrivate() const override;
+
+	protected:
+
+		//! Forbidden. Manage with the Ptr<> template.
+        ~NodeMeshTransform();
+
+	private:
+
+		Private* m_pD;
+
+	};
+
+
+}
diff --git a/Engine/Plugins/Experimental/Mutable/Source/MutableTools/Public/MuT/NodeMeshVariation.h b/Engine/Plugins/Experimental/Mutable/Source/MutableTools/Public/MuT/NodeMeshVariation.h
new file mode 100644
--- /dev/null
+++ b/Engine/Plugins/Experimental/Mutable/Source/MutableTools/Public/MuT/NodeMeshVariation.h
@@ -0,0 +1,83 @@
+// Copyright Epic Games, Inc. All Rights Reserved.
+
+#pragma once
+
+#include "MuT/NodeMesh.h"
+
+
+namespace mu
+{
+
+    // Forward definitions
+    class NodeMeshVariation;
+    typedef Ptr<NodeMeshVariation> NodeMeshVariationPtr;
+    typedef Ptr<const NodeMeshVariation> NodeMeshVariationPtrConst;
+
+
+    //!
+    //! \ingroup model
+    class MUTABLETOOLS_API NodeMeshVariation : public NodeMesh
+    {
+    public:
+        //-----------------------------------------------------------------------------------------
+        // Life cycle
+        //-----------------------------------------------------------------------------------------
+
+        NodeMeshVariation();
+
+		void SerialiseWrapper(OutputArchive& arch) const override;
+		static void Serialise( const NodeMeshVariation* pNode, OutputArchive& arch );
+        static NodeMeshVariationPtr StaticUnserialise( InputArchive& arch );
+
+        //-----------------------------------------------------------------------------------------
+        // Node Interface
+        //-----------------------------------------------------------------------------------------
+
+        NodePtr Clone( NodeMapPtr pMap = 0 ) const override;
+
+        const NODE_TYPE* GetType() const override;
+        static const NODE_TYPE* GetStaticType();
+
+        virtual int GetInputCount() const override;
+        virtual Node* GetInputNode( int i ) const override;
+        void SetInputNode( int i, NodePtr pNode ) override;
+
+        //-----------------------------------------------------------------------------------------
+        // Own Interface
+        //-----------------------------------------------------------------------------------------
+
+        //!
+        void SetDefaultMesh( NodeMesh* Mesh );
+
+        //! Set the number of tags to consider in this variation
+        void SetVariationCount( int count );
+
+        //!
+        int GetVariationCount() const;
+
+        //! Set the tag or state name that will enable a specific vartiation
+        void SetVariationTag( int index, const char* strTag );
+
+        //!
+        void SetVariationMesh( int index, NodeMesh* Mesh );
+
+        //!}
+
+
+        //-----------------------------------------------------------------------------------------
+        // Interface pattern
+        //-----------------------------------------------------------------------------------------
+        class Private;
+        Private* GetPrivate() const;
+        Node::Private* GetBasePrivate() const override;
+
+    protected:
+        //! Forbidden. Manage with the Ptr<> template.
+        ~NodeMeshVariation();
+
+    private:
+        Private* m_pD;
+    };
+
+
+} // namespace mu
diff --git a/Engine/Plugins/Experimental/Mutable/Source/MutableTools/Public/MuT/NodeModifier.h b/Engine/Plugins/Experimental/Mutable/Source/MutableTools/Public/MuT/NodeModifier.h
new file mode 100644
--- /dev/null
+++ b/Engine/Plugins/Experimental/Mutable/Source/MutableTools/Public/MuT/NodeModifier.h
@@ -0,0 +1,84 @@
+// Copyright Epic Games, Inc. All Rights Reserved.
+
+#pragma once
+
+#include "MuT/Node.h"
+
+
+namespace mu
+{
+
+	// Forward definitions
+	class NodeModifier;
+	typedef Ptr<NodeModifier> NodeModifierPtr;
+	typedef Ptr<const NodeModifier> NodeModifierConst;
+
+
+	//! This class is the parent of all nodes that output a component.
+	//! \ingroup model
+	class MUTABLETOOLS_API NodeModifier : public Node
+	{
+	public:
+
+		// Possible subclasses
+		enum class EType : uint8
+		{
+			MeshClipMorphPlane = 0,
+			MeshClipWithMesh = 1,
+			MeshClipDeform = 2,
+			
+			None
+		};
+
+		//-----------------------------------------------------------------------------------------
+		// Life cycle
+		//-----------------------------------------------------------------------------------------
+
+		static void Serialise( const NodeModifier* pNode, OutputArchive& arch );
+		static NodeModifierPtr StaticUnserialise( InputArchive& arch );
+
+		//-----------------------------------------------------------------------------------------
+        // Node interface
+		//-----------------------------------------------------------------------------------------
+
+		const NODE_TYPE* GetType() const override;
+		static const NODE_TYPE* GetStaticType();
+
+        //-----------------------------------------------------------------------------------------
+        // Own interface
+        //-----------------------------------------------------------------------------------------
+
+        //! \name Tags
+        //! \{
+
+        //! Add a tag to the surface, which will be affected by modifier nodes with the same tag
+        void AddTag(const char* tagName);
+
+        //! Get the number of tags added to the Surface.
+        int GetTagCount() const;
+
+        //! Get a tag string from an index (0 to GetTagCount-1)
+        const char* GetTag( int ) const;
+
+        //! \}
+
+		//! Set the stage to apply this modifier in. Default is before normal operations.
+		void SetStage( bool bBeforeNormalOperation );
+
+		//-----------------------------------------------------------------------------------------
+		// Interface pattern
+		//-----------------------------------------------------------------------------------------
+		class Private;
+
+	protected:
+
+		//! Forbidden. Manage with the Ptr<> template.
+		inline ~NodeModifier() {}
+
+		//!
+		EType Type = EType::None;
+
+	};
+
+
+}
diff --git a/Engine/Plugins/Experimental/Mutable/Source/MutableTools/Public/MuT/NodeModifierMeshClipDeform.h b/Engine/Plugins/Experimental/Mutable/Source/MutableTools/Public/MuT/NodeModifierMeshClipDeform.h
new file mode 100644
--- /dev/null
+++ b/Engine/Plugins/Experimental/Mutable/Source/MutableTools/Public/MuT/NodeModifierMeshClipDeform.h
@@ -0,0 +1,77 @@
+// Copyright Epic Games, Inc. All Rights Reserved.
+
+#pragma once
+
+#include "MuT/NodeModifier.h"
+
+namespace mu
+{
+
+	// Forward definitions
+	class NodeModifierMeshClipDeform;
+	typedef Ptr<NodeModifierMeshClipDeform> NodeModifierMeshClipDeformPtr;
+	typedef Ptr<const NodeModifierMeshClipDeform> NodeModifierMeshClipDeformPtrConst;
+
+	class NodeMesh;
+	typedef Ptr<NodeMesh> NodeMeshPtr;
+	typedef Ptr<const NodeMesh> NodeMeshPtrConst;
+
+	enum class EShapeBindingMethod : uint32;
+	
+	//! This node makes a new component from several meshes and images.
+	//! \ingroup model
+	class MUTABLETOOLS_API NodeModifierMeshClipDeform : public NodeModifier
+	{
+	public:
+
+		//-----------------------------------------------------------------------------------------
+		// Life cycle
+		//-----------------------------------------------------------------------------------------
+
+		NodeModifierMeshClipDeform();
+
+		void SerialiseWrapper(OutputArchive& arch) const override;
+		static void Serialise( const NodeModifierMeshClipDeform* pNode, OutputArchive& arch );
+		static NodeModifierMeshClipDeformPtr StaticUnserialise( InputArchive& arch );
+
+		//-----------------------------------------------------------------------------------------
+        // Node interface
+		//-----------------------------------------------------------------------------------------
+
+        NodePtr Clone( NodeMapPtr pMap = 0 ) const override;
+
+        const NODE_TYPE* GetType() const override;
+		static const NODE_TYPE* GetStaticType();
+
+        int GetInputCount() const override;
+        Node* GetInputNode( int i ) const override;
+        void SetInputNode( int i, NodePtr pNode ) override;
+
+		//-----------------------------------------------------------------------------------------
+        // Own interface
+		//-----------------------------------------------------------------------------------------
+
+		//! \param 
+		void SetClipMesh( NodeMesh* InClipMesh);
+		void SetBindingMethod(EShapeBindingMethod BindingMethod);
+	
+		//-----------------------------------------------------------------------------------------
+		// Interface pattern
+		//-----------------------------------------------------------------------------------------
+		class Private;
+		Private* GetPrivate() const;
+        Node::Private* GetBasePrivate() const override;
+
+	protected:
+
+		//! Forbidden. Manage with the Ptr<> template.
+		~NodeModifierMeshClipDeform();
+
+	private:
+
+		Private* m_pD;
+
+	};
+
+
+}
diff --git a/Engine/Plugins/Experimental/Mutable/Source/MutableTools/Public/MuT/NodeModifierMeshClipMorphPlane.h b/Engine/Plugins/Experimental/Mutable/Source/MutableTools/Public/MuT/NodeModifierMeshClipMorphPlane.h
new file mode 100644
--- /dev/null
+++ b/Engine/Plugins/Experimental/Mutable/Source/MutableTools/Public/MuT/NodeModifierMeshClipMorphPlane.h
@@ -0,0 +1,82 @@
+// Copyright Epic Games, Inc. All Rights Reserved.
+
+#pragma once
+
+#include "MuT/NodeModifier.h"
+
+
+namespace mu
+{
+
+	// Forward definitions
+	class NodeModifierMeshClipMorphPlane;
+	typedef Ptr<NodeModifierMeshClipMorphPlane> NodeModifierMeshClipMorphPlanePtr;
+	typedef Ptr<const NodeModifierMeshClipMorphPlane> NodeModifierMeshClipMorphPlanePtrConst;
+
+
+	//! This node makes a new component from several meshes and images.
+	//! \ingroup model
+	class MUTABLETOOLS_API NodeModifierMeshClipMorphPlane : public NodeModifier
+	{
+	public:
+
+		//-----------------------------------------------------------------------------------------
+		// Life cycle
+		//-----------------------------------------------------------------------------------------
+
+		NodeModifierMeshClipMorphPlane();
+
+		void SerialiseWrapper(OutputArchive& arch) const override;
+		static void Serialise( const NodeModifierMeshClipMorphPlane* pNode, OutputArchive& arch );
+		static NodeModifierMeshClipMorphPlanePtr StaticUnserialise( InputArchive& arch );
+
+		//-----------------------------------------------------------------------------------------
+        // Node interface
+		//-----------------------------------------------------------------------------------------
+
+        NodePtr Clone( NodeMapPtr pMap = 0 ) const override;
+
+        const NODE_TYPE* GetType() const override;
+		static const NODE_TYPE* GetStaticType();
+
+        int GetInputCount() const override;
+        Node* GetInputNode( int i ) const override;
+        void SetInputNode( int i, NodePtr pNode ) override;
+
+		//-----------------------------------------------------------------------------------------
+        // Own interface
+		//-----------------------------------------------------------------------------------------
+
+		//! \param 
+		void SetPlane(float centerX, float centerY, float centerZ, float normalX, float normalY, float normalZ);
+		void SetParams(float dist, float factor);
+		void SetMorphEllipse(float radius1, float radius2, float rotation);
+
+		//! Define an axis-aligned box that will select the vertices to be morphed.
+		//! Only one of Box or Bone Hierarchy can be used (the last one set)
+		void SetVertexSelectionBox(float centerX, float centerY, float centerZ, float radiusX, float radiusY, float radiusZ);
+
+		//! Define the root bone of the subhierarchy of the mesh that will be affected.
+		//! Only one of Box or Bone Hierarchy can be used (the last one set)
+		void SetVertexSelectionBone(const char* strBoneName, float maxEffectRadius);
+
+		//-----------------------------------------------------------------------------------------
+		// Interface pattern
+		//-----------------------------------------------------------------------------------------
+		class Private;
+		Private* GetPrivate() const;
+        Node::Private* GetBasePrivate() const override;
+
+	protected:
+
+		//! Forbidden. Manage with the Ptr<> template.
+		~NodeModifierMeshClipMorphPlane();
+
+	private:
+
+		Private* m_pD;
+
+	};
+
+
+}
diff --git a/Engine/Plugins/Experimental/Mutable/Source/MutableTools/Public/MuT/NodeModifierMeshClipWithMesh.h b/Engine/Plugins/Experimental/Mutable/Source/MutableTools/Public/MuT/NodeModifierMeshClipWithMesh.h
new file mode 100644
--- /dev/null
+++ b/Engine/Plugins/Experimental/Mutable/Source/MutableTools/Public/MuT/NodeModifierMeshClipWithMesh.h
@@ -0,0 +1,76 @@
+// Copyright Epic Games, Inc. All Rights Reserved.
+
+#pragma once
+
+#include "MuT/NodeModifier.h"
+
+
+namespace mu
+{
+
+	// Forward definitions
+	class NodeModifierMeshClipWithMesh;
+	typedef Ptr<NodeModifierMeshClipWithMesh> NodeModifierMeshClipWithMeshPtr;
+	typedef Ptr<const NodeModifierMeshClipWithMesh> NodeModifierMeshClipWithMeshPtrConst;
+
+	class NodeMesh;
+	typedef Ptr<NodeMesh> NodeMeshPtr;
+	typedef Ptr<const NodeMesh> NodeMeshPtrConst;
+
+
+	//! This node makes a new component from several meshes and images.
+	//! \ingroup model
+	class MUTABLETOOLS_API NodeModifierMeshClipWithMesh : public NodeModifier
+	{
+	public:
+
+		//-----------------------------------------------------------------------------------------
+		// Life cycle
+		//-----------------------------------------------------------------------------------------
+
+		NodeModifierMeshClipWithMesh();
+
+		void SerialiseWrapper(OutputArchive& arch) const override;
+		static void Serialise( const NodeModifierMeshClipWithMesh* pNode, OutputArchive& arch );
+		static NodeModifierMeshClipWithMeshPtr StaticUnserialise( InputArchive& arch );
+
+		//-----------------------------------------------------------------------------------------
+        // Node interface
+		//-----------------------------------------------------------------------------------------
+
+        NodePtr Clone( NodeMapPtr pMap = 0 ) const override;
+
+        const NODE_TYPE* GetType() const override;
+		static const NODE_TYPE* GetStaticType();
+
+        int GetInputCount() const override;
+        Node* GetInputNode( int i ) const override;
+        void SetInputNode( int i, NodePtr pNode ) override;
+
+		//-----------------------------------------------------------------------------------------
+        // Own interface
+		//-----------------------------------------------------------------------------------------
+
+		//! \param 
+		void SetClipMesh(NodeMesh*);
+
+		//-----------------------------------------------------------------------------------------
+		// Interface pattern
+		//-----------------------------------------------------------------------------------------
+		class Private;
+		Private* GetPrivate() const;
+        Node::Private* GetBasePrivate() const override;
+
+	protected:
+
+		//! Forbidden. Manage with the Ptr<> template.
+		~NodeModifierMeshClipWithMesh();
+
+	private:
+
+		Private* m_pD;
+
+	};
+
+
+}
diff --git a/Engine/Plugins/Experimental/Mutable/Source/MutableTools/Public/MuT/NodeObject.h b/Engine/Plugins/Experimental/Mutable/Source/MutableTools/Public/MuT/NodeObject.h
new file mode 100644
--- /dev/null
+++ b/Engine/Plugins/Experimental/Mutable/Source/MutableTools/Public/MuT/NodeObject.h
@@ -0,0 +1,80 @@
+// Copyright Epic Games, Inc. All Rights Reserved.
+
+#pragma once
+
+#include "MuT/Node.h"
+
+
+namespace mu
+{
+
+	// Forward definitions
+	class NodeObject;
+	typedef Ptr<NodeObject> NodeObjectPtr;
+	typedef Ptr<const NodeObject> NodeObjectPtrConst;
+
+
+    //! %Base class for any node that outputs an object.
+	//! \ingroup model
+	class MUTABLETOOLS_API NodeObject : public Node
+	{
+	public:
+
+		// Possible subclasses
+		enum class EType : uint8
+		{
+			New = 0,
+			Group = 1,
+			State = 2,
+
+			None
+		};
+
+		//-----------------------------------------------------------------------------------------
+		// Life cycle
+		//-----------------------------------------------------------------------------------------
+
+		static void Serialise( const NodeObject* pNode, OutputArchive& arch );
+		static NodeObjectPtr StaticUnserialise( InputArchive& arch );
+
+
+		//-----------------------------------------------------------------------------------------
+		// Node Interface
+		//-----------------------------------------------------------------------------------------
+
+		const NODE_TYPE* GetType() const override;
+		static const NODE_TYPE* GetStaticType();
+
+		//-----------------------------------------------------------------------------------------
+		// Own interface
+		//-----------------------------------------------------------------------------------------
+
+		//! Get the name of the object.
+		virtual const char* GetName() const = 0;
+
+		//! Set the name of the object.
+		virtual void SetName( const char* strName ) = 0;
+
+		//! Get the uid of the object.
+		virtual const char* GetUid() const = 0;
+
+		//! Set the uid of the object.
+		virtual void SetUid( const char* strUid ) = 0;
+
+		//-----------------------------------------------------------------------------------------
+		// Interface pattern
+		//-----------------------------------------------------------------------------------------
+		class Private;
+
+	protected:
+
+		//! Forbidden. Manage with the Ptr<> template.
+		inline ~NodeObject() {}
+
+		//!
+		EType Type = EType::None;
+
+	};
+
+
+}
diff --git a/Engine/Plugins/Experimental/Mutable/Source/MutableTools/Public/MuT/NodeObjectGroup.h b/Engine/Plugins/Experimental/Mutable/Source/MutableTools/Public/MuT/NodeObjectGroup.h
new file mode 100644
--- /dev/null
+++ b/Engine/Plugins/Experimental/Mutable/Source/MutableTools/Public/MuT/NodeObjectGroup.h
@@ -0,0 +1,117 @@
+// Copyright Epic Games, Inc. All Rights Reserved.
+
+#pragma once
+
+#include "MuT/NodeObject.h"
+
+
+namespace mu
+{
+
+	// Forward definitions
+	class NodeObjectGroup;
+	typedef Ptr<NodeObjectGroup> NodeObjectGroupPtr;
+	typedef Ptr<const NodeObjectGroup>NodeObjectGroupPtrConst;
+
+
+	//! Node that creates a group of objects and describes how they are selected.
+	//! \ingroup model
+	class MUTABLETOOLS_API NodeObjectGroup : public NodeObject
+	{
+	public:
+
+		//-----------------------------------------------------------------------------------------
+		// Life cycle
+		//-----------------------------------------------------------------------------------------
+
+		NodeObjectGroup();
+
+		void SerialiseWrapper(OutputArchive& arch) const override;
+		static void Serialise( const NodeObjectGroup* pNode, OutputArchive& arch );
+		static NodeObjectGroupPtr StaticUnserialise( InputArchive& arch );
+
+
+		//-----------------------------------------------------------------------------------------
+		// Node Interface
+		//-----------------------------------------------------------------------------------------
+
+        NodePtr Clone( NodeMapPtr pMap = 0 ) const override;
+
+        const NODE_TYPE* GetType() const override;
+		static const NODE_TYPE* GetStaticType();
+
+        int GetInputCount() const override;
+        Node* GetInputNode( int i ) const override;
+        void SetInputNode( int i, NodePtr pNode ) override;
+
+		//-----------------------------------------------------------------------------------------
+		// NodeObject Interface
+		//-----------------------------------------------------------------------------------------
+        const char* GetName() const override;
+        void SetName( const char* strName ) override;
+        const char* GetUid() const override;
+        void SetUid( const char* strUid ) override;
+
+		//-----------------------------------------------------------------------------------------
+		// Own Interface
+		//-----------------------------------------------------------------------------------------
+
+		//! Typed of child selection
+		typedef enum
+		{
+			//! All objects in the group will always be enabled, and no parameter will be generated.
+			CS_ALWAYS_ALL,
+
+			//! Only one children may be selected, but it is allowed to have none.
+			//! An enumeration parameter will be generated but it may have a null value
+			CS_ONE_OR_NONE,
+
+			//! One and only one children has to be selected at all times
+			//! An enumeraation parameter will be generated and it cannot be null
+			CS_ALWAYS_ONE,
+
+			//! Each child in the group can be enabled or disabled individually.
+			//! A boolean parameter will be generated for every child.
+			CS_TOGGLE_EACH
+
+		} CHILD_SELECTION;
+
+		//! Get child selection type
+		CHILD_SELECTION GetSelectionType() const;
+
+		//! Set the child selection type
+		void SetSelectionType( CHILD_SELECTION );
+
+		//! Get the number of child objects
+		int GetChildCount() const;
+
+		//! Set the number of child objects
+		void SetChildCount( int );
+
+		//! Get a child object node
+		NodeObjectPtr GetChild( int index ) const;
+
+		//! Set a child object node
+		void SetChild( int index, NodeObjectPtr );
+
+		//-----------------------------------------------------------------------------------------
+		// Interface pattern
+		//-----------------------------------------------------------------------------------------
+		class Private;
+		Private* GetPrivate() const;
+        Node::Private* GetBasePrivate() const override;
+
+	protected:
+
+		//! Forbidden. Manage with the Ptr<> template.
+		~NodeObjectGroup();
+
+	private:
+
+		Private* m_pD;
+
+	};
+
+
+
+}
diff --git a/Engine/Plugins/Experimental/Mutable/Source/MutableTools/Public/MuT/NodeObjectNew.h b/Engine/Plugins/Experimental/Mutable/Source/MutableTools/Public/MuT/NodeObjectNew.h
new file mode 100644
--- /dev/null
+++ b/Engine/Plugins/Experimental/Mutable/Source/MutableTools/Public/MuT/NodeObjectNew.h
@@ -0,0 +1,124 @@
+// Copyright Epic Games, Inc. All Rights Reserved.
+
+#pragma once
+#include "MuT/NodeObject.h"
+
+
+namespace mu
+{
+
+	// Forward definitions
+	class NodeObjectNew;
+	typedef Ptr<NodeObjectNew> NodeObjectNewPtr;
+	typedef Ptr<const NodeObjectNew> NodeObjectNewPtrConst;
+
+	class NodeLOD;
+	typedef Ptr<NodeLOD> NodeLODPtr;
+	typedef Ptr<const NodeLOD> NodeLODPtrConst;
+
+
+	//! Node that creates a new object by setting its levels-of-detail and its children.
+	//! \ingroup model
+	class MUTABLETOOLS_API NodeObjectNew : public NodeObject
+	{
+	public:
+
+		//-----------------------------------------------------------------------------------------
+		// Life cycle
+		//-----------------------------------------------------------------------------------------
+
+		NodeObjectNew();
+
+		void SerialiseWrapper(OutputArchive& arch) const override;
+		static void Serialise( const NodeObjectNew* pNode, OutputArchive& arch );
+		static NodeObjectNewPtr StaticUnserialise( InputArchive& arch );
+
+
+		//-----------------------------------------------------------------------------------------
+		// Node Interface
+		//-----------------------------------------------------------------------------------------
+
+        NodePtr Clone( NodeMapPtr pMap = 0 ) const override;
+
+        const NODE_TYPE* GetType() const override;
+		static const NODE_TYPE* GetStaticType();
+
+        int GetInputCount() const override;
+        Node* GetInputNode( int i ) const override;
+        void SetInputNode( int i, NodePtr pNode ) override;
+
+		//-----------------------------------------------------------------------------------------
+		// NodeObject Interface
+		//-----------------------------------------------------------------------------------------
+        const char* GetName() const override;
+        void SetName( const char* strName ) override;
+        const char* GetUid() const override;
+        void SetUid( const char* strUid ) override;
+
+		//-----------------------------------------------------------------------------------------
+		// Own Interface
+		//-----------------------------------------------------------------------------------------
+
+		//! Get the number of levels of detail in the object.
+		int GetLODCount() const;
+		void SetLODCount( int );
+
+		//! Get a level of detail node from the object.
+		NodeLODPtr GetLOD( int index ) const;
+		void SetLOD( int index, NodeLODPtr );
+
+		//! Get the number of child objects
+		int GetChildCount() const;
+		void SetChildCount( int );
+
+		//! Get a child object node
+		NodeObjectPtr GetChild( int index ) const;
+		void SetChild( int index, NodeObjectPtr );
+
+		//! Set the number of states that the model can be in.
+		int GetStateCount() const;
+		void SetStateCount( int c );
+
+		//! Set the name of a state
+		const char* GetStateName( int s ) const;
+		void SetStateName( int s, const char* n );
+
+		//! State optimization options
+		void SetStateGPUOptimisation( int s, bool internal, bool external, bool returnGPU = false );
+		bool GetStateGPUOptimisationInternal( int s ) const;
+		bool GetStateGPUOptimisationExternal( int s ) const;
+		bool GetStateGPUOptimisationReturnImages( int s ) const;
+
+		//! See if a state has a parameter as runtime.
+		bool HasStateParam( int s, const char* param ) const;
+
+		//! Add a runtime parameter to the state.
+		void AddStateParam( int s, const char* param );
+
+		//! Remove a runtime parameter from a state.
+		void RemoveStateParam( int s, const char* param );
+
+        //! Set the optimisation properties of a state
+        void SetStateProperties( int s, bool avoidRuntimeCompression, bool onlyFirstLOD, int firstLOD );
+
+		//-----------------------------------------------------------------------------------------
+		// Interface pattern
+		//-----------------------------------------------------------------------------------------
+		class Private;
+		Private* GetPrivate() const;
+        Node::Private* GetBasePrivate() const override;
+
+	protected:
+
+		//! Forbidden. Manage with the Ptr<> template.
+		~NodeObjectNew();
+
+	private:
+
+		Private* m_pD;
+
+	};
+
+
+
+}
diff --git a/Engine/Plugins/Experimental/Mutable/Source/MutableTools/Public/MuT/NodeObjectState.h b/Engine/Plugins/Experimental/Mutable/Source/MutableTools/Public/MuT/NodeObjectState.h
new file mode 100644
--- /dev/null
+++ b/Engine/Plugins/Experimental/Mutable/Source/MutableTools/Public/MuT/NodeObjectState.h
@@ -0,0 +1,108 @@
+// Copyright Epic Games, Inc. All Rights Reserved.
+
+#pragma once
+
+#include "MuT/NodeObject.h"
+
+
+namespace mu
+{
+
+	//! Forward definitions
+	class NodeObjectState;
+	typedef Ptr<NodeObjectState> NodeObjectStatePtr;
+	typedef Ptr<const NodeObjectState> NodeObjectStatePtrConst;
+
+
+
+	//!
+	//! \ingroup model
+	class MUTABLETOOLS_API NodeObjectState : public NodeObject
+	{
+	public:
+
+		//-----------------------------------------------------------------------------------------
+		// Life cycle
+		//-----------------------------------------------------------------------------------------
+
+		NodeObjectState();
+
+		void SerialiseWrapper(OutputArchive& arch) const override;
+		static void Serialise( const NodeObjectState* pNode, OutputArchive& arch );
+		static NodeObjectStatePtr StaticUnserialise( InputArchive& arch );
+
+
+		//-----------------------------------------------------------------------------------------
+		// Node interface
+		//-----------------------------------------------------------------------------------------
+
+        NodePtr Clone( NodeMapPtr pMap = 0 ) const override;
+
+        const NODE_TYPE* GetType() const override;
+		static const NODE_TYPE* GetStaticType();
+
+        int GetInputCount() const override;
+        Node* GetInputNode( int i ) const override;
+        void SetInputNode( int i, NodePtr pNode ) override;
+
+
+		//-----------------------------------------------------------------------------------------
+		// NodeObject interface
+		//-----------------------------------------------------------------------------------------
+        const char* GetName() const override;
+        void SetName( const char* strName ) override;
+        const char* GetUid() const override;
+        void SetUid( const char* strUid ) override;
+
+		//-----------------------------------------------------------------------------------------
+		// Own interface
+		//-----------------------------------------------------------------------------------------
+
+		//! Source object
+		NodeObjectPtr GetSource() const;
+		void SetSource( NodeObjectPtr );
+
+		//! State diagram to add
+		NodeObjectPtr GetStateRoot() const;
+		void SetStateRoot( NodeObjectPtr );
+
+		//! Set the name of a state
+		const char* GetStateName() const;
+		void SetStateName( const char* n );
+
+		//! State optimization options
+		void SetStateGPUOptimisation( bool internal, bool external, bool returnGPU = false );
+		bool GetStateGPUOptimisationInternal() const;
+		bool GetStateGPUOptimisationExternal() const;
+		bool GetStateGPUOptimisationReturnImages() const;
+
+		//! See if a state has a parameter as runtime.
+		bool HasStateParam( const char* param ) const;
+
+		//! Add a runtime parameter to the state.
+		void AddStateParam( const char* param );
+
+		//! Remove a runtime parameter from a state.
+		void RemoveStateParam( const char* param );
+
+		//-----------------------------------------------------------------------------------------
+		// Interface pattern
+		//-----------------------------------------------------------------------------------------
+		class Private;
+		Private* GetPrivate() const;
+        Node::Private* GetBasePrivate() const override;
+
+	protected:
+
+		//! Forbidden. Manage with the Ptr<> template.
+		~NodeObjectState();
+
+	private:
+
+		Private* m_pD;
+
+	};
+
+
+
+}
diff --git a/Engine/Plugins/Experimental/Mutable/Source/MutableTools/Public/MuT/NodePatchMesh.h b/Engine/Plugins/Experimental/Mutable/Source/MutableTools/Public/MuT/NodePatchMesh.h
new file mode 100644
--- /dev/null
+++ b/Engine/Plugins/Experimental/Mutable/Source/MutableTools/Public/MuT/NodePatchMesh.h
@@ -0,0 +1,77 @@
+// Copyright Epic Games, Inc. All Rights Reserved.
+
+#pragma once
+
+#include "MuT/NodeMesh.h"
+
+
+namespace mu
+{
+
+	// Forward definitions
+	class NodePatchMesh;
+	typedef Ptr<NodePatchMesh> NodePatchMeshPtr;
+	typedef Ptr<const NodePatchMesh> NodePatchMeshPtrConst;
+
+
+	//!
+	//! \ingroup model
+	class MUTABLETOOLS_API NodePatchMesh : public Node
+	{
+	public:
+
+		//-----------------------------------------------------------------------------------------
+		// Life cycle
+		//-----------------------------------------------------------------------------------------
+
+		NodePatchMesh();
+
+		void SerialiseWrapper(OutputArchive& arch) const override;
+		static void Serialise( const NodePatchMesh* pNode, OutputArchive& arch );
+		static NodePatchMeshPtr StaticUnserialise( InputArchive& arch );
+
+		//-----------------------------------------------------------------------------------------
+		// Node Interface
+		//-----------------------------------------------------------------------------------------
+
+        NodePtr Clone( NodeMapPtr pMap = 0 ) const override;
+
+        const NODE_TYPE* GetType() const override;
+		static const NODE_TYPE* GetStaticType();
+
+        int GetInputCount() const override;
+        Node* GetInputNode( int i ) const override;
+        void SetInputNode( int i, NodePtr pNode ) override;
+
+		//-----------------------------------------------------------------------------------------
+		// Own Interface
+		//-----------------------------------------------------------------------------------------
+
+		//!
+        NodeMesh* GetRemove() const;
+        void SetRemove( NodeMesh* );
+
+		//!
+        NodeMesh* GetAdd() const;
+        void SetAdd( NodeMesh* );
+
+		//-----------------------------------------------------------------------------------------
+		// Interface pattern
+		//-----------------------------------------------------------------------------------------
+		class Private;
+		Private* GetPrivate() const;
+        Node::Private* GetBasePrivate() const override;
+
+	protected:
+
+		//! Forbidden. Manage with the Ptr<> template.
+		~NodePatchMesh();
+
+	private:
+
+		Private* m_pD;
+
+	};
+
+
+}
diff --git a/Engine/Plugins/Experimental/Mutable/Source/MutableTools/Public/MuT/NodeProjector.h b/Engine/Plugins/Experimental/Mutable/Source/MutableTools/Public/MuT/NodeProjector.h
new file mode 100644
--- /dev/null
+++ b/Engine/Plugins/Experimental/Mutable/Source/MutableTools/Public/MuT/NodeProjector.h
@@ -0,0 +1,230 @@
+// Copyright Epic Games, Inc. All Rights Reserved.
+
+#pragma once
+
+#include "MuT/Node.h"
+#include "MuR/Parameters.h"
+
+
+namespace mu
+{
+
+	class NodeProjector;
+    using NodeProjectorPtr = Ptr<NodeProjector>;
+    using NodeProjectorPtrConst = Ptr<const NodeProjector>;
+
+	class NodeProjectorConstant;
+    using NodeProjectorConstantPtr = Ptr<NodeProjectorConstant>;
+    using NodeProjectorConstantPtrConst = Ptr<const NodeProjectorConstant>;
+
+	class NodeProjectorParameter;
+    using NodeProjectorParameterPtr = Ptr<NodeProjectorParameter>;
+    using NodeProjectorParameterPtrConst = Ptr<const NodeProjectorParameter>;
+
+    class NodeRange;
+    using NodeRangePtr = Ptr<NodeRange>;
+    using NodeRangePtrConst = Ptr<const NodeRange>;
+
+
+    //! Base class of any node that outputs a Projector.
+	//! \ingroup model
+	class MUTABLETOOLS_API NodeProjector : public Node
+	{
+	public:
+
+		// Possible subclasses
+		enum class EType : uint8
+		{
+			Constant = 0,
+			Parameter = 1,
+
+			None
+		};
+
+		//-----------------------------------------------------------------------------------------
+		// Life cycle
+		//-----------------------------------------------------------------------------------------
+
+		static void Serialise( const NodeProjector* pNode, OutputArchive& arch );
+		static NodeProjectorPtr StaticUnserialise( InputArchive& arch );
+
+
+		//-----------------------------------------------------------------------------------------
+		// Node Interface
+		//-----------------------------------------------------------------------------------------
+
+        const NODE_TYPE* GetType() const override;
+		static const NODE_TYPE* GetStaticType();
+
+
+		//-----------------------------------------------------------------------------------------
+		// Interface pattern
+		//-----------------------------------------------------------------------------------------
+		class Private;
+
+	protected:
+
+		//! Forbidden. Manage with the Ptr<> template.
+		inline ~NodeProjector() {}
+		
+		//!
+		EType Type = EType::None;
+
+	};
+
+
+	//! This node outputs a predefined Projector value.
+	//! \ingroup model
+	class MUTABLETOOLS_API NodeProjectorConstant : public NodeProjector
+	{
+	public:
+
+		//-----------------------------------------------------------------------------------------
+		// Life cycle
+		//-----------------------------------------------------------------------------------------
+
+		NodeProjectorConstant();
+
+		void SerialiseWrapper(OutputArchive& arch) const override;
+		static void Serialise( const NodeProjectorConstant* pNode, OutputArchive& arch );
+		static NodeProjectorConstantPtr StaticUnserialise( InputArchive& arch );
+
+		//-----------------------------------------------------------------------------------------
+		// Node Interface
+		//-----------------------------------------------------------------------------------------
+
+        NodePtr Clone( NodeMapPtr pMap = 0 ) const override;
+
+        const NODE_TYPE* GetType() const override;
+		static const NODE_TYPE* GetStaticType();
+
+        int GetInputCount() const override;
+        Node* GetInputNode( int i ) const override;
+        void SetInputNode( int i, NodePtr pNode ) override;
+
+		//-----------------------------------------------------------------------------------------
+		// Own Interface
+		//-----------------------------------------------------------------------------------------
+
+		//! Get the value that this node returns
+        void GetValue( PROJECTOR_TYPE* pType,
+                       float* pPosX, float* pPosY, float* pPosZ,
+					   float* pDirX, float* pDirY, float* pDirZ,
+					   float* pUpX, float* pUpY, float *pUpZ,
+                       float* pScaleU, float* pScaleV, float* pScaleW,
+                       float* pProjectionAngle ) const;
+
+		//! Set the value to be returned by this node
+        void SetValue( PROJECTOR_TYPE type,
+                       float posX, float posY, float posZ,
+					   float dirX, float dirY, float dirZ,
+					   float upX, float upY, float upZ,
+                       float scaleU, float scaleV, float scaleW,
+                       float projectionAngle );
+
+		//-----------------------------------------------------------------------------------------
+		// Interface pattern
+		//-----------------------------------------------------------------------------------------
+		class Private;
+		Private* GetPrivate() const;
+        Node::Private* GetBasePrivate() const override;
+
+	protected:
+
+		//! Forbidden. Manage with the Ptr<> template.
+		~NodeProjectorConstant();
+
+	private:
+
+		Private* m_pD;
+
+	};
+
+
+	//! Node that defines a Projector model parameter.
+	//! \ingroup model
+	class MUTABLETOOLS_API NodeProjectorParameter : public NodeProjector
+	{
+	public:
+
+		//-----------------------------------------------------------------------------------------
+		// Life cycle
+		//-----------------------------------------------------------------------------------------
+
+		NodeProjectorParameter();
+
+		void SerialiseWrapper(OutputArchive& arch) const override;
+		static void Serialise( const NodeProjectorParameter* pNode, OutputArchive& arch );
+		static NodeProjectorParameterPtr StaticUnserialise( InputArchive& arch );
+
+
+		//-----------------------------------------------------------------------------------------
+		// Node Interface
+		//-----------------------------------------------------------------------------------------
+
+        NodePtr Clone( NodeMapPtr pMap = 0 ) const override;
+
+        const NODE_TYPE* GetType() const override;
+		static const NODE_TYPE* GetStaticType();
+
+        int GetInputCount() const override;
+        Node* GetInputNode( int i ) const override;
+        void SetInputNode( int i, NodePtr pNode ) override;
+
+		//-----------------------------------------------------------------------------------------
+		// Own Interface
+		//-----------------------------------------------------------------------------------------
+
+		//! Get the name of the parameter. It will be exposed in the final compiled data.
+		const char* GetName() const;
+
+		//! Set the name of the parameter.
+		void SetName( const char* );
+
+		//! Get the uid of the parameter. It will be exposed in the final compiled data.
+		const char* GetUid() const;
+
+		//! Set the uid of the parameter.
+		void SetUid( const char* );
+
+		//! Get the default value of the parameter.
+        void GetDefaultValue( PROJECTOR_TYPE* pType,
+                              float* posX, float* posY, float* posZ,
+							  float* pDirX, float* pDirY, float* pDirZ,
+							  float* pUpX, float* pUpY, float *pUpZ,
+                              float *pScaleU, float *pScaleV, float* pScaleW,
+                              float* pProjectionAngle ) const;
+
+		//! Set the default value of the parameter.
+        void SetDefaultValue( PROJECTOR_TYPE type,
+                              float posX, float posY, float posZ,
+							  float dirX, float dirY, float dirZ,
+							  float upX, float upY, float upZ,
+                              float scaleU, float scaleV, float scaleW,
+                              float projectionAngle );
+
+        //! Set the number of ranges (dimensions) for this parameter.
+        //! By default a parameter has 0 ranges, meaning it only has one value.
+        void SetRangeCount( int i );
+        void SetRange( int i, NodeRangePtr pRange );
+
+		//-----------------------------------------------------------------------------------------
+		// Interface pattern
+		//-----------------------------------------------------------------------------------------
+		class Private;
+		Private* GetPrivate() const;
+        Node::Private* GetBasePrivate() const override;
+
+	protected:
+
+		//! Forbidden. Manage with the Ptr<> template.
+		~NodeProjectorParameter();
+
+	private:
+
+		Private* m_pD;
+
+	};
+
+
+}
diff --git a/Engine/Plugins/Experimental/Mutable/Source/MutableTools/Public/MuT/NodeRange.h b/Engine/Plugins/Experimental/Mutable/Source/MutableTools/Public/MuT/NodeRange.h
new file mode 100644
--- /dev/null
+++ b/Engine/Plugins/Experimental/Mutable/Source/MutableTools/Public/MuT/NodeRange.h
@@ -0,0 +1,64 @@
+// Copyright Epic Games, Inc. All Rights Reserved.
+
+#pragma once
+
+#include "MuT/Node.h"
+
+
+namespace mu
+{
+
+    // Forward definitions
+    class NodeRange;
+    typedef Ptr<NodeRange> NodeRangePtr;
+    typedef Ptr<const NodeRange> NodeRangePtrConst;
+
+
+    //! %Base class of any node that outputs a range.
+	//! \ingroup model
+    class MUTABLETOOLS_API NodeRange : public Node
+	{
+	public:
+
+		// Possible subclasses
+		enum class EType : uint8
+		{
+			FromScalar = 0,
+
+			None
+		};
+
+		//-----------------------------------------------------------------------------------------
+		// Life cycle
+		//-----------------------------------------------------------------------------------------
+
+        static void Serialise( const NodeRange* pNode, OutputArchive& arch );
+        static NodeRangePtr StaticUnserialise( InputArchive& arch );
+
+
+		//-----------------------------------------------------------------------------------------
+		// Node Interface
+		//-----------------------------------------------------------------------------------------
+
+		const NODE_TYPE* GetType() const override;
+		static const NODE_TYPE* GetStaticType();
+
+
+		//-----------------------------------------------------------------------------------------
+		// Interface pattern
+		//-----------------------------------------------------------------------------------------
+		class Private;
+
+	protected:
+
+		//! Forbidden. Manage with the Ptr<> template.
+        inline ~NodeRange() {}
+
+		//!
+		EType Type = EType::None;
+
+	};
+
+
+}
+
diff --git a/Engine/Plugins/Experimental/Mutable/Source/MutableTools/Public/MuT/NodeRangeFromScalar.h b/Engine/Plugins/Experimental/Mutable/Source/MutableTools/Public/MuT/NodeRangeFromScalar.h
new file mode 100644
--- /dev/null
+++ b/Engine/Plugins/Experimental/Mutable/Source/MutableTools/Public/MuT/NodeRangeFromScalar.h
@@ -0,0 +1,76 @@
+// Copyright Epic Games, Inc. All Rights Reserved.
+
+#pragma once
+
+#include "MuT/NodeRange.h"
+
+namespace mu
+{
+
+	// Forward definitions
+    class NodeScalar;
+
+
+	//! 
+	//! \ingroup model
+    class MUTABLETOOLS_API NodeRangeFromScalar : public NodeRange
+	{
+	public:
+
+		//-----------------------------------------------------------------------------------------
+		// Life cycle
+		//-----------------------------------------------------------------------------------------
+
+        NodeRangeFromScalar();
+
+		void SerialiseWrapper(OutputArchive& arch) const override;
+		static void Serialise( const NodeRangeFromScalar* pNode, OutputArchive& arch );
+        static Ptr<NodeRangeFromScalar> StaticUnserialise( InputArchive& arch );
+
+
+		//-----------------------------------------------------------------------------------------
+		// Node Interface
+		//-----------------------------------------------------------------------------------------
+
+        NodePtr Clone( NodeMapPtr pMap = nullptr ) const override;
+
+		const NODE_TYPE* GetType() const override;
+		static const NODE_TYPE* GetStaticType();
+
+        int GetInputCount() const override;
+        Node* GetInputNode( int i ) const override;
+		void SetInputNode( int i, NodePtr pNode ) override;
+
+
+		//-----------------------------------------------------------------------------------------
+		// Own Interface
+		//-----------------------------------------------------------------------------------------
+
+        //!
+        Ptr<NodeScalar> GetSize() const;
+        void SetSize( const Ptr<NodeScalar>& );
+
+        //!
+        const char* GetName() const;
+        void SetName( const char* strName );
+
+		//-----------------------------------------------------------------------------------------
+		// Interface pattern
+		//-----------------------------------------------------------------------------------------
+		class Private;
+		Private* GetPrivate() const;
+        Node::Private* GetBasePrivate() const override;
+
+	protected:
+
+		//! Forbidden. Manage with the Ptr<> template.
+        ~NodeRangeFromScalar();
+
+	private:
+
+		Private* m_pD;
+
+	};
+
+
+}
diff --git a/Engine/Plugins/Experimental/Mutable/Source/MutableTools/Public/MuT/NodeScalar.h b/Engine/Plugins/Experimental/Mutable/Source/MutableTools/Public/MuT/NodeScalar.h
new file mode 100644
--- /dev/null
+++ b/Engine/Plugins/Experimental/Mutable/Source/MutableTools/Public/MuT/NodeScalar.h
@@ -0,0 +1,71 @@
+// Copyright Epic Games, Inc. All Rights Reserved.
+
+#pragma once
+
+#include "MuT/Node.h"
+#include "MuR/Parameters.h"
+
+
+namespace mu
+{
+
+	// Forward definitions
+	class NodeScalar;
+	typedef Ptr<NodeScalar> NodeScalarPtr;
+	typedef Ptr<const NodeScalar> NodeScalarPtrConst;
+
+
+    //! %Base class of any node that outputs a scalar value.
+	//! \ingroup model
+	class MUTABLETOOLS_API NodeScalar : public Node
+	{
+	public:
+
+		// Possible subclasses
+		enum class EType : uint8
+		{
+			Constant = 0,
+			Parameter = 1,
+			EnumParameter = 2,
+			Curve = 3,
+			Switch = 4,
+			ArithmeticOperation = 5,
+			Variation = 6,
+			Table = 7,
+
+			None
+		};
+
+		//-----------------------------------------------------------------------------------------
+		// Life cycle
+		//-----------------------------------------------------------------------------------------
+
+		static void Serialise( const NodeScalar* pNode, OutputArchive& arch );
+		static NodeScalarPtr StaticUnserialise( InputArchive& arch );
+
+
+		//-----------------------------------------------------------------------------------------
+		// Node Interface
+		//-----------------------------------------------------------------------------------------
+
+		const NODE_TYPE* GetType() const override;
+		static const NODE_TYPE* GetStaticType();
+
+
+		//-----------------------------------------------------------------------------------------
+		// Interface pattern
+		//-----------------------------------------------------------------------------------------
+		class Private;
+
+	protected:
+
+		//! Forbidden. Manage with the Ptr<> template.
+		inline ~NodeScalar() {}
+
+		//!
+		EType Type = EType::None;
+
+	};
+
+
+}
diff --git a/Engine/Plugins/Experimental/Mutable/Source/MutableTools/Public/MuT/NodeScalarArithmeticOperation.h b/Engine/Plugins/Experimental/Mutable/Source/MutableTools/Public/MuT/NodeScalarArithmeticOperation.h
new file mode 100644
--- /dev/null
+++ b/Engine/Plugins/Experimental/Mutable/Source/MutableTools/Public/MuT/NodeScalarArithmeticOperation.h
@@ -0,0 +1,98 @@
+// Copyright Epic Games, Inc. All Rights Reserved.
+
+#pragma once
+
+#include "MuT/NodeScalar.h"
+
+
+namespace mu
+{
+
+	// Forward definitions
+    class NodeScalarArithmeticOperation;
+    typedef Ptr<NodeScalarArithmeticOperation> NodeScalarArithmeticOperationPtr;
+    typedef Ptr<const NodeScalarArithmeticOperation> NodeScalarArithmeticOperationPtrConst;
+
+
+    //! Perform an arithmetic operation between two scalars.
+	//! \ingroup model
+    class MUTABLETOOLS_API NodeScalarArithmeticOperation : public NodeScalar
+	{
+	public:
+
+		//-----------------------------------------------------------------------------------------
+		// Life cycle
+		//-----------------------------------------------------------------------------------------
+
+        NodeScalarArithmeticOperation();
+
+		void SerialiseWrapper(OutputArchive& arch) const override;
+		static void Serialise( const NodeScalarArithmeticOperation* pNode, OutputArchive& arch );
+        static NodeScalarArithmeticOperationPtr StaticUnserialise( InputArchive& arch );
+
+
+		//-----------------------------------------------------------------------------------------
+		// Node Interface
+		//-----------------------------------------------------------------------------------------
+
+        NodePtr Clone( NodeMapPtr pMap = 0 ) const override;
+
+        const NODE_TYPE* GetType() const override;
+		static const NODE_TYPE* GetStaticType();
+
+        int GetInputCount() const override;
+        Node* GetInputNode( int i ) const override;
+        void SetInputNode( int i, NodePtr pNode ) override;
+
+
+		//-----------------------------------------------------------------------------------------
+		// Own Interface
+		//-----------------------------------------------------------------------------------------
+
+		//! Possible arithmetic operations
+		typedef enum
+		{
+            AO_ADD = 0,
+			AO_SUBTRACT,
+			AO_MULTIPLY,
+            AO_DIVIDE,
+            _AO_COUNT
+		} OPERATION;
+
+        static const char* s_opTypeName[ _AO_COUNT ];
+
+		//! Get child selection type
+		OPERATION GetOperation() const;
+
+		//! Set the child selection type
+		void SetOperation(OPERATION);
+
+
+		//! Get the first operand
+        NodeScalarPtr GetA() const;
+        void SetA(NodeScalarPtr);
+
+		// Get the second operand
+        NodeScalarPtr GetB() const;
+        void SetB(NodeScalarPtr);
+
+		//-----------------------------------------------------------------------------------------
+		// Interface pattern
+		//-----------------------------------------------------------------------------------------
+		class Private;
+		Private* GetPrivate() const;
+        Node::Private* GetBasePrivate() const override;
+
+	protected:
+
+		//! Forbidden. Manage with the Ptr<> template.
+        ~NodeScalarArithmeticOperation();
+
+	private:
+
+		Private* m_pD;
+
+	};
+
+
+}
diff --git a/Engine/Plugins/Experimental/Mutable/Source/MutableTools/Public/MuT/NodeScalarConstant.h b/Engine/Plugins/Experimental/Mutable/Source/MutableTools/Public/MuT/NodeScalarConstant.h
new file mode 100644
--- /dev/null
+++ b/Engine/Plugins/Experimental/Mutable/Source/MutableTools/Public/MuT/NodeScalarConstant.h
@@ -0,0 +1,78 @@
+// Copyright Epic Games, Inc. All Rights Reserved.
+
+#pragma once
+
+#include "MuT/NodeScalar.h"
+#include "MuR/Parameters.h"
+
+
+namespace mu
+{
+
+	// Forward definitions
+	class NodeScalarConstant;
+	typedef Ptr<NodeScalarConstant> NodeScalarConstantPtr;
+	typedef Ptr<const NodeScalarConstant> NodeScalarConstantPtrConst;
+
+
+	//! Node returning a scalar constant value.
+	//! \ingroup model
+	class MUTABLETOOLS_API NodeScalarConstant : public NodeScalar
+	{
+	public:
+
+		//-----------------------------------------------------------------------------------------
+		// Life cycle
+		//-----------------------------------------------------------------------------------------
+
+		NodeScalarConstant();
+
+		void SerialiseWrapper(OutputArchive& arch) const override;
+		static void Serialise( const NodeScalarConstant* pNode, OutputArchive& arch );
+		static NodeScalarConstantPtr StaticUnserialise( InputArchive& arch );
+
+
+		//-----------------------------------------------------------------------------------------
+		// Node Interface
+		//-----------------------------------------------------------------------------------------
+
+        NodePtr Clone( NodeMapPtr pMap = 0 ) const override;
+
+        const NODE_TYPE* GetType() const override;
+		static const NODE_TYPE* GetStaticType();
+
+        int GetInputCount() const override;
+        Node* GetInputNode( int i ) const override;
+        void SetInputNode( int i, NodePtr pNode ) override;
+
+		//-----------------------------------------------------------------------------------------
+		// Own Interface
+		//-----------------------------------------------------------------------------------------
+
+		//! Get the value to be returned by the node.
+		float GetValue() const;
+
+		//! Set the value to be returned by the node.
+		void SetValue( float v );
+
+		//-----------------------------------------------------------------------------------------
+		// Interface pattern
+		//-----------------------------------------------------------------------------------------
+		class Private;
+		Private* GetPrivate() const;
+        Node::Private* GetBasePrivate() const override;
+
+	protected:
+
+		//! Forbidden. Manage with the Ptr<> template.
+		~NodeScalarConstant();
+
+	private:
+
+		Private* m_pD;
+
+	};
+
+
+
+}
diff --git a/Engine/Plugins/Experimental/Mutable/Source/MutableTools/Public/MuT/NodeScalarCurve.h b/Engine/Plugins/Experimental/Mutable/Source/MutableTools/Public/MuT/NodeScalarCurve.h
new file mode 100644
--- /dev/null
+++ b/Engine/Plugins/Experimental/Mutable/Source/MutableTools/Public/MuT/NodeScalarCurve.h
@@ -0,0 +1,88 @@
+// Copyright Epic Games, Inc. All Rights Reserved.
+
+
+#pragma once
+
+#include "MuT/NodeScalar.h"
+
+namespace mu
+{
+
+	// Forward definitions
+    class NodeScalarCurve;
+    typedef Ptr<NodeScalarCurve> NodeScalarCurvePtr;
+    typedef Ptr<const NodeScalarCurve> NodeScalarCurvePtrConst;
+
+
+    //! This node makes a new scalar value transforming another scalar value with a curve.
+	//! \ingroup model
+    class MUTABLETOOLS_API NodeScalarCurve : public NodeScalar
+	{
+	public:
+
+		//-----------------------------------------------------------------------------------------
+		// Life cycle
+		//-----------------------------------------------------------------------------------------
+
+        NodeScalarCurve();
+
+		void SerialiseWrapper(OutputArchive& arch) const override;
+		static void Serialise( const NodeScalarCurve* pNode, OutputArchive& arch );
+        static NodeScalarCurvePtr StaticUnserialise( InputArchive& arch );
+
+		//-----------------------------------------------------------------------------------------
+		// Node Interface
+		//-----------------------------------------------------------------------------------------
+        NodePtr Clone( NodeMapPtr pMap = nullptr ) const override;
+
+        const NODE_TYPE* GetType() const override;
+		static const NODE_TYPE* GetStaticType();
+
+        int GetInputCount() const override;
+        Node* GetInputNode( int i ) const override;
+        void SetInputNode( int i, NodePtr pNode ) override;
+
+		//-----------------------------------------------------------------------------------------
+		// Own Interface
+		//-----------------------------------------------------------------------------------------
+
+        //! Get the number of key frames in the Curve.
+		int GetKeyFrameCount() const;
+
+        //! Set the number of key frames in the Curve.
+		void SetKeyFrameCount( int );
+
+		//! Set the default value of the curve
+		void SetDefaultValue(float);
+
+		//! \param index index of the curve, from 0 to GetKeyFrameCount()-1
+        void SetKeyFrame(int index, float time, float value,
+                         float in_tangent, float in_tangent_weight,
+                         float out_tangent, float out_tangent_weight,
+                         uint8_t interp_mode,
+                         uint8_t tangent_mode, uint8_t tangent_weight_mode);
+
+        //! Get the curve time parameter
+		NodeScalarPtr GetT() const;
+		void SetT(NodeScalarPtr);
+
+		//-----------------------------------------------------------------------------------------
+		// Interface pattern
+		//-----------------------------------------------------------------------------------------
+		class Private;
+		Private* GetPrivate() const;
+        Node::Private* GetBasePrivate() const override;
+
+	protected:
+
+		//! Forbidden. Manage with the Ptr<> template.
+        ~NodeScalarCurve();
+
+	private:
+
+		Private* m_pD;
+
+	};
+
+
+}
diff --git a/Engine/Plugins/Experimental/Mutable/Source/MutableTools/Public/MuT/NodeScalarEnumParameter.h b/Engine/Plugins/Experimental/Mutable/Source/MutableTools/Public/MuT/NodeScalarEnumParameter.h
new file mode 100644
--- /dev/null
+++ b/Engine/Plugins/Experimental/Mutable/Source/MutableTools/Public/MuT/NodeScalarEnumParameter.h
@@ -0,0 +1,106 @@
+// Copyright Epic Games, Inc. All Rights Reserved.
+
+#pragma once
+
+#include "MuT/NodeScalar.h"
+#include "MuR/Parameters.h"
+#include "MuT/NodeRange.h"
+
+
+namespace mu
+{
+
+	// Forward definitions
+	class NodeImage;
+    using NodeImagePtr = Ptr<NodeImage>;
+    using NodeImagePtrConst = Ptr<const NodeImage>;
+
+	class NodeScalarEnumParameter;
+    using NodeScalarEnumParameterPtr = Ptr<NodeScalarEnumParameter>;
+    using NodeScalarEnumParameterPtrConst = Ptr<const NodeScalarEnumParameter>;
+
+    class NodeRange;
+    using NodeRangePtr = Ptr<NodeRange>;
+    using NodeRangePtrConst = Ptr<const NodeRange>;
+
+
+	//! Node that defines a scalar model parameter to be selected from a set of named values.
+	//! \ingroup model
+	class MUTABLETOOLS_API NodeScalarEnumParameter : public NodeScalar
+	{
+	public:
+
+		//-----------------------------------------------------------------------------------------
+		// Life cycle
+		//-----------------------------------------------------------------------------------------
+
+		NodeScalarEnumParameter();
+
+		void SerialiseWrapper(OutputArchive& arch) const override;
+		static void Serialise( const NodeScalarEnumParameter* pNode, OutputArchive& arch );
+		static NodeScalarEnumParameterPtr StaticUnserialise( InputArchive& arch );
+
+
+		//-----------------------------------------------------------------------------------------
+		// Node Interface
+		//-----------------------------------------------------------------------------------------
+
+        NodePtr Clone( NodeMapPtr pMap = 0 ) const override;
+
+        const NODE_TYPE* GetType() const override;
+		static const NODE_TYPE* GetStaticType();
+
+        int GetInputCount() const override;
+        Node* GetInputNode( int i ) const override;
+        void SetInputNode( int i, NodePtr pNode ) override;
+
+		//-----------------------------------------------------------------------------------------
+		// Own Interface
+		//-----------------------------------------------------------------------------------------
+
+		//! Get the name of the parameter. It will be exposed in the final compiled data.
+		const char* GetName() const;
+		void SetName( const char* );
+
+		//! Get the uid of the parameter. It will be exposed in the final compiled data.
+		const char* GetUid() const;
+		void SetUid( const char* );
+
+		//! Get the index of the default value of the parameter.
+		int GetDefaultValueIndex() const;
+		void SetDefaultValueIndex( int v );
+
+		//! Set the number of possible values.
+		void SetValueCount( int i );
+
+		//! Get the number of possible values.
+		int GetValueCount() const;
+
+		//! Set the data of one of the possible values of the parameter.
+		void SetValue( int i, float value, const char* strName );
+
+        //! Set the number of ranges (dimensions) for this parameter.
+        //! By default a parameter has 0 ranges, meaning it only has one value.
+        void SetRangeCount( int i );
+        void SetRange( int i, NodeRangePtr pRange );
+
+		//-----------------------------------------------------------------------------------------
+		// Interface pattern
+		//-----------------------------------------------------------------------------------------
+		class Private;
+		Private* GetPrivate() const;
+        Node::Private* GetBasePrivate() const override;
+
+	protected:
+
+		//! Forbidden. Manage with the Ptr<> template.
+		~NodeScalarEnumParameter();
+
+	private:
+
+		Private* m_pD;
+
+	};
+
+
+}
diff --git a/Engine/Plugins/Experimental/Mutable/Source/MutableTools/Public/MuT/NodeScalarParameter.h b/Engine/Plugins/Experimental/Mutable/Source/MutableTools/Public/MuT/NodeScalarParameter.h
new file mode 100644
--- /dev/null
+++ b/Engine/Plugins/Experimental/Mutable/Source/MutableTools/Public/MuT/NodeScalarParameter.h
@@ -0,0 +1,106 @@
+// Copyright Epic Games, Inc. All Rights Reserved.
+
+#pragma once
+
+#include "MuT/NodeScalar.h"
+#include "MuR/Parameters.h"
+
+
+namespace mu
+{
+
+	// Forward definitions
+    class NodeImage;
+    using NodeImagePtr = Ptr<NodeImage>;
+    using NodeImagePtrConst = Ptr<const NodeImage>;
+
+    class NodeRange;
+    using NodeRangePtr = Ptr<NodeRange>;
+    using NodeRangePtrConst = Ptr<const NodeRange>;
+
+    class NodeScalarParameter;
+    using NodeScalarParameterPtr = Ptr<NodeScalarParameter>;
+    using NodeScalarParameterPtrConst = Ptr<const NodeScalarParameter>;
+
+
+	//! Node that defines a scalar model parameter.
+	//! \ingroup model
+	class MUTABLETOOLS_API NodeScalarParameter : public NodeScalar
+	{
+	public:
+
+		//-----------------------------------------------------------------------------------------
+		// Life cycle
+		//-----------------------------------------------------------------------------------------
+
+		NodeScalarParameter();
+
+		void SerialiseWrapper(OutputArchive& arch) const override;
+		static void Serialise( const NodeScalarParameter* pNode, OutputArchive& arch );
+		static NodeScalarParameterPtr StaticUnserialise( InputArchive& arch );
+
+
+		//-----------------------------------------------------------------------------------------
+		// Node Interface
+		//-----------------------------------------------------------------------------------------
+
+        NodePtr Clone( NodeMapPtr pMap = 0 ) const override;
+
+        const NODE_TYPE* GetType() const override;
+		static const NODE_TYPE* GetStaticType();
+
+        int GetInputCount() const override;
+        Node* GetInputNode( int i ) const override;
+        void SetInputNode( int i, NodePtr pNode ) override;
+
+		//-----------------------------------------------------------------------------------------
+		// Own Interface
+		//-----------------------------------------------------------------------------------------
+
+		//! Get the name of the parameter. It will be exposed in the final compiled data.
+		const char* GetName() const;
+		void SetName( const char* );
+
+		//! Get the uid of the parameter. It will be exposed in the final compiled data.
+		const char* GetUid() const;
+		void SetUid( const char* );
+
+		//! Get the default value of the parameter.
+		float GetDefaultValue() const;
+		void SetDefaultValue( float v );
+
+		//! Get the additional information about the type
+		PARAMETER_DETAILED_TYPE GetDetailedType() const ;
+		void SetDetailedType( PARAMETER_DETAILED_TYPE type );
+
+        //! Set the number of additional description images associated to this parameter.
+        //! These images can be built from a parameter and, for example, they can be used to build
+        //! user-interfaces for customisation.
+        void SetDescImageCount( int i );
+        void SetDescImage( int i, NodeImagePtr pImage );
+
+        //! Set the number of ranges (dimensions) for this parameter.
+        //! By default a parameter has 0 ranges, meaning it only has one value.
+        void SetRangeCount( int i );
+        void SetRange( int i, NodeRangePtr pRange );
+
+		//-----------------------------------------------------------------------------------------
+		// Interface pattern
+		//-----------------------------------------------------------------------------------------
+		class Private;
+		Private* GetPrivate() const;
+        Node::Private* GetBasePrivate() const override;
+
+	protected:
+
+		//! Forbidden. Manage with the Ptr<> template.
+		~NodeScalarParameter();
+
+	private:
+
+		Private* m_pD;
+
+	};
+
+
+}
diff --git a/Engine/Plugins/Experimental/Mutable/Source/MutableTools/Public/MuT/NodeScalarSwitch.h b/Engine/Plugins/Experimental/Mutable/Source/MutableTools/Public/MuT/NodeScalarSwitch.h
new file mode 100644
--- /dev/null
+++ b/Engine/Plugins/Experimental/Mutable/Source/MutableTools/Public/MuT/NodeScalarSwitch.h
@@ -0,0 +1,87 @@
+// Copyright Epic Games, Inc. All Rights Reserved.
+
+#pragma once
+
+#include "MuT/NodeScalar.h"
+
+
+namespace mu
+{
+
+	// Forward definitions
+	class NodeScalar;
+	typedef Ptr<NodeScalar> NodeScalarPtr;
+	typedef Ptr<const NodeScalar> NodeScalarPtrConst;
+
+    class NodeScalarSwitch;
+    typedef Ptr<NodeScalarSwitch> NodeScalarSwitchPtr;
+    typedef Ptr<const NodeScalarSwitch> NodeScalarSwitchPtrConst;
+
+
+    //! This node selects an output Scalar from a set of input Scalars based on a parameter.
+	//! \ingroup model
+    class MUTABLETOOLS_API NodeScalarSwitch : public NodeScalar
+	{
+	public:
+
+		//-----------------------------------------------------------------------------------------
+		// Life cycle
+		//-----------------------------------------------------------------------------------------
+
+        NodeScalarSwitch();
+
+		void SerialiseWrapper(OutputArchive& arch) const override;
+		static void Serialise( const NodeScalarSwitch* pNode, OutputArchive& arch );
+        static NodeScalarSwitchPtr StaticUnserialise( InputArchive& arch );
+
+
+		//-----------------------------------------------------------------------------------------
+		// Node Interface
+		//-----------------------------------------------------------------------------------------
+
+        NodePtr Clone( NodeMapPtr pMap = nullptr ) const override;
+
+        const NODE_TYPE* GetType() const override;
+		static const NODE_TYPE* GetStaticType();
+
+        virtual int GetInputCount() const override;
+        virtual Node* GetInputNode( int i ) const override;
+        void SetInputNode( int i, NodePtr pNode ) override;
+
+		//-----------------------------------------------------------------------------------------
+		// Own Interface
+		//-----------------------------------------------------------------------------------------
+
+		//! Get the node generating the parameter used to select the option.
+		NodeScalarPtr GetParameter() const;
+		void SetParameter( NodeScalarPtr );
+
+        //! Set the number of option Scalars. It will keep the currently set targets and initialise
+		//! the new ones as null.
+		void SetOptionCount( int );
+
+        //! Get the node generating the t-th option Scalar.
+        NodeScalarPtr GetOption( int t ) const;
+        void SetOption( int t, NodeScalarPtr );
+
+
+		//-----------------------------------------------------------------------------------------
+		// Interface pattern
+		//-----------------------------------------------------------------------------------------
+		class Private;
+		Private* GetPrivate() const;
+        Node::Private* GetBasePrivate() const override;
+
+	protected:
+
+		//! Forbidden. Manage with the Ptr<> template.
+        ~NodeScalarSwitch();
+
+	private:
+
+		Private* m_pD;
+
+	};
+
+
+}
diff --git a/Engine/Plugins/Experimental/Mutable/Source/MutableTools/Public/MuT/NodeScalarTable.h b/Engine/Plugins/Experimental/Mutable/Source/MutableTools/Public/MuT/NodeScalarTable.h
new file mode 100644
--- /dev/null
+++ b/Engine/Plugins/Experimental/Mutable/Source/MutableTools/Public/MuT/NodeScalarTable.h
@@ -0,0 +1,86 @@
+// Copyright Epic Games, Inc. All Rights Reserved.
+
+#pragma once
+
+#include "MuT/NodeScalar.h"
+
+
+namespace mu
+{
+
+	// Forward definitions
+	class NodeScalarTable;
+	typedef Ptr<NodeScalarTable> NodeScalarTablePtr;
+	typedef Ptr<const NodeScalarTable> NodeScalarTablePtrConst;
+
+	class Table;
+	typedef Ptr<Table> TablePtr;
+	typedef Ptr<const Table> TablePtrConst;
+
+
+	//! This node provides the meshes stored in the column of a table.
+	//! \ingroup transform
+	class MUTABLETOOLS_API NodeScalarTable : public NodeScalar
+	{
+	public:
+
+		//-----------------------------------------------------------------------------------------
+		// Life cycle
+		//-----------------------------------------------------------------------------------------
+
+		NodeScalarTable();
+
+		void SerialiseWrapper(OutputArchive& arch) const override;
+		static void Serialise( const NodeScalarTable* pNode, OutputArchive& arch );
+		static NodeScalarTablePtr StaticUnserialise( InputArchive& arch );
+
+
+		//-----------------------------------------------------------------------------------------
+		// Node Interface
+		//-----------------------------------------------------------------------------------------
+
+		NodePtr Clone( NodeMapPtr pMap = 0 ) const override;
+
+		const NODE_TYPE* GetType() const override;
+		static const NODE_TYPE* GetStaticType();
+
+		int GetInputCount() const override;
+		Node* GetInputNode( int i ) const override;
+		void SetInputNode( int i, NodePtr pNode ) override;
+
+		//-----------------------------------------------------------------------------------------
+		// Own Interface
+		//-----------------------------------------------------------------------------------------
+
+		//! Set the name of the implicit table parameter.
+		void SetParameterName( const char* strName );
+
+		//!
+		TablePtr GetTable() const;
+		void SetTable( TablePtr );
+
+		//!
+		const char* GetColumn() const;
+		void SetColumn( const char* strName );
+
+		//-----------------------------------------------------------------------------------------
+		// Interface pattern
+		//-----------------------------------------------------------------------------------------
+		class Private;
+		Private* GetPrivate() const;
+        Node::Private* GetBasePrivate() const override;
+
+
+	protected:
+
+		//! Forbidden. Manage with the Ptr<> template.
+		~NodeScalarTable();
+
+	private:
+
+		Private* m_pD;
+
+	};
+
+
+}
diff --git a/Engine/Plugins/Experimental/Mutable/Source/MutableTools/Public/MuT/NodeScalarVariation.h b/Engine/Plugins/Experimental/Mutable/Source/MutableTools/Public/MuT/NodeScalarVariation.h
new file mode 100644
--- /dev/null
+++ b/Engine/Plugins/Experimental/Mutable/Source/MutableTools/Public/MuT/NodeScalarVariation.h
@@ -0,0 +1,83 @@
+// Copyright Epic Games, Inc. All Rights Reserved.
+
+#pragma once
+
+#include "MuT/NodeScalar.h"
+
+
+namespace mu
+{
+
+    // Forward definitions
+    class NodeScalarVariation;
+    typedef Ptr<NodeScalarVariation> NodeScalarVariationPtr;
+    typedef Ptr<const NodeScalarVariation> NodeScalarVariationPtrConst;
+
+
+    //!
+    //! \ingroup model
+    class MUTABLETOOLS_API NodeScalarVariation : public NodeScalar
+    {
+    public:
+        //-----------------------------------------------------------------------------------------
+        // Life cycle
+        //-----------------------------------------------------------------------------------------
+
+        NodeScalarVariation();
+
+		void SerialiseWrapper(OutputArchive& arch) const override;
+		static void Serialise( const NodeScalarVariation* pNode, OutputArchive& arch );
+        static NodeScalarVariationPtr StaticUnserialise( InputArchive& arch );
+
+        //-----------------------------------------------------------------------------------------
+        // Node Interface
+        //-----------------------------------------------------------------------------------------
+
+        NodePtr Clone( NodeMapPtr pMap = 0 ) const override;
+
+        const NODE_TYPE* GetType() const override;
+        static const NODE_TYPE* GetStaticType();
+
+        virtual int GetInputCount() const override;
+        virtual Node* GetInputNode( int i ) const override;
+        void SetInputNode( int i, NodePtr pNode ) override;
+
+        //-----------------------------------------------------------------------------------------
+        // Own Interface
+        //-----------------------------------------------------------------------------------------
+
+        //!
+        void SetDefaultScalar( NodeScalar* Scalar );
+
+        //! Set the number of tags to consider in this variation
+        void SetVariationCount( int count );
+
+        //!
+        int GetVariationCount() const;
+
+        //! Set the tag or state name that will enable a specific vartiation
+        void SetVariationTag( int index, const char* strTag );
+
+        //!
+        void SetVariationScalar( int index, NodeScalar* Scalar );
+
+        //!}
+
+
+        //-----------------------------------------------------------------------------------------
+        // Interface pattern
+        //-----------------------------------------------------------------------------------------
+        class Private;
+        Private* GetPrivate() const;
+        Node::Private* GetBasePrivate() const override;
+
+    protected:
+        //! Forbidden. Manage with the Ptr<> template.
+        ~NodeScalarVariation();
+
+    private:
+        Private* m_pD;
+    };
+
+
+} // namespace mu
diff --git a/Engine/Plugins/Experimental/Mutable/Source/MutableTools/Public/MuT/NodeString.h b/Engine/Plugins/Experimental/Mutable/Source/MutableTools/Public/MuT/NodeString.h
new file mode 100644
--- /dev/null
+++ b/Engine/Plugins/Experimental/Mutable/Source/MutableTools/Public/MuT/NodeString.h
@@ -0,0 +1,66 @@
+// Copyright Epic Games, Inc. All Rights Reserved.
+
+#pragma once
+
+#include "MuT/Node.h"
+#include "MuR/Parameters.h"
+
+
+namespace mu
+{
+
+	// Forward definitions
+	class NodeString;
+	typedef Ptr<NodeString> NodeStringPtr;
+	typedef Ptr<const NodeString> NodeStringPtrConst;
+
+
+    //! %Base class of any node that outputs a string value.
+	//! \ingroup model
+	class MUTABLETOOLS_API NodeString : public Node
+	{
+	public:
+
+		// Possible subclasses
+		enum class EType : uint8
+		{
+			Constant = 0,
+			Parameter = 1,
+
+			None
+		};
+
+		//-----------------------------------------------------------------------------------------
+		// Life cycle
+		//-----------------------------------------------------------------------------------------
+
+		static void Serialise( const NodeString* pNode, OutputArchive& arch );
+		static NodeStringPtr StaticUnserialise( InputArchive& arch );
+
+
+		//-----------------------------------------------------------------------------------------
+		// Node Interface
+		//-----------------------------------------------------------------------------------------
+
+		const NODE_TYPE* GetType() const override;
+		static const NODE_TYPE* GetStaticType();
+
+
+		//-----------------------------------------------------------------------------------------
+		// Interface pattern
+		//-----------------------------------------------------------------------------------------
+		class Private;
+
+	protected:
+
+		//! Forbidden. Manage with the Ptr<> template.
+		inline ~NodeString() {}
+
+		//!
+		EType Type = EType::None;
+
+
+	};
+
+
+}
diff --git a/Engine/Plugins/Experimental/Mutable/Source/MutableTools/Public/MuT/NodeStringConstant.h b/Engine/Plugins/Experimental/Mutable/Source/MutableTools/Public/MuT/NodeStringConstant.h
new file mode 100644
--- /dev/null
+++ b/Engine/Plugins/Experimental/Mutable/Source/MutableTools/Public/MuT/NodeStringConstant.h
@@ -0,0 +1,78 @@
+// Copyright Epic Games, Inc. All Rights Reserved.
+
+#pragma once
+
+#include "MuT/NodeString.h"
+#include "MuR/Parameters.h"
+
+
+namespace mu
+{
+
+	// Forward definitions
+	class NodeStringConstant;
+	typedef Ptr<NodeStringConstant> NodeStringConstantPtr;
+	typedef Ptr<const NodeStringConstant> NodeStringConstantPtrConst;
+
+
+	//! Node returning a string constant value.
+	//! \ingroup model
+	class MUTABLETOOLS_API NodeStringConstant : public NodeString
+	{
+	public:
+
+		//-----------------------------------------------------------------------------------------
+		// Life cycle
+		//-----------------------------------------------------------------------------------------
+
+		NodeStringConstant();
+
+		void SerialiseWrapper(OutputArchive& arch) const override;
+		static void Serialise( const NodeStringConstant* pNode, OutputArchive& arch );
+		static NodeStringConstantPtr StaticUnserialise( InputArchive& arch );
+
+
+		//-----------------------------------------------------------------------------------------
+		// Node Interface
+		//-----------------------------------------------------------------------------------------
+
+        NodePtr Clone( NodeMapPtr pMap = 0 ) const override;
+
+        const NODE_TYPE* GetType() const override;
+		static const NODE_TYPE* GetStaticType();
+
+        int GetInputCount() const override;
+        Node* GetInputNode( int i ) const override;
+        void SetInputNode( int i, NodePtr pNode ) override;
+
+		//-----------------------------------------------------------------------------------------
+		// Own Interface
+		//-----------------------------------------------------------------------------------------
+
+		//! Get the value to be returned by the node.
+		const char* GetValue() const;
+
+		//! Set the value to be returned by the node.
+		void SetValue( const char* v );
+
+		//-----------------------------------------------------------------------------------------
+		// Interface pattern
+		//-----------------------------------------------------------------------------------------
+		class Private;
+		Private* GetPrivate() const;
+        Node::Private* GetBasePrivate() const override;
+
+	protected:
+
+		//! Forbidden. Manage with the Ptr<> template.
+		~NodeStringConstant();
+
+	private:
+
+		Private* m_pD;
+
+	};
+
+
+
+}
diff --git a/Engine/Plugins/Experimental/Mutable/Source/MutableTools/Public/MuT/NodeStringParameter.h b/Engine/Plugins/Experimental/Mutable/Source/MutableTools/Public/MuT/NodeStringParameter.h
new file mode 100644
--- /dev/null
+++ b/Engine/Plugins/Experimental/Mutable/Source/MutableTools/Public/MuT/NodeStringParameter.h
@@ -0,0 +1,100 @@
+// Copyright Epic Games, Inc. All Rights Reserved.
+
+#pragma once
+
+#include "MuT/NodeString.h"
+#include "MuR/Parameters.h"
+
+
+namespace mu
+{
+
+	// Forward definitions
+    class NodeImage;
+    using NodeImagePtr = Ptr<NodeImage>;
+    using NodeImagePtrConst = Ptr<const NodeImage>;
+
+    class NodeRange;
+    using NodeRangePtr = Ptr<NodeRange>;
+    using NodeRangePtrConst = Ptr<const NodeRange>;
+
+    class NodeStringParameter;
+    using NodeStringParameterPtr = Ptr<NodeStringParameter>;
+    using NodeStringParameterPtrConst = Ptr<const NodeStringParameter>;
+
+
+	//! Node that defines a string model parameter.
+	//! \ingroup model
+	class MUTABLETOOLS_API NodeStringParameter : public NodeString
+	{
+	public:
+
+		//-----------------------------------------------------------------------------------------
+		// Life cycle
+		//-----------------------------------------------------------------------------------------
+
+		NodeStringParameter();
+
+		void SerialiseWrapper(OutputArchive& arch) const override;
+		static void Serialise( const NodeStringParameter* pNode, OutputArchive& arch );
+		static NodeStringParameterPtr StaticUnserialise( InputArchive& arch );
+
+
+		//-----------------------------------------------------------------------------------------
+		// Node Interface
+		//-----------------------------------------------------------------------------------------
+
+        NodePtr Clone( NodeMapPtr pMap = 0 ) const override;
+
+        const NODE_TYPE* GetType() const override;
+		static const NODE_TYPE* GetStaticType();
+
+        int GetInputCount() const override;
+        Node* GetInputNode( int i ) const override;
+        void SetInputNode( int i, NodePtr pNode ) override;
+
+		//-----------------------------------------------------------------------------------------
+		// Own Interface
+		//-----------------------------------------------------------------------------------------
+
+		//! Get the name of the parameter. It will be exposed in the final compiled data.
+		const char* GetName() const;
+		void SetName( const char* );
+
+		//! Get the uid of the parameter. It will be exposed in the final compiled data.
+		const char* GetUid() const;
+		void SetUid( const char* );
+
+		//! Get the default value of the parameter.
+		const char* GetDefaultValue() const;
+		void SetDefaultValue( const char* v );
+
+		//! Get the additional information about the type
+		PARAMETER_DETAILED_TYPE GetDetailedType() const ;
+		void SetDetailedType( PARAMETER_DETAILED_TYPE type );
+
+        //! Set the number of ranges (dimensions) for this parameter.
+        //! By default a parameter has 0 ranges, meaning it only has one value.
+        void SetRangeCount( int i );
+        void SetRange( int i, NodeRangePtr pRange );
+
+		//-----------------------------------------------------------------------------------------
+		// Interface pattern
+		//-----------------------------------------------------------------------------------------
+		class Private;
+		Private* GetPrivate() const;
+        Node::Private* GetBasePrivate() const override;
+
+	protected:
+
+		//! Forbidden. Manage with the Ptr<> template.
+		~NodeStringParameter();
+
+	private:
+
+		Private* m_pD;
+
+	};
+
+
+}
diff --git a/Engine/Plugins/Experimental/Mutable/Source/MutableTools/Public/MuT/NodeSurface.h b/Engine/Plugins/Experimental/Mutable/Source/MutableTools/Public/MuT/NodeSurface.h
new file mode 100644
--- /dev/null
+++ b/Engine/Plugins/Experimental/Mutable/Source/MutableTools/Public/MuT/NodeSurface.h
@@ -0,0 +1,60 @@
+// Copyright Epic Games, Inc. All Rights Reserved.
+
+#pragma once
+
+#include "MuT/Node.h"
+
+
+namespace mu
+{
+
+	// Forward definitions
+    class NodeSurface;
+    typedef Ptr<NodeSurface> NodeSurfacePtr;
+    typedef Ptr<const NodeSurface> NodeSurfacePtrConst;
+
+
+
+    //! This class is the parent of all nodes that output a Surface.
+	//! \ingroup model
+    class MUTABLETOOLS_API NodeSurface : public Node
+	{
+	public:
+
+		// Possible subclasses
+		enum class EType : uint8
+		{
+			New = 0,
+			Edit = 1,
+			Variation = 3,
+
+			None
+		};
+
+		//-----------------------------------------------------------------------------------------
+		// Life cycle
+		//-----------------------------------------------------------------------------------------
+
+        static void Serialise( const NodeSurface* pNode, OutputArchive& arch );
+        static NodeSurfacePtr StaticUnserialise( InputArchive& arch );
+
+		//-----------------------------------------------------------------------------------------
+		// Node Interface
+		//-----------------------------------------------------------------------------------------
+
+        const NODE_TYPE* GetType() const override;
+		static const NODE_TYPE* GetStaticType();
+
+
+	protected:
+
+		//! Forbidden. Manage with the Ptr<> template.
+        inline ~NodeSurface() {}
+
+		//!
+		EType Type = EType::None;
+
+	};
+
+
+}
diff --git a/Engine/Plugins/Experimental/Mutable/Source/MutableTools/Public/MuT/NodeSurfaceEdit.h b/Engine/Plugins/Experimental/Mutable/Source/MutableTools/Public/MuT/NodeSurfaceEdit.h
new file mode 100644
--- /dev/null
+++ b/Engine/Plugins/Experimental/Mutable/Source/MutableTools/Public/MuT/NodeSurfaceEdit.h
@@ -0,0 +1,163 @@
+// Copyright Epic Games, Inc. All Rights Reserved.
+
+#pragma once
+
+#include "MuT/NodeSurface.h"
+
+
+namespace mu
+{
+
+	// Forward definitions
+    class NodeMesh;
+    typedef Ptr<NodeMesh> NodeMeshPtr;
+    typedef Ptr<const NodeMesh> NodeMeshPtrConst;
+
+    class NodeImage;
+	typedef Ptr<NodeImage> NodeImagePtr;
+	typedef Ptr<const NodeImage> NodeImagePtrConst;
+
+	class NodePatchMesh;
+	typedef Ptr<NodePatchMesh> NodePatchMeshPtr;
+	typedef Ptr<const NodePatchMesh> NodePatchMeshPtrConst;
+
+	class NodePatchImage;
+	typedef Ptr<NodePatchImage> NodePatchImagePtr;
+	typedef Ptr<const NodePatchImage> NodePatchImagePtrConst;
+
+    class NodeSurfaceEdit;
+    typedef Ptr<NodeSurfaceEdit> NodeSurfaceEditPtr;
+    typedef Ptr<const NodeSurfaceEdit> NodeSurfaceEditPtrConst;
+
+	class NodeScalar;
+	typedef Ptr<NodeScalar> NodeScalarPtr;
+	typedef Ptr<const NodeScalar> NodeScalarPtrConst;
+
+
+
+	//! This node modifies a node of the parent object of the object that this node belongs to.
+    //! It allows to extend, cut and morph the parent Surface's meshes.
+    //! It also allows to patch the parent Surface's textures.
+	//! \ingroup model
+    class MUTABLETOOLS_API NodeSurfaceEdit : public NodeSurface
+	{
+	public:
+
+		//-----------------------------------------------------------------------------------------
+		// Life cycle
+		//-----------------------------------------------------------------------------------------
+
+        NodeSurfaceEdit();
+
+		void SerialiseWrapper(OutputArchive& arch) const override;
+		static void Serialise( const NodeSurfaceEdit* pNode, OutputArchive& arch );
+        static NodeSurfaceEditPtr StaticUnserialise( InputArchive& arch );
+
+		//-----------------------------------------------------------------------------------------
+		// Node Interface
+		//-----------------------------------------------------------------------------------------
+
+        NodePtr Clone( NodeMapPtr pMap = 0 ) const override;
+
+        const NODE_TYPE* GetType() const override;
+		static const NODE_TYPE* GetStaticType();
+
+        virtual int GetInputCount() const override;
+        virtual Node* GetInputNode( int i ) const override;
+        void SetInputNode( int i, NodePtr pNode ) override;
+
+		//-----------------------------------------------------------------------------------------
+		// Own Interface
+		//-----------------------------------------------------------------------------------------
+
+        //! Set the parent Surface to modify. It should be a node from the same LOD of one of
+		//! the parent objects of this node's object.
+        void SetParent( NodeSurface* );
+        NodeSurface* GetParent() const;
+
+        //! \name Tags
+        //! \{
+
+        //! Add a tag to the surface:
+        //! - the surface will be affected by modifier nodes with the same tag
+        //! - the tag will be enabled when the surface is added to an object, and it can activate
+        //! variations for any surface.
+        void AddTag(const char* tagName);
+
+        //! Get the number of tags added to the Surface.
+        int GetTagCount() const;
+
+        //! Get a tag string from an index (0 to GetTagCount-1)
+        const char* GetTag( int ) const;
+
+        //! \}
+
+		//! \name Mesh modification
+		//! \{
+
+		//! Mesh a node to extend the parent mesh
+        void SetMesh( NodePatchMesh* );
+
+		//! Get the node to extend the parent mesh
+        NodePatchMesh* GetMesh() const;
+
+		//! Set a node to morph the parent mesh
+        void SetMorph( NodeMeshPtr );
+
+		//! Get the node to morph the parent mesh
+        NodeMeshPtr GetMorph() const;
+
+		//! Set the factor: weight used to select (and combine) the morphs to apply.
+		void SetFactor( NodeScalarPtr );
+
+		//! Get the factor: weight used to select (and combine) the morphs to apply.
+		NodeScalarPtr GetFactor() const;
+
+		//! \}
+
+		//! \name Image modification
+		//! \{
+
+		//! Get the number of images to modify. Should match the parent image count.
+		int GetImageCount() const;
+
+		//! Set the number of images to modify. Should match the parent image count.
+		void SetImageCount( int );
+
+		//! Get the image node that extends a parent image, adding new blocks
+		NodeImagePtr GetImage( int index ) const;
+
+		//! Set the image node that extends a parent image, adding new blocks
+		void SetImage( int index, NodeImagePtr );
+
+		//! Get the image node that patches a parent image, blending on existing blocks
+		NodePatchImagePtr GetPatch( int index ) const;
+
+		//! Set the image node that patches a parent image, blending on existing blocks
+		void SetPatch( int index, NodePatchImagePtr );
+
+		//!}
+
+
+		//-----------------------------------------------------------------------------------------
+		// Interface pattern
+		//-----------------------------------------------------------------------------------------
+		class Private;
+		Private* GetPrivate() const;
+        Node::Private* GetBasePrivate() const override;
+
+	protected:
+
+		//! Forbidden. Manage with the Ptr<> template.
+        ~NodeSurfaceEdit();
+
+	private:
+
+		Private* m_pD;
+
+	};
+
+
+
+}
+
