Updated more paths in Mutable to contain the module name.

diff --git a/Engine/Plugins/Experimental/Mutable/Source/CustomizableObject/Private/MuCO/CustomizableObjectSystem.cpp b/Engine/Plugins/Experimental/Mutable/Source/CustomizableObject/Private/MuCO/CustomizableObjectSystem.cpp
--- a/Engine/Plugins/Experimental/Mutable/Source/CustomizableObject/Private/MuCO/CustomizableObjectSystem.cpp
+++ b/Engine/Plugins/Experimental/Mutable/Source/CustomizableObject/Private/MuCO/CustomizableObjectSystem.cpp
@@ -6,9 +6,9 @@
 #include "MuCO/CustomizableObject.h"
 #include "MuCO/CustomizableObjectPrivate.h"
 #include "MuCO/CustomizableInstancePrivateData.h"
-#include "UnrealMutableImageProvider.h"
+#include "MuCO/UnrealMutableImageProvider.h"
 #include "MuCO/UnrealMutableModelDiskStreamer.h"
-#include "LogBenchmarkUtil.h"
+#include "MuCO/LogBenchmarkUtil.h"
 #include "MuCO/ICustomizableObjectModule.h"
 #include "MuCO/CustomizableInstanceLODManagement.h"
 #include "MuCO/UnrealBakeHelpers.h"
@@ -16,7 +16,7 @@
 #include "AssetRegistry/AssetRegistryModule.h"
 #include "Engine/SkeletalMesh.h"
 #include "Animation/Skeleton.h"
-#include "LogInformationUtil.h"
+#include "MuCO/LogInformationUtil.h"
 #include "TimerManager.h"
 #include "SceneManagement.h"
 #include "GameFramework/PlayerController.h"
diff --git a/Engine/Plugins/Experimental/Mutable/Source/CustomizableObject/Private/MuCO/LogBenchmarkUtil.cpp b/Engine/Plugins/Experimental/Mutable/Source/CustomizableObject/Private/MuCO/LogBenchmarkUtil.cpp
--- a/Engine/Plugins/Experimental/Mutable/Source/CustomizableObject/Private/MuCO/LogBenchmarkUtil.cpp
+++ b/Engine/Plugins/Experimental/Mutable/Source/CustomizableObject/Private/MuCO/LogBenchmarkUtil.cpp
@@ -1,6 +1,6 @@
 // Copyright Epic Games, Inc. All Rights Reserved.
 
-#include "LogBenchmarkUtil.h"
+#include "MuCO/LogBenchmarkUtil.h"
 #include "Misc/FileHelper.h"
 #include "Containers/UnrealString.h"
 #include "HAL/PlatformFileManager.h"
@@ -120,4 +120,4 @@ bool LogBenchmarkUtil::isLoggingActive()
 	return bLoggingActive;
 }
 
-LogBenchmarkUtil::LogBenchmarkUtil() {}
\ No newline at end of file
+LogBenchmarkUtil::LogBenchmarkUtil() {}
diff --git a/Engine/Plugins/Experimental/Mutable/Source/CustomizableObject/Private/MuCO/LogInformationUtil.cpp b/Engine/Plugins/Experimental/Mutable/Source/CustomizableObject/Private/MuCO/LogInformationUtil.cpp
--- a/Engine/Plugins/Experimental/Mutable/Source/CustomizableObject/Private/MuCO/LogInformationUtil.cpp
+++ b/Engine/Plugins/Experimental/Mutable/Source/CustomizableObject/Private/MuCO/LogInformationUtil.cpp
@@ -1,6 +1,6 @@
 // Copyright Epic Games, Inc. All Rights Reserved.
 
-#include "LogInformationUtil.h"
+#include "MuCO/LogInformationUtil.h"
 #include "Engine/SkeletalMesh.h"
 #include "Engine/Texture2D.h"
 #include "MuCO/CustomizableObject.h"
diff --git a/Engine/Plugins/Experimental/Mutable/Source/CustomizableObject/Private/MuCO/UnrealMutableImageProvider.cpp b/Engine/Plugins/Experimental/Mutable/Source/CustomizableObject/Private/MuCO/UnrealMutableImageProvider.cpp
--- a/Engine/Plugins/Experimental/Mutable/Source/CustomizableObject/Private/MuCO/UnrealMutableImageProvider.cpp
+++ b/Engine/Plugins/Experimental/Mutable/Source/CustomizableObject/Private/MuCO/UnrealMutableImageProvider.cpp
@@ -1,6 +1,6 @@
 // Copyright Epic Games, Inc. All Rights Reserved.
 
-#include "UnrealMutableImageProvider.h"
+#include "MuCO/UnrealMutableImageProvider.h"
 
 #include "MuCO/CustomizableObject.h"
 #include "MuCO/CustomizableObjectSystem.h"
diff --git a/Engine/Plugins/Experimental/Mutable/Source/CustomizableObjectEditor/Classes/MuCOE/Nodes/CustomizableObjectNodeEditMaterialBase.h b/Engine/Plugins/Experimental/Mutable/Source/CustomizableObjectEditor/Classes/MuCOE/Nodes/CustomizableObjectNodeEditMaterialBase.h
--- a/Engine/Plugins/Experimental/Mutable/Source/CustomizableObjectEditor/Classes/MuCOE/Nodes/CustomizableObjectNodeEditMaterialBase.h
+++ b/Engine/Plugins/Experimental/Mutable/Source/CustomizableObjectEditor/Classes/MuCOE/Nodes/CustomizableObjectNodeEditMaterialBase.h
@@ -3,7 +3,7 @@
 #pragma once
 
 #include "MuCOE/Nodes/CustomizableObjectNode.h"
-#include "CustomizableObjectNodeParentedMaterial.h"
+#include "MuCOE/Nodes/CustomizableObjectNodeParentedMaterial.h"
 
 #include "CustomizableObjectNodeEditMaterialBase.generated.h"
 
diff --git a/Engine/Plugins/Experimental/Mutable/Source/CustomizableObjectEditor/Classes/MuCOE/Nodes/CustomizableObjectNodeMeshClipDeform.h b/Engine/Plugins/Experimental/Mutable/Source/CustomizableObjectEditor/Classes/MuCOE/Nodes/CustomizableObjectNodeMeshClipDeform.h
--- a/Engine/Plugins/Experimental/Mutable/Source/CustomizableObjectEditor/Classes/MuCOE/Nodes/CustomizableObjectNodeMeshClipDeform.h
+++ b/Engine/Plugins/Experimental/Mutable/Source/CustomizableObjectEditor/Classes/MuCOE/Nodes/CustomizableObjectNodeMeshClipDeform.h
@@ -2,7 +2,7 @@
 
 #pragma once
 
-#include "CustomizableObjectNodeModifierBase.h"
+#include "MuCOE/Nodes/CustomizableObjectNodeModifierBase.h"
 
 #include "CustomizableObjectNodeMeshClipDeform.generated.h"
 
diff --git a/Engine/Plugins/Experimental/Mutable/Source/CustomizableObjectEditor/Classes/MuCOE/Nodes/CustomizableObjectNodeMeshClipMorph.h b/Engine/Plugins/Experimental/Mutable/Source/CustomizableObjectEditor/Classes/MuCOE/Nodes/CustomizableObjectNodeMeshClipMorph.h
--- a/Engine/Plugins/Experimental/Mutable/Source/CustomizableObjectEditor/Classes/MuCOE/Nodes/CustomizableObjectNodeMeshClipMorph.h
+++ b/Engine/Plugins/Experimental/Mutable/Source/CustomizableObjectEditor/Classes/MuCOE/Nodes/CustomizableObjectNodeMeshClipMorph.h
@@ -2,7 +2,7 @@
 
 #pragma once
 
-#include "CustomizableObjectNodeModifierBase.h"
+#include "MuCOE/Nodes/CustomizableObjectNodeModifierBase.h"
 
 #include "CustomizableObjectNodeMeshClipMorph.generated.h"
 
diff --git a/Engine/Plugins/Experimental/Mutable/Source/CustomizableObjectEditor/Classes/MuCOE/Nodes/CustomizableObjectNodeMeshClipWithMesh.h b/Engine/Plugins/Experimental/Mutable/Source/CustomizableObjectEditor/Classes/MuCOE/Nodes/CustomizableObjectNodeMeshClipWithMesh.h
--- a/Engine/Plugins/Experimental/Mutable/Source/CustomizableObjectEditor/Classes/MuCOE/Nodes/CustomizableObjectNodeMeshClipWithMesh.h
+++ b/Engine/Plugins/Experimental/Mutable/Source/CustomizableObjectEditor/Classes/MuCOE/Nodes/CustomizableObjectNodeMeshClipWithMesh.h
@@ -2,8 +2,8 @@
 
 #pragma once
 
-#include "CustomizableObjectNodeModifierBase.h"
-#include "CustomizableObjectNodeMaterial.h"
+#include "MuCOE/Nodes/CustomizableObjectNodeModifierBase.h"
+#include "MuCOE/Nodes/CustomizableObjectNodeMaterial.h"
 
 #include "CustomizableObjectNodeMeshClipWithMesh.generated.h"
 
diff --git a/Engine/Plugins/Experimental/Mutable/Source/CustomizableObjectEditor/Classes/MuCOE/Nodes/CustomizableObjectNodeMeshMorphStackApplication.h b/Engine/Plugins/Experimental/Mutable/Source/CustomizableObjectEditor/Classes/MuCOE/Nodes/CustomizableObjectNodeMeshMorphStackApplication.h
--- a/Engine/Plugins/Experimental/Mutable/Source/CustomizableObjectEditor/Classes/MuCOE/Nodes/CustomizableObjectNodeMeshMorphStackApplication.h
+++ b/Engine/Plugins/Experimental/Mutable/Source/CustomizableObjectEditor/Classes/MuCOE/Nodes/CustomizableObjectNodeMeshMorphStackApplication.h
@@ -2,7 +2,7 @@
 
 #pragma once
 
-#include "CustomizableObjectNode.h"
+#include "MuCOE/Nodes/CustomizableObjectNode.h"
 #include "CustomizableObjectNodeMeshMorphStackApplication.generated.h"
 
 
diff --git a/Engine/Plugins/Experimental/Mutable/Source/CustomizableObjectEditor/Classes/MuCOE/Nodes/CustomizableObjectNodeMeshMorphStackDefinition.h b/Engine/Plugins/Experimental/Mutable/Source/CustomizableObjectEditor/Classes/MuCOE/Nodes/CustomizableObjectNodeMeshMorphStackDefinition.h
--- a/Engine/Plugins/Experimental/Mutable/Source/CustomizableObjectEditor/Classes/MuCOE/Nodes/CustomizableObjectNodeMeshMorphStackDefinition.h
+++ b/Engine/Plugins/Experimental/Mutable/Source/CustomizableObjectEditor/Classes/MuCOE/Nodes/CustomizableObjectNodeMeshMorphStackDefinition.h
@@ -2,7 +2,7 @@
 
 #pragma once
 
-#include "CustomizableObjectNode.h"
+#include "MuCOE/Nodes/CustomizableObjectNode.h"
 #include "CustomizableObjectNodeMeshMorphStackDefinition.generated.h"
 
 
diff --git a/Engine/Plugins/Experimental/Mutable/Source/CustomizableObjectEditor/Classes/MuCOE/Nodes/CustomizableObjectNodeParentedMaterial.h b/Engine/Plugins/Experimental/Mutable/Source/CustomizableObjectEditor/Classes/MuCOE/Nodes/CustomizableObjectNodeParentedMaterial.h
--- a/Engine/Plugins/Experimental/Mutable/Source/CustomizableObjectEditor/Classes/MuCOE/Nodes/CustomizableObjectNodeParentedMaterial.h
+++ b/Engine/Plugins/Experimental/Mutable/Source/CustomizableObjectEditor/Classes/MuCOE/Nodes/CustomizableObjectNodeParentedMaterial.h
@@ -2,7 +2,7 @@
 
 #pragma once
 
-#include "CustomizableObjectNodeParentedNode.h"
+#include "MuCOE/Nodes/CustomizableObjectNodeParentedNode.h"
 #include "MuCOE/Nodes/CustomizableObjectNode.h"
 #include "MuCOE/Nodes/CustomizableObjectNodeMaterial.h"
 
diff --git a/Engine/Plugins/Experimental/Mutable/Source/CustomizableObjectEditor/Classes/MuCOE/Nodes/CustomizableObjectNodeTable.h b/Engine/Plugins/Experimental/Mutable/Source/CustomizableObjectEditor/Classes/MuCOE/Nodes/CustomizableObjectNodeTable.h
--- a/Engine/Plugins/Experimental/Mutable/Source/CustomizableObjectEditor/Classes/MuCOE/Nodes/CustomizableObjectNodeTable.h
+++ b/Engine/Plugins/Experimental/Mutable/Source/CustomizableObjectEditor/Classes/MuCOE/Nodes/CustomizableObjectNodeTable.h
@@ -2,7 +2,7 @@
 
 #pragma once
 
-#include "CustomizableObjectNodeMaterial.h"
+#include "MuCOE/Nodes/CustomizableObjectNodeMaterial.h"
 #include "Engine/DataTable.h"
 
 #include "MuCOE/Nodes/CustomizableObjectNode.h"
diff --git a/Engine/Plugins/Experimental/Mutable/Source/CustomizableObjectEditor/Classes/MuCOE/Nodes/CustomizableObjectNodeUseMaterial.h b/Engine/Plugins/Experimental/Mutable/Source/CustomizableObjectEditor/Classes/MuCOE/Nodes/CustomizableObjectNodeUseMaterial.h
--- a/Engine/Plugins/Experimental/Mutable/Source/CustomizableObjectEditor/Classes/MuCOE/Nodes/CustomizableObjectNodeUseMaterial.h
+++ b/Engine/Plugins/Experimental/Mutable/Source/CustomizableObjectEditor/Classes/MuCOE/Nodes/CustomizableObjectNodeUseMaterial.h
@@ -2,7 +2,7 @@
 
 #pragma once
 
-#include "CustomizableObjectNodeParentedMaterial.h"
+#include "MuCOE/Nodes/CustomizableObjectNodeParentedMaterial.h"
 #include "MuCOE/Nodes/CustomizableObjectNode.h"
 
 
diff --git a/Engine/Plugins/Experimental/Mutable/Source/CustomizableObjectEditor/Private/MuCOE/AssetTypeActions_CustomizableObject.cpp b/Engine/Plugins/Experimental/Mutable/Source/CustomizableObjectEditor/Private/MuCOE/AssetTypeActions_CustomizableObject.cpp
--- a/Engine/Plugins/Experimental/Mutable/Source/CustomizableObjectEditor/Private/MuCOE/AssetTypeActions_CustomizableObject.cpp
+++ b/Engine/Plugins/Experimental/Mutable/Source/CustomizableObjectEditor/Private/MuCOE/AssetTypeActions_CustomizableObject.cpp
@@ -1,6 +1,6 @@
 // Copyright Epic Games, Inc. All Rights Reserved.
 
-#include "AssetTypeActions_CustomizableObject.h"
+#include "MuCOE/AssetTypeActions_CustomizableObject.h"
 #include "MuCOE/CustomizableObjectEditorModule.h"
 #include "ContentBrowserModule.h"
 #include "AssetRegistry/AssetRegistryModule.h"
diff --git a/Engine/Plugins/Experimental/Mutable/Source/CustomizableObjectEditor/Private/MuCOE/AssetTypeActions_CustomizableObjectInstance.cpp b/Engine/Plugins/Experimental/Mutable/Source/CustomizableObjectEditor/Private/MuCOE/AssetTypeActions_CustomizableObjectInstance.cpp
--- a/Engine/Plugins/Experimental/Mutable/Source/CustomizableObjectEditor/Private/MuCOE/AssetTypeActions_CustomizableObjectInstance.cpp
+++ b/Engine/Plugins/Experimental/Mutable/Source/CustomizableObjectEditor/Private/MuCOE/AssetTypeActions_CustomizableObjectInstance.cpp
@@ -1,6 +1,6 @@
 // Copyright Epic Games, Inc. All Rights Reserved.
 
-#include "AssetTypeActions_CustomizableObjectInstance.h"
+#include "MuCOE/AssetTypeActions_CustomizableObjectInstance.h"
 #include "MuCOE/CustomizableObjectEditorModule.h"
 #include "Framework/MultiBox/MultiBoxBuilder.h"
 #include "MuCOE/UnrealEditorPortabilityHelpers.h"
diff --git a/Engine/Plugins/Experimental/Mutable/Source/CustomizableObjectEditor/Private/MuCOE/CustomizableObjectBakeHelpers.cpp b/Engine/Plugins/Experimental/Mutable/Source/CustomizableObjectEditor/Private/MuCOE/CustomizableObjectBakeHelpers.cpp
--- a/Engine/Plugins/Experimental/Mutable/Source/CustomizableObjectEditor/Private/MuCOE/CustomizableObjectBakeHelpers.cpp
+++ b/Engine/Plugins/Experimental/Mutable/Source/CustomizableObjectEditor/Private/MuCOE/CustomizableObjectBakeHelpers.cpp
@@ -1,6 +1,6 @@
 // Copyright Epic Games, Inc. All Rights Reserved.
 
-#include "CustomizableObjectBakeHelpers.h"
+#include "MuCOE/CustomizableObjectBakeHelpers.h"
 
 #include "MuCOE/CustomizableObjectEditor.h"
 #include "MuCO/CustomizableObjectMipDataProvider.h"
diff --git a/Engine/Plugins/Experimental/Mutable/Source/CustomizableObjectEditor/Private/MuCOE/CustomizableObjectCustomSettings.cpp b/Engine/Plugins/Experimental/Mutable/Source/CustomizableObjectEditor/Private/MuCOE/CustomizableObjectCustomSettings.cpp
--- a/Engine/Plugins/Experimental/Mutable/Source/CustomizableObjectEditor/Private/MuCOE/CustomizableObjectCustomSettings.cpp
+++ b/Engine/Plugins/Experimental/Mutable/Source/CustomizableObjectEditor/Private/MuCOE/CustomizableObjectCustomSettings.cpp
@@ -1,6 +1,6 @@
 // Copyright Epic Games, Inc. All Rights Reserved.
 
-#include "CustomizableObjectCustomSettings.h"
+#include "MuCOE/CustomizableObjectCustomSettings.h"
 #include "MuCOE/SCustomizableObjectEditorViewport.h"
 #include "MuCOE/CustomizableObjectEditorViewportClient.h"
 #include "Widgets/Input/SEditableTextBox.h"
diff --git a/Engine/Plugins/Experimental/Mutable/Source/CustomizableObjectEditor/Private/MuCOE/CustomizableObjectCustomSettings.h b/Engine/Plugins/Experimental/Mutable/Source/CustomizableObjectEditor/Private/MuCOE/CustomizableObjectCustomSettings.h
--- a/Engine/Plugins/Experimental/Mutable/Source/CustomizableObjectEditor/Private/MuCOE/CustomizableObjectCustomSettings.h
+++ b/Engine/Plugins/Experimental/Mutable/Source/CustomizableObjectEditor/Private/MuCOE/CustomizableObjectCustomSettings.h
@@ -7,7 +7,7 @@
 #include "Widgets/Views/SListView.h"
 #include "AssetRegistry/AssetData.h"
 //#include "Engine/DataTable.h"
-#include "CustomizableObjectEditorViewportLights.h"
+#include "MuCOE/CustomizableObjectEditorViewportLights.h"
 #include "CustomizableObjectCustomSettings.generated.h"
 
 UCLASS()
@@ -82,4 +82,4 @@ private:
 	TSharedPtr<class SExpandableArea> PointLightProperties;
 	TSharedPtr<class SExpandableArea> SpotLightProperties;
 
-};
\ No newline at end of file
+};
diff --git a/Engine/Plugins/Experimental/Mutable/Source/CustomizableObjectEditor/Private/MuCOE/CustomizableObjectCustomSettingsDetails.cpp b/Engine/Plugins/Experimental/Mutable/Source/CustomizableObjectEditor/Private/MuCOE/CustomizableObjectCustomSettingsDetails.cpp
--- a/Engine/Plugins/Experimental/Mutable/Source/CustomizableObjectEditor/Private/MuCOE/CustomizableObjectCustomSettingsDetails.cpp
+++ b/Engine/Plugins/Experimental/Mutable/Source/CustomizableObjectEditor/Private/MuCOE/CustomizableObjectCustomSettingsDetails.cpp
@@ -1,10 +1,10 @@
 // Copyright Epic Games, Inc. All Rights Reserved.
 
-#include "CustomizableObjectCustomSettingsDetails.h"
+#include "MuCOE/CustomizableObjectCustomSettingsDetails.h"
 #include "DetailCategoryBuilder.h"
 #include "DetailLayoutBuilder.h"
 #include "DetailWidgetRow.h"
-#include "CustomizableObjectCustomSettings.h"
+#include "MuCOE/CustomizableObjectCustomSettings.h"
 
 TSharedRef<IDetailCustomization> FCustomizableObjectCustomSettingsDetails::MakeInstance()
 {
diff --git a/Engine/Plugins/Experimental/Mutable/Source/CustomizableObjectEditor/Private/MuCOE/CustomizableObjectDebugger.cpp b/Engine/Plugins/Experimental/Mutable/Source/CustomizableObjectEditor/Private/MuCOE/CustomizableObjectDebugger.cpp
--- a/Engine/Plugins/Experimental/Mutable/Source/CustomizableObjectEditor/Private/MuCOE/CustomizableObjectDebugger.cpp
+++ b/Engine/Plugins/Experimental/Mutable/Source/CustomizableObjectEditor/Private/MuCOE/CustomizableObjectDebugger.cpp
@@ -7,9 +7,9 @@
 #include "MuCOE/CustomizableObjectEditor.h"
 
 #include "MuCOE/SCustomizableObjectProperties.h"
-#include "SMutableGraphViewer.h"
-#include "SMutableCodeViewer.h"
-#include "SMutableObjectViewer.h"
+#include "MuCOE/SMutableGraphViewer.h"
+#include "MuCOE/SMutableCodeViewer.h"
+#include "MuCOE/SMutableObjectViewer.h"
 
 #include "Editor.h"
 #include "AssetRegistry/AssetRegistryModule.h"
diff --git a/Engine/Plugins/Experimental/Mutable/Source/CustomizableObjectEditor/Private/MuCOE/CustomizableObjectEditor.cpp b/Engine/Plugins/Experimental/Mutable/Source/CustomizableObjectEditor/Private/MuCOE/CustomizableObjectEditor.cpp
--- a/Engine/Plugins/Experimental/Mutable/Source/CustomizableObjectEditor/Private/MuCOE/CustomizableObjectEditor.cpp
+++ b/Engine/Plugins/Experimental/Mutable/Source/CustomizableObjectEditor/Private/MuCOE/CustomizableObjectEditor.cpp
@@ -65,12 +65,12 @@
 
 #include "Runtime/Launch/Resources/Version.h"
 #include "AdvancedPreviewSceneModule.h"
-#include "CustomizableObjectCustomSettings.h"
+#include "MuCOE/CustomizableObjectCustomSettings.h"
 #include "MuCOE/SCustomizableObjectEditorAdvancedPreviewSettings.h"
-#include "CustomizableObjectEditorNodeContextCommands.h"
+#include "MuCOE/CustomizableObjectEditorNodeContextCommands.h"
 #include "Animation/AnimSingleNodeInstance.h"
 #include "AdvancedPreviewScene.h"
-#include "CustomizableObjectBakeHelpers.h"
+#include "MuCOE/CustomizableObjectBakeHelpers.h"
 #include "MuCOE/UnrealEditorPortabilityHelpers.h"
 #include "GenericPlatform/GenericPlatformApplicationMisc.h"
 #include "ISinglePropertyView.h"
diff --git a/Engine/Plugins/Experimental/Mutable/Source/CustomizableObjectEditor/Private/MuCOE/CustomizableObjectEditor.h b/Engine/Plugins/Experimental/Mutable/Source/CustomizableObjectEditor/Private/MuCOE/CustomizableObjectEditor.h
--- a/Engine/Plugins/Experimental/Mutable/Source/CustomizableObjectEditor/Private/MuCOE/CustomizableObjectEditor.h
+++ b/Engine/Plugins/Experimental/Mutable/Source/CustomizableObjectEditor/Private/MuCOE/CustomizableObjectEditor.h
@@ -11,7 +11,7 @@
 #include "MuCOE/CustomizableObjectCompiler.h"
 #include "Widgets/Input/SNumericEntryBox.h"
 #include "Framework/MultiBox/MultiBoxBuilder.h"
-#include "StatePerformanceTesting.h"
+#include "MuCOE/StatePerformanceTesting.h"
 #include "MuCO/CustomizableObjectInstance.h"
 #include "MuCOE/SCustomizableObjectEditorViewport.h"
 #include "MuCOE/CustomizableObjectEditorViewportClient.h"
diff --git a/Engine/Plugins/Experimental/Mutable/Source/CustomizableObjectEditor/Private/MuCOE/CustomizableObjectEditorModule.cpp b/Engine/Plugins/Experimental/Mutable/Source/CustomizableObjectEditor/Private/MuCOE/CustomizableObjectEditorModule.cpp
--- a/Engine/Plugins/Experimental/Mutable/Source/CustomizableObjectEditor/Private/MuCOE/CustomizableObjectEditorModule.cpp
+++ b/Engine/Plugins/Experimental/Mutable/Source/CustomizableObjectEditor/Private/MuCOE/CustomizableObjectEditorModule.cpp
@@ -41,8 +41,8 @@
 #include "Modules/ModuleManager.h"
 #include "Interfaces/IPluginManager.h"
 #include "AssetToolsModule.h"
-#include "AssetTypeActions_CustomizableObject.h"
-#include "AssetTypeActions_CustomizableObjectInstance.h"
+#include "MuCOE/AssetTypeActions_CustomizableObject.h"
+#include "MuCOE/AssetTypeActions_CustomizableObjectInstance.h"
 #include "Styling/SlateStyle.h"
 #include "Brushes/SlateImageBrush.h"
 #include "MessageLogModule.h"
@@ -50,11 +50,11 @@
 #include "MuCOE/CustomizableObjectEditorSettings.h"
 #include "ISettingsModule.h"
 #include "ISettingsSection.h"
-#include "CustomizableObjectCustomSettingsDetails.h"
+#include "MuCOE/CustomizableObjectCustomSettingsDetails.h"
 #include "AssetRegistry/AssetRegistryModule.h"
 #include "HAL/FileManager.h"
 #include "CoreGlobals.h"
-#include "CustomizableObjectCustomSettings.h"
+#include "MuCOE/CustomizableObjectCustomSettings.h"
 #include "Settings/ProjectPackagingSettings.h"
 #include "MuCOE/UnrealEditorPortabilityHelpers.h"
 #include "CookOnTheSide/CookOnTheFlyServer.h" // This is necessary to register the cook delegate.
diff --git a/Engine/Plugins/Experimental/Mutable/Source/CustomizableObjectEditor/Private/MuCOE/CustomizableObjectEditorNodeContextCommands.cpp b/Engine/Plugins/Experimental/Mutable/Source/CustomizableObjectEditor/Private/MuCOE/CustomizableObjectEditorNodeContextCommands.cpp
--- a/Engine/Plugins/Experimental/Mutable/Source/CustomizableObjectEditor/Private/MuCOE/CustomizableObjectEditorNodeContextCommands.cpp
+++ b/Engine/Plugins/Experimental/Mutable/Source/CustomizableObjectEditor/Private/MuCOE/CustomizableObjectEditorNodeContextCommands.cpp
@@ -1,6 +1,6 @@
 // Copyright Epic Games, Inc. All Rights Reserved.
 
-#include "CustomizableObjectEditorNodeContextCommands.h"
+#include "MuCOE/CustomizableObjectEditorNodeContextCommands.h"
 
 #define LOCTEXT_NAMESPACE "CustomizableObjectEditorNodeContextCommands"
 
diff --git a/Engine/Plugins/Experimental/Mutable/Source/CustomizableObjectEditor/Private/MuCOE/CustomizableObjectEditorViewportClient.cpp b/Engine/Plugins/Experimental/Mutable/Source/CustomizableObjectEditor/Private/MuCOE/CustomizableObjectEditorViewportClient.cpp
--- a/Engine/Plugins/Experimental/Mutable/Source/CustomizableObjectEditor/Private/MuCOE/CustomizableObjectEditorViewportClient.cpp
+++ b/Engine/Plugins/Experimental/Mutable/Source/CustomizableObjectEditor/Private/MuCOE/CustomizableObjectEditorViewportClient.cpp
@@ -2,7 +2,7 @@
 
 #include "MuCOE/CustomizableObjectEditorViewportClient.h"
 #include "MuCOE/CustomizableObjectEditorActions.h"
-#include "CustomizableObjectBakeHelpers.h"
+#include "MuCOE/CustomizableObjectBakeHelpers.h"
 
 #include "MouseDeltaTracker.h"
 #include "PreviewScene.h"
diff --git a/Engine/Plugins/Experimental/Mutable/Source/CustomizableObjectEditor/Private/MuCOE/EdGraphSchema_CustomizableObject.cpp b/Engine/Plugins/Experimental/Mutable/Source/CustomizableObjectEditor/Private/MuCOE/EdGraphSchema_CustomizableObject.cpp
--- a/Engine/Plugins/Experimental/Mutable/Source/CustomizableObjectEditor/Private/MuCOE/EdGraphSchema_CustomizableObject.cpp
+++ b/Engine/Plugins/Experimental/Mutable/Source/CustomizableObjectEditor/Private/MuCOE/EdGraphSchema_CustomizableObject.cpp
@@ -24,7 +24,7 @@
 #include "Engine/SkeletalMesh.h"
 #include "Engine/StaticMesh.h"
 
-#include "CustomizableObjectEditorNodeContextCommands.h"
+#include "MuCOE/CustomizableObjectEditorNodeContextCommands.h"
 #include "Settings/EditorStyleSettings.h"
 
 
diff --git a/Engine/Plugins/Experimental/Mutable/Source/CustomizableObjectEditor/Private/MuCOE/FMutableViewportClient.cpp b/Engine/Plugins/Experimental/Mutable/Source/CustomizableObjectEditor/Private/MuCOE/FMutableViewportClient.cpp
--- a/Engine/Plugins/Experimental/Mutable/Source/CustomizableObjectEditor/Private/MuCOE/FMutableViewportClient.cpp
+++ b/Engine/Plugins/Experimental/Mutable/Source/CustomizableObjectEditor/Private/MuCOE/FMutableViewportClient.cpp
@@ -1,6 +1,6 @@
 // Copyright Epic Games, Inc. All Rights Reserved.
 
-#include "FMutableViewportClient.h"
+#include "MuCOE/FMutableViewportClient.h"
 
 #include "AdvancedPreviewScene.h"
 
diff --git a/Engine/Plugins/Experimental/Mutable/Source/CustomizableObjectEditor/Private/MuCOE/Nodes/CustomizableObjectNodeCopyMaterialDetails.cpp b/Engine/Plugins/Experimental/Mutable/Source/CustomizableObjectEditor/Private/MuCOE/Nodes/CustomizableObjectNodeCopyMaterialDetails.cpp
--- a/Engine/Plugins/Experimental/Mutable/Source/CustomizableObjectEditor/Private/MuCOE/Nodes/CustomizableObjectNodeCopyMaterialDetails.cpp
+++ b/Engine/Plugins/Experimental/Mutable/Source/CustomizableObjectEditor/Private/MuCOE/Nodes/CustomizableObjectNodeCopyMaterialDetails.cpp
@@ -1,6 +1,6 @@
 // Copyright Epic Games, Inc. All Rights Reserved.
 
-#include "CustomizableObjectNodeCopyMaterialDetails.h"
+#include "MuCOE/Nodes/CustomizableObjectNodeCopyMaterialDetails.h"
 
 #include "DetailLayoutBuilder.h"
 
diff --git a/Engine/Plugins/Experimental/Mutable/Source/CustomizableObjectEditor/Private/MuCOE/Nodes/CustomizableObjectNodeEditMaterialBaseDetails.cpp b/Engine/Plugins/Experimental/Mutable/Source/CustomizableObjectEditor/Private/MuCOE/Nodes/CustomizableObjectNodeEditMaterialBaseDetails.cpp
--- a/Engine/Plugins/Experimental/Mutable/Source/CustomizableObjectEditor/Private/MuCOE/Nodes/CustomizableObjectNodeEditMaterialBaseDetails.cpp
+++ b/Engine/Plugins/Experimental/Mutable/Source/CustomizableObjectEditor/Private/MuCOE/Nodes/CustomizableObjectNodeEditMaterialBaseDetails.cpp
@@ -1,6 +1,6 @@
 // Copyright Epic Games, Inc. All Rights Reserved.
 
-#include "CustomizableObjectNodeEditMaterialBaseDetails.h"
+#include "MuCOE/Nodes/CustomizableObjectNodeEditMaterialBaseDetails.h"
 #include "MuCOE/CustomizableObjectEditorModule.h"
 #include "PropertyCustomizationHelpers.h"
 #include "MuCOE/CustomizableObjectEditorUtilities.h"
diff --git a/Engine/Plugins/Experimental/Mutable/Source/CustomizableObjectEditor/Private/MuCOE/Nodes/CustomizableObjectNodeEditMaterialBaseDetails.h b/Engine/Plugins/Experimental/Mutable/Source/CustomizableObjectEditor/Private/MuCOE/Nodes/CustomizableObjectNodeEditMaterialBaseDetails.h
--- a/Engine/Plugins/Experimental/Mutable/Source/CustomizableObjectEditor/Private/MuCOE/Nodes/CustomizableObjectNodeEditMaterialBaseDetails.h
+++ b/Engine/Plugins/Experimental/Mutable/Source/CustomizableObjectEditor/Private/MuCOE/Nodes/CustomizableObjectNodeEditMaterialBaseDetails.h
@@ -2,7 +2,7 @@
 
 #pragma once
 
-#include "CustomizableObjectNodeParentedMaterialDetails.h"
+#include "MuCOE/Nodes/CustomizableObjectNodeParentedMaterialDetails.h"
 
 #include "Templates/SharedPointer.h"
 #include "UObject/WeakObjectPtr.h"
diff --git a/Engine/Plugins/Experimental/Mutable/Source/CustomizableObjectEditor/Private/MuCOE/Nodes/CustomizableObjectNodeExternalPinDetails.cpp b/Engine/Plugins/Experimental/Mutable/Source/CustomizableObjectEditor/Private/MuCOE/Nodes/CustomizableObjectNodeExternalPinDetails.cpp
--- a/Engine/Plugins/Experimental/Mutable/Source/CustomizableObjectEditor/Private/MuCOE/Nodes/CustomizableObjectNodeExternalPinDetails.cpp
+++ b/Engine/Plugins/Experimental/Mutable/Source/CustomizableObjectEditor/Private/MuCOE/Nodes/CustomizableObjectNodeExternalPinDetails.cpp
@@ -1,6 +1,6 @@
 // Copyright Epic Games, Inc. All Rights Reserved.
 
-#include "CustomizableObjectNodeExternalPinDetails.h"
+#include "MuCOE/Nodes/CustomizableObjectNodeExternalPinDetails.h"
 #include "MuCOE/CustomizableObjectEditorModule.h"
 #include "PropertyCustomizationHelpers.h"
 #include "MuCOE/CustomizableObjectEditorUtilities.h"
diff --git a/Engine/Plugins/Experimental/Mutable/Source/CustomizableObjectEditor/Private/MuCOE/Nodes/CustomizableObjectNodeLayoutBlocksDetails.cpp b/Engine/Plugins/Experimental/Mutable/Source/CustomizableObjectEditor/Private/MuCOE/Nodes/CustomizableObjectNodeLayoutBlocksDetails.cpp
--- a/Engine/Plugins/Experimental/Mutable/Source/CustomizableObjectEditor/Private/MuCOE/Nodes/CustomizableObjectNodeLayoutBlocksDetails.cpp
+++ b/Engine/Plugins/Experimental/Mutable/Source/CustomizableObjectEditor/Private/MuCOE/Nodes/CustomizableObjectNodeLayoutBlocksDetails.cpp
@@ -1,6 +1,6 @@
 // Copyright Epic Games, Inc. All Rights Reserved.
 
-#include "CustomizableObjectNodeLayoutBlocksDetails.h"
+#include "MuCOE/Nodes/CustomizableObjectNodeLayoutBlocksDetails.h"
 #include "MuCOE/CustomizableObjectEditorModule.h"
 #include "PropertyCustomizationHelpers.h"
 
diff --git a/Engine/Plugins/Experimental/Mutable/Source/CustomizableObjectEditor/Private/MuCOE/Nodes/CustomizableObjectNodeMaterial.cpp b/Engine/Plugins/Experimental/Mutable/Source/CustomizableObjectEditor/Private/MuCOE/Nodes/CustomizableObjectNodeMaterial.cpp
--- a/Engine/Plugins/Experimental/Mutable/Source/CustomizableObjectEditor/Private/MuCOE/Nodes/CustomizableObjectNodeMaterial.cpp
+++ b/Engine/Plugins/Experimental/Mutable/Source/CustomizableObjectEditor/Private/MuCOE/Nodes/CustomizableObjectNodeMaterial.cpp
@@ -7,7 +7,7 @@
 #include "MuCOE/EdGraphSchema_CustomizableObject.h"
 #include "MuCOE/GraphTraversal.h"
 #include "ObjectEditorUtils.h"
-#include "SPinViewerNodeMaterialPinImageDetails.h"
+#include "MuCOE/Nodes/SPinViewerNodeMaterialPinImageDetails.h"
 #include "Materials/MaterialExpressionTextureCoordinate.h"
 #include "MuCOE/Nodes/CustomizableObjectNodeSkeletalMesh.h"
 #include "MuCOE/Nodes/CustomizableObjectNodeStaticMesh.h"
diff --git a/Engine/Plugins/Experimental/Mutable/Source/CustomizableObjectEditor/Private/MuCOE/Nodes/CustomizableObjectNodeMeshClipMorphDetails.cpp b/Engine/Plugins/Experimental/Mutable/Source/CustomizableObjectEditor/Private/MuCOE/Nodes/CustomizableObjectNodeMeshClipMorphDetails.cpp
--- a/Engine/Plugins/Experimental/Mutable/Source/CustomizableObjectEditor/Private/MuCOE/Nodes/CustomizableObjectNodeMeshClipMorphDetails.cpp
+++ b/Engine/Plugins/Experimental/Mutable/Source/CustomizableObjectEditor/Private/MuCOE/Nodes/CustomizableObjectNodeMeshClipMorphDetails.cpp
@@ -1,6 +1,6 @@
 // Copyright Epic Games, Inc. All Rights Reserved.
 
-#include "CustomizableObjectNodeMeshClipMorphDetails.h"
+#include "MuCOE/Nodes/CustomizableObjectNodeMeshClipMorphDetails.h"
 #include "MuCOE/CustomizableObjectEditorModule.h"
 #include "PropertyCustomizationHelpers.h"
 #include "MuCOE/CustomizableObjectEditorUtilities.h"
diff --git a/Engine/Plugins/Experimental/Mutable/Source/CustomizableObjectEditor/Private/MuCOE/Nodes/CustomizableObjectNodeMeshClipWithMeshDetails.cpp b/Engine/Plugins/Experimental/Mutable/Source/CustomizableObjectEditor/Private/MuCOE/Nodes/CustomizableObjectNodeMeshClipWithMeshDetails.cpp
--- a/Engine/Plugins/Experimental/Mutable/Source/CustomizableObjectEditor/Private/MuCOE/Nodes/CustomizableObjectNodeMeshClipWithMeshDetails.cpp
+++ b/Engine/Plugins/Experimental/Mutable/Source/CustomizableObjectEditor/Private/MuCOE/Nodes/CustomizableObjectNodeMeshClipWithMeshDetails.cpp
@@ -1,6 +1,6 @@
 // Copyright Epic Games, Inc. All Rights Reserved.
 
-#include "CustomizableObjectNodeMeshClipWithMeshDetails.h"
+#include "MuCOE/Nodes/CustomizableObjectNodeMeshClipWithMeshDetails.h"
 #include "MuCOE/CustomizableObjectEditorModule.h"
 #include "PropertyCustomizationHelpers.h"
 #include "MuCOE/CustomizableObjectEditorUtilities.h"
diff --git a/Engine/Plugins/Experimental/Mutable/Source/CustomizableObjectEditor/Private/MuCOE/Nodes/CustomizableObjectNodeMeshMorphDetails.cpp b/Engine/Plugins/Experimental/Mutable/Source/CustomizableObjectEditor/Private/MuCOE/Nodes/CustomizableObjectNodeMeshMorphDetails.cpp
--- a/Engine/Plugins/Experimental/Mutable/Source/CustomizableObjectEditor/Private/MuCOE/Nodes/CustomizableObjectNodeMeshMorphDetails.cpp
+++ b/Engine/Plugins/Experimental/Mutable/Source/CustomizableObjectEditor/Private/MuCOE/Nodes/CustomizableObjectNodeMeshMorphDetails.cpp
@@ -1,6 +1,6 @@
 // Copyright Epic Games, Inc. All Rights Reserved.
 
-#include "CustomizableObjectNodeMeshMorphDetails.h"
+#include "MuCOE/Nodes/CustomizableObjectNodeMeshMorphDetails.h"
 #include "MuCOE/CustomizableObjectEditorModule.h"
 #include "MuCOE/CustomizableObjectCompiler.h"
 #include "MuCOE/CustomizableObjectEditorUtilities.h"
diff --git a/Engine/Plugins/Experimental/Mutable/Source/CustomizableObjectEditor/Private/MuCOE/Nodes/CustomizableObjectNodeMeshReshapeDetails.cpp b/Engine/Plugins/Experimental/Mutable/Source/CustomizableObjectEditor/Private/MuCOE/Nodes/CustomizableObjectNodeMeshReshapeDetails.cpp
--- a/Engine/Plugins/Experimental/Mutable/Source/CustomizableObjectEditor/Private/MuCOE/Nodes/CustomizableObjectNodeMeshReshapeDetails.cpp
+++ b/Engine/Plugins/Experimental/Mutable/Source/CustomizableObjectEditor/Private/MuCOE/Nodes/CustomizableObjectNodeMeshReshapeDetails.cpp
@@ -1,6 +1,6 @@
 // Copyright Epic Games, Inc. All Rights Reserved.
 
-#include "CustomizableObjectNodeMeshReshapeDetails.h"
+#include "MuCOE/Nodes/CustomizableObjectNodeMeshReshapeDetails.h"
 #include "MuCOE/Nodes/CustomizableObjectNodeMeshReshape.h"
 #include "MuCOE/Nodes/CustomizableObjectNodeSkeletalMesh.h"
 #include "MuCOE/Nodes/CustomizableObjectNodeTable.h"
diff --git a/Engine/Plugins/Experimental/Mutable/Source/CustomizableObjectEditor/Private/MuCOE/Nodes/CustomizableObjectNodeMeshReshapeSelectionDetails.cpp b/Engine/Plugins/Experimental/Mutable/Source/CustomizableObjectEditor/Private/MuCOE/Nodes/CustomizableObjectNodeMeshReshapeSelectionDetails.cpp
--- a/Engine/Plugins/Experimental/Mutable/Source/CustomizableObjectEditor/Private/MuCOE/Nodes/CustomizableObjectNodeMeshReshapeSelectionDetails.cpp
+++ b/Engine/Plugins/Experimental/Mutable/Source/CustomizableObjectEditor/Private/MuCOE/Nodes/CustomizableObjectNodeMeshReshapeSelectionDetails.cpp
@@ -1,6 +1,6 @@
 // Copyright Epic Games, Inc. All Rights Reserved.
 
-#include "CustomizableObjectNodeMeshReshapeSelectionDetails.h"
+#include "MuCOE/Nodes/CustomizableObjectNodeMeshReshapeSelectionDetails.h"
 
 #include "MuCOE/CustomizableObjectEditorModule.h"
 #include "MuCOE/CustomizableObjectCompiler.h"
diff --git a/Engine/Plugins/Experimental/Mutable/Source/CustomizableObjectEditor/Private/MuCOE/Nodes/CustomizableObjectNodeMorphMaterialDetails.cpp b/Engine/Plugins/Experimental/Mutable/Source/CustomizableObjectEditor/Private/MuCOE/Nodes/CustomizableObjectNodeMorphMaterialDetails.cpp
--- a/Engine/Plugins/Experimental/Mutable/Source/CustomizableObjectEditor/Private/MuCOE/Nodes/CustomizableObjectNodeMorphMaterialDetails.cpp
+++ b/Engine/Plugins/Experimental/Mutable/Source/CustomizableObjectEditor/Private/MuCOE/Nodes/CustomizableObjectNodeMorphMaterialDetails.cpp
@@ -1,6 +1,6 @@
 // Copyright Epic Games, Inc. All Rights Reserved.
 
-#include "CustomizableObjectNodeMorphMaterialDetails.h"
+#include "MuCOE/Nodes/CustomizableObjectNodeMorphMaterialDetails.h"
 #include "MuCOE/CustomizableObjectEditorModule.h"
 #include "MuCOE/CustomizableObjectCompiler.h"
 #include "PropertyCustomizationHelpers.h"
diff --git a/Engine/Plugins/Experimental/Mutable/Source/CustomizableObjectEditor/Private/MuCOE/Nodes/CustomizableObjectNodeObjectDetails.cpp b/Engine/Plugins/Experimental/Mutable/Source/CustomizableObjectEditor/Private/MuCOE/Nodes/CustomizableObjectNodeObjectDetails.cpp
--- a/Engine/Plugins/Experimental/Mutable/Source/CustomizableObjectEditor/Private/MuCOE/Nodes/CustomizableObjectNodeObjectDetails.cpp
+++ b/Engine/Plugins/Experimental/Mutable/Source/CustomizableObjectEditor/Private/MuCOE/Nodes/CustomizableObjectNodeObjectDetails.cpp
@@ -1,6 +1,6 @@
 // Copyright Epic Games, Inc. All Rights Reserved.
 
-#include "CustomizableObjectNodeObjectDetails.h"
+#include "MuCOE/Nodes/CustomizableObjectNodeObjectDetails.h"
 
 #include "MuCOE/CustomizableObjectEditorModule.h"
 #include "PropertyCustomizationHelpers.h"
diff --git a/Engine/Plugins/Experimental/Mutable/Source/CustomizableObjectEditor/Private/MuCOE/Nodes/CustomizableObjectNodeParentedMaterialDetails.cpp b/Engine/Plugins/Experimental/Mutable/Source/CustomizableObjectEditor/Private/MuCOE/Nodes/CustomizableObjectNodeParentedMaterialDetails.cpp
--- a/Engine/Plugins/Experimental/Mutable/Source/CustomizableObjectEditor/Private/MuCOE/Nodes/CustomizableObjectNodeParentedMaterialDetails.cpp
+++ b/Engine/Plugins/Experimental/Mutable/Source/CustomizableObjectEditor/Private/MuCOE/Nodes/CustomizableObjectNodeParentedMaterialDetails.cpp
@@ -1,6 +1,6 @@
 // Copyright Epic Games, Inc. All Rights Reserved.
 
-#include "CustomizableObjectNodeParentedMaterialDetails.h"
+#include "MuCOE/Nodes/CustomizableObjectNodeParentedMaterialDetails.h"
 
 #include "MuCOE/CustomizableObjectEditorModule.h"
 #include "PropertyCustomizationHelpers.h"
diff --git a/Engine/Plugins/Experimental/Mutable/Source/CustomizableObjectEditor/Private/MuCOE/Nodes/CustomizableObjectNodePinViewerDetails.cpp b/Engine/Plugins/Experimental/Mutable/Source/CustomizableObjectEditor/Private/MuCOE/Nodes/CustomizableObjectNodePinViewerDetails.cpp
--- a/Engine/Plugins/Experimental/Mutable/Source/CustomizableObjectEditor/Private/MuCOE/Nodes/CustomizableObjectNodePinViewerDetails.cpp
+++ b/Engine/Plugins/Experimental/Mutable/Source/CustomizableObjectEditor/Private/MuCOE/Nodes/CustomizableObjectNodePinViewerDetails.cpp
@@ -1,6 +1,6 @@
 // Copyright Epic Games, Inc. All Rights Reserved.
 
-#include "CustomizableObjectNodePinViewerDetails.h"
+#include "MuCOE/Nodes/CustomizableObjectNodePinViewerDetails.h"
 
 #include "DetailLayoutBuilder.h"
 
@@ -20,4 +20,4 @@ void FCustomizableObjectNodePinViewerDetails::CustomizeDetails(IDetailLayoutBuil
 	PinViewerAttachToDetailCustomization(DetailBuilder);
 }
 
-#undef LOCTEXT_NAMESPACE
\ No newline at end of file
+#undef LOCTEXT_NAMESPACE
diff --git a/Engine/Plugins/Experimental/Mutable/Source/CustomizableObjectEditor/Private/MuCOE/Nodes/CustomizableObjectNodeProjectorParameterDetails.cpp b/Engine/Plugins/Experimental/Mutable/Source/CustomizableObjectEditor/Private/MuCOE/Nodes/CustomizableObjectNodeProjectorParameterDetails.cpp
--- a/Engine/Plugins/Experimental/Mutable/Source/CustomizableObjectEditor/Private/MuCOE/Nodes/CustomizableObjectNodeProjectorParameterDetails.cpp
+++ b/Engine/Plugins/Experimental/Mutable/Source/CustomizableObjectEditor/Private/MuCOE/Nodes/CustomizableObjectNodeProjectorParameterDetails.cpp
@@ -1,6 +1,6 @@
 // Copyright Epic Games, Inc. All Rights Reserved.
 
-#include "CustomizableObjectNodeProjectorParameterDetails.h"
+#include "MuCOE/Nodes/CustomizableObjectNodeProjectorParameterDetails.h"
 #include "MuCOE/CustomizableObjectEditorModule.h"
 #include "PropertyCustomizationHelpers.h"
 #include "MuCOE/CustomizableObjectEditorUtilities.h"
diff --git a/Engine/Plugins/Experimental/Mutable/Source/CustomizableObjectEditor/Private/MuCOE/Nodes/CustomizableObjectNodeSkeletalMeshDetails.cpp b/Engine/Plugins/Experimental/Mutable/Source/CustomizableObjectEditor/Private/MuCOE/Nodes/CustomizableObjectNodeSkeletalMeshDetails.cpp
--- a/Engine/Plugins/Experimental/Mutable/Source/CustomizableObjectEditor/Private/MuCOE/Nodes/CustomizableObjectNodeSkeletalMeshDetails.cpp
+++ b/Engine/Plugins/Experimental/Mutable/Source/CustomizableObjectEditor/Private/MuCOE/Nodes/CustomizableObjectNodeSkeletalMeshDetails.cpp
@@ -1,6 +1,6 @@
 // Copyright Epic Games, Inc. All Rights Reserved.
 
-#include "CustomizableObjectNodeSkeletalMeshDetails.h"
+#include "MuCOE/Nodes/CustomizableObjectNodeSkeletalMeshDetails.h"
 #include "MuCOE/CustomizableObjectEditorModule.h"
 #include "PropertyCustomizationHelpers.h"
 
diff --git a/Engine/Plugins/Experimental/Mutable/Source/CustomizableObjectEditor/Private/MuCOE/Nodes/CustomizableObjectNodeTableDetails.cpp b/Engine/Plugins/Experimental/Mutable/Source/CustomizableObjectEditor/Private/MuCOE/Nodes/CustomizableObjectNodeTableDetails.cpp
--- a/Engine/Plugins/Experimental/Mutable/Source/CustomizableObjectEditor/Private/MuCOE/Nodes/CustomizableObjectNodeTableDetails.cpp
+++ b/Engine/Plugins/Experimental/Mutable/Source/CustomizableObjectEditor/Private/MuCOE/Nodes/CustomizableObjectNodeTableDetails.cpp
@@ -1,6 +1,6 @@
 // Copyright Epic Games, Inc. All Rights Reserved.
 
-#include "CustomizableObjectNodeTableDetails.h"
+#include "MuCOE/Nodes/CustomizableObjectNodeTableDetails.h"
 #include "MuCOE/Nodes/CustomizableObjectNode.h"
 #include "MuCOE/Nodes/CustomizableObjectNodeTable.h"
 
diff --git a/Engine/Plugins/Experimental/Mutable/Source/CustomizableObjectEditor/Private/MuCOE/Nodes/SPinViewerNodeMaterialPinImageDetails.cpp b/Engine/Plugins/Experimental/Mutable/Source/CustomizableObjectEditor/Private/MuCOE/Nodes/SPinViewerNodeMaterialPinImageDetails.cpp
--- a/Engine/Plugins/Experimental/Mutable/Source/CustomizableObjectEditor/Private/MuCOE/Nodes/SPinViewerNodeMaterialPinImageDetails.cpp
+++ b/Engine/Plugins/Experimental/Mutable/Source/CustomizableObjectEditor/Private/MuCOE/Nodes/SPinViewerNodeMaterialPinImageDetails.cpp
@@ -1,6 +1,6 @@
 // Copyright Epic Games, Inc. All Rights Reserved.
 
-#include "SPinViewerNodeMaterialPinImageDetails.h"
+#include "MuCOE/Nodes/SPinViewerNodeMaterialPinImageDetails.h"
 
 #include "ISinglePropertyView.h"
 #include "PerPlatformPropertyCustomization.h"
@@ -204,4 +204,4 @@ void SPinViewerNodeMaterialPinImageDetails::UVLayoutOnValueChanged(int32 Value)
 	PinData->UVLayout = Value;
 }
 
-#undef LOCTEXT_NAMESPACE
\ No newline at end of file
+#undef LOCTEXT_NAMESPACE
diff --git a/Engine/Plugins/Experimental/Mutable/Source/CustomizableObjectEditor/Private/MuCOE/PinViewer/SPinViewer.cpp b/Engine/Plugins/Experimental/Mutable/Source/CustomizableObjectEditor/Private/MuCOE/PinViewer/SPinViewer.cpp
--- a/Engine/Plugins/Experimental/Mutable/Source/CustomizableObjectEditor/Private/MuCOE/PinViewer/SPinViewer.cpp
+++ b/Engine/Plugins/Experimental/Mutable/Source/CustomizableObjectEditor/Private/MuCOE/PinViewer/SPinViewer.cpp
@@ -9,7 +9,7 @@
 #include "Widgets/Layout/SScrollBox.h"
 #include "MuCOE/EdGraphSchema_CustomizableObject.h"
 #include "ScopedTransaction.h"
-#include "SPinViewerListRow.h"
+#include "MuCOE/PinViewer/SPinViewerListRow.h"
 
 #define LOCTEXT_NAMESPACE "CustomizableObjectEditor"
 
diff --git a/Engine/Plugins/Experimental/Mutable/Source/CustomizableObjectEditor/Private/MuCOE/PinViewer/SPinViewerDetailRowIndent.cpp b/Engine/Plugins/Experimental/Mutable/Source/CustomizableObjectEditor/Private/MuCOE/PinViewer/SPinViewerDetailRowIndent.cpp
--- a/Engine/Plugins/Experimental/Mutable/Source/CustomizableObjectEditor/Private/MuCOE/PinViewer/SPinViewerDetailRowIndent.cpp
+++ b/Engine/Plugins/Experimental/Mutable/Source/CustomizableObjectEditor/Private/MuCOE/PinViewer/SPinViewerDetailRowIndent.cpp
@@ -1,6 +1,6 @@
 // Copyright Epic Games, Inc. All Rights Reserved.
 
-#include "SPinViewerDetailRowIndent.h"
+#include "MuCOE/PinViewer/SPinViewerDetailRowIndent.h"
 
 #include "MuCOE/Widgets/SMutableExpandableTableRow.h"
 #include "Widgets/Layout/SBox.h"
diff --git a/Engine/Plugins/Experimental/Mutable/Source/CustomizableObjectEditor/Private/MuCOE/SCustomizableInstanceProperties.cpp b/Engine/Plugins/Experimental/Mutable/Source/CustomizableObjectEditor/Private/MuCOE/SCustomizableInstanceProperties.cpp
--- a/Engine/Plugins/Experimental/Mutable/Source/CustomizableObjectEditor/Private/MuCOE/SCustomizableInstanceProperties.cpp
+++ b/Engine/Plugins/Experimental/Mutable/Source/CustomizableObjectEditor/Private/MuCOE/SCustomizableInstanceProperties.cpp
@@ -10,7 +10,7 @@
 #include "MuCOE/CustomizableObjectEditorUtilities.h"
 
 #include "Widgets/Input/STextComboBox.h"
-#include "SMutableTextSearchBox.h"
+#include "MuCOE/SMutableTextSearchBox.h"
 #include "Widgets/Colors/SColorPicker.h"
 #include "Widgets/Colors/SColorBlock.h"
 #include "Widgets/Input/SCheckBox.h"
diff --git a/Engine/Plugins/Experimental/Mutable/Source/CustomizableObjectEditor/Private/MuCOE/SCustomizableObjectEditorPerformanceReport.cpp b/Engine/Plugins/Experimental/Mutable/Source/CustomizableObjectEditor/Private/MuCOE/SCustomizableObjectEditorPerformanceReport.cpp
--- a/Engine/Plugins/Experimental/Mutable/Source/CustomizableObjectEditor/Private/MuCOE/SCustomizableObjectEditorPerformanceReport.cpp
+++ b/Engine/Plugins/Experimental/Mutable/Source/CustomizableObjectEditor/Private/MuCOE/SCustomizableObjectEditorPerformanceReport.cpp
@@ -3,7 +3,7 @@
 
 #include "MuCOE/SCustomizableObjectEditorPerformanceReport.h"
 #include "MuCO/CustomizableObjectInstance.h"
-#include "StatePerformanceTesting.h"
+#include "MuCOE/StatePerformanceTesting.h"
 #include "PropertyCustomizationHelpers.h"
 #include "IDetailsView.h"
 #include "Widgets/Layout/SScrollBox.h"
diff --git a/Engine/Plugins/Experimental/Mutable/Source/CustomizableObjectEditor/Private/MuCOE/SCustomizableObjectEditorPerformanceReport.h b/Engine/Plugins/Experimental/Mutable/Source/CustomizableObjectEditor/Private/MuCOE/SCustomizableObjectEditorPerformanceReport.h
--- a/Engine/Plugins/Experimental/Mutable/Source/CustomizableObjectEditor/Private/MuCOE/SCustomizableObjectEditorPerformanceReport.h
+++ b/Engine/Plugins/Experimental/Mutable/Source/CustomizableObjectEditor/Private/MuCOE/SCustomizableObjectEditorPerformanceReport.h
@@ -6,7 +6,7 @@
 #include "Widgets/Views/SListView.h"
 #include "Misc/NotifyHook.h"
 #include "MuCO/CustomizableObject.h"
-#include "StressTest.h"
+#include "MuCOE/StressTest.h"
 #include "MuCOE/Widgets/SMutableExpandableTableRow.h"
 #include "MuCOE/SCustomizableObjectEditorTextureAnalyzer.h"
 #include "SCustomizableObjectEditorPerformanceReport.generated.h"
diff --git a/Engine/Plugins/Experimental/Mutable/Source/CustomizableObjectEditor/Private/MuCOE/SMutableBoolViewer.cpp b/Engine/Plugins/Experimental/Mutable/Source/CustomizableObjectEditor/Private/MuCOE/SMutableBoolViewer.cpp
--- a/Engine/Plugins/Experimental/Mutable/Source/CustomizableObjectEditor/Private/MuCOE/SMutableBoolViewer.cpp
+++ b/Engine/Plugins/Experimental/Mutable/Source/CustomizableObjectEditor/Private/MuCOE/SMutableBoolViewer.cpp
@@ -1,6 +1,6 @@
 // Copyright Epic Games, Inc. All Rights Reserved.
 
-#include "SMutableBoolViewer.h"
+#include "MuCOE/SMutableBoolViewer.h"
 
 #include "Widgets/SBoxPanel.h"
 #include "Internationalization/Text.h"
@@ -53,4 +53,4 @@ FSlateColor SMutableBoolViewer::GetColorForValue() const
 	return this->bBoolValue ? this->TrueValueColor : this->FalseValueColor;
 }
 
-#undef LOCTEXT_NAMESPACE
\ No newline at end of file
+#undef LOCTEXT_NAMESPACE
diff --git a/Engine/Plugins/Experimental/Mutable/Source/CustomizableObjectEditor/Private/MuCOE/SMutableCodeViewer.cpp b/Engine/Plugins/Experimental/Mutable/Source/CustomizableObjectEditor/Private/MuCOE/SMutableCodeViewer.cpp
--- a/Engine/Plugins/Experimental/Mutable/Source/CustomizableObjectEditor/Private/MuCOE/SMutableCodeViewer.cpp
+++ b/Engine/Plugins/Experimental/Mutable/Source/CustomizableObjectEditor/Private/MuCOE/SMutableCodeViewer.cpp
@@ -1,9 +1,9 @@
 // Copyright Epic Games, Inc. All Rights Reserved.
 
-#include "SMutableCodeViewer.h"
+#include "MuCOE/SMutableCodeViewer.h"
 
-#include "SMutableImageViewer.h"
-#include "SMutableMeshViewer.h"
+#include "MuCOE/SMutableImageViewer.h"
+#include "MuCOE/SMutableMeshViewer.h"
 #include "MuCOE/CustomizableObjectCompiler.h"
 #include "MuT/ErrorLogPrivate.h"
 #include "MuT/Streams.h"
@@ -29,11 +29,11 @@
 #include "HAL/PlatformApplicationMisc.h"
 #include "Internationalization/Text.h"
 
-#include "SMutableColorPreviewBox.h"
-#include "SMutableConstantsWidget.h"
+#include "MuCOE/SMutableColorPreviewBox.h"
+#include "MuCOE/SMutableConstantsWidget.h"
 
 #include "Widgets/Views/SHeaderRow.h"
-#include "SMutableColorPreviewBox.h"
+#include "MuCOE/SMutableColorPreviewBox.h"
 
 
 #define LOCTEXT_NAMESPACE "SMutableDebugger"
diff --git a/Engine/Plugins/Experimental/Mutable/Source/CustomizableObjectEditor/Private/MuCOE/SMutableCodeViewer.h b/Engine/Plugins/Experimental/Mutable/Source/CustomizableObjectEditor/Private/MuCOE/SMutableCodeViewer.h
--- a/Engine/Plugins/Experimental/Mutable/Source/CustomizableObjectEditor/Private/MuCOE/SMutableCodeViewer.h
+++ b/Engine/Plugins/Experimental/Mutable/Source/CustomizableObjectEditor/Private/MuCOE/SMutableCodeViewer.h
@@ -3,18 +3,18 @@
 #pragma once
 
 #include "MuCO/CustomizableObject.h"
-#include "SMutableBoolViewer.h"
-#include "SMutableColorViewer.h"
-#include "SMutableImageViewer.h"
-#include "SMutableIntViewer.h"
-#include "SMutableLayoutViewer.h"
-#include "SMutableMeshViewer.h"
-#include "SMutableParametersWidget.h"
-#include "SMutableProjectorViewer.h"
-#include "SMutableScalarViewer.h"
-#include "SMutableStringViewer.h"
-#include "SMutableSkeletonViewer.h"
-#include "SMutableCurveViewer.h"
+#include "MuCOE/SMutableBoolViewer.h"
+#include "MuCOE/SMutableColorViewer.h"
+#include "MuCOE/SMutableImageViewer.h"
+#include "MuCOE/SMutableIntViewer.h"
+#include "MuCOE/SMutableLayoutViewer.h"
+#include "MuCOE/SMutableMeshViewer.h"
+#include "MuCOE/SMutableParametersWidget.h"
+#include "MuCOE/SMutableProjectorViewer.h"
+#include "MuCOE/SMutableScalarViewer.h"
+#include "MuCOE/SMutableStringViewer.h"
+#include "MuCOE/SMutableSkeletonViewer.h"
+#include "MuCOE/SMutableCurveViewer.h"
 #include "Framework/Docking/TabManager.h"
 #include "MuR/Operations.h"
 #include "MuR/Model.h"
diff --git a/Engine/Plugins/Experimental/Mutable/Source/CustomizableObjectEditor/Private/MuCOE/SMutableColorPreviewBox.cpp b/Engine/Plugins/Experimental/Mutable/Source/CustomizableObjectEditor/Private/MuCOE/SMutableColorPreviewBox.cpp
--- a/Engine/Plugins/Experimental/Mutable/Source/CustomizableObjectEditor/Private/MuCOE/SMutableColorPreviewBox.cpp
+++ b/Engine/Plugins/Experimental/Mutable/Source/CustomizableObjectEditor/Private/MuCOE/SMutableColorPreviewBox.cpp
@@ -1,6 +1,6 @@
 // Copyright Epic Games, Inc. All Rights Reserved.
 
-#include "SMutableColorPreviewBox.h"
+#include "MuCOE/SMutableColorPreviewBox.h"
 
 #include "Widgets/SBoxPanel.h"
 #include "Widgets/Layout/SBorder.h"
diff --git a/Engine/Plugins/Experimental/Mutable/Source/CustomizableObjectEditor/Private/MuCOE/SMutableColorViewer.cpp b/Engine/Plugins/Experimental/Mutable/Source/CustomizableObjectEditor/Private/MuCOE/SMutableColorViewer.cpp
--- a/Engine/Plugins/Experimental/Mutable/Source/CustomizableObjectEditor/Private/MuCOE/SMutableColorViewer.cpp
+++ b/Engine/Plugins/Experimental/Mutable/Source/CustomizableObjectEditor/Private/MuCOE/SMutableColorViewer.cpp
@@ -1,7 +1,7 @@
 // Copyright Epic Games, Inc. All Rights Reserved.
 
-#include "SMutableColorViewer.h"
-#include "SMutableColorPreviewBox.h"
+#include "MuCOE/SMutableColorViewer.h"
+#include "MuCOE/SMutableColorPreviewBox.h"
 
 #include "GenericPlatform/GenericPlatform.h"
 #include "Widgets/SBoxPanel.h"
@@ -160,4 +160,4 @@ FText SMutableColorViewer::GetBlueValue() const
 	return FText::AsNumber(this->BlueValue);
 }
 
-#undef LOCTEXT_NAMESPACE
\ No newline at end of file
+#undef LOCTEXT_NAMESPACE
diff --git a/Engine/Plugins/Experimental/Mutable/Source/CustomizableObjectEditor/Private/MuCOE/SMutableColorViewer.h b/Engine/Plugins/Experimental/Mutable/Source/CustomizableObjectEditor/Private/MuCOE/SMutableColorViewer.h
--- a/Engine/Plugins/Experimental/Mutable/Source/CustomizableObjectEditor/Private/MuCOE/SMutableColorViewer.h
+++ b/Engine/Plugins/Experimental/Mutable/Source/CustomizableObjectEditor/Private/MuCOE/SMutableColorViewer.h
@@ -2,7 +2,7 @@
 
 #pragma once
 
-#include "SMutableColorPreviewBox.h"
+#include "MuCOE/SMutableColorPreviewBox.h"
 #include "Widgets/SCompoundWidget.h"
 #include "Widgets/DeclarativeSyntaxSupport.h"
 #include "Internationalization/Text.h"
diff --git a/Engine/Plugins/Experimental/Mutable/Source/CustomizableObjectEditor/Private/MuCOE/SMutableConstantsWidget.cpp b/Engine/Plugins/Experimental/Mutable/Source/CustomizableObjectEditor/Private/MuCOE/SMutableConstantsWidget.cpp
--- a/Engine/Plugins/Experimental/Mutable/Source/CustomizableObjectEditor/Private/MuCOE/SMutableConstantsWidget.cpp
+++ b/Engine/Plugins/Experimental/Mutable/Source/CustomizableObjectEditor/Private/MuCOE/SMutableConstantsWidget.cpp
@@ -1,6 +1,6 @@
 // Copyright Epic Games, Inc. All Rights Reserved.
 
-#include "SMutableConstantsWidget.h"
+#include "MuCOE/SMutableConstantsWidget.h"
 
 #include "MuR/MutableString.h"
 #include "MuR/ModelPrivate.h"
@@ -9,11 +9,11 @@
 #include "MuR/Layout.h"
 #include "MuR/Mesh.h"
 
-#include "SMutableImageViewer.h"
-#include "SMutableMeshViewer.h"
-#include "SMutableStringViewer.h"
+#include "MuCOE/SMutableImageViewer.h"
+#include "MuCOE/SMutableMeshViewer.h"
+#include "MuCOE/SMutableStringViewer.h"
 
-#include "SMutableCodeViewer.h"
+#include "MuCOE/SMutableCodeViewer.h"
 #include "Algo/StableSort.h"
 #include "Components/TextBlock.h"			
 #include "Internationalization/Text.h"		
diff --git a/Engine/Plugins/Experimental/Mutable/Source/CustomizableObjectEditor/Private/MuCOE/SMutableCurveViewer.cpp b/Engine/Plugins/Experimental/Mutable/Source/CustomizableObjectEditor/Private/MuCOE/SMutableCurveViewer.cpp
--- a/Engine/Plugins/Experimental/Mutable/Source/CustomizableObjectEditor/Private/MuCOE/SMutableCurveViewer.cpp
+++ b/Engine/Plugins/Experimental/Mutable/Source/CustomizableObjectEditor/Private/MuCOE/SMutableCurveViewer.cpp
@@ -1,6 +1,6 @@
 // Copyright Epic Games, Inc. All Rights Reserved.
 
-#include "SMutableCurveViewer.h"
+#include "MuCOE/SMutableCurveViewer.h"
 
 #include "Math/Color.h"
 #include "MuR/ParametersPrivate.h"
@@ -570,4 +570,4 @@ void SMutableCurveViewer::SetupMutableCurveGraph() const
 	CurveEditor->PinCurve(CurveId);
 }
 
-#undef LOCTEXT_NAMESPACE
\ No newline at end of file
+#undef LOCTEXT_NAMESPACE
diff --git a/Engine/Plugins/Experimental/Mutable/Source/CustomizableObjectEditor/Private/MuCOE/SMutableGraphViewer.cpp b/Engine/Plugins/Experimental/Mutable/Source/CustomizableObjectEditor/Private/MuCOE/SMutableGraphViewer.cpp
--- a/Engine/Plugins/Experimental/Mutable/Source/CustomizableObjectEditor/Private/MuCOE/SMutableGraphViewer.cpp
+++ b/Engine/Plugins/Experimental/Mutable/Source/CustomizableObjectEditor/Private/MuCOE/SMutableGraphViewer.cpp
@@ -1,10 +1,10 @@
 // Copyright Epic Games, Inc. All Rights Reserved.
 
-#include "SMutableGraphViewer.h"
+#include "MuCOE/SMutableGraphViewer.h"
 
 #include "MuCOE/CustomizableObjectEditorStyle.h"
 #include "MuCOE/CustomizableObjectCompiler.h"
-#include "SMutableCodeViewer.h"
+#include "MuCOE/SMutableCodeViewer.h"
 #include "MuT/Streams.h"
 
 #include "Misc/App.h"
diff --git a/Engine/Plugins/Experimental/Mutable/Source/CustomizableObjectEditor/Private/MuCOE/SMutableImageViewer.cpp b/Engine/Plugins/Experimental/Mutable/Source/CustomizableObjectEditor/Private/MuCOE/SMutableImageViewer.cpp
--- a/Engine/Plugins/Experimental/Mutable/Source/CustomizableObjectEditor/Private/MuCOE/SMutableImageViewer.cpp
+++ b/Engine/Plugins/Experimental/Mutable/Source/CustomizableObjectEditor/Private/MuCOE/SMutableImageViewer.cpp
@@ -1,6 +1,6 @@
 // Copyright Epic Games, Inc. All Rights Reserved.
 
-#include "SMutableImageViewer.h"
+#include "MuCOE/SMutableImageViewer.h"
 #include "MuCOE/CustomizableObjectEditorModule.h"
 #include "Styling/AppStyle.h"
 #include "Async/TaskGraphInterfaces.h"
diff --git a/Engine/Plugins/Experimental/Mutable/Source/CustomizableObjectEditor/Private/MuCOE/SMutableIntViewer.cpp b/Engine/Plugins/Experimental/Mutable/Source/CustomizableObjectEditor/Private/MuCOE/SMutableIntViewer.cpp
--- a/Engine/Plugins/Experimental/Mutable/Source/CustomizableObjectEditor/Private/MuCOE/SMutableIntViewer.cpp
+++ b/Engine/Plugins/Experimental/Mutable/Source/CustomizableObjectEditor/Private/MuCOE/SMutableIntViewer.cpp
@@ -1,6 +1,6 @@
 // Copyright Epic Games, Inc. All Rights Reserved.
 
-#include "SMutableIntViewer.h"
+#include "MuCOE/SMutableIntViewer.h"
 
 #include "GenericPlatform/GenericPlatform.h"
 #include "Widgets/SBoxPanel.h"
@@ -46,4 +46,4 @@ FText SMutableIntViewer::GetValue() const
 	return FText::AsNumber(this->IntValue);
 }
 
-#undef LOCTEXT_NAMESPACE
\ No newline at end of file
+#undef LOCTEXT_NAMESPACE
diff --git a/Engine/Plugins/Experimental/Mutable/Source/CustomizableObjectEditor/Private/MuCOE/SMutableLayoutViewer.cpp b/Engine/Plugins/Experimental/Mutable/Source/CustomizableObjectEditor/Private/MuCOE/SMutableLayoutViewer.cpp
--- a/Engine/Plugins/Experimental/Mutable/Source/CustomizableObjectEditor/Private/MuCOE/SMutableLayoutViewer.cpp
+++ b/Engine/Plugins/Experimental/Mutable/Source/CustomizableObjectEditor/Private/MuCOE/SMutableLayoutViewer.cpp
@@ -1,6 +1,6 @@
 // Copyright Epic Games, Inc. All Rights Reserved.
 
-#include "SMutableLayoutViewer.h"
+#include "MuCOE/SMutableLayoutViewer.h"
 #include "MuCOE/CustomizableObjectEditorModule.h"
 #include "Styling/AppStyle.h"
 #include "Async/TaskGraphInterfaces.h"
diff --git a/Engine/Plugins/Experimental/Mutable/Source/CustomizableObjectEditor/Private/MuCOE/SMutableMeshViewer.cpp b/Engine/Plugins/Experimental/Mutable/Source/CustomizableObjectEditor/Private/MuCOE/SMutableMeshViewer.cpp
--- a/Engine/Plugins/Experimental/Mutable/Source/CustomizableObjectEditor/Private/MuCOE/SMutableMeshViewer.cpp
+++ b/Engine/Plugins/Experimental/Mutable/Source/CustomizableObjectEditor/Private/MuCOE/SMutableMeshViewer.cpp
@@ -1,6 +1,6 @@
 // Copyright Epic Games, Inc. All Rights Reserved.
 
-#include "SMutableMeshViewer.h"
+#include "MuCOE/SMutableMeshViewer.h"
 #include "MuCOE/CustomizableObjectEditorModule.h"
 #include "Styling/AppStyle.h"
 #include "Widgets/Views/SListView.h"
@@ -15,7 +15,7 @@
 #include "Widgets/Views/SHeaderRow.h"
 #include "Templates/SharedPointer.h"
 #include "PropertyCustomizationHelpers.h"
-#include "SMutableSkeletonViewer.h"
+#include "MuCOE/SMutableSkeletonViewer.h"
 #include "MuT/TypeInfo.h"
 #include "MuR/Skeleton.h"
 
diff --git a/Engine/Plugins/Experimental/Mutable/Source/CustomizableObjectEditor/Private/MuCOE/SMutableMeshViewer.h b/Engine/Plugins/Experimental/Mutable/Source/CustomizableObjectEditor/Private/MuCOE/SMutableMeshViewer.h
--- a/Engine/Plugins/Experimental/Mutable/Source/CustomizableObjectEditor/Private/MuCOE/SMutableMeshViewer.h
+++ b/Engine/Plugins/Experimental/Mutable/Source/CustomizableObjectEditor/Private/MuCOE/SMutableMeshViewer.h
@@ -14,12 +14,12 @@
 
 #include "Templates/SharedPointer.h"
 
-#include "SMutableMeshViewport.h"
+#include "MuCOE/SMutableMeshViewport.h"
 #include "Widgets/Input/SSearchBox.h"
 #include "Widgets/Layout/SScrollBox.h"
 
 #include "PropertyCustomizationHelpers.h"
-#include "SMutableSkeletonViewer.h"
+#include "MuCOE/SMutableSkeletonViewer.h"
 
 
 /** Container designed to hold the buffer channel data of a mutable mesh buffer to be later used by the UI */
diff --git a/Engine/Plugins/Experimental/Mutable/Source/CustomizableObjectEditor/Private/MuCOE/SMutableMeshViewport.cpp b/Engine/Plugins/Experimental/Mutable/Source/CustomizableObjectEditor/Private/MuCOE/SMutableMeshViewport.cpp
--- a/Engine/Plugins/Experimental/Mutable/Source/CustomizableObjectEditor/Private/MuCOE/SMutableMeshViewport.cpp
+++ b/Engine/Plugins/Experimental/Mutable/Source/CustomizableObjectEditor/Private/MuCOE/SMutableMeshViewport.cpp
@@ -1,14 +1,14 @@
 // Copyright Epic Games, Inc. All Rights Reserved.
 
-#include "SMutableMeshViewport.h"
+#include "MuCOE/SMutableMeshViewport.h"
 
-#include "FMutableViewportClient.h"
+#include "MuCOE/FMutableViewportClient.h"
 #include "MuCO/CustomizableObjectInstance.h"
 #include "SEditorViewport.h"
 #include "Containers/UnrealString.h"
 #include "Internationalization/Text.h"
 #include "Engine/SkeletalMesh.h"
-#include "FMutableViewportClient.h"
+#include "MuCOE/FMutableViewportClient.h"
 #include "MuR/Mesh.h"
 #include "MuR/MeshPrivate.h"
 #include "MuR/Parameters.h"
@@ -143,4 +143,4 @@ bool SMutableMeshViewport::GenerateUnrealMesh()
 	return true;
 } 
 
-#undef LOCTEXT_NAMESPACE
\ No newline at end of file
+#undef LOCTEXT_NAMESPACE
diff --git a/Engine/Plugins/Experimental/Mutable/Source/CustomizableObjectEditor/Private/MuCOE/SMutableMeshViewport.h b/Engine/Plugins/Experimental/Mutable/Source/CustomizableObjectEditor/Private/MuCOE/SMutableMeshViewport.h
--- a/Engine/Plugins/Experimental/Mutable/Source/CustomizableObjectEditor/Private/MuCOE/SMutableMeshViewport.h
+++ b/Engine/Plugins/Experimental/Mutable/Source/CustomizableObjectEditor/Private/MuCOE/SMutableMeshViewport.h
@@ -3,7 +3,7 @@
 #pragma once
 
 #include "MuR/Mesh.h"
-#include "FMutableViewportClient.h"
+#include "MuCOE/FMutableViewportClient.h"
 #include "Widgets/SCompoundWidget.h"
 #include "AdvancedPreviewScene.h"
 #include "Components/SkeletalMeshComponent.h"
diff --git a/Engine/Plugins/Experimental/Mutable/Source/CustomizableObjectEditor/Private/MuCOE/SMutableObjectViewer.cpp b/Engine/Plugins/Experimental/Mutable/Source/CustomizableObjectEditor/Private/MuCOE/SMutableObjectViewer.cpp
--- a/Engine/Plugins/Experimental/Mutable/Source/CustomizableObjectEditor/Private/MuCOE/SMutableObjectViewer.cpp
+++ b/Engine/Plugins/Experimental/Mutable/Source/CustomizableObjectEditor/Private/MuCOE/SMutableObjectViewer.cpp
@@ -1,11 +1,11 @@
 // Copyright Epic Games, Inc. All Rights Reserved.
 
-#include "SMutableObjectViewer.h"
+#include "MuCOE/SMutableObjectViewer.h"
 
 #include "MuCOE/CustomizableObjectEditorStyle.h"
 #include "MuCOE/CustomizableObjectCompiler.h"
-#include "SMutableGraphViewer.h"
-#include "SMutableCodeViewer.h"
+#include "MuCOE/SMutableGraphViewer.h"
+#include "MuCOE/SMutableCodeViewer.h"
 #include "MuT/Streams.h"
 
 #include "Misc/App.h"
diff --git a/Engine/Plugins/Experimental/Mutable/Source/CustomizableObjectEditor/Private/MuCOE/SMutableParametersWidget.cpp b/Engine/Plugins/Experimental/Mutable/Source/CustomizableObjectEditor/Private/MuCOE/SMutableParametersWidget.cpp
--- a/Engine/Plugins/Experimental/Mutable/Source/CustomizableObjectEditor/Private/MuCOE/SMutableParametersWidget.cpp
+++ b/Engine/Plugins/Experimental/Mutable/Source/CustomizableObjectEditor/Private/MuCOE/SMutableParametersWidget.cpp
@@ -1,6 +1,6 @@
 // Copyright Epic Games, Inc. All Rights Reserved.
 
-#include "SMutableParametersWidget.h"
+#include "MuCOE/SMutableParametersWidget.h"
 #include "MuCOE/CustomizableObjectEditorModule.h"
 #include "Styling/AppStyle.h"
 #include "Async/TaskGraphInterfaces.h"
diff --git a/Engine/Plugins/Experimental/Mutable/Source/CustomizableObjectEditor/Private/MuCOE/SMutableParametersWidget.h b/Engine/Plugins/Experimental/Mutable/Source/CustomizableObjectEditor/Private/MuCOE/SMutableParametersWidget.h
--- a/Engine/Plugins/Experimental/Mutable/Source/CustomizableObjectEditor/Private/MuCOE/SMutableParametersWidget.h
+++ b/Engine/Plugins/Experimental/Mutable/Source/CustomizableObjectEditor/Private/MuCOE/SMutableParametersWidget.h
@@ -2,7 +2,7 @@
 
 #pragma once
 
-#include "SMutableTextSearchBox.h"
+#include "MuCOE/SMutableTextSearchBox.h"
 #include "Widgets/SCompoundWidget.h"
 #include "Widgets/SBoxPanel.h"
 #include "Widgets/Input/SCheckBox.h"
diff --git a/Engine/Plugins/Experimental/Mutable/Source/CustomizableObjectEditor/Private/MuCOE/SMutableProjectorViewer.cpp b/Engine/Plugins/Experimental/Mutable/Source/CustomizableObjectEditor/Private/MuCOE/SMutableProjectorViewer.cpp
--- a/Engine/Plugins/Experimental/Mutable/Source/CustomizableObjectEditor/Private/MuCOE/SMutableProjectorViewer.cpp
+++ b/Engine/Plugins/Experimental/Mutable/Source/CustomizableObjectEditor/Private/MuCOE/SMutableProjectorViewer.cpp
@@ -1,6 +1,6 @@
 // Copyright Epic Games, Inc. All Rights Reserved.
 
-#include "SMutableProjectorViewer.h"
+#include "MuCOE/SMutableProjectorViewer.h"
 
 #include "Widgets/Input/SVectorInputBox.h"
 #include "Internationalization/Text.h"
@@ -246,4 +246,4 @@ FText SMutableProjectorViewer::GetProjectorTypeAsText() const
 	return FText::FromString(* FString(mu::TypeInfo::s_projectorTypeName[TypeIndex]));
 }
 
-#undef LOCTEXT_NAMESPACE
\ No newline at end of file
+#undef LOCTEXT_NAMESPACE
diff --git a/Engine/Plugins/Experimental/Mutable/Source/CustomizableObjectEditor/Private/MuCOE/SMutableScalarViewer.cpp b/Engine/Plugins/Experimental/Mutable/Source/CustomizableObjectEditor/Private/MuCOE/SMutableScalarViewer.cpp
--- a/Engine/Plugins/Experimental/Mutable/Source/CustomizableObjectEditor/Private/MuCOE/SMutableScalarViewer.cpp
+++ b/Engine/Plugins/Experimental/Mutable/Source/CustomizableObjectEditor/Private/MuCOE/SMutableScalarViewer.cpp
@@ -1,6 +1,6 @@
 // Copyright Epic Games, Inc. All Rights Reserved.
 
-#include "SMutableScalarViewer.h"
+#include "MuCOE/SMutableScalarViewer.h"
 
 #include "Widgets/SBoxPanel.h"
 #include "Internationalization/Text.h"
@@ -45,4 +45,4 @@ FText SMutableScalarViewer::GetValue() const
 	return FText::AsNumber(this->ScalarValue);
 }
 
-#undef LOCTEXT_NAMESPACE
\ No newline at end of file
+#undef LOCTEXT_NAMESPACE
diff --git a/Engine/Plugins/Experimental/Mutable/Source/CustomizableObjectEditor/Private/MuCOE/SMutableSkeletonViewer.cpp b/Engine/Plugins/Experimental/Mutable/Source/CustomizableObjectEditor/Private/MuCOE/SMutableSkeletonViewer.cpp
--- a/Engine/Plugins/Experimental/Mutable/Source/CustomizableObjectEditor/Private/MuCOE/SMutableSkeletonViewer.cpp
+++ b/Engine/Plugins/Experimental/Mutable/Source/CustomizableObjectEditor/Private/MuCOE/SMutableSkeletonViewer.cpp
@@ -1,6 +1,6 @@
 // Copyright Epic Games, Inc. All Rights Reserved.
 
-#include "SMutableSkeletonViewer.h"
+#include "MuCOE/SMutableSkeletonViewer.h"
 #include "Widgets/Layout/SScrollBox.h"
 #include "MuT/TypeInfo.h"
 #include "MuR/Skeleton.h"
diff --git a/Engine/Plugins/Experimental/Mutable/Source/CustomizableObjectEditor/Private/MuCOE/SMutableStringViewer.cpp b/Engine/Plugins/Experimental/Mutable/Source/CustomizableObjectEditor/Private/MuCOE/SMutableStringViewer.cpp
--- a/Engine/Plugins/Experimental/Mutable/Source/CustomizableObjectEditor/Private/MuCOE/SMutableStringViewer.cpp
+++ b/Engine/Plugins/Experimental/Mutable/Source/CustomizableObjectEditor/Private/MuCOE/SMutableStringViewer.cpp
@@ -1,6 +1,6 @@
 // Copyright Epic Games, Inc. All Rights Reserved.
 
-#include "SMutableStringViewer.h"
+#include "MuCOE/SMutableStringViewer.h"
 
 #include "Widgets/SBoxPanel.h"
 #include "Internationalization/Text.h"
@@ -51,4 +51,4 @@ FText SMutableStringViewer::GetValue() const
 	return this->StringValue;
 }
 
-#undef LOCTEXT_NAMESPACE
\ No newline at end of file
+#undef LOCTEXT_NAMESPACE
diff --git a/Engine/Plugins/Experimental/Mutable/Source/CustomizableObjectEditor/Private/MuCOE/StressTest.cpp b/Engine/Plugins/Experimental/Mutable/Source/CustomizableObjectEditor/Private/MuCOE/StressTest.cpp
--- a/Engine/Plugins/Experimental/Mutable/Source/CustomizableObjectEditor/Private/MuCOE/StressTest.cpp
+++ b/Engine/Plugins/Experimental/Mutable/Source/CustomizableObjectEditor/Private/MuCOE/StressTest.cpp
@@ -1,6 +1,6 @@
 // Copyright Epic Games, Inc. All Rights Reserved.
 
-#include "StressTest.h"
+#include "MuCOE/StressTest.h"
 #include "UObject/Package.h"
 #include "MuCO/CustomizableObject.h"
 #include "MuCO/CustomizableObjectInstance.h"
@@ -1076,4 +1076,4 @@ FString FRunningStressTest::StringyfyEPixelFormat(EPixelFormat Value)
 			return FString("Error");
 		}
 	}
-}
\ No newline at end of file
+}
diff --git a/Engine/Plugins/Experimental/Mutable/Source/CustomizableObjectEditor/Private/MuCOE/StressTest.h b/Engine/Plugins/Experimental/Mutable/Source/CustomizableObjectEditor/Private/MuCOE/StressTest.h
--- a/Engine/Plugins/Experimental/Mutable/Source/CustomizableObjectEditor/Private/MuCOE/StressTest.h
+++ b/Engine/Plugins/Experimental/Mutable/Source/CustomizableObjectEditor/Private/MuCOE/StressTest.h
@@ -3,7 +3,7 @@
 #pragma once
 
 #include "Delegates/IDelegateInstance.h"
-#include "StatePerformanceTesting.h"
+#include "MuCOE/StatePerformanceTesting.h"
 #include "MuCO/CustomizableObject.h"
 #include "Containers/Ticker.h"
 #include "StressTest.generated.h"
@@ -289,4 +289,4 @@ protected:
 
 	/** Number of LOD of the Customizable Object */
 	uint32 NumLOD;
-};
\ No newline at end of file
+};
diff --git a/Engine/Plugins/Experimental/Mutable/Source/CustomizableObjectEditor/Private/MuCOE/Widgets/SMutableExpandableTableRow.cpp b/Engine/Plugins/Experimental/Mutable/Source/CustomizableObjectEditor/Private/MuCOE/Widgets/SMutableExpandableTableRow.cpp
--- a/Engine/Plugins/Experimental/Mutable/Source/CustomizableObjectEditor/Private/MuCOE/Widgets/SMutableExpandableTableRow.cpp
+++ b/Engine/Plugins/Experimental/Mutable/Source/CustomizableObjectEditor/Private/MuCOE/Widgets/SMutableExpandableTableRow.cpp
@@ -1,6 +1,6 @@
 // Copyright Epic Games, Inc. All Rights Reserved.
 
-#include "SMutableExpandableTableRow.h"
+#include "MuCOE/Widgets/SMutableExpandableTableRow.h"
 
 
 FSlateColor GetRowBackgroundColor(const int32 IndentLevel, const bool IsHovered)
diff --git a/Engine/Plugins/Experimental/Mutable/Source/CustomizableObjectEditor/Public/MuCOE/RemapPins/CustomizableObjectNodeRemapPinsByName.h b/Engine/Plugins/Experimental/Mutable/Source/CustomizableObjectEditor/Public/MuCOE/RemapPins/CustomizableObjectNodeRemapPinsByName.h
--- a/Engine/Plugins/Experimental/Mutable/Source/CustomizableObjectEditor/Public/MuCOE/RemapPins/CustomizableObjectNodeRemapPinsByName.h
+++ b/Engine/Plugins/Experimental/Mutable/Source/CustomizableObjectEditor/Public/MuCOE/RemapPins/CustomizableObjectNodeRemapPinsByName.h
@@ -2,7 +2,7 @@
 
 #pragma once
 
-#include "CustomizableObjectNodeRemapPins.h"
+#include "MuCOE/RemapPins/CustomizableObjectNodeRemapPins.h"
 
 #include "CustomizableObjectNodeRemapPinsByName.generated.h"
 
@@ -23,4 +23,4 @@ public:
 	virtual bool Equal(const UEdGraphPin& OldPin, const UEdGraphPin& NewPin) const;
 	
 	virtual void RemapPins(const TArray<UEdGraphPin*>& OldPins, const TArray<UEdGraphPin*>& NewPins, TMap<UEdGraphPin*, UEdGraphPin*>& PinsToRemap, TArray<UEdGraphPin*>& PinsToOrphan) override;
-};
\ No newline at end of file
+};
diff --git a/Engine/Plugins/Experimental/Mutable/Source/CustomizableObjectEditor/Public/MuCOE/RemapPins/CustomizableObjectNodeRemapPinsByNameDefaultPin.h b/Engine/Plugins/Experimental/Mutable/Source/CustomizableObjectEditor/Public/MuCOE/RemapPins/CustomizableObjectNodeRemapPinsByNameDefaultPin.h
--- a/Engine/Plugins/Experimental/Mutable/Source/CustomizableObjectEditor/Public/MuCOE/RemapPins/CustomizableObjectNodeRemapPinsByNameDefaultPin.h
+++ b/Engine/Plugins/Experimental/Mutable/Source/CustomizableObjectEditor/Public/MuCOE/RemapPins/CustomizableObjectNodeRemapPinsByNameDefaultPin.h
@@ -2,7 +2,7 @@
 
 #pragma once
 
-#include "CustomizableObjectNodeRemapPinsByName.h"
+#include "MuCOE/RemapPins/CustomizableObjectNodeRemapPinsByName.h"
 
 #include "CustomizableObjectNodeRemapPinsByNameDefaultPin.generated.h"
 
diff --git a/Engine/Plugins/Experimental/Mutable/Source/CustomizableObjectEditor/Public/MuCOE/RemapPins/CustomizableObjectNodeRemapPinsByPosition.h b/Engine/Plugins/Experimental/Mutable/Source/CustomizableObjectEditor/Public/MuCOE/RemapPins/CustomizableObjectNodeRemapPinsByPosition.h
--- a/Engine/Plugins/Experimental/Mutable/Source/CustomizableObjectEditor/Public/MuCOE/RemapPins/CustomizableObjectNodeRemapPinsByPosition.h
+++ b/Engine/Plugins/Experimental/Mutable/Source/CustomizableObjectEditor/Public/MuCOE/RemapPins/CustomizableObjectNodeRemapPinsByPosition.h
@@ -2,7 +2,7 @@
 
 #pragma once
 
-#include "CustomizableObjectNodeRemapPins.h"
+#include "MuCOE/RemapPins/CustomizableObjectNodeRemapPins.h"
 
 #include "CustomizableObjectNodeRemapPinsByPosition.generated.h"
 
@@ -21,4 +21,4 @@ public:
 	GENERATED_BODY()
 
 	virtual void RemapPins(const TArray<UEdGraphPin*>& OldPins, const TArray<UEdGraphPin*>& NewPins, TMap<UEdGraphPin*, UEdGraphPin*>& PinsToRemap, TArray<UEdGraphPin*>& PinsToOrphan) override;
-};
\ No newline at end of file
+};
diff --git a/Engine/Plugins/Experimental/Mutable/Source/CustomizableObjectPopulation/Private/MuCOP/CustomizableObjectPopulationConstraint.cpp b/Engine/Plugins/Experimental/Mutable/Source/CustomizableObjectPopulation/Private/MuCOP/CustomizableObjectPopulationConstraint.cpp
--- a/Engine/Plugins/Experimental/Mutable/Source/CustomizableObjectPopulation/Private/MuCOP/CustomizableObjectPopulationConstraint.cpp
+++ b/Engine/Plugins/Experimental/Mutable/Source/CustomizableObjectPopulation/Private/MuCOP/CustomizableObjectPopulationConstraint.cpp
@@ -1,6 +1,6 @@
 // Copyright Epic Games, Inc. All Rights Reserved.
 
-#include "CustomizableObjectPopulationConstraint.h"
+#include "MuCOP/CustomizableObjectPopulationConstraint.h"
 
 #define LOCTEXT_NAMESPACE "CustomizableObjectPopulationConstraint"
 
diff --git a/Engine/Plugins/Experimental/Mutable/Source/CustomizableObjectPopulationEditor/Private/MuCOPE/AssetTypeActions_CustomizableObjectPopulation.cpp b/Engine/Plugins/Experimental/Mutable/Source/CustomizableObjectPopulationEditor/Private/MuCOPE/AssetTypeActions_CustomizableObjectPopulation.cpp
--- a/Engine/Plugins/Experimental/Mutable/Source/CustomizableObjectPopulationEditor/Private/MuCOPE/AssetTypeActions_CustomizableObjectPopulation.cpp
+++ b/Engine/Plugins/Experimental/Mutable/Source/CustomizableObjectPopulationEditor/Private/MuCOPE/AssetTypeActions_CustomizableObjectPopulation.cpp
@@ -1,6 +1,6 @@
 // Copyright Epic Games, Inc. All Rights Reserved.
 
-#include "AssetTypeActions_CustomizableObjectPopulation.h"
+#include "MuCOPE/AssetTypeActions_CustomizableObjectPopulation.h"
 
 #include "ContentBrowserModule.h"
 #include "IContentBrowserSingleton.h"
diff --git a/Engine/Plugins/Experimental/Mutable/Source/CustomizableObjectPopulationEditor/Private/MuCOPE/AssetTypeActions_CustomizableObjectPopulationClass.cpp b/Engine/Plugins/Experimental/Mutable/Source/CustomizableObjectPopulationEditor/Private/MuCOPE/AssetTypeActions_CustomizableObjectPopulationClass.cpp
--- a/Engine/Plugins/Experimental/Mutable/Source/CustomizableObjectPopulationEditor/Private/MuCOPE/AssetTypeActions_CustomizableObjectPopulationClass.cpp
+++ b/Engine/Plugins/Experimental/Mutable/Source/CustomizableObjectPopulationEditor/Private/MuCOPE/AssetTypeActions_CustomizableObjectPopulationClass.cpp
@@ -1,10 +1,10 @@
 // Copyright Epic Games, Inc. All Rights Reserved.
 
 
-#include "AssetTypeActions_CustomizableObjectPopulationClass.h"
+#include "MuCOPE/AssetTypeActions_CustomizableObjectPopulationClass.h"
 
 #include "MuCOE/CustomizableObjectEditorModule.h"
-#include "CustomizableObjectPopulationClassEditor.h"
+#include "MuCOPE/CustomizableObjectPopulationClassEditor.h"
 #include "MuCOPE/CustomizableObjectPopulationEditorModule.h"
 
 #define LOCTEXT_NAMESPACE "AssetTypeActions_CustomizableObjectPopulationClass"
@@ -50,4 +50,4 @@ void FAssetTypeActions_CustomizableObjectPopulationClass::OpenAssetEditor(const
 }
 
 
-#undef LOCTEXT_NAMESPACE
\ No newline at end of file
+#undef LOCTEXT_NAMESPACE
diff --git a/Engine/Plugins/Experimental/Mutable/Source/CustomizableObjectPopulationEditor/Private/MuCOPE/CustomizableObjectPopulationClassDetails.cpp b/Engine/Plugins/Experimental/Mutable/Source/CustomizableObjectPopulationEditor/Private/MuCOPE/CustomizableObjectPopulationClassDetails.cpp
--- a/Engine/Plugins/Experimental/Mutable/Source/CustomizableObjectPopulationEditor/Private/MuCOPE/CustomizableObjectPopulationClassDetails.cpp
+++ b/Engine/Plugins/Experimental/Mutable/Source/CustomizableObjectPopulationEditor/Private/MuCOPE/CustomizableObjectPopulationClassDetails.cpp
@@ -1,6 +1,6 @@
 // Copyright Epic Games, Inc. All Rights Reserved.
 
-#include "CustomizableObjectPopulationClassDetails.h"
+#include "MuCOPE/CustomizableObjectPopulationClassDetails.h"
 #include "MuCOP/CustomizableObjectPopulationClass.h"
 #include "MuCO/CustomizableObject.h"
 #include "MuCO/CustomizableObjectInstance.h"
diff --git a/Engine/Plugins/Experimental/Mutable/Source/CustomizableObjectPopulationEditor/Private/MuCOPE/CustomizableObjectPopulationClassEditor.cpp b/Engine/Plugins/Experimental/Mutable/Source/CustomizableObjectPopulationEditor/Private/MuCOPE/CustomizableObjectPopulationClassEditor.cpp
--- a/Engine/Plugins/Experimental/Mutable/Source/CustomizableObjectPopulationEditor/Private/MuCOPE/CustomizableObjectPopulationClassEditor.cpp
+++ b/Engine/Plugins/Experimental/Mutable/Source/CustomizableObjectPopulationEditor/Private/MuCOPE/CustomizableObjectPopulationClassEditor.cpp
@@ -1,6 +1,6 @@
 // Copyright Epic Games, Inc. All Rights Reserved.
 
-#include "CustomizableObjectPopulationClassEditor.h"
+#include "MuCOPE/CustomizableObjectPopulationClassEditor.h"
 #include "MuCOPE/CustomizableObjectPopulationClassEditorActions.h"
 #include "MuCOE/CustomizableObjectEditorModule.h"
 #include "MuCOPE/CustomizableObjectPopulationEditorModule.h"
@@ -15,7 +15,7 @@
 
 #include "MuCOP/CustomizableObjectPopulationClass.h"
 #include "MuCOPE/SCustomizableObjectPopulationClassTagsTool.h"
-#include "CustomizableObjectPopulationClassDetails.h"
+#include "MuCOPE/CustomizableObjectPopulationClassDetails.h"
 
 #include "MuCOP/CustomizableObjectPopulation.h"
 #include "MuCOP/CustomizableObjectPopulationGenerator.h"
diff --git a/Engine/Plugins/Experimental/Mutable/Source/CustomizableObjectPopulationEditor/Private/MuCOPE/CustomizableObjectPopulationClassEditorActions.cpp b/Engine/Plugins/Experimental/Mutable/Source/CustomizableObjectPopulationEditor/Private/MuCOPE/CustomizableObjectPopulationClassEditorActions.cpp
--- a/Engine/Plugins/Experimental/Mutable/Source/CustomizableObjectPopulationEditor/Private/MuCOPE/CustomizableObjectPopulationClassEditorActions.cpp
+++ b/Engine/Plugins/Experimental/Mutable/Source/CustomizableObjectPopulationEditor/Private/MuCOPE/CustomizableObjectPopulationClassEditorActions.cpp
@@ -1,12 +1,24 @@
 // Copyright Epic Games, Inc. All Rights Reserved.
 
 #include "MuCOPE/CustomizableObjectPopulationClassEditorActions.h"
+
+#include "MuCOPE/CustomizableObjectPopulationEditorStyle.h"
 #include "MuCOP/CustomizableObjectPopulationClass.h"
 #include "Framework/Commands/InputChord.h"
 
-
 #define LOCTEXT_NAMESPACE "CustomizableObjectEditor"
 
+FCustomizableObjectPopulationClassEditorCommands::FCustomizableObjectPopulationClassEditorCommands() : TCommands<FCustomizableObjectPopulationClassEditorCommands>
+(
+	"CustomizableObjectPopulationClassEditor", // Context name for fast lookup
+	NSLOCTEXT("Contexts", "CustomizableObjectPopulationClassEditor", "CustomizableObject Population Class Editor"), // Localized context name for displaying
+	NAME_None, // Parent
+	FCustomizableObjectPopulationEditorStyle::GetStyleSetName()
+	)
+{
+
+}
+
 void FCustomizableObjectPopulationClassEditorCommands::RegisterCommands()
 {
 	UI_COMMAND(Save, "Save", "Save the Population Class.", EUserInterfaceActionType::Button, FInputChord());
diff --git a/Engine/Plugins/Experimental/Mutable/Source/CustomizableObjectPopulationEditor/Private/MuCOPE/CustomizableObjectPopulationEditorActions.cpp b/Engine/Plugins/Experimental/Mutable/Source/CustomizableObjectPopulationEditor/Private/MuCOPE/CustomizableObjectPopulationEditorActions.cpp
--- a/Engine/Plugins/Experimental/Mutable/Source/CustomizableObjectPopulationEditor/Private/MuCOPE/CustomizableObjectPopulationEditorActions.cpp
+++ b/Engine/Plugins/Experimental/Mutable/Source/CustomizableObjectPopulationEditor/Private/MuCOPE/CustomizableObjectPopulationEditorActions.cpp
@@ -1,11 +1,23 @@
 // Copyright Epic Games, Inc. All Rights Reserved.
 
 #include "MuCOPE/CustomizableObjectPopulationEditorActions.h"
+#include "MuCOPE/CustomizableObjectPopulationEditorStyle.h"
 #include "MuCOP/CustomizableObjectPopulation.h"
 #include "Framework/Commands/InputChord.h"
 
 #define LOCTEXT_NAMESPACE "CustomizableObjectPopulationEditor"
 
+FCustomizableObjectPopulationEditorCommands::FCustomizableObjectPopulationEditorCommands() : TCommands<FCustomizableObjectPopulationEditorCommands>
+(
+	"CustomizableObjectPopulationEditor", // Context name for fast lookup
+	NSLOCTEXT("Contexts", "CustomizableObjectPopulationEditor", "CustomizableObject Population Editor"), // Localized context name for displaying
+	NAME_None, // Parent
+	FCustomizableObjectPopulationEditorStyle::GetStyleSetName()
+	)
+{
+
+}
+
 void FCustomizableObjectPopulationEditorCommands::RegisterCommands()
 {
 	UI_COMMAND(Save, "Save", "Save the Population.", EUserInterfaceActionType::Button, FInputChord());
diff --git a/Engine/Plugins/Experimental/Mutable/Source/CustomizableObjectPopulationEditor/Private/MuCOPE/CustomizableObjectPopulationEditorModule.cpp b/Engine/Plugins/Experimental/Mutable/Source/CustomizableObjectPopulationEditor/Private/MuCOPE/CustomizableObjectPopulationEditorModule.cpp
--- a/Engine/Plugins/Experimental/Mutable/Source/CustomizableObjectPopulationEditor/Private/MuCOPE/CustomizableObjectPopulationEditorModule.cpp
+++ b/Engine/Plugins/Experimental/Mutable/Source/CustomizableObjectPopulationEditor/Private/MuCOPE/CustomizableObjectPopulationEditorModule.cpp
@@ -20,13 +20,13 @@
 #include "CoreGlobals.h"
 #include "Settings/ProjectPackagingSettings.h"
 #include "MuCOE/UnrealEditorPortabilityHelpers.h"
-#include "AssetTypeActions_CustomizableObjectPopulation.h"
-#include "AssetTypeActions_CustomizableObjectPopulationClass.h"
-#include "CustomizableObjectPopulationClassEditor.h"
+#include "MuCOPE/AssetTypeActions_CustomizableObjectPopulation.h"
+#include "MuCOPE/AssetTypeActions_CustomizableObjectPopulationClass.h"
+#include "MuCOPE/CustomizableObjectPopulationClassEditor.h"
 #include "MuCOPE/CustomizableObjectPopulationEditor.h"
-#include "CustomizableObjectPopulationEditorStyle.h"
+#include "MuCOPE/CustomizableObjectPopulationEditorStyle.h"
 
-#include "CustomizableObjectPopulationClassDetails.h"
+#include "MuCOPE/CustomizableObjectPopulationClassDetails.h"
 
 // This is necessary to register the cook delegate.
 #include "CookOnTheSide/CookOnTheFlyServer.h"
@@ -141,4 +141,4 @@ TSharedRef<ICustomizableObjectPopulationClassEditor> FCustomizableObjectPopulati
 }
 
 
-#undef LOCTEXT_NAMESPACE
\ No newline at end of file
+#undef LOCTEXT_NAMESPACE
diff --git a/Engine/Plugins/Experimental/Mutable/Source/CustomizableObjectPopulationEditor/Private/MuCOPE/CustomizableObjectPopulationEditorStyle.cpp b/Engine/Plugins/Experimental/Mutable/Source/CustomizableObjectPopulationEditor/Private/MuCOPE/CustomizableObjectPopulationEditorStyle.cpp
--- a/Engine/Plugins/Experimental/Mutable/Source/CustomizableObjectPopulationEditor/Private/MuCOPE/CustomizableObjectPopulationEditorStyle.cpp
+++ b/Engine/Plugins/Experimental/Mutable/Source/CustomizableObjectPopulationEditor/Private/MuCOPE/CustomizableObjectPopulationEditorStyle.cpp
@@ -1,5 +1,5 @@
 // Copyright Epic Games, Inc. All Rights Reserved.
-#include "CustomizableObjectPopulationEditorStyle.h"
+#include "MuCOPE/CustomizableObjectPopulationEditorStyle.h"
 
 #include "Framework/Application/SlateApplication.h"
 #include "Styling/AppStyle.h"
@@ -92,4 +92,4 @@ TSharedRef<FSlateStyleSet> FCustomizableObjectPopulationEditorStyle::Create()
 
 #undef IMAGE_PLUGIN_BRUSH
 #undef IMAGE_BRUSH
-#undef IMAGE_BRUSH_SVG
\ No newline at end of file
+#undef IMAGE_BRUSH_SVG
diff --git a/Engine/Plugins/Experimental/Mutable/Source/CustomizableObjectPopulationEditor/Private/MuCOPE/SCustomizableObjectPopulationClassTagsTool.cpp b/Engine/Plugins/Experimental/Mutable/Source/CustomizableObjectPopulationEditor/Private/MuCOPE/SCustomizableObjectPopulationClassTagsTool.cpp
--- a/Engine/Plugins/Experimental/Mutable/Source/CustomizableObjectPopulationEditor/Private/MuCOPE/SCustomizableObjectPopulationClassTagsTool.cpp
+++ b/Engine/Plugins/Experimental/Mutable/Source/CustomizableObjectPopulationEditor/Private/MuCOPE/SCustomizableObjectPopulationClassTagsTool.cpp
@@ -1,7 +1,7 @@
 // Copyright Epic Games, Inc. All Rights Reserved.
 
 #include "MuCOPE/SCustomizableObjectPopulationClassTagsTool.h"
-#include "CustomizableObjectPopulationClassDetails.h"
+#include "MuCOPE/CustomizableObjectPopulationClassDetails.h"
 #include "MuCOP/CustomizableObjectPopulationClass.h"
 #include "MuCO/CustomizableObject.h"
 
@@ -790,4 +790,4 @@ FReply SSingleTagWidget::OnRemoveTagButtonPressed()
 	return FReply::Handled();
 }
 
-#undef LOCTEXT_NAMESPACE
\ No newline at end of file
+#undef LOCTEXT_NAMESPACE
diff --git a/Engine/Plugins/Experimental/Mutable/Source/CustomizableObjectPopulationEditor/Public/MuCOPE/CustomizableObjectPopulationClassEditorActions.h b/Engine/Plugins/Experimental/Mutable/Source/CustomizableObjectPopulationEditor/Public/MuCOPE/CustomizableObjectPopulationClassEditorActions.h
--- a/Engine/Plugins/Experimental/Mutable/Source/CustomizableObjectPopulationEditor/Public/MuCOPE/CustomizableObjectPopulationClassEditorActions.h
+++ b/Engine/Plugins/Experimental/Mutable/Source/CustomizableObjectPopulationEditor/Public/MuCOPE/CustomizableObjectPopulationClassEditorActions.h
@@ -5,22 +5,12 @@
 #include "Styling/AppStyle.h"
 #include "Framework/Commands/UICommandInfo.h"
 #include "Framework/Commands/Commands.h"
-#include "CustomizableObjectPopulationEditorStyle.h"
-
 
 class FCustomizableObjectPopulationClassEditorCommands : public TCommands<FCustomizableObjectPopulationClassEditorCommands>
 {
 public:
 
-	FCustomizableObjectPopulationClassEditorCommands() : TCommands<FCustomizableObjectPopulationClassEditorCommands>
-		(
-			"CustomizableObjectPopulationClassEditor", // Context name for fast lookup
-			NSLOCTEXT("Contexts", "CustomizableObjectPopulationClassEditor", "CustomizableObject Population Class Editor"), // Localized context name for displaying
-			NAME_None, // Parent
-			FCustomizableObjectPopulationEditorStyle::GetStyleSetName()
-			)
-	{
-	}
+	FCustomizableObjectPopulationClassEditorCommands();
 
 	TSharedPtr< FUICommandInfo > Save;
 	TSharedPtr< FUICommandInfo > SaveCustomizableObject;
diff --git a/Engine/Plugins/Experimental/Mutable/Source/CustomizableObjectPopulationEditor/Public/MuCOPE/CustomizableObjectPopulationEditorActions.h b/Engine/Plugins/Experimental/Mutable/Source/CustomizableObjectPopulationEditor/Public/MuCOPE/CustomizableObjectPopulationEditorActions.h
--- a/Engine/Plugins/Experimental/Mutable/Source/CustomizableObjectPopulationEditor/Public/MuCOPE/CustomizableObjectPopulationEditorActions.h
+++ b/Engine/Plugins/Experimental/Mutable/Source/CustomizableObjectPopulationEditor/Public/MuCOPE/CustomizableObjectPopulationEditorActions.h
@@ -5,21 +5,12 @@
 #include "Styling/AppStyle.h"
 #include "Framework/Commands/UICommandInfo.h"
 #include "Framework/Commands/Commands.h"
-#include "CustomizableObjectPopulationEditorStyle.h"
 
 class FCustomizableObjectPopulationEditorCommands : public TCommands<FCustomizableObjectPopulationEditorCommands>
 {
 
 public:
-	FCustomizableObjectPopulationEditorCommands() : TCommands<FCustomizableObjectPopulationEditorCommands>
-		(
-			"CustomizableObjectPopulationEditor", // Context name for fast lookup
-			NSLOCTEXT("Contexts", "CustomizableObjectPopulationEditor", "CustomizableObject Population Editor"), // Localized context name for displaying
-			NAME_None, // Parent
-			FCustomizableObjectPopulationEditorStyle::GetStyleSetName()
-			)
-	{
-	}
+	FCustomizableObjectPopulationEditorCommands();
 
 	TSharedPtr< FUICommandInfo > Save;
 	TSharedPtr< FUICommandInfo > TestPopulation;
diff --git a/Engine/Plugins/Experimental/Mutable/Source/MutableRuntime/Private/MuR/BlockCompression/Miro/Miro.cpp b/Engine/Plugins/Experimental/Mutable/Source/MutableRuntime/Private/MuR/BlockCompression/Miro/Miro.cpp
--- a/Engine/Plugins/Experimental/Mutable/Source/MutableRuntime/Private/MuR/BlockCompression/Miro/Miro.cpp
+++ b/Engine/Plugins/Experimental/Mutable/Source/MutableRuntime/Private/MuR/BlockCompression/Miro/Miro.cpp
@@ -1,6 +1,6 @@
 // Copyright Epic Games, Inc. All Rights Reserved.
 
-#include "Miro.h"
+#include "MuR/BlockCompression/Miro/Miro.h"
 
 #include "MuR/Platform.h"
 #include "Misc/AssertionMacros.h"
diff --git a/Engine/Plugins/Experimental/Mutable/Source/MutableRuntime/Private/MuR/ImageRLE.cpp b/Engine/Plugins/Experimental/Mutable/Source/MutableRuntime/Private/MuR/ImageRLE.cpp
--- a/Engine/Plugins/Experimental/Mutable/Source/MutableRuntime/Private/MuR/ImageRLE.cpp
+++ b/Engine/Plugins/Experimental/Mutable/Source/MutableRuntime/Private/MuR/ImageRLE.cpp
@@ -1,6 +1,6 @@
 // Copyright Epic Games, Inc. All Rights Reserved.
 
-#include "ImageRLE.h"
+#include "MuR/ImageRLE.h"
 #include "MuR/Platform.h"
 
 namespace mu
diff --git a/Engine/Plugins/Experimental/Mutable/Source/MutableRuntime/Private/MuR/OpEvalCurve.h b/Engine/Plugins/Experimental/Mutable/Source/MutableRuntime/Private/MuR/OpEvalCurve.h
--- a/Engine/Plugins/Experimental/Mutable/Source/MutableRuntime/Private/MuR/OpEvalCurve.h
+++ b/Engine/Plugins/Experimental/Mutable/Source/MutableRuntime/Private/MuR/OpEvalCurve.h
@@ -2,7 +2,7 @@
 
 #pragma once
 
-#include "ConvertData.h"
+#include "MuR/ConvertData.h"
 #include "MuR/Platform.h"
 
 
diff --git a/Engine/Plugins/Experimental/Mutable/Source/MutableRuntime/Private/MuR/OpImageDisplace.h b/Engine/Plugins/Experimental/Mutable/Source/MutableRuntime/Private/MuR/OpImageDisplace.h
--- a/Engine/Plugins/Experimental/Mutable/Source/MutableRuntime/Private/MuR/OpImageDisplace.h
+++ b/Engine/Plugins/Experimental/Mutable/Source/MutableRuntime/Private/MuR/OpImageDisplace.h
@@ -4,7 +4,7 @@
 
 #include "MuR/ImagePrivate.h"
 #include "MuR/Platform.h"
-#include "ConvertData.h"
+#include "MuR/ConvertData.h"
 #include "Async/ParallelFor.h"
 
 
diff --git a/Engine/Plugins/Experimental/Mutable/Source/MutableRuntime/Private/MuR/OpImagePixelFormat.cpp b/Engine/Plugins/Experimental/Mutable/Source/MutableRuntime/Private/MuR/OpImagePixelFormat.cpp
--- a/Engine/Plugins/Experimental/Mutable/Source/MutableRuntime/Private/MuR/OpImagePixelFormat.cpp
+++ b/Engine/Plugins/Experimental/Mutable/Source/MutableRuntime/Private/MuR/OpImagePixelFormat.cpp
@@ -3,7 +3,7 @@
 
 #include "MuR/OpImagePixelFormat.h"
 
-#include "ImageRLE.h"
+#include "MuR/ImageRLE.h"
 #include "MuR/BlockCompression/BlockCompression.h"
 
 #include "MuR/MutableTrace.h"
diff --git a/Engine/Plugins/Experimental/Mutable/Source/MutableRuntime/Private/MuR/OpImageProject.cpp b/Engine/Plugins/Experimental/Mutable/Source/MutableRuntime/Private/MuR/OpImageProject.cpp
--- a/Engine/Plugins/Experimental/Mutable/Source/MutableRuntime/Private/MuR/OpImageProject.cpp
+++ b/Engine/Plugins/Experimental/Mutable/Source/MutableRuntime/Private/MuR/OpImageProject.cpp
@@ -8,10 +8,10 @@
 #include "MuR/Platform.h"
 #include "MuR/ImagePrivate.h"
 #include "MuR/MeshPrivate.h"
-#include "Raster.h"
-#include "ConvertData.h"
+#include "MuR/Raster.h"
+#include "MuR/ConvertData.h"
 #include "MuR/MutableMath.h"
-#include "Octree.h"
+#include "MuR/Octree.h"
 
 #include <queue>
 #include <unordered_set>
@@ -2742,4 +2742,4 @@ MeshPtr CreateMeshOptimisedForWrappingProjection( int layout )
     return pFormatMesh;
 }
 
-}
\ No newline at end of file
+}
diff --git a/Engine/Plugins/Experimental/Mutable/Source/MutableRuntime/Private/MuR/OpImageProject.h b/Engine/Plugins/Experimental/Mutable/Source/MutableRuntime/Private/MuR/OpImageProject.h
--- a/Engine/Plugins/Experimental/Mutable/Source/MutableRuntime/Private/MuR/OpImageProject.h
+++ b/Engine/Plugins/Experimental/Mutable/Source/MutableRuntime/Private/MuR/OpImageProject.h
@@ -3,7 +3,7 @@
 #pragma once
 
 #include "MuR/Platform.h"
-#include "Raster.h"
+#include "MuR/Raster.h"
 #include "MuR/ParametersPrivate.h"
 #include "MuR/Image.h"
 #include "MuR/Mesh.h"
diff --git a/Engine/Plugins/Experimental/Mutable/Source/MutableRuntime/Private/MuR/OpImageRasterMesh.h b/Engine/Plugins/Experimental/Mutable/Source/MutableRuntime/Private/MuR/OpImageRasterMesh.h
--- a/Engine/Plugins/Experimental/Mutable/Source/MutableRuntime/Private/MuR/OpImageRasterMesh.h
+++ b/Engine/Plugins/Experimental/Mutable/Source/MutableRuntime/Private/MuR/OpImageRasterMesh.h
@@ -5,8 +5,8 @@
 #include "MuR/ImagePrivate.h"
 #include "MuR/Platform.h"
 #include "MuR/MeshPrivate.h"
-#include "Raster.h"
-#include "ConvertData.h"
+#include "MuR/Raster.h"
+#include "MuR/ConvertData.h"
 
 
 namespace mu
diff --git a/Engine/Plugins/Experimental/Mutable/Source/MutableRuntime/Private/MuR/OpMeshMorph.h b/Engine/Plugins/Experimental/Mutable/Source/MutableRuntime/Private/MuR/OpMeshMorph.h
--- a/Engine/Plugins/Experimental/Mutable/Source/MutableRuntime/Private/MuR/OpMeshMorph.h
+++ b/Engine/Plugins/Experimental/Mutable/Source/MutableRuntime/Private/MuR/OpMeshMorph.h
@@ -3,10 +3,10 @@
 #pragma once
 
 #include "MuR/MeshPrivate.h"
-#include "ConvertData.h"
+#include "MuR/ConvertData.h"
 #include "MuR/Platform.h"
 
-#include "SparseIndexMap.h"
+#include "MuR/SparseIndexMap.h"
 
 namespace mu
 {
diff --git a/Engine/Plugins/Experimental/Mutable/Source/MutableRuntime/Private/MuR/Settings.cpp b/Engine/Plugins/Experimental/Mutable/Source/MutableRuntime/Private/MuR/Settings.cpp
--- a/Engine/Plugins/Experimental/Mutable/Source/MutableRuntime/Private/MuR/Settings.cpp
+++ b/Engine/Plugins/Experimental/Mutable/Source/MutableRuntime/Private/MuR/Settings.cpp
@@ -4,7 +4,7 @@
 #include "MuR/Platform.h"
 
 #include "MuR/SettingsPrivate.h"
-#include "Config.h"
+#include "MuR/Config.h"
 
 namespace mu
 {
diff --git a/Engine/Plugins/Experimental/Mutable/Source/MutableRuntime/Private/MuR/System.cpp b/Engine/Plugins/Experimental/Mutable/Source/MutableRuntime/Private/MuR/System.cpp
--- a/Engine/Plugins/Experimental/Mutable/Source/MutableRuntime/Private/MuR/System.cpp
+++ b/Engine/Plugins/Experimental/Mutable/Source/MutableRuntime/Private/MuR/System.cpp
@@ -16,7 +16,7 @@
 #include "MuR/Operations.h"
 #include "MuR/CodeVisitor.h"
 
-#include "Config.h"
+#include "MuR/Config.h"
 #include "MuR/MutableTrace.h"
 
 namespace mu
diff --git a/Engine/Plugins/Experimental/Mutable/Source/MutableTools/Private/MuT/ASTOpConditional.cpp b/Engine/Plugins/Experimental/Mutable/Source/MutableTools/Private/MuT/ASTOpConditional.cpp
--- a/Engine/Plugins/Experimental/Mutable/Source/MutableTools/Private/MuT/ASTOpConditional.cpp
+++ b/Engine/Plugins/Experimental/Mutable/Source/MutableTools/Private/MuT/ASTOpConditional.cpp
@@ -1,7 +1,7 @@
 // Copyright Epic Games, Inc. All Rights Reserved.
 
-#include "ASTOpConditional.h"
-#include "ASTOpConstantBool.h"
+#include "MuT/ASTOpConditional.h"
+#include "MuT/ASTOpConstantBool.h"
 #include "MuT/ErrorLogPrivate.h"
 #include "MuT/CodeOptimiser.h"
 #include "MuT/StreamsPrivate.h"
diff --git a/Engine/Plugins/Experimental/Mutable/Source/MutableTools/Private/MuT/ASTOpConstantBool.cpp b/Engine/Plugins/Experimental/Mutable/Source/MutableTools/Private/MuT/ASTOpConstantBool.cpp
--- a/Engine/Plugins/Experimental/Mutable/Source/MutableTools/Private/MuT/ASTOpConstantBool.cpp
+++ b/Engine/Plugins/Experimental/Mutable/Source/MutableTools/Private/MuT/ASTOpConstantBool.cpp
@@ -1,6 +1,6 @@
 // Copyright Epic Games, Inc. All Rights Reserved.
 
-#include "ASTOpConstantBool.h"
+#include "MuT/ASTOpConstantBool.h"
 #include "MuT/ErrorLogPrivate.h"
 #include "MuT/CodeOptimiser.h"
 #include "Hash/CityHash.h"
diff --git a/Engine/Plugins/Experimental/Mutable/Source/MutableTools/Private/MuT/ASTOpConstantProjector.cpp b/Engine/Plugins/Experimental/Mutable/Source/MutableTools/Private/MuT/ASTOpConstantProjector.cpp
--- a/Engine/Plugins/Experimental/Mutable/Source/MutableTools/Private/MuT/ASTOpConstantProjector.cpp
+++ b/Engine/Plugins/Experimental/Mutable/Source/MutableTools/Private/MuT/ASTOpConstantProjector.cpp
@@ -1,6 +1,6 @@
 // Copyright Epic Games, Inc. All Rights Reserved.
 
-#include "ASTOpConstantProjector.h"
+#include "MuT/ASTOpConstantProjector.h"
 #include "MuT/ErrorLogPrivate.h"
 #include "MuT/CodeOptimiser.h"
 #include "MuT/StreamsPrivate.h"
diff --git a/Engine/Plugins/Experimental/Mutable/Source/MutableTools/Private/MuT/ASTOpConstantResource.cpp b/Engine/Plugins/Experimental/Mutable/Source/MutableTools/Private/MuT/ASTOpConstantResource.cpp
--- a/Engine/Plugins/Experimental/Mutable/Source/MutableTools/Private/MuT/ASTOpConstantResource.cpp
+++ b/Engine/Plugins/Experimental/Mutable/Source/MutableTools/Private/MuT/ASTOpConstantResource.cpp
@@ -1,6 +1,6 @@
 // Copyright Epic Games, Inc. All Rights Reserved.
 
-#include "ASTOpConstantResource.h"
+#include "MuT/ASTOpConstantResource.h"
 #include "MuT/ErrorLogPrivate.h"
 #include "MuT/CodeOptimiser.h"
 #include "MuT/StreamsPrivate.h"
diff --git a/Engine/Plugins/Experimental/Mutable/Source/MutableTools/Private/MuT/ASTOpConstantString.cpp b/Engine/Plugins/Experimental/Mutable/Source/MutableTools/Private/MuT/ASTOpConstantString.cpp
--- a/Engine/Plugins/Experimental/Mutable/Source/MutableTools/Private/MuT/ASTOpConstantString.cpp
+++ b/Engine/Plugins/Experimental/Mutable/Source/MutableTools/Private/MuT/ASTOpConstantString.cpp
@@ -1,6 +1,6 @@
 // Copyright Epic Games, Inc. All Rights Reserved.
 
-#include "ASTOpConstantString.h"
+#include "MuT/ASTOpConstantString.h"
 #include "MuT/ErrorLogPrivate.h"
 #include "MuT/CodeOptimiser.h"
 #include "MuT/StreamsPrivate.h"
diff --git a/Engine/Plugins/Experimental/Mutable/Source/MutableTools/Private/MuT/ASTOpImageCompose.cpp b/Engine/Plugins/Experimental/Mutable/Source/MutableTools/Private/MuT/ASTOpImageCompose.cpp
--- a/Engine/Plugins/Experimental/Mutable/Source/MutableTools/Private/MuT/ASTOpImageCompose.cpp
+++ b/Engine/Plugins/Experimental/Mutable/Source/MutableTools/Private/MuT/ASTOpImageCompose.cpp
@@ -1,6 +1,6 @@
 // Copyright Epic Games, Inc. All Rights Reserved.
 
-#include "ASTOpImageCompose.h"
+#include "MuT/ASTOpImageCompose.h"
 #include "MuT/ErrorLogPrivate.h"
 #include "MuT/CodeOptimiser.h"
 #include "MuT/StreamsPrivate.h"
diff --git a/Engine/Plugins/Experimental/Mutable/Source/MutableTools/Private/MuT/ASTOpImageMipmap.cpp b/Engine/Plugins/Experimental/Mutable/Source/MutableTools/Private/MuT/ASTOpImageMipmap.cpp
--- a/Engine/Plugins/Experimental/Mutable/Source/MutableTools/Private/MuT/ASTOpImageMipmap.cpp
+++ b/Engine/Plugins/Experimental/Mutable/Source/MutableTools/Private/MuT/ASTOpImageMipmap.cpp
@@ -1,12 +1,12 @@
 // Copyright Epic Games, Inc. All Rights Reserved.
 
-#include "ASTOpImageMipmap.h"
+#include "MuT/ASTOpImageMipmap.h"
 
-#include "ASTOpImagePixelFormat.h"
-#include "ASTOpImageCompose.h"
-#include "ASTOpImagePatch.h"
-#include "ASTOpSwitch.h"
-#include "ASTOpConditional.h"
+#include "MuT/ASTOpImagePixelFormat.h"
+#include "MuT/ASTOpImageCompose.h"
+#include "MuT/ASTOpImagePatch.h"
+#include "MuT/ASTOpSwitch.h"
+#include "MuT/ASTOpConditional.h"
 
 #include "MuT/ErrorLogPrivate.h"
 #include "MuT/CodeOptimiser.h"
diff --git a/Engine/Plugins/Experimental/Mutable/Source/MutableTools/Private/MuT/ASTOpImageMultiLayer.cpp b/Engine/Plugins/Experimental/Mutable/Source/MutableTools/Private/MuT/ASTOpImageMultiLayer.cpp
--- a/Engine/Plugins/Experimental/Mutable/Source/MutableTools/Private/MuT/ASTOpImageMultiLayer.cpp
+++ b/Engine/Plugins/Experimental/Mutable/Source/MutableTools/Private/MuT/ASTOpImageMultiLayer.cpp
@@ -1,6 +1,6 @@
 // Copyright Epic Games, Inc. All Rights Reserved.
 
-#include "ASTOpImageMultiLayer.h"
+#include "MuT/ASTOpImageMultiLayer.h"
 #include "MuT/ErrorLogPrivate.h"
 #include "MuT/CodeOptimiser.h"
 #include "MuT/StreamsPrivate.h"
diff --git a/Engine/Plugins/Experimental/Mutable/Source/MutableTools/Private/MuT/ASTOpImageNormalComposite.cpp b/Engine/Plugins/Experimental/Mutable/Source/MutableTools/Private/MuT/ASTOpImageNormalComposite.cpp
--- a/Engine/Plugins/Experimental/Mutable/Source/MutableTools/Private/MuT/ASTOpImageNormalComposite.cpp
+++ b/Engine/Plugins/Experimental/Mutable/Source/MutableTools/Private/MuT/ASTOpImageNormalComposite.cpp
@@ -1,6 +1,6 @@
 // Copyright Epic Games, Inc. All Rights Reserved.
 
-#include "ASTOpImageNormalComposite.h"
+#include "MuT/ASTOpImageNormalComposite.h"
 #include "MuT/ErrorLogPrivate.h"
 #include "MuT/CodeOptimiser.h"
 #include "MuT/StreamsPrivate.h"
diff --git a/Engine/Plugins/Experimental/Mutable/Source/MutableTools/Private/MuT/ASTOpImagePatch.cpp b/Engine/Plugins/Experimental/Mutable/Source/MutableTools/Private/MuT/ASTOpImagePatch.cpp
--- a/Engine/Plugins/Experimental/Mutable/Source/MutableTools/Private/MuT/ASTOpImagePatch.cpp
+++ b/Engine/Plugins/Experimental/Mutable/Source/MutableTools/Private/MuT/ASTOpImagePatch.cpp
@@ -1,6 +1,6 @@
 // Copyright Epic Games, Inc. All Rights Reserved.
 
-#include "ASTOpImagePatch.h"
+#include "MuT/ASTOpImagePatch.h"
 #include "MuT/ErrorLogPrivate.h"
 #include "MuT/CodeOptimiser.h"
 #include "MuT/StreamsPrivate.h"
diff --git a/Engine/Plugins/Experimental/Mutable/Source/MutableTools/Private/MuT/ASTOpImagePixelFormat.cpp b/Engine/Plugins/Experimental/Mutable/Source/MutableTools/Private/MuT/ASTOpImagePixelFormat.cpp
--- a/Engine/Plugins/Experimental/Mutable/Source/MutableTools/Private/MuT/ASTOpImagePixelFormat.cpp
+++ b/Engine/Plugins/Experimental/Mutable/Source/MutableTools/Private/MuT/ASTOpImagePixelFormat.cpp
@@ -1,11 +1,11 @@
 // Copyright Epic Games, Inc. All Rights Reserved.
 
-#include "ASTOpImagePixelFormat.h"
-#include "ASTOpImageMipmap.h"
-#include "ASTOpImageCompose.h"
-#include "ASTOpImagePatch.h"
-#include "ASTOpConditional.h"
-#include "ASTOpSwitch.h"
+#include "MuT/ASTOpImagePixelFormat.h"
+#include "MuT/ASTOpImageMipmap.h"
+#include "MuT/ASTOpImageCompose.h"
+#include "MuT/ASTOpImagePatch.h"
+#include "MuT/ASTOpConditional.h"
+#include "MuT/ASTOpSwitch.h"
 #include "MuT/ErrorLogPrivate.h"
 #include "MuT/CodeOptimiser.h"
 #include "MuT/StreamsPrivate.h"
