Updated more paths in Mutable to contain the module name.

diff --git a/Engine/Plugins/Experimental/Mutable/Source/MutableTools/Private/MuT/ASTOpImageTransform.cpp b/Engine/Plugins/Experimental/Mutable/Source/MutableTools/Private/MuT/ASTOpImageTransform.cpp
--- a/Engine/Plugins/Experimental/Mutable/Source/MutableTools/Private/MuT/ASTOpImageTransform.cpp
+++ b/Engine/Plugins/Experimental/Mutable/Source/MutableTools/Private/MuT/ASTOpImageTransform.cpp
@@ -1,6 +1,6 @@
 // Copyright Epic Games, Inc. All Rights Reserved.
 
-#include "ASTOpImageTransform.h"
+#include "MuT/ASTOpImageTransform.h"
 #include "MuT/ErrorLogPrivate.h"
 #include "MuT/CodeOptimiser.h"
 #include "MuT/StreamsPrivate.h"
diff --git a/Engine/Plugins/Experimental/Mutable/Source/MutableTools/Private/MuT/ASTOpInstanceAdd.cpp b/Engine/Plugins/Experimental/Mutable/Source/MutableTools/Private/MuT/ASTOpInstanceAdd.cpp
--- a/Engine/Plugins/Experimental/Mutable/Source/MutableTools/Private/MuT/ASTOpInstanceAdd.cpp
+++ b/Engine/Plugins/Experimental/Mutable/Source/MutableTools/Private/MuT/ASTOpInstanceAdd.cpp
@@ -1,6 +1,6 @@
 // Copyright Epic Games, Inc. All Rights Reserved.
 
-#include "ASTOpInstanceAdd.h"
+#include "MuT/ASTOpInstanceAdd.h"
 #include "MuT/ErrorLogPrivate.h"
 #include "MuT/CodeOptimiser.h"
 #include "MuT/StreamsPrivate.h"
diff --git a/Engine/Plugins/Experimental/Mutable/Source/MutableTools/Private/MuT/ASTOpMeshApplyPose.cpp b/Engine/Plugins/Experimental/Mutable/Source/MutableTools/Private/MuT/ASTOpMeshApplyPose.cpp
--- a/Engine/Plugins/Experimental/Mutable/Source/MutableTools/Private/MuT/ASTOpMeshApplyPose.cpp
+++ b/Engine/Plugins/Experimental/Mutable/Source/MutableTools/Private/MuT/ASTOpMeshApplyPose.cpp
@@ -1,6 +1,6 @@
 // Copyright Epic Games, Inc. All Rights Reserved.
 
-#include "ASTOpMeshApplyPose.h"
+#include "MuT/ASTOpMeshApplyPose.h"
 #include "MuT/ErrorLogPrivate.h"
 #include "MuT/CodeOptimiser.h"
 #include "MuT/StreamsPrivate.h"
diff --git a/Engine/Plugins/Experimental/Mutable/Source/MutableTools/Private/MuT/ASTOpMeshApplyShape.cpp b/Engine/Plugins/Experimental/Mutable/Source/MutableTools/Private/MuT/ASTOpMeshApplyShape.cpp
--- a/Engine/Plugins/Experimental/Mutable/Source/MutableTools/Private/MuT/ASTOpMeshApplyShape.cpp
+++ b/Engine/Plugins/Experimental/Mutable/Source/MutableTools/Private/MuT/ASTOpMeshApplyShape.cpp
@@ -1,6 +1,6 @@
 // Copyright Epic Games, Inc. All Rights Reserved.
 
-#include "ASTOpMeshApplyShape.h"
+#include "MuT/ASTOpMeshApplyShape.h"
 #include "MuT/ErrorLogPrivate.h"
 #include "MuT/CodeOptimiser.h"
 #include "MuT/StreamsPrivate.h"
diff --git a/Engine/Plugins/Experimental/Mutable/Source/MutableTools/Private/MuT/ASTOpMeshBindShape.cpp b/Engine/Plugins/Experimental/Mutable/Source/MutableTools/Private/MuT/ASTOpMeshBindShape.cpp
--- a/Engine/Plugins/Experimental/Mutable/Source/MutableTools/Private/MuT/ASTOpMeshBindShape.cpp
+++ b/Engine/Plugins/Experimental/Mutable/Source/MutableTools/Private/MuT/ASTOpMeshBindShape.cpp
@@ -1,7 +1,7 @@
 // Copyright Epic Games, Inc. All Rights Reserved.
 
-#include "ASTOpMeshBindShape.h"
-#include "ASTOpSwitch.h"
+#include "MuT/ASTOpMeshBindShape.h"
+#include "MuT/ASTOpSwitch.h"
 #include "MuT/ErrorLogPrivate.h"
 #include "MuT/CodeOptimiser.h"
 #include "MuT/StreamsPrivate.h"
diff --git a/Engine/Plugins/Experimental/Mutable/Source/MutableTools/Private/MuT/ASTOpMeshClipDeform.cpp b/Engine/Plugins/Experimental/Mutable/Source/MutableTools/Private/MuT/ASTOpMeshClipDeform.cpp
--- a/Engine/Plugins/Experimental/Mutable/Source/MutableTools/Private/MuT/ASTOpMeshClipDeform.cpp
+++ b/Engine/Plugins/Experimental/Mutable/Source/MutableTools/Private/MuT/ASTOpMeshClipDeform.cpp
@@ -1,6 +1,6 @@
 // Copyright Epic Games, Inc. All Rights Reserved.
 
-#include "ASTOpMeshClipDeform.h"
+#include "MuT/ASTOpMeshClipDeform.h"
 #include "MuT/ErrorLogPrivate.h"
 #include "MuT/CodeOptimiser.h"
 #include "MuT/StreamsPrivate.h"
diff --git a/Engine/Plugins/Experimental/Mutable/Source/MutableTools/Private/MuT/ASTOpMeshClipMorphPlane.cpp b/Engine/Plugins/Experimental/Mutable/Source/MutableTools/Private/MuT/ASTOpMeshClipMorphPlane.cpp
--- a/Engine/Plugins/Experimental/Mutable/Source/MutableTools/Private/MuT/ASTOpMeshClipMorphPlane.cpp
+++ b/Engine/Plugins/Experimental/Mutable/Source/MutableTools/Private/MuT/ASTOpMeshClipMorphPlane.cpp
@@ -1,6 +1,6 @@
 // Copyright Epic Games, Inc. All Rights Reserved.
 
-#include "ASTOpMeshClipMorphPlane.h"
+#include "MuT/ASTOpMeshClipMorphPlane.h"
 #include "MuT/ErrorLogPrivate.h"
 #include "MuT/CodeOptimiser.h"
 #include "MuT/StreamsPrivate.h"
diff --git a/Engine/Plugins/Experimental/Mutable/Source/MutableTools/Private/MuT/ASTOpMeshExtractLayoutBlocks.cpp b/Engine/Plugins/Experimental/Mutable/Source/MutableTools/Private/MuT/ASTOpMeshExtractLayoutBlocks.cpp
--- a/Engine/Plugins/Experimental/Mutable/Source/MutableTools/Private/MuT/ASTOpMeshExtractLayoutBlocks.cpp
+++ b/Engine/Plugins/Experimental/Mutable/Source/MutableTools/Private/MuT/ASTOpMeshExtractLayoutBlocks.cpp
@@ -1,6 +1,6 @@
 // Copyright Epic Games, Inc. All Rights Reserved.
 
-#include "ASTOpMeshExtractLayoutBlocks.h"
+#include "MuT/ASTOpMeshExtractLayoutBlocks.h"
 #include "MuT/ErrorLogPrivate.h"
 #include "MuT/CodeOptimiser.h"
 #include "MuT/StreamsPrivate.h"
diff --git a/Engine/Plugins/Experimental/Mutable/Source/MutableTools/Private/MuT/ASTOpMeshGeometryOperation.cpp b/Engine/Plugins/Experimental/Mutable/Source/MutableTools/Private/MuT/ASTOpMeshGeometryOperation.cpp
--- a/Engine/Plugins/Experimental/Mutable/Source/MutableTools/Private/MuT/ASTOpMeshGeometryOperation.cpp
+++ b/Engine/Plugins/Experimental/Mutable/Source/MutableTools/Private/MuT/ASTOpMeshGeometryOperation.cpp
@@ -1,6 +1,6 @@
 // Copyright Epic Games, Inc. All Rights Reserved.
 
-#include "ASTOpMeshGeometryOperation.h"
+#include "MuT/ASTOpMeshGeometryOperation.h"
 #include "MuT/ErrorLogPrivate.h"
 #include "MuT/CodeOptimiser.h"
 #include "MuT/StreamsPrivate.h"
@@ -87,4 +87,4 @@ namespace mu
 
 	}
 
-}
\ No newline at end of file
+}
diff --git a/Engine/Plugins/Experimental/Mutable/Source/MutableTools/Private/MuT/ASTOpMeshMorphReshape.cpp b/Engine/Plugins/Experimental/Mutable/Source/MutableTools/Private/MuT/ASTOpMeshMorphReshape.cpp
--- a/Engine/Plugins/Experimental/Mutable/Source/MutableTools/Private/MuT/ASTOpMeshMorphReshape.cpp
+++ b/Engine/Plugins/Experimental/Mutable/Source/MutableTools/Private/MuT/ASTOpMeshMorphReshape.cpp
@@ -1,6 +1,6 @@
 // Copyright Epic Games, Inc. All Rights Reserved.
 
-#include "ASTOpMeshMorphReshape.h"
+#include "MuT/ASTOpMeshMorphReshape.h"
 #include "MuT/ErrorLogPrivate.h"
 #include "MuT/CodeOptimiser.h"
 #include "MuT/StreamsPrivate.h"
diff --git a/Engine/Plugins/Experimental/Mutable/Source/MutableTools/Private/MuT/ASTOpMeshRemapIndices.cpp b/Engine/Plugins/Experimental/Mutable/Source/MutableTools/Private/MuT/ASTOpMeshRemapIndices.cpp
--- a/Engine/Plugins/Experimental/Mutable/Source/MutableTools/Private/MuT/ASTOpMeshRemapIndices.cpp
+++ b/Engine/Plugins/Experimental/Mutable/Source/MutableTools/Private/MuT/ASTOpMeshRemapIndices.cpp
@@ -1,6 +1,6 @@
 // Copyright Epic Games, Inc. All Rights Reserved.
 
-#include "ASTOpMeshRemapIndices.h"
+#include "MuT/ASTOpMeshRemapIndices.h"
 #include "MuT/ErrorLogPrivate.h"
 #include "MuT/CodeOptimiser.h"
 #include "MuT/StreamsPrivate.h"
diff --git a/Engine/Plugins/Experimental/Mutable/Source/MutableTools/Private/MuT/ASTOpMeshRemoveMask.cpp b/Engine/Plugins/Experimental/Mutable/Source/MutableTools/Private/MuT/ASTOpMeshRemoveMask.cpp
--- a/Engine/Plugins/Experimental/Mutable/Source/MutableTools/Private/MuT/ASTOpMeshRemoveMask.cpp
+++ b/Engine/Plugins/Experimental/Mutable/Source/MutableTools/Private/MuT/ASTOpMeshRemoveMask.cpp
@@ -1,6 +1,6 @@
 // Copyright Epic Games, Inc. All Rights Reserved.
 
-#include "ASTOpMeshRemoveMask.h"
+#include "MuT/ASTOpMeshRemoveMask.h"
 #include "MuT/ErrorLogPrivate.h"
 #include "MuT/CodeOptimiser.h"
 #include "MuT/StreamsPrivate.h"
diff --git a/Engine/Plugins/Experimental/Mutable/Source/MutableTools/Private/MuT/ASTOpMeshTransform.cpp b/Engine/Plugins/Experimental/Mutable/Source/MutableTools/Private/MuT/ASTOpMeshTransform.cpp
--- a/Engine/Plugins/Experimental/Mutable/Source/MutableTools/Private/MuT/ASTOpMeshTransform.cpp
+++ b/Engine/Plugins/Experimental/Mutable/Source/MutableTools/Private/MuT/ASTOpMeshTransform.cpp
@@ -1,6 +1,6 @@
 // Copyright Epic Games, Inc. All Rights Reserved.
 
-#include "ASTOpMeshTransform.h"
+#include "MuT/ASTOpMeshTransform.h"
 #include "MuT/ErrorLogPrivate.h"
 #include "MuT/CodeOptimiser.h"
 #include "MuT/StreamsPrivate.h"
diff --git a/Engine/Plugins/Experimental/Mutable/Source/MutableTools/Private/MuT/ASTOpParameter.cpp b/Engine/Plugins/Experimental/Mutable/Source/MutableTools/Private/MuT/ASTOpParameter.cpp
--- a/Engine/Plugins/Experimental/Mutable/Source/MutableTools/Private/MuT/ASTOpParameter.cpp
+++ b/Engine/Plugins/Experimental/Mutable/Source/MutableTools/Private/MuT/ASTOpParameter.cpp
@@ -1,6 +1,6 @@
 // Copyright Epic Games, Inc. All Rights Reserved.
 
-#include "ASTOpParameter.h"
+#include "MuT/ASTOpParameter.h"
 #include "MuT/ErrorLogPrivate.h"
 #include "MuT/CodeOptimiser.h"
 #include "MuT/StreamsPrivate.h"
diff --git a/Engine/Plugins/Experimental/Mutable/Source/MutableTools/Private/MuT/ASTOpScalarCurve.cpp b/Engine/Plugins/Experimental/Mutable/Source/MutableTools/Private/MuT/ASTOpScalarCurve.cpp
--- a/Engine/Plugins/Experimental/Mutable/Source/MutableTools/Private/MuT/ASTOpScalarCurve.cpp
+++ b/Engine/Plugins/Experimental/Mutable/Source/MutableTools/Private/MuT/ASTOpScalarCurve.cpp
@@ -1,6 +1,6 @@
 // Copyright Epic Games, Inc. All Rights Reserved.
 
-#include "ASTOpScalarCurve.h"
+#include "MuT/ASTOpScalarCurve.h"
 #include "MuT/ErrorLogPrivate.h"
 #include "MuT/CodeOptimiser.h"
 #include "MuT/StreamsPrivate.h"
diff --git a/Engine/Plugins/Experimental/Mutable/Source/MutableTools/Private/MuT/ASTOpSwitch.cpp b/Engine/Plugins/Experimental/Mutable/Source/MutableTools/Private/MuT/ASTOpSwitch.cpp
--- a/Engine/Plugins/Experimental/Mutable/Source/MutableTools/Private/MuT/ASTOpSwitch.cpp
+++ b/Engine/Plugins/Experimental/Mutable/Source/MutableTools/Private/MuT/ASTOpSwitch.cpp
@@ -1,6 +1,6 @@
 // Copyright Epic Games, Inc. All Rights Reserved.
 
-#include "ASTOpSwitch.h"
+#include "MuT/ASTOpSwitch.h"
 #include "MuT/ErrorLogPrivate.h"
 #include "MuT/CodeOptimiser.h"
 #include "MuT/StreamsPrivate.h"
diff --git a/Engine/Plugins/Experimental/Mutable/Source/MutableTools/Private/MuT/CodeGenerator.cpp b/Engine/Plugins/Experimental/Mutable/Source/MutableTools/Private/MuT/CodeGenerator.cpp
--- a/Engine/Plugins/Experimental/Mutable/Source/MutableTools/Private/MuT/CodeGenerator.cpp
+++ b/Engine/Plugins/Experimental/Mutable/Source/MutableTools/Private/MuT/CodeGenerator.cpp
@@ -1,36 +1,36 @@
 // Copyright Epic Games, Inc. All Rights Reserved.
 
-#include "CodeGenerator.h"
+#include "MuT/CodeGenerator.h"
 
-#include "CodeGenerator_SecondPass.h"
+#include "MuT/CodeGenerator_SecondPass.h"
 #include "MuT/CodeOptimiser.h"
 #include "MuR/Platform.h"
 #include "MuT/CompilerPrivate.h"
-#include "ImageDescGenerator.h"
-#include "TablePrivate.h"
+#include "MuT/ImageDescGenerator.h"
+#include "MuT/TablePrivate.h"
 #include "MuT/NodeLayoutPrivate.h"
 
 #include "MuT/NodeBoolPrivate.h"
 
-#include "NodeRangeFromScalarPrivate.h"
+#include "MuT/NodeRangeFromScalarPrivate.h"
 
-#include "NodeProjectorPrivate.h"
+#include "MuT/NodeProjectorPrivate.h"
 
 #include "MuT/NodeMeshTablePrivate.h"
 #include "MuT/NodeImageTablePrivate.h"
 
-#include "NodeSurfaceNewPrivate.h"
-#include "NodeSurfaceEditPrivate.h"
-#include "NodeSurfaceVariationPrivate.h"
+#include "MuT/NodeSurfaceNewPrivate.h"
+#include "MuT/NodeSurfaceEditPrivate.h"
+#include "MuT/NodeSurfaceVariationPrivate.h"
 
-#include "NodeComponentPrivate.h"
+#include "MuT/NodeComponentPrivate.h"
 #include "MuT/NodeComponentNewPrivate.h"
-#include "NodeComponentEditPrivate.h"
+#include "MuT/NodeComponentEditPrivate.h"
 
-#include "NodeLODPrivate.h"
+#include "MuT/NodeLODPrivate.h"
 #include "MuT/NodeObjectNewPrivate.h"
 #include "MuT/NodeObjectStatePrivate.h"
-#include "NodeObjectGroupPrivate.h"
+#include "MuT/NodeObjectGroupPrivate.h"
 
 #include "MuT/NodeImageFormatPrivate.h"
 #include "MuT/NodeImageMipmapPrivate.h"
diff --git a/Engine/Plugins/Experimental/Mutable/Source/MutableTools/Private/MuT/CodeGenerator_Bool.cpp b/Engine/Plugins/Experimental/Mutable/Source/MutableTools/Private/MuT/CodeGenerator_Bool.cpp
--- a/Engine/Plugins/Experimental/Mutable/Source/MutableTools/Private/MuT/CodeGenerator_Bool.cpp
+++ b/Engine/Plugins/Experimental/Mutable/Source/MutableTools/Private/MuT/CodeGenerator_Bool.cpp
@@ -1,6 +1,6 @@
 // Copyright Epic Games, Inc. All Rights Reserved.
 
-#include "CodeGenerator.h"
+#include "MuT/CodeGenerator.h"
 
 #include "MuR/Platform.h"
 #include "MuT/CompilerPrivate.h"
diff --git a/Engine/Plugins/Experimental/Mutable/Source/MutableTools/Private/MuT/CodeGenerator_Color.cpp b/Engine/Plugins/Experimental/Mutable/Source/MutableTools/Private/MuT/CodeGenerator_Color.cpp
--- a/Engine/Plugins/Experimental/Mutable/Source/MutableTools/Private/MuT/CodeGenerator_Color.cpp
+++ b/Engine/Plugins/Experimental/Mutable/Source/MutableTools/Private/MuT/CodeGenerator_Color.cpp
@@ -1,6 +1,6 @@
 // Copyright Epic Games, Inc. All Rights Reserved.
 
-#include "CodeGenerator.h"
+#include "MuT/CodeGenerator.h"
 
 #include "MuR/Platform.h"
 #include "MuT/CompilerPrivate.h"
diff --git a/Engine/Plugins/Experimental/Mutable/Source/MutableTools/Private/MuT/CodeGenerator_FirstPass.cpp b/Engine/Plugins/Experimental/Mutable/Source/MutableTools/Private/MuT/CodeGenerator_FirstPass.cpp
--- a/Engine/Plugins/Experimental/Mutable/Source/MutableTools/Private/MuT/CodeGenerator_FirstPass.cpp
+++ b/Engine/Plugins/Experimental/Mutable/Source/MutableTools/Private/MuT/CodeGenerator_FirstPass.cpp
@@ -1,34 +1,34 @@
 // Copyright Epic Games, Inc. All Rights Reserved.
 
-#include "CodeGenerator_FirstPass.h"
+#include "MuT/CodeGenerator_FirstPass.h"
 
 #include "MuR/Platform.h"
 #include "MuT/CompilerPrivate.h"
 
-#include "NodeSurfaceNewPrivate.h"
-#include "NodeSurfaceEditPrivate.h"
-#include "NodeSurfaceVariationPrivate.h"
+#include "MuT/NodeSurfaceNewPrivate.h"
+#include "MuT/NodeSurfaceEditPrivate.h"
+#include "MuT/NodeSurfaceVariationPrivate.h"
 
-#include "NodeComponentPrivate.h"
+#include "MuT/NodeComponentPrivate.h"
 #include "MuT/NodeComponentNewPrivate.h"
-#include "NodeComponentEditPrivate.h"
+#include "MuT/NodeComponentEditPrivate.h"
 
-#include "NodeLODPrivate.h"
+#include "MuT/NodeLODPrivate.h"
 #include "MuT/NodeObjectNewPrivate.h"
 #include "MuT/NodeObjectStatePrivate.h"
-#include "NodeObjectGroupPrivate.h"
+#include "MuT/NodeObjectGroupPrivate.h"
 
 #include "MuT/NodeMeshClipMorphPlane.h"
 #include "MuT/NodeMeshClipWithMesh.h"
 
-#include "NodeModifierPrivate.h"
-#include "NodeModifierMeshClipMorphPlanePrivate.h"
-#include "NodeModifierMeshClipWithMeshPrivate.h"
-#include "NodeModifierMeshClipDeformPrivate.h"
+#include "MuT/NodeModifierPrivate.h"
+#include "MuT/NodeModifierMeshClipMorphPlanePrivate.h"
+#include "MuT/NodeModifierMeshClipWithMeshPrivate.h"
+#include "MuT/NodeModifierMeshClipDeformPrivate.h"
 
 #include "MuT/AST.h"
-#include "ASTOpConstantBool.h"
-#include "ASTOpParameter.h"
+#include "MuT/ASTOpConstantBool.h"
+#include "MuT/ASTOpParameter.h"
 
 #include "MuR/ModelPrivate.h"
 #include "MuR/MutableTrace.h"
diff --git a/Engine/Plugins/Experimental/Mutable/Source/MutableTools/Private/MuT/CodeGenerator_FirstPass.h b/Engine/Plugins/Experimental/Mutable/Source/MutableTools/Private/MuT/CodeGenerator_FirstPass.h
--- a/Engine/Plugins/Experimental/Mutable/Source/MutableTools/Private/MuT/CodeGenerator_FirstPass.h
+++ b/Engine/Plugins/Experimental/Mutable/Source/MutableTools/Private/MuT/CodeGenerator_FirstPass.h
@@ -2,7 +2,7 @@
 
 #pragma once
 
-#include "Visitor.h"
+#include "MuT/Visitor.h"
 
 #include "MuT/AST.h"
 #include "MuT/CompilerPrivate.h"
diff --git a/Engine/Plugins/Experimental/Mutable/Source/MutableTools/Private/MuT/CodeGenerator_Image.cpp b/Engine/Plugins/Experimental/Mutable/Source/MutableTools/Private/MuT/CodeGenerator_Image.cpp
--- a/Engine/Plugins/Experimental/Mutable/Source/MutableTools/Private/MuT/CodeGenerator_Image.cpp
+++ b/Engine/Plugins/Experimental/Mutable/Source/MutableTools/Private/MuT/CodeGenerator_Image.cpp
@@ -1,11 +1,11 @@
 // Copyright Epic Games, Inc. All Rights Reserved.
 
-#include "CodeGenerator.h"
+#include "MuT/CodeGenerator.h"
 
 #include "MuR/Platform.h"
 #include "MuT/CompilerPrivate.h"
-#include "ImageDescGenerator.h"
-#include "TablePrivate.h"
+#include "MuT/ImageDescGenerator.h"
+#include "MuT/TablePrivate.h"
 #include "MuT/NodeLayoutPrivate.h"
 
 #include "MuT/NodeImagePrivate.h"
diff --git a/Engine/Plugins/Experimental/Mutable/Source/MutableTools/Private/MuT/CodeGenerator_Mesh.cpp b/Engine/Plugins/Experimental/Mutable/Source/MutableTools/Private/MuT/CodeGenerator_Mesh.cpp
--- a/Engine/Plugins/Experimental/Mutable/Source/MutableTools/Private/MuT/CodeGenerator_Mesh.cpp
+++ b/Engine/Plugins/Experimental/Mutable/Source/MutableTools/Private/MuT/CodeGenerator_Mesh.cpp
@@ -1,10 +1,10 @@
 // Copyright Epic Games, Inc. All Rights Reserved.
 
-#include "CodeGenerator.h"
+#include "MuT/CodeGenerator.h"
 
 #include "MuR/Platform.h"
 #include "MuT/CompilerPrivate.h"
-#include "TablePrivate.h"
+#include "MuT/TablePrivate.h"
 #include "MuT/NodeLayoutPrivate.h"
 
 #include "MuT/NodeMeshPrivate.h"
@@ -14,7 +14,7 @@
 #include "MuT/NodeMeshInterpolatePrivate.h"
 #include "MuT/NodeMeshMorphPrivate.h"
 #include "MuT/NodeMeshMakeMorphPrivate.h"
-#include "NodePatchMeshPrivate.h"
+#include "MuT/NodePatchMeshPrivate.h"
 #include "MuT/NodeMeshSubtractPrivate.h"
 #include "MuT/NodeMeshFormatPrivate.h"
 #include "MuT/NodeMeshFragmentPrivate.h"
diff --git a/Engine/Plugins/Experimental/Mutable/Source/MutableTools/Private/MuT/CodeGenerator_Projector.cpp b/Engine/Plugins/Experimental/Mutable/Source/MutableTools/Private/MuT/CodeGenerator_Projector.cpp
--- a/Engine/Plugins/Experimental/Mutable/Source/MutableTools/Private/MuT/CodeGenerator_Projector.cpp
+++ b/Engine/Plugins/Experimental/Mutable/Source/MutableTools/Private/MuT/CodeGenerator_Projector.cpp
@@ -1,11 +1,11 @@
 // Copyright Epic Games, Inc. All Rights Reserved.
 
-#include "CodeGenerator.h"
+#include "MuT/CodeGenerator.h"
 
 #include "MuR/Platform.h"
 #include "MuT/CompilerPrivate.h"
 
-#include "NodeProjectorPrivate.h"
+#include "MuT/NodeProjectorPrivate.h"
 #include "MuT/ASTOpConstantProjector.h"
 
 #include "MuR/ModelPrivate.h"
diff --git a/Engine/Plugins/Experimental/Mutable/Source/MutableTools/Private/MuT/CodeGenerator_Scalar.cpp b/Engine/Plugins/Experimental/Mutable/Source/MutableTools/Private/MuT/CodeGenerator_Scalar.cpp
--- a/Engine/Plugins/Experimental/Mutable/Source/MutableTools/Private/MuT/CodeGenerator_Scalar.cpp
+++ b/Engine/Plugins/Experimental/Mutable/Source/MutableTools/Private/MuT/CodeGenerator_Scalar.cpp
@@ -1,17 +1,17 @@
 // Copyright Epic Games, Inc. All Rights Reserved.
 
-#include "CodeGenerator.h"
+#include "MuT/CodeGenerator.h"
 
 #include "MuR/Platform.h"
 #include "MuT/CompilerPrivate.h"
 
-#include "NodeScalarPrivate.h"
-#include "NodeScalarConstantPrivate.h"
-#include "NodeScalarParameterPrivate.h"
-#include "NodeScalarSwitchPrivate.h"
-#include "NodeScalarEnumParameterPrivate.h"
-#include "NodeScalarCurvePrivate.h"
-#include "NodeScalarArithmeticOperationPrivate.h"
+#include "MuT/NodeScalarPrivate.h"
+#include "MuT/NodeScalarConstantPrivate.h"
+#include "MuT/NodeScalarParameterPrivate.h"
+#include "MuT/NodeScalarSwitchPrivate.h"
+#include "MuT/NodeScalarEnumParameterPrivate.h"
+#include "MuT/NodeScalarCurvePrivate.h"
+#include "MuT/NodeScalarArithmeticOperationPrivate.h"
 #include "MuT/NodeScalarVariationPrivate.h"
 #include "MuT/NodeScalarTablePrivate.h"
 
diff --git a/Engine/Plugins/Experimental/Mutable/Source/MutableTools/Private/MuT/CodeGenerator_SecondPass.cpp b/Engine/Plugins/Experimental/Mutable/Source/MutableTools/Private/MuT/CodeGenerator_SecondPass.cpp
--- a/Engine/Plugins/Experimental/Mutable/Source/MutableTools/Private/MuT/CodeGenerator_SecondPass.cpp
+++ b/Engine/Plugins/Experimental/Mutable/Source/MutableTools/Private/MuT/CodeGenerator_SecondPass.cpp
@@ -1,30 +1,30 @@
 // Copyright Epic Games, Inc. All Rights Reserved.
 
-#include "CodeGenerator_SecondPass.h"
-#include "CodeGenerator_FirstPass.h"
-#include "CodeGenerator.h"
+#include "MuT/CodeGenerator_SecondPass.h"
+#include "MuT/CodeGenerator_FirstPass.h"
+#include "MuT/CodeGenerator.h"
 
 #include "MuR/Platform.h"
 #include "MuT/CompilerPrivate.h"
 
-#include "NodeSurfaceNewPrivate.h"
-#include "NodeSurfaceEditPrivate.h"
+#include "MuT/NodeSurfaceNewPrivate.h"
+#include "MuT/NodeSurfaceEditPrivate.h"
 
-#include "NodeComponentPrivate.h"
+#include "MuT/NodeComponentPrivate.h"
 #include "MuT/NodeComponentNewPrivate.h"
-#include "NodeComponentEditPrivate.h"
+#include "MuT/NodeComponentEditPrivate.h"
 
-#include "NodeLODPrivate.h"
+#include "MuT/NodeLODPrivate.h"
 #include "MuT/NodeObjectNewPrivate.h"
 #include "MuT/NodeObjectStatePrivate.h"
-#include "NodeObjectGroupPrivate.h"
+#include "MuT/NodeObjectGroupPrivate.h"
 
 #include "MuT/NodeMeshClipMorphPlane.h"
 #include "MuT/NodeMeshClipWithMesh.h"
 
-#include "NodeModifierPrivate.h"
-#include "NodeModifierMeshClipMorphPlanePrivate.h"
-#include "NodeModifierMeshClipWithMeshPrivate.h"
+#include "MuT/NodeModifierPrivate.h"
+#include "MuT/NodeModifierMeshClipMorphPlanePrivate.h"
+#include "MuT/NodeModifierMeshClipWithMeshPrivate.h"
 #include "MuT/NodeComponentNewPrivate.h"
 
 #include "MuT/ASTOpConstantBool.h"
diff --git a/Engine/Plugins/Experimental/Mutable/Source/MutableTools/Private/MuT/CodeGenerator_SecondPass.h b/Engine/Plugins/Experimental/Mutable/Source/MutableTools/Private/MuT/CodeGenerator_SecondPass.h
--- a/Engine/Plugins/Experimental/Mutable/Source/MutableTools/Private/MuT/CodeGenerator_SecondPass.h
+++ b/Engine/Plugins/Experimental/Mutable/Source/MutableTools/Private/MuT/CodeGenerator_SecondPass.h
@@ -2,8 +2,8 @@
 
 #pragma once
 
-#include "Visitor.h"
-#include "CodeGenerator_FirstPass.h"
+#include "MuT/Visitor.h"
+#include "MuT/CodeGenerator_FirstPass.h"
 
 #include "MuT/AST.h"
 #include "MuT/CompilerPrivate.h"
diff --git a/Engine/Plugins/Experimental/Mutable/Source/MutableTools/Private/MuT/CodeOptimiser.h b/Engine/Plugins/Experimental/Mutable/Source/MutableTools/Private/MuT/CodeOptimiser.h
--- a/Engine/Plugins/Experimental/Mutable/Source/MutableTools/Private/MuT/CodeOptimiser.h
+++ b/Engine/Plugins/Experimental/Mutable/Source/MutableTools/Private/MuT/CodeOptimiser.h
@@ -8,7 +8,7 @@
 #include "MuR/CodeVisitor.h"
 #include "MuT/CompilerPrivate.h"
 #include "MuT/AST.h"
-#include "TaskManager.h"
+#include "MuT/TaskManager.h"
 
 #include <memory>
 
diff --git a/Engine/Plugins/Experimental/Mutable/Source/MutableTools/Private/MuT/CodeOptimiser_GPU.cpp b/Engine/Plugins/Experimental/Mutable/Source/MutableTools/Private/MuT/CodeOptimiser_GPU.cpp
--- a/Engine/Plugins/Experimental/Mutable/Source/MutableTools/Private/MuT/CodeOptimiser_GPU.cpp
+++ b/Engine/Plugins/Experimental/Mutable/Source/MutableTools/Private/MuT/CodeOptimiser_GPU.cpp
@@ -1,8 +1,8 @@
 // Copyright Epic Games, Inc. All Rights Reserved.
 
-#include "CodeOptimiser_GPU.h"
+#include "MuT/CodeOptimiser_GPU.h"
 #include "MuT/CodeOptimiser.h"
-#include "CodeGenerator.h"
+#include "MuT/CodeGenerator.h"
 #include "MuT/ModelReportPrivate.h"
 
 #include "MuR/CodeVisitor.h"
diff --git a/Engine/Plugins/Experimental/Mutable/Source/MutableTools/Private/MuT/CodeOptimiser_Parameters.cpp b/Engine/Plugins/Experimental/Mutable/Source/MutableTools/Private/MuT/CodeOptimiser_Parameters.cpp
--- a/Engine/Plugins/Experimental/Mutable/Source/MutableTools/Private/MuT/CodeOptimiser_Parameters.cpp
+++ b/Engine/Plugins/Experimental/Mutable/Source/MutableTools/Private/MuT/CodeOptimiser_Parameters.cpp
@@ -1,10 +1,10 @@
 // Copyright Epic Games, Inc. All Rights Reserved.
 
 #include "MuT/CodeOptimiser.h"
-#include "CodeOptimiser_GPU.h"
+#include "MuT/CodeOptimiser_GPU.h"
 #include "MuT/NodeObjectNewPrivate.h"
-#include "CodeGenerator.h"
-#include "DataPacker.h"
+#include "MuT/CodeGenerator.h"
+#include "MuT/DataPacker.h"
 #include "MuT/ModelReportPrivate.h"
 #include "MuT/CompilerPrivate.h"
 
diff --git a/Engine/Plugins/Experimental/Mutable/Source/MutableTools/Private/MuT/Compiler.cpp b/Engine/Plugins/Experimental/Mutable/Source/MutableTools/Private/MuT/Compiler.cpp
--- a/Engine/Plugins/Experimental/Mutable/Source/MutableTools/Private/MuT/Compiler.cpp
+++ b/Engine/Plugins/Experimental/Mutable/Source/MutableTools/Private/MuT/Compiler.cpp
@@ -4,9 +4,9 @@
 #include "MuT/Compiler.h"
 
 #include "MuT/CompilerPrivate.h"
-#include "CodeGenerator.h"
+#include "MuT/CodeGenerator.h"
 #include "MuT/CodeOptimiser.h"
-#include "DataPacker.h"
+#include "MuT/DataPacker.h"
 #include "MuT/ErrorLogPrivate.h"
 #include "Hash/CityHash.h"
 
diff --git a/Engine/Plugins/Experimental/Mutable/Source/MutableTools/Private/MuT/DataPacker.cpp b/Engine/Plugins/Experimental/Mutable/Source/MutableTools/Private/MuT/DataPacker.cpp
--- a/Engine/Plugins/Experimental/Mutable/Source/MutableTools/Private/MuT/DataPacker.cpp
+++ b/Engine/Plugins/Experimental/Mutable/Source/MutableTools/Private/MuT/DataPacker.cpp
@@ -1,6 +1,6 @@
 // Copyright Epic Games, Inc. All Rights Reserved.
 
-#include "DataPacker.h"
+#include "MuT/DataPacker.h"
 #include "MuT/ErrorLogPrivate.h"
 #include "MuT/Streams.h"
 
diff --git a/Engine/Plugins/Experimental/Mutable/Source/MutableTools/Private/MuT/ImageDescGenerator.cpp b/Engine/Plugins/Experimental/Mutable/Source/MutableTools/Private/MuT/ImageDescGenerator.cpp
--- a/Engine/Plugins/Experimental/Mutable/Source/MutableTools/Private/MuT/ImageDescGenerator.cpp
+++ b/Engine/Plugins/Experimental/Mutable/Source/MutableTools/Private/MuT/ImageDescGenerator.cpp
@@ -1,7 +1,7 @@
 // Copyright Epic Games, Inc. All Rights Reserved.
 
 
-#include "ImageDescGenerator.h"
+#include "MuT/ImageDescGenerator.h"
 
 #include "MuT/NodeImagePrivate.h"
 #include "MuT/NodeImageConstantPrivate.h"
@@ -32,7 +32,7 @@
 #include "MuT/NodeImageNormalCompositePrivate.h"
 #include "MuT/NodeImageTransformPrivate.h"
 
-#include "TablePrivate.h"
+#include "MuT/TablePrivate.h"
 
 #include "MuR/MeshPrivate.h"
 #include "MuR/ImagePrivate.h"
diff --git a/Engine/Plugins/Experimental/Mutable/Source/MutableTools/Private/MuT/ImageDescGenerator.h b/Engine/Plugins/Experimental/Mutable/Source/MutableTools/Private/MuT/ImageDescGenerator.h
--- a/Engine/Plugins/Experimental/Mutable/Source/MutableTools/Private/MuT/ImageDescGenerator.h
+++ b/Engine/Plugins/Experimental/Mutable/Source/MutableTools/Private/MuT/ImageDescGenerator.h
@@ -2,7 +2,7 @@
 
 #pragma once
 
-#include "Visitor.h"
+#include "MuT/Visitor.h"
 
 #include "MuT/AST.h"
 #include "MuT/NodeImage.h"
diff --git a/Engine/Plugins/Experimental/Mutable/Source/MutableTools/Private/MuT/NodeComponent.cpp b/Engine/Plugins/Experimental/Mutable/Source/MutableTools/Private/MuT/NodeComponent.cpp
--- a/Engine/Plugins/Experimental/Mutable/Source/MutableTools/Private/MuT/NodeComponent.cpp
+++ b/Engine/Plugins/Experimental/Mutable/Source/MutableTools/Private/MuT/NodeComponent.cpp
@@ -2,7 +2,7 @@
 
 
 #include "MuT/NodeComponent.h"
-#include "NodeComponentPrivate.h"
+#include "MuT/NodeComponentPrivate.h"
 
 #include "MuT/NodeMesh.h"
 #include "MuT/NodeImage.h"
diff --git a/Engine/Plugins/Experimental/Mutable/Source/MutableTools/Private/MuT/NodeComponentEdit.cpp b/Engine/Plugins/Experimental/Mutable/Source/MutableTools/Private/MuT/NodeComponentEdit.cpp
--- a/Engine/Plugins/Experimental/Mutable/Source/MutableTools/Private/MuT/NodeComponentEdit.cpp
+++ b/Engine/Plugins/Experimental/Mutable/Source/MutableTools/Private/MuT/NodeComponentEdit.cpp
@@ -2,7 +2,7 @@
 
 
 #include "MuT/NodeComponentEdit.h"
-#include "NodeComponentEditPrivate.h"
+#include "MuT/NodeComponentEditPrivate.h"
 
 #include "MuT/NodeSurface.h"
 
diff --git a/Engine/Plugins/Experimental/Mutable/Source/MutableTools/Private/MuT/NodeComponentEditPrivate.h b/Engine/Plugins/Experimental/Mutable/Source/MutableTools/Private/MuT/NodeComponentEditPrivate.h
--- a/Engine/Plugins/Experimental/Mutable/Source/MutableTools/Private/MuT/NodeComponentEditPrivate.h
+++ b/Engine/Plugins/Experimental/Mutable/Source/MutableTools/Private/MuT/NodeComponentEditPrivate.h
@@ -3,9 +3,9 @@
 
 #pragma once
 
-#include "NodeComponentPrivate.h"
+#include "MuT/NodeComponentPrivate.h"
 #include "MuT/NodeComponentEdit.h"
-#include "NodePatchImagePrivate.h"
+#include "MuT/NodePatchImagePrivate.h"
 
 #include "MuT/NodeMesh.h"
 #include "MuT/NodePatchMesh.h"
diff --git a/Engine/Plugins/Experimental/Mutable/Source/MutableTools/Private/MuT/NodeComponentNewPrivate.h b/Engine/Plugins/Experimental/Mutable/Source/MutableTools/Private/MuT/NodeComponentNewPrivate.h
--- a/Engine/Plugins/Experimental/Mutable/Source/MutableTools/Private/MuT/NodeComponentNewPrivate.h
+++ b/Engine/Plugins/Experimental/Mutable/Source/MutableTools/Private/MuT/NodeComponentNewPrivate.h
@@ -4,9 +4,9 @@
 #pragma once
 
 
-#include "NodeComponentPrivate.h"
+#include "MuT/NodeComponentPrivate.h"
 #include "MuT/NodeComponentNew.h"
-#include "NodePatchImagePrivate.h"
+#include "MuT/NodePatchImagePrivate.h"
 
 #include "MuT/NodeMesh.h"
 #include "MuT/NodePatchMesh.h"
diff --git a/Engine/Plugins/Experimental/Mutable/Source/MutableTools/Private/MuT/NodeLODPrivate.h b/Engine/Plugins/Experimental/Mutable/Source/MutableTools/Private/MuT/NodeLODPrivate.h
--- a/Engine/Plugins/Experimental/Mutable/Source/MutableTools/Private/MuT/NodeLODPrivate.h
+++ b/Engine/Plugins/Experimental/Mutable/Source/MutableTools/Private/MuT/NodeLODPrivate.h
@@ -4,8 +4,8 @@
 
 #include "MuT/NodePrivate.h"
 #include "MuT/NodeLOD.h"
-#include "NodeComponentPrivate.h"
-#include "NodeModifierPrivate.h"
+#include "MuT/NodeComponentPrivate.h"
+#include "MuT/NodeModifierPrivate.h"
 #include "MuT/AST.h"
 
 #include "MuR/MemoryPrivate.h"
diff --git a/Engine/Plugins/Experimental/Mutable/Source/MutableTools/Private/MuT/NodeModifier.cpp b/Engine/Plugins/Experimental/Mutable/Source/MutableTools/Private/MuT/NodeModifier.cpp
--- a/Engine/Plugins/Experimental/Mutable/Source/MutableTools/Private/MuT/NodeModifier.cpp
+++ b/Engine/Plugins/Experimental/Mutable/Source/MutableTools/Private/MuT/NodeModifier.cpp
@@ -2,7 +2,7 @@
 
 
 #include "MuT/NodeModifier.h"
-#include "NodeModifierPrivate.h"
+#include "MuT/NodeModifierPrivate.h"
 
 #include "MuT/NodeModifierMeshClipMorphPlane.h"
 #include "MuT/NodeModifierMeshClipWithMesh.h"
diff --git a/Engine/Plugins/Experimental/Mutable/Source/MutableTools/Private/MuT/NodeModifierMeshClipDeform.cpp b/Engine/Plugins/Experimental/Mutable/Source/MutableTools/Private/MuT/NodeModifierMeshClipDeform.cpp
--- a/Engine/Plugins/Experimental/Mutable/Source/MutableTools/Private/MuT/NodeModifierMeshClipDeform.cpp
+++ b/Engine/Plugins/Experimental/Mutable/Source/MutableTools/Private/MuT/NodeModifierMeshClipDeform.cpp
@@ -2,7 +2,7 @@
 
 
 #include "MuT/NodeModifierMeshClipDeform.h"
-#include "NodeModifierMeshClipDeformPrivate.h"
+#include "MuT/NodeModifierMeshClipDeformPrivate.h"
 #include "MuR/MeshPrivate.h"
 
 
diff --git a/Engine/Plugins/Experimental/Mutable/Source/MutableTools/Private/MuT/NodeModifierMeshClipDeformPrivate.h b/Engine/Plugins/Experimental/Mutable/Source/MutableTools/Private/MuT/NodeModifierMeshClipDeformPrivate.h
--- a/Engine/Plugins/Experimental/Mutable/Source/MutableTools/Private/MuT/NodeModifierMeshClipDeformPrivate.h
+++ b/Engine/Plugins/Experimental/Mutable/Source/MutableTools/Private/MuT/NodeModifierMeshClipDeformPrivate.h
@@ -2,7 +2,7 @@
 
 #pragma once
 
-#include "NodeModifierPrivate.h"
+#include "MuT/NodeModifierPrivate.h"
 #include "MuT/NodeModifierMeshClipDeform.h"
 #include "MuT/NodeMesh.h"
 #include "MuT/AST.h"
diff --git a/Engine/Plugins/Experimental/Mutable/Source/MutableTools/Private/MuT/NodeModifierMeshClipMorphPlane.cpp b/Engine/Plugins/Experimental/Mutable/Source/MutableTools/Private/MuT/NodeModifierMeshClipMorphPlane.cpp
--- a/Engine/Plugins/Experimental/Mutable/Source/MutableTools/Private/MuT/NodeModifierMeshClipMorphPlane.cpp
+++ b/Engine/Plugins/Experimental/Mutable/Source/MutableTools/Private/MuT/NodeModifierMeshClipMorphPlane.cpp
@@ -2,7 +2,7 @@
 
 
 #include "MuT/NodeModifierMeshClipMorphPlane.h"
-#include "NodeModifierMeshClipMorphPlanePrivate.h"
+#include "MuT/NodeModifierMeshClipMorphPlanePrivate.h"
 #include "MuR/MeshPrivate.h"
 
 
diff --git a/Engine/Plugins/Experimental/Mutable/Source/MutableTools/Private/MuT/NodeModifierMeshClipWithMesh.cpp b/Engine/Plugins/Experimental/Mutable/Source/MutableTools/Private/MuT/NodeModifierMeshClipWithMesh.cpp
--- a/Engine/Plugins/Experimental/Mutable/Source/MutableTools/Private/MuT/NodeModifierMeshClipWithMesh.cpp
+++ b/Engine/Plugins/Experimental/Mutable/Source/MutableTools/Private/MuT/NodeModifierMeshClipWithMesh.cpp
@@ -2,7 +2,7 @@
 
 
 #include "MuT/NodeModifierMeshClipWithMesh.h"
-#include "NodeModifierMeshClipWithMeshPrivate.h"
+#include "MuT/NodeModifierMeshClipWithMeshPrivate.h"
 #include "MuR/MeshPrivate.h"
 
 
diff --git a/Engine/Plugins/Experimental/Mutable/Source/MutableTools/Private/MuT/NodeObjectGroup.cpp b/Engine/Plugins/Experimental/Mutable/Source/MutableTools/Private/MuT/NodeObjectGroup.cpp
--- a/Engine/Plugins/Experimental/Mutable/Source/MutableTools/Private/MuT/NodeObjectGroup.cpp
+++ b/Engine/Plugins/Experimental/Mutable/Source/MutableTools/Private/MuT/NodeObjectGroup.cpp
@@ -2,7 +2,7 @@
 
 
 #include "MuT/NodeObjectGroup.h"
-#include "NodeObjectGroupPrivate.h"
+#include "MuT/NodeObjectGroupPrivate.h"
 
 #include "MuT/NodeLayout.h"
 
diff --git a/Engine/Plugins/Experimental/Mutable/Source/MutableTools/Private/MuT/NodeObjectNewPrivate.h b/Engine/Plugins/Experimental/Mutable/Source/MutableTools/Private/MuT/NodeObjectNewPrivate.h
--- a/Engine/Plugins/Experimental/Mutable/Source/MutableTools/Private/MuT/NodeObjectNewPrivate.h
+++ b/Engine/Plugins/Experimental/Mutable/Source/MutableTools/Private/MuT/NodeObjectNewPrivate.h
@@ -3,7 +3,7 @@
 #pragma once
 
 #include "MuT/NodeObjectPrivate.h"
-#include "NodeLODPrivate.h"
+#include "MuT/NodeLODPrivate.h"
 #include "MuT/CompilerPrivate.h"
 
 #include "MuT/NodeObjectNew.h"
diff --git a/Engine/Plugins/Experimental/Mutable/Source/MutableTools/Private/MuT/NodePatchImage.cpp b/Engine/Plugins/Experimental/Mutable/Source/MutableTools/Private/MuT/NodePatchImage.cpp
--- a/Engine/Plugins/Experimental/Mutable/Source/MutableTools/Private/MuT/NodePatchImage.cpp
+++ b/Engine/Plugins/Experimental/Mutable/Source/MutableTools/Private/MuT/NodePatchImage.cpp
@@ -2,7 +2,7 @@
 
 
 #include "MuT/NodePatchImage.h"
-#include "NodePatchImagePrivate.h"
+#include "MuT/NodePatchImagePrivate.h"
 
 #include "MuT/NodeImage.h"
 
diff --git a/Engine/Plugins/Experimental/Mutable/Source/MutableTools/Private/MuT/NodePatchMesh.cpp b/Engine/Plugins/Experimental/Mutable/Source/MutableTools/Private/MuT/NodePatchMesh.cpp
--- a/Engine/Plugins/Experimental/Mutable/Source/MutableTools/Private/MuT/NodePatchMesh.cpp
+++ b/Engine/Plugins/Experimental/Mutable/Source/MutableTools/Private/MuT/NodePatchMesh.cpp
@@ -2,7 +2,7 @@
 
 
 #include "MuT/NodePatchMesh.h"
-#include "NodePatchMeshPrivate.h"
+#include "MuT/NodePatchMeshPrivate.h"
 #include "MuR/MeshPrivate.h"
 
 
diff --git a/Engine/Plugins/Experimental/Mutable/Source/MutableTools/Private/MuT/NodeProjectorConstant.cpp b/Engine/Plugins/Experimental/Mutable/Source/MutableTools/Private/MuT/NodeProjectorConstant.cpp
--- a/Engine/Plugins/Experimental/Mutable/Source/MutableTools/Private/MuT/NodeProjectorConstant.cpp
+++ b/Engine/Plugins/Experimental/Mutable/Source/MutableTools/Private/MuT/NodeProjectorConstant.cpp
@@ -1,7 +1,7 @@
 // Copyright Epic Games, Inc. All Rights Reserved.
 
 
-#include "NodeProjectorPrivate.h"
+#include "MuT/NodeProjectorPrivate.h"
 
 #include "MuR/MutableMath.h"
 
diff --git a/Engine/Plugins/Experimental/Mutable/Source/MutableTools/Private/MuT/NodeProjectorParameter.cpp b/Engine/Plugins/Experimental/Mutable/Source/MutableTools/Private/MuT/NodeProjectorParameter.cpp
--- a/Engine/Plugins/Experimental/Mutable/Source/MutableTools/Private/MuT/NodeProjectorParameter.cpp
+++ b/Engine/Plugins/Experimental/Mutable/Source/MutableTools/Private/MuT/NodeProjectorParameter.cpp
@@ -1,7 +1,7 @@
 // Copyright Epic Games, Inc. All Rights Reserved.
 
 
-#include "NodeProjectorPrivate.h"
+#include "MuT/NodeProjectorPrivate.h"
 
 #include "MuR/MutableMath.h"
 #include "MuR/MemoryPrivate.h"
diff --git a/Engine/Plugins/Experimental/Mutable/Source/MutableTools/Private/MuT/NodeRangeFromScalar.cpp b/Engine/Plugins/Experimental/Mutable/Source/MutableTools/Private/MuT/NodeRangeFromScalar.cpp
--- a/Engine/Plugins/Experimental/Mutable/Source/MutableTools/Private/MuT/NodeRangeFromScalar.cpp
+++ b/Engine/Plugins/Experimental/Mutable/Source/MutableTools/Private/MuT/NodeRangeFromScalar.cpp
@@ -2,7 +2,7 @@
 
 
 #include "MuT/NodeRangeFromScalar.h"
-#include "NodeRangeFromScalarPrivate.h"
+#include "MuT/NodeRangeFromScalarPrivate.h"
 
 #include "MuT/NodeScalar.h"
 
diff --git a/Engine/Plugins/Experimental/Mutable/Source/MutableTools/Private/MuT/NodeRangeFromScalarPrivate.h b/Engine/Plugins/Experimental/Mutable/Source/MutableTools/Private/MuT/NodeRangeFromScalarPrivate.h
--- a/Engine/Plugins/Experimental/Mutable/Source/MutableTools/Private/MuT/NodeRangeFromScalarPrivate.h
+++ b/Engine/Plugins/Experimental/Mutable/Source/MutableTools/Private/MuT/NodeRangeFromScalarPrivate.h
@@ -2,7 +2,7 @@
 
 #pragma once
 
-#include "NodeRangePrivate.h"
+#include "MuT/NodeRangePrivate.h"
 #include "MuT/NodeRangeFromScalar.h"
 #include "MuT/NodeRange.h"
 #include "MuT/NodeScalar.h"
diff --git a/Engine/Plugins/Experimental/Mutable/Source/MutableTools/Private/MuT/NodeScalarArithmeticOperation.cpp b/Engine/Plugins/Experimental/Mutable/Source/MutableTools/Private/MuT/NodeScalarArithmeticOperation.cpp
--- a/Engine/Plugins/Experimental/Mutable/Source/MutableTools/Private/MuT/NodeScalarArithmeticOperation.cpp
+++ b/Engine/Plugins/Experimental/Mutable/Source/MutableTools/Private/MuT/NodeScalarArithmeticOperation.cpp
@@ -2,7 +2,7 @@
 
 
 #include "MuT/NodeScalarArithmeticOperation.h"
-#include "NodeScalarArithmeticOperationPrivate.h"
+#include "MuT/NodeScalarArithmeticOperationPrivate.h"
 
 #include "MuT/NodeScalar.h"
 
diff --git a/Engine/Plugins/Experimental/Mutable/Source/MutableTools/Private/MuT/NodeScalarArithmeticOperationPrivate.h b/Engine/Plugins/Experimental/Mutable/Source/MutableTools/Private/MuT/NodeScalarArithmeticOperationPrivate.h
--- a/Engine/Plugins/Experimental/Mutable/Source/MutableTools/Private/MuT/NodeScalarArithmeticOperationPrivate.h
+++ b/Engine/Plugins/Experimental/Mutable/Source/MutableTools/Private/MuT/NodeScalarArithmeticOperationPrivate.h
@@ -2,7 +2,7 @@
 
 #pragma once
 
-#include "NodeScalarPrivate.h"
+#include "MuT/NodeScalarPrivate.h"
 #include "MuT/NodeScalarArithmeticOperation.h"
 
 #include "MuR/MemoryPrivate.h"
diff --git a/Engine/Plugins/Experimental/Mutable/Source/MutableTools/Private/MuT/NodeScalarConstant.cpp b/Engine/Plugins/Experimental/Mutable/Source/MutableTools/Private/MuT/NodeScalarConstant.cpp
--- a/Engine/Plugins/Experimental/Mutable/Source/MutableTools/Private/MuT/NodeScalarConstant.cpp
+++ b/Engine/Plugins/Experimental/Mutable/Source/MutableTools/Private/MuT/NodeScalarConstant.cpp
@@ -2,7 +2,7 @@
 
 
 #include "MuT/NodeScalarConstant.h"
-#include "NodeScalarConstantPrivate.h"
+#include "MuT/NodeScalarConstantPrivate.h"
 
 #include "MuR/MutableMath.h"
 
diff --git a/Engine/Plugins/Experimental/Mutable/Source/MutableTools/Private/MuT/NodeScalarConstantPrivate.h b/Engine/Plugins/Experimental/Mutable/Source/MutableTools/Private/MuT/NodeScalarConstantPrivate.h
--- a/Engine/Plugins/Experimental/Mutable/Source/MutableTools/Private/MuT/NodeScalarConstantPrivate.h
+++ b/Engine/Plugins/Experimental/Mutable/Source/MutableTools/Private/MuT/NodeScalarConstantPrivate.h
@@ -2,7 +2,7 @@
 
 #pragma once
 
-#include "NodeScalarPrivate.h"
+#include "MuT/NodeScalarPrivate.h"
 
 #include "MuT/NodeScalarConstant.h"
 #include "MuT/NodeImage.h"
diff --git a/Engine/Plugins/Experimental/Mutable/Source/MutableTools/Private/MuT/NodeScalarCurve.cpp b/Engine/Plugins/Experimental/Mutable/Source/MutableTools/Private/MuT/NodeScalarCurve.cpp
--- a/Engine/Plugins/Experimental/Mutable/Source/MutableTools/Private/MuT/NodeScalarCurve.cpp
+++ b/Engine/Plugins/Experimental/Mutable/Source/MutableTools/Private/MuT/NodeScalarCurve.cpp
@@ -2,7 +2,7 @@
 
 
 #include "MuT/NodeScalarCurve.h"
-#include "NodeScalarCurvePrivate.h"
+#include "MuT/NodeScalarCurvePrivate.h"
 
 #include "MuT/NodeScalar.h"
 
diff --git a/Engine/Plugins/Experimental/Mutable/Source/MutableTools/Private/MuT/NodeScalarEnumParameter.cpp b/Engine/Plugins/Experimental/Mutable/Source/MutableTools/Private/MuT/NodeScalarEnumParameter.cpp
--- a/Engine/Plugins/Experimental/Mutable/Source/MutableTools/Private/MuT/NodeScalarEnumParameter.cpp
+++ b/Engine/Plugins/Experimental/Mutable/Source/MutableTools/Private/MuT/NodeScalarEnumParameter.cpp
@@ -2,7 +2,7 @@
 
 
 #include "MuT/NodeScalarEnumParameter.h"
-#include "NodeScalarEnumParameterPrivate.h"
+#include "MuT/NodeScalarEnumParameterPrivate.h"
 
 #include "MuR/MutableMath.h"
 #include "MuR/MemoryPrivate.h"
diff --git a/Engine/Plugins/Experimental/Mutable/Source/MutableTools/Private/MuT/NodeScalarParameter.cpp b/Engine/Plugins/Experimental/Mutable/Source/MutableTools/Private/MuT/NodeScalarParameter.cpp
--- a/Engine/Plugins/Experimental/Mutable/Source/MutableTools/Private/MuT/NodeScalarParameter.cpp
+++ b/Engine/Plugins/Experimental/Mutable/Source/MutableTools/Private/MuT/NodeScalarParameter.cpp
@@ -2,7 +2,7 @@
 
 
 #include "MuT/NodeScalarParameter.h"
-#include "NodeScalarParameterPrivate.h"
+#include "MuT/NodeScalarParameterPrivate.h"
 
 #include "MuR/MutableMath.h"
 #include "MuR/MemoryPrivate.h"
diff --git a/Engine/Plugins/Experimental/Mutable/Source/MutableTools/Private/MuT/NodeScalarSwitch.cpp b/Engine/Plugins/Experimental/Mutable/Source/MutableTools/Private/MuT/NodeScalarSwitch.cpp
--- a/Engine/Plugins/Experimental/Mutable/Source/MutableTools/Private/MuT/NodeScalarSwitch.cpp
+++ b/Engine/Plugins/Experimental/Mutable/Source/MutableTools/Private/MuT/NodeScalarSwitch.cpp
@@ -2,7 +2,7 @@
 
 
 #include "MuT/NodeScalarSwitch.h"
-#include "NodeScalarSwitchPrivate.h"
+#include "MuT/NodeScalarSwitchPrivate.h"
 
 #include "MuT/NodeScalar.h"
 
diff --git a/Engine/Plugins/Experimental/Mutable/Source/MutableTools/Private/MuT/NodeStringConstantPrivate.h b/Engine/Plugins/Experimental/Mutable/Source/MutableTools/Private/MuT/NodeStringConstantPrivate.h
--- a/Engine/Plugins/Experimental/Mutable/Source/MutableTools/Private/MuT/NodeStringConstantPrivate.h
+++ b/Engine/Plugins/Experimental/Mutable/Source/MutableTools/Private/MuT/NodeStringConstantPrivate.h
@@ -2,7 +2,7 @@
 
 #pragma once
 
-#include "NodeStringPrivate.h"
+#include "MuT/NodeStringPrivate.h"
 
 #include "MuT/NodeStringConstant.h"
 #include "MuT/NodeImage.h"
diff --git a/Engine/Plugins/Experimental/Mutable/Source/MutableTools/Private/MuT/NodeSurfaceEditPrivate.h b/Engine/Plugins/Experimental/Mutable/Source/MutableTools/Private/MuT/NodeSurfaceEditPrivate.h
--- a/Engine/Plugins/Experimental/Mutable/Source/MutableTools/Private/MuT/NodeSurfaceEditPrivate.h
+++ b/Engine/Plugins/Experimental/Mutable/Source/MutableTools/Private/MuT/NodeSurfaceEditPrivate.h
@@ -4,7 +4,7 @@
 
 #include "MuT/NodePrivate.h"
 #include "MuT/NodeSurfaceEdit.h"
-#include "NodePatchImagePrivate.h"
+#include "MuT/NodePatchImagePrivate.h"
 
 #include "MuT/NodeMesh.h"
 #include "MuT/NodePatchMesh.h"
diff --git a/Engine/Plugins/Experimental/Mutable/Source/MutableTools/Private/MuT/NodeSurfaceVariationPrivate.h b/Engine/Plugins/Experimental/Mutable/Source/MutableTools/Private/MuT/NodeSurfaceVariationPrivate.h
--- a/Engine/Plugins/Experimental/Mutable/Source/MutableTools/Private/MuT/NodeSurfaceVariationPrivate.h
+++ b/Engine/Plugins/Experimental/Mutable/Source/MutableTools/Private/MuT/NodeSurfaceVariationPrivate.h
@@ -4,7 +4,7 @@
 
 #include "MuT/NodePrivate.h"
 #include "MuT/NodeSurfaceVariation.h"
-#include "NodeSurfaceNewPrivate.h"
+#include "MuT/NodeSurfaceNewPrivate.h"
 #include "MuT/NodeModifier.h"
 
 #include "MuR/MemoryPrivate.h"
diff --git a/Engine/Plugins/Experimental/Mutable/Source/MutableTools/Private/MuT/Table.cpp b/Engine/Plugins/Experimental/Mutable/Source/MutableTools/Private/MuT/Table.cpp
--- a/Engine/Plugins/Experimental/Mutable/Source/MutableTools/Private/MuT/Table.cpp
+++ b/Engine/Plugins/Experimental/Mutable/Source/MutableTools/Private/MuT/Table.cpp
@@ -2,7 +2,7 @@
 
 
 #include "MuT/Table.h"
-#include "TablePrivate.h"
+#include "MuT/TablePrivate.h"
 
 #include "MuR/Platform.h"
 #include "MuT/ErrorLogPrivate.h"
diff --git a/Engine/Plugins/Experimental/Mutable/Source/MutableTools/Private/MuT/TaskManager.cpp b/Engine/Plugins/Experimental/Mutable/Source/MutableTools/Private/MuT/TaskManager.cpp
--- a/Engine/Plugins/Experimental/Mutable/Source/MutableTools/Private/MuT/TaskManager.cpp
+++ b/Engine/Plugins/Experimental/Mutable/Source/MutableTools/Private/MuT/TaskManager.cpp
@@ -1,6 +1,6 @@
 // Copyright Epic Games, Inc. All Rights Reserved.
 
-#include "TaskManager.h"
+#include "MuT/TaskManager.h"
 #include "MuR/SystemPrivate.h"
 
 
