Removed include to NodeScalarTablePrivate.h because it doesn't compile on Linux.

diff --git a/Engine/Plugins/Experimental/Mutable/Source/CustomizableObjectEditor/Private/MuCOE/CustomizableObjectPin.cpp b/Engine/Plugins/Experimental/Mutable/Source/CustomizableObjectEditor/Private/MuCOE/CustomizableObjectPin.cpp
--- a/Engine/Plugins/Experimental/Mutable/Source/CustomizableObjectEditor/Private/MuCOE/CustomizableObjectPin.cpp
+++ b/Engine/Plugins/Experimental/Mutable/Source/CustomizableObjectEditor/Private/MuCOE/CustomizableObjectPin.cpp
@@ -1,7 +1,6 @@
 // Copyright Epic Games, Inc. All Rights Reserved.
 
 #include "MuCOE/CustomizableObjectPin.h"
-#include "MuT/NodeScalarTablePrivate.h"
 
 bool IsPinOrphan(const UEdGraphPin &Pin)
 {
diff --git a/Engine/Plugins/Experimental/Mutable/Source/CustomizableObjectEditor/Private/MuCOE/GraphTraversal.cpp b/Engine/Plugins/Experimental/Mutable/Source/CustomizableObjectEditor/Private/MuCOE/GraphTraversal.cpp
--- a/Engine/Plugins/Experimental/Mutable/Source/CustomizableObjectEditor/Private/MuCOE/GraphTraversal.cpp
+++ b/Engine/Plugins/Experimental/Mutable/Source/CustomizableObjectEditor/Private/MuCOE/GraphTraversal.cpp
@@ -2,7 +2,6 @@
 
 #include "MuCOE/GraphTraversal.h"
 #include "MuCO/CustomizableObject.h"
-#include "MuT/NodeScalarTablePrivate.h"
 #include "MuCOE/CustomizableObjectPin.h"
 #include "MuCOE/Nodes/CustomizableObjectNodeEnumParameter.h"
 #include "MuCOE/Nodes/CustomizableObjectNodeExtendMaterial.h"
diff --git a/Engine/Plugins/Experimental/Mutable/Source/CustomizableObjectEditor/Private/MuCOE/Nodes/CustomizableObjectNode.cpp b/Engine/Plugins/Experimental/Mutable/Source/CustomizableObjectEditor/Private/MuCOE/Nodes/CustomizableObjectNode.cpp
--- a/Engine/Plugins/Experimental/Mutable/Source/CustomizableObjectEditor/Private/MuCOE/Nodes/CustomizableObjectNode.cpp
+++ b/Engine/Plugins/Experimental/Mutable/Source/CustomizableObjectEditor/Private/MuCOE/Nodes/CustomizableObjectNode.cpp
@@ -15,7 +15,6 @@
 
 #include "MuCOE/CustomizableObjectGraph.h"
 #include "MuCOE/CustomizableObjectPin.h"
-#include "MuT/NodeScalarTablePrivate.h"
 #include "MuCOE/EdGraphSchema_CustomizableObject.h"
 #include "MuCOE/GraphTraversal.h"
 #include "Runtime/Launch/Resources/Version.h" // Include engine version macros
diff --git a/Engine/Plugins/Experimental/Mutable/Source/CustomizableObjectEditor/Private/MuCOE/Nodes/CustomizableObjectNodeExternalPin.cpp b/Engine/Plugins/Experimental/Mutable/Source/CustomizableObjectEditor/Private/MuCOE/Nodes/CustomizableObjectNodeExternalPin.cpp
--- a/Engine/Plugins/Experimental/Mutable/Source/CustomizableObjectEditor/Private/MuCOE/Nodes/CustomizableObjectNodeExternalPin.cpp
+++ b/Engine/Plugins/Experimental/Mutable/Source/CustomizableObjectEditor/Private/MuCOE/Nodes/CustomizableObjectNodeExternalPin.cpp
@@ -4,7 +4,6 @@
 #include "MuCOE/CustomizableObjectEditorModule.h"
 #include "MuCOE/CustomizableObjectPin.h"
 #include "MuCO/CustomizableObject.h"
-#include "MuT/NodeScalarTablePrivate.h"
 
 #define LOCTEXT_NAMESPACE "CustomizableObjectEditor"
 
diff --git a/Engine/Plugins/Experimental/Mutable/Source/CustomizableObjectEditor/Private/MuCOE/Nodes/CustomizableObjectNodeUseMaterial.cpp b/Engine/Plugins/Experimental/Mutable/Source/CustomizableObjectEditor/Private/MuCOE/Nodes/CustomizableObjectNodeUseMaterial.cpp
--- a/Engine/Plugins/Experimental/Mutable/Source/CustomizableObjectEditor/Private/MuCOE/Nodes/CustomizableObjectNodeUseMaterial.cpp
+++ b/Engine/Plugins/Experimental/Mutable/Source/CustomizableObjectEditor/Private/MuCOE/Nodes/CustomizableObjectNodeUseMaterial.cpp
@@ -2,7 +2,6 @@
 
 #include "MuCOE/Nodes/CustomizableObjectNodeUseMaterial.h"
 
-#include "MuT/NodeScalarTablePrivate.h"
 #include "MuCOE/CustomizableObjectPin.h"
 #include "MuCOE/GraphTraversal.h"
 #include "MuCOE/ICustomizableObjectEditor.h"
