Ran IWYU on Mutable plugin

diff --git a/Engine/Plugins/Experimental/Mutable/Source/CustomizableObject/Classes/MuCO/CustomizableInstanceLODManagement.h b/Engine/Plugins/Experimental/Mutable/Source/CustomizableObject/Classes/MuCO/CustomizableInstanceLODManagement.h
--- a/Engine/Plugins/Experimental/Mutable/Source/CustomizableObject/Classes/MuCO/CustomizableInstanceLODManagement.h
+++ b/Engine/Plugins/Experimental/Mutable/Source/CustomizableObject/Classes/MuCO/CustomizableInstanceLODManagement.h
@@ -1,8 +1,14 @@
 // Copyright Epic Games, Inc. All Rights Reserved.
 
 #pragma once
-#include "HAL/IConsoleManager.h"
+#include "Containers/Set.h"
 #include "GameFramework/Actor.h"
+#include "HAL/Platform.h"
+#include "Misc/AssertionMacros.h"
+#include "UObject/Object.h"
+#include "UObject/UObjectGlobals.h"
+#include "UObject/WeakObjectPtr.h"
+#include "UObject/WeakObjectPtrTemplates.h"
 
 #include "CustomizableInstanceLODManagement.generated.h"
 
diff --git a/Engine/Plugins/Experimental/Mutable/Source/CustomizableObject/Classes/MuCO/CustomizableObject.h b/Engine/Plugins/Experimental/Mutable/Source/CustomizableObject/Classes/MuCO/CustomizableObject.h
--- a/Engine/Plugins/Experimental/Mutable/Source/CustomizableObject/Classes/MuCO/CustomizableObject.h
+++ b/Engine/Plugins/Experimental/Mutable/Source/CustomizableObject/Classes/MuCO/CustomizableObject.h
@@ -2,15 +2,41 @@
  
 #pragma once
 
+#include "Containers/Array.h"
+#include "Containers/BitArray.h"
+#include "Containers/ContainerAllocationPolicies.h"
+#include "Containers/ContainersFwd.h"
+#include "Containers/EnumAsByte.h"
+#include "Containers/Map.h"
+#include "Containers/Set.h"
+#include "Containers/UnrealString.h"
+#include "Delegates/Delegate.h"
 #include "Engine/Texture.h"
-#include "EdGraph/EdGraphPin.h"
-#include "EdGraph/EdGraph.h"
-#include "EdGraph/EdGraphNode.h"
+#include "Engine/TextureDefines.h"
+#include "Input/Reply.h"
+#include "Logging/LogMacros.h"
+#include "Math/BoxSphereBounds.h"
+#include "Math/Rotator.h"
+#include "Math/UnrealMathSSE.h"
+#include "Math/Vector.h"
+#include "Math/Vector4.h"
+#include "Misc/AssertionMacros.h"
+#include "Misc/Guid.h"
+#include "MuCO/CustomizableObjectClothingTypes.h"
 #include "MuCO/CustomizableObjectIdentifier.h"
 #include "MuCO/CustomizableObjectParameterTypeDefinitions.h"
-#include "MuCO/CustomizableObjectClothingTypes.h"
 #include "MuCO/CustomizableObjectUIData.h"
-#include "Animation/AnimInstance.h"
+#include "RHIDefinitions.h"
+#include "Serialization/Archive.h"
+#include "Serialization/StructuredArchiveAdapters.h"
+#include "SkeletalMeshTypes.h"
+#include "Templates/SharedPointer.h"
+#include "UObject/NameTypes.h"
+#include "UObject/Object.h"
+#include "UObject/ObjectPtr.h"
+#include "UObject/SoftObjectPath.h"
+#include "UObject/SoftObjectPtr.h"
+#include "UObject/UObjectGlobals.h"
 
 #if WITH_EDITORONLY_DATA
 #include "PerPlatformProperties.h"
@@ -20,7 +46,22 @@
 
 #include "CustomizableObject.generated.h"
 
+class FMemoryReaderView;
+class FMemoryWriter64;
+class FObjectPreSaveContext;
+class FObjectPreSaveRootContext;
+class FText;
+class IAsyncReadFileHandle;
+class ITargetPlatform;
+class UAnimInstance;
 class UCustomizableObject;
+class UEdGraph;
+class UMaterialInterface;
+class UPhysicsAsset;
+class USkeletalMesh;
+class USkeleton;
+struct FFrame;
+template <typename FuncType> class TFunctionRef;
 
 DECLARE_MULTICAST_DELEGATE(FPostCompileDelegate)
  
diff --git a/Engine/Plugins/Experimental/Mutable/Source/CustomizableObject/Classes/MuCO/CustomizableObjectIdentifier.h b/Engine/Plugins/Experimental/Mutable/Source/CustomizableObject/Classes/MuCO/CustomizableObjectIdentifier.h
--- a/Engine/Plugins/Experimental/Mutable/Source/CustomizableObject/Classes/MuCO/CustomizableObjectIdentifier.h
+++ b/Engine/Plugins/Experimental/Mutable/Source/CustomizableObject/Classes/MuCO/CustomizableObjectIdentifier.h
@@ -2,6 +2,9 @@
 
 #pragma once
 
+#include "Containers/UnrealString.h"
+#include "Serialization/Archive.h"
+
 #include "CustomizableObjectIdentifier.generated.h"
 
 USTRUCT(BlueprintType)
diff --git a/Engine/Plugins/Experimental/Mutable/Source/CustomizableObject/Classes/MuCO/CustomizableObjectInstance.h b/Engine/Plugins/Experimental/Mutable/Source/CustomizableObject/Classes/MuCO/CustomizableObjectInstance.h
--- a/Engine/Plugins/Experimental/Mutable/Source/CustomizableObject/Classes/MuCO/CustomizableObjectInstance.h
+++ b/Engine/Plugins/Experimental/Mutable/Source/CustomizableObject/Classes/MuCO/CustomizableObjectInstance.h
@@ -2,18 +2,38 @@
 
 #pragma once
 
-#include "MuCO/CustomizableSkeletalComponent.h"
-#include "MuCO/CustomizableObject.h"
+#include "Async/TaskGraphInterfaces.h"
+#include "Containers/Array.h"
+#include "Containers/Map.h"
+#include "Containers/UnrealString.h"
+#include "Delegates/Delegate.h"
+#include "Math/Color.h"
+#include "Math/UnrealMathSSE.h"
+#include "Misc/AssertionMacros.h"
 #include "MuCO/CustomizableObjectInstanceDescriptor.h"
-#include "MuCO/MultilayerProjector.h"
-#include "Engine/Texture2D.h"
-
 #include "MuCO/CustomizableObjectParameterTypeDefinitions.h"
-#include "GameplayTagContainer.h"
+#include "MuCO/MultilayerProjector.h"
+#include "Serialization/Archive.h"
+#include "Templates/SubclassOf.h"
+#include "UObject/NameTypes.h"
+#include "UObject/Object.h"
+#include "UObject/ObjectPtr.h"
+#include "UObject/UObjectGlobals.h"
+#include "UObject/WeakObjectPtr.h"
+#include "UObject/WeakObjectPtrTemplates.h"
 
 #include "CustomizableObjectInstance.generated.h"
 
+class AActor;
+class FProperty;
+class UAnimInstance;
 class UCustomizableInstancePrivateData; // This is used to hide Mutable SDK members in the public headers.
+class UCustomizableObject;
+class UCustomizableSkeletalComponent;
+class UTexture2D;
+struct FFrame;
+struct FGameplayTagContainer;
+struct FPropertyChangedEvent;
 
 //! Order of the unreal vertex buffers when in mutable data
 #define MUTABLE_VERTEXBUFFER_POSITION	0
@@ -62,12 +82,6 @@ namespace EProjectorState
 };
 
 
-namespace mu
-{
-	class Instance;
-};
-
-
 /** FString with the possible errors from skeletal mesh update */
 namespace ESkeletalMeshState
 {
diff --git a/Engine/Plugins/Experimental/Mutable/Source/CustomizableObject/Classes/MuCO/CustomizableObjectSystem.h b/Engine/Plugins/Experimental/Mutable/Source/CustomizableObject/Classes/MuCO/CustomizableObjectSystem.h
--- a/Engine/Plugins/Experimental/Mutable/Source/CustomizableObject/Classes/MuCO/CustomizableObjectSystem.h
+++ b/Engine/Plugins/Experimental/Mutable/Source/CustomizableObject/Classes/MuCO/CustomizableObjectSystem.h
@@ -2,21 +2,37 @@
 
 #pragma once
 
+#include "Containers/Array.h"
+#include "Containers/Map.h"
+#include "Containers/UnrealString.h"
+#include "Delegates/Delegate.h"
+#include "Engine/EngineTypes.h"
 #include "Engine/StreamableManager.h"
+#include "Math/IntVector.h"
+#include "Math/UnrealMathSSE.h"
+#include "PixelFormat.h"
+#include "Stats/Stats2.h"
 #include "Templates/SharedPointer.h"
-
-#include "Containers/Ticker.h"
-#include "MuCO/CustomizableObjectInstance.h"
-
-#include "Runtime/Launch/Resources/Version.h"
+#include "UObject/Object.h"
+#include "UObject/ObjectPtr.h"
+#include "UObject/UObjectGlobals.h"
 
 #if WITH_EDITOR
 #include "Framework/Notifications/NotificationManager.h"
-#include "Widgets/Notifications/SNotificationList.h"
 #endif
 
 #include "CustomizableObjectSystem.generated.h"
 
+class FCustomizableObjectCompilerBase;
+class ITargetPlatform;
+class SNotificationItem;
+class UCustomizableObject;
+class USkeletalMesh;
+class UTexture2D;
+struct FAssetData;
+struct FFrame;
+struct FGuid;
+
 // This sets the amount of memory in bytes used to keep streaming data after using it, to reduce the streaming load.
 // High values use more memory, but save object construction time.
 // Setting it to 0 was the original behaviour, and keeps all the data loaded for a mutable operation until the
diff --git a/Engine/Plugins/Experimental/Mutable/Source/CustomizableObject/Classes/MuCO/CustomizableSkeletalComponent.h b/Engine/Plugins/Experimental/Mutable/Source/CustomizableObject/Classes/MuCO/CustomizableSkeletalComponent.h
--- a/Engine/Plugins/Experimental/Mutable/Source/CustomizableObject/Classes/MuCO/CustomizableSkeletalComponent.h
+++ b/Engine/Plugins/Experimental/Mutable/Source/CustomizableObject/Classes/MuCO/CustomizableSkeletalComponent.h
@@ -3,14 +3,21 @@
 #pragma once
 
 #include "Components/SceneComponent.h"
-#include "CustomizableSkeletalComponent.generated.h"
+#include "Containers/UnrealString.h"
+#include "Delegates/Delegate.h"
+#include "Engine/EngineBaseTypes.h"
+#include "Math/UnrealMathSSE.h"
+#include "UObject/ObjectPtr.h"
+#include "UObject/UObjectGlobals.h"
 
-namespace mu
-{
-	class Instance;
-};
+#include "CustomizableSkeletalComponent.generated.h"
 
+class AActor;
+class FObjectPreSaveContext;
+class UCustomizableSkeletalComponent;
+class UObject;
 class USkeletalMesh;
+struct FFrame;
 
 
 DECLARE_DELEGATE_TwoParams(FCustomizableSkeletalComponentPreUpdateDelegate, UCustomizableSkeletalComponent* Component, USkeletalMesh* NextMesh);
diff --git a/Engine/Plugins/Experimental/Mutable/Source/CustomizableObject/Classes/MuCO/CustomizableSkeletalMeshActor.h b/Engine/Plugins/Experimental/Mutable/Source/CustomizableObject/Classes/MuCO/CustomizableSkeletalMeshActor.h
--- a/Engine/Plugins/Experimental/Mutable/Source/CustomizableObject/Classes/MuCO/CustomizableSkeletalMeshActor.h
+++ b/Engine/Plugins/Experimental/Mutable/Source/CustomizableObject/Classes/MuCO/CustomizableSkeletalMeshActor.h
@@ -2,12 +2,17 @@
 
 #pragma once
 
-#include "CoreMinimal.h"
-#include "UObject/ObjectMacros.h"
 #include "Animation/SkeletalMeshActor.h"
-#include "MuCO/CustomizableSkeletalComponent.h"
+#include "Containers/Array.h"
+#include "HAL/Platform.h"
+#include "UObject/ObjectMacros.h"
+#include "UObject/ObjectPtr.h"
+#include "UObject/UObjectGlobals.h"
+
 #include "CustomizableSkeletalMeshActor.generated.h"
 
+class UObject;
+
 UCLASS(ClassGroup = CustomizableObject, Blueprintable, ComponentWrapperClass, ConversionRoot, meta = (ChildCanTick))
 class CUSTOMIZABLEOBJECT_API ACustomizableSkeletalMeshActor : public ASkeletalMeshActor
 {
diff --git a/Engine/Plugins/Experimental/Mutable/Source/CustomizableObject/Private/MuCO/CustomizableInstanceLODManagement.cpp b/Engine/Plugins/Experimental/Mutable/Source/CustomizableObject/Private/MuCO/CustomizableInstanceLODManagement.cpp
--- a/Engine/Plugins/Experimental/Mutable/Source/CustomizableObject/Private/MuCO/CustomizableInstanceLODManagement.cpp
+++ b/Engine/Plugins/Experimental/Mutable/Source/CustomizableObject/Private/MuCO/CustomizableInstanceLODManagement.cpp
@@ -2,16 +2,37 @@
 
 #include "MuCO/CustomizableInstanceLODManagement.h"
 
-//#include "MuCO/CustomizableObjectSystemPrivate.h"
+#include "Camera/CameraActor.h"
+#include "Camera/CameraComponent.h"
+#include "Camera/CameraTypes.h"
+#include "Containers/Array.h"
+#include "Containers/EnumAsByte.h"
+#include "Containers/Map.h"
+#include "Engine/EngineTypes.h"
+#include "Engine/SkeletalMesh.h"
+#include "Engine/World.h"
 #include "GameFramework/PlayerController.h"
+#include "HAL/IConsoleManager.h"
+#include "HAL/PlatformCrt.h"
+#include "Kismet/GameplayStatics.h"
+#include "Math/BoxSphereBounds.h"
+#include "Math/Matrix.h"
+#include "Math/UnrealMathSSE.h"
+#include "Math/Vector.h"
+#include "MuCO/CustomizableObject.h"
+#include "MuCO/CustomizableObjectInstance.h"
 #include "MuCO/CustomizableObjectSystem.h"
 #include "MuCO/CustomizableSkeletalComponent.h"
-#include "MuCO/CustomizableObject.h"
-#include "Camera/CameraComponent.h"
-#include "Camera/CameraActor.h"
-#include "Kismet/GameplayStatics.h"
+#include "PerPlatformProperties.h"
+#include "SceneManagement.h"
+#include "Serialization/StructuredArchiveAdapters.h"
+#include "Templates/Casts.h"
+#include "UObject/ObjectPtr.h"
+#include "UObject/UObjectIterator.h"
 
 #if WITH_EDITOR
+#include "Editor.h"
+#include "Editor/EditorEngine.h"
 #include "LevelEditorViewport.h"
 #endif
 
diff --git a/Engine/Plugins/Experimental/Mutable/Source/CustomizableObject/Private/MuCO/CustomizableObject.cpp b/Engine/Plugins/Experimental/Mutable/Source/CustomizableObject/Private/MuCO/CustomizableObject.cpp
--- a/Engine/Plugins/Experimental/Mutable/Source/CustomizableObject/Private/MuCO/CustomizableObject.cpp
+++ b/Engine/Plugins/Experimental/Mutable/Source/CustomizableObject/Private/MuCO/CustomizableObject.cpp
@@ -2,37 +2,56 @@
 
 #include "MuCO/CustomizableObject.h"
 
+#include "Algo/Copy.h"
+#include "AssetRegistry/AssetData.h"
 #include "AssetRegistry/AssetRegistryModule.h"
-#include "Serialization/MemoryReader.h"
-#include "Serialization/MemoryWriter.h"
-#include "Misc/Paths.h"
+#include "AssetRegistry/IAssetRegistry.h"
+#include "Async/AsyncFileHandle.h"
+#include "Containers/ArrayView.h"
+#include "Containers/StringConv.h"
+#include "CoreGlobals.h"
+#include "EdGraph/EdGraph.h"
+#include "EdGraph/EdGraphPin.h"
+#include "GenericPlatform/GenericPlatformFile.h"
 #include "HAL/FileManager.h"
 #include "HAL/PlatformFileManager.h"
+#include "HAL/PlatformMath.h"
+#include "HAL/PlatformProperties.h"
 #include "Interfaces/ITargetPlatform.h"
-#include "Runtime/Launch/Resources/Version.h"
-#include "Algo/Count.h"
-#include "Algo/Copy.h"
-#include "PhysicsEngine/PhysicsAsset.h"
-#include "ClothConfig.h"
-#include "UObject/ObjectSaveContext.h"
-
-#include "MuR/Model.h"
-#include "MuR/Mesh.h"
-#include "MuR/Image.h"
-#include "MuR/Layout.h"
-#include "MuR/System.h"
-#include "MuR/MutableTrace.h"
-#include "MuR/ModelPrivate.h"
-
-#include "MuCO/ICustomizableObjectModule.h"
+#include "Internationalization/Text.h"
+#include "Logging/LogCategory.h"
+#include "Misc/CString.h"
+#include "Misc/CoreMisc.h"
+#include "Misc/PackageName.h"
+#include "Misc/Paths.h"
+#include "Modules/ModuleManager.h"
+#include "MuCO/CustomizableObjectInstance.h"
 #include "MuCO/CustomizableObjectPrivate.h"
 #include "MuCO/CustomizableObjectSystem.h"
+#include "MuCO/ICustomizableObjectModule.h"
 #include "MuCO/UnrealMutableModelDiskStreamer.h"
-
-#include "Materials/MaterialInstance.h"
-#include "Misc/SecureHash.h"
-
-#include "ClothingAsset.h"
+#include "MuR/Model.h"
+#include "MuR/ModelPrivate.h"
+#include "MuR/MutableTrace.h"
+#include "MuR/Parameters.h"
+#include "MuR/Ptr.h"
+#include "MuR/Serialisation.h"
+#include "MuR/System.h"
+#include "PhysicsEngine/PhysicsAsset.h"
+#include "Serialization/MemoryReader.h"
+#include "Serialization/MemoryWriter.h"
+#include "Templates/Function.h"
+#include "Templates/UnrealTemplate.h"
+#include "Trace/Detail/Channel.h"
+#include "UObject/ObjectSaveContext.h"
+#include "UObject/Package.h"
+#include "UObject/UObjectMarks.h"
+#include "UObject/UnrealNames.h"
+
+class UMaterialInterface;
+class UPhysicsAsset;
+class USkeletalMesh;
+class USkeleton;
 
 #define LOCTEXT_NAMESPACE "CustomizableObject"
 
diff --git a/Engine/Plugins/Experimental/Mutable/Source/CustomizableObject/Private/MuCO/CustomizableObjectCustomVersion.cpp b/Engine/Plugins/Experimental/Mutable/Source/CustomizableObject/Private/MuCO/CustomizableObjectCustomVersion.cpp
--- a/Engine/Plugins/Experimental/Mutable/Source/CustomizableObject/Private/MuCO/CustomizableObjectCustomVersion.cpp
+++ b/Engine/Plugins/Experimental/Mutable/Source/CustomizableObject/Private/MuCO/CustomizableObjectCustomVersion.cpp
@@ -1,6 +1,9 @@
 // Copyright Epic Games, Inc. All Rights Reserved.
 
 #include "MuCO/CustomizableObjectCustomVersion.h"
+
+#include "HAL/Platform.h"
+#include "Misc/Guid.h"
 #include "Serialization/CustomVersion.h"
 
 const FGuid FCustomizableObjectCustomVersion::GUID(0xC1270312, 0xAB730B1E, 0x0B1EC745, 0x71C70F73);
diff --git a/Engine/Plugins/Experimental/Mutable/Source/CustomizableObject/Private/MuCO/CustomizableObjectDGGUI.cpp b/Engine/Plugins/Experimental/Mutable/Source/CustomizableObject/Private/MuCO/CustomizableObjectDGGUI.cpp
--- a/Engine/Plugins/Experimental/Mutable/Source/CustomizableObject/Private/MuCO/CustomizableObjectDGGUI.cpp
+++ b/Engine/Plugins/Experimental/Mutable/Source/CustomizableObject/Private/MuCO/CustomizableObjectDGGUI.cpp
@@ -1,12 +1,28 @@
 // Copyright Epic Games, Inc. All Rights Reserved.
 
 #include "MuCO/CustomizableObjectDGGUI.h"
-#include "MuCO/CustomizableSkeletalComponent.h"
+
+#include "Animation/AnimClassInterface.h"
+#include "Animation/AnimInstance.h"
+#include "Animation/Skeleton.h"
+#include "Components/SkeletalMeshComponent.h"
+#include "Containers/Array.h"
+#include "Containers/UnrealString.h"
+#include "CoreTypes.h"
+#include "Delegates/Delegate.h"
 #include "Engine/SkeletalMesh.h"
-#include "HAL/IConsoleManager.h"
+#include "Engine/World.h"
+#include "GameFramework/PlayerController.h"
 #include "Kismet/GameplayStatics.h"
-#include "Animation/AnimInstance.h"
+#include "Misc/AssertionMacros.h"
 #include "MuCO/CustomizableObjectInstance.h"
+#include "MuCO/CustomizableSkeletalComponent.h"
+#include "Templates/Casts.h"
+#include "UObject/Class.h"
+#include "UObject/ObjectPtr.h"
+#include "UObject/SoftObjectPath.h"
+#include "UObject/SoftObjectPtr.h"
+#include "UObject/UObjectIterator.h"
 
 void UDGGUI::OpenDGGUI(const int32 SlotID, UCustomizableSkeletalComponent* SelectedCustomizableSkeletalComponent, const UWorld* CurrentWorld, const int32 PlayerIndex)
 {
diff --git a/Engine/Plugins/Experimental/Mutable/Source/CustomizableObject/Private/MuCO/CustomizableObjectInstance.cpp b/Engine/Plugins/Experimental/Mutable/Source/CustomizableObject/Private/MuCO/CustomizableObjectInstance.cpp
--- a/Engine/Plugins/Experimental/Mutable/Source/CustomizableObject/Private/MuCO/CustomizableObjectInstance.cpp
+++ b/Engine/Plugins/Experimental/Mutable/Source/CustomizableObject/Private/MuCO/CustomizableObjectInstance.cpp
@@ -2,72 +2,130 @@
 
 #include "MuCO/CustomizableObjectInstance.h"
 
-#include "StaticMeshResources.h"
-#include "SkeletalMeshTypes.h"
-#include "Engine/SkeletalMesh.h"
-#include "Animation/Skeleton.h"
-#include "Materials/MaterialInstanceDynamic.h"
-#include "Components/SkinnedMeshComponent.h"
-#include "Engine/SkeletalMeshSocket.h"
+#include "Algo/MaxElement.h"
+#include "Animation/AnimInstance.h"
 #include "Animation/MorphTarget.h"
-#include "ClothingAsset.h"
+#include "Animation/Skeleton.h"
 #include "ClothConfig.h"
-#include "SkeletalMergingLibrary.h"
-
+#include "ClothLODData.h"
+#include "ClothPhysicalMeshData.h"
+#include "ClothTetherData.h"
+#include "ClothVertBoneData.h"
+#include "ClothingAsset.h"
+#include "Components.h"
+#include "Containers/ArrayView.h"
+#include "Containers/BitArray.h"
+#include "Containers/EnumAsByte.h"
+#include "Containers/IndirectArray.h"
+#include "Containers/Set.h"
+#include "CoreGlobals.h"
+#include "Engine/EngineTypes.h"
+#include "Engine/SkeletalMesh.h"
+#include "Engine/SkeletalMeshSocket.h"
+#include "Engine/StreamableManager.h"
+#include "Engine/Texture.h"
+#include "Engine/Texture2D.h"
+#include "Engine/TextureStreamingTypes.h"
+#include "GameplayTagContainer.h"
+#include "HAL/IConsoleManager.h"
+#include "HAL/PlatformCrt.h"
+#include "HAL/PlatformTime.h"
+#include "HAL/UnrealMemory.h"
+#include "Internationalization/Text.h"
+#include "Logging/LogCategory.h"
+#include "Logging/LogMacros.h"
+#include "Logging/TokenizedMessage.h"
+#include "MaterialShared.h"
+#include "MaterialTypes.h"
+#include "Materials/Material.h"
+#include "Materials/MaterialInstanceDynamic.h"
+#include "Materials/MaterialInterface.h"
+#include "Math/Matrix.h"
+#include "Math/NumericLimits.h"
+#include "Math/Plane.h"
+#include "Math/Quat.h"
+#include "Math/Rotator.h"
+#include "Math/Transform.h"
+#include "Math/Vector.h"
+#include "Math/Vector4.h"
+#include "Misc/CString.h"
+#include "Misc/Guid.h"
+#include "Modules/ModuleManager.h"
+#include "MuCO/CustomizableInstanceLODManagement.h"
+#include "MuCO/CustomizableInstancePrivateData.h"
 #include "MuCO/CustomizableObject.h"
+#include "MuCO/CustomizableObjectClothingTypes.h"
+#include "MuCO/CustomizableObjectCustomVersion.h"
+#include "MuCO/CustomizableObjectMipDataProvider.h"
+#include "MuCO/CustomizableObjectParameterTypeDefinitions.h"
 #include "MuCO/CustomizableObjectPrivate.h"
-#include "MuCO/CustomizableSkeletalComponent.h"
-#include "MuCO/CustomizableInstancePrivateData.h"
-#include "MuCO/ICustomizableObjectModule.h"
 #include "MuCO/CustomizableObjectSystem.h"
 #include "MuCO/CustomizableObjectSystemPrivate.h"
-#include "MuCO/CustomizableInstanceLODManagement.h"
-#include "MuCO/CustomizableObjectParameterTypeDefinitions.h"
-#include "MuCO/CustomizableObjectCustomVersion.h"
-#include "GameplayTagContainer.h"
-#include "MuCO/UnrealConversionUtils.h"
+#include "MuCO/CustomizableSkeletalComponent.h"
 #include "MuCO/MultilayerProjector.h"
-
+#include "MuCO/UnrealConversionUtils.h"
+#include "MuR/MeshBufferSet.h"
+#include "MuR/PhysicsBody.h"
+#include "MuR/Ptr.h"
+#include "MuR/RefCounted.h"
+#include "PerPlatformProperties.h"
+#include "PhysicsEngine/AggregateGeom.h"
+#include "PhysicsEngine/BoxElem.h"
+#include "PhysicsEngine/ConvexElem.h"
+#include "PhysicsEngine/PhysicsAsset.h"
+#include "PhysicsEngine/RigidBodyIndexPair.h"
+#include "PhysicsEngine/SphereElem.h"
+#include "PhysicsEngine/SphylElem.h"
+#include "PhysicsEngine/TaperedCapsuleElem.h"
+#include "PixelFormat.h"
+#include "PointWeightMap.h"
+#include "RHI.h"
+#include "RHICommandList.h"
+#include "RHIDefinitions.h"
+#include "ReferenceSkeleton.h"
+#include "Rendering/SkeletalMeshLODRenderData.h"
+#include "Rendering/SkeletalMeshRenderData.h"
+#include "Rendering/SkeletalMeshVertexClothBuffer.h"
 #include "Rendering/Texture2DResource.h"
-#include "Engine/Texture2D.h"
-#include "UObject/Package.h"
+#include "RenderingThread.h"
 #include "Serialization/BulkData.h"
-#include "RHI.h"
-
-#include "Algo/MaxElement.h"
-#include "Algo/MinElement.h"
-#include "Math/NumericLimits.h"
+#include "Serialization/MemoryReader.h"
+#include "SkeletalMergingLibrary.h"
+#include "SkeletalMeshTypes.h"
+#include "Stats/Stats2.h"
+#include "Templates/Casts.h"
 #include "Templates/IsTrivial.h"
+#include "Templates/RefCounting.h"
+#include "Templates/SharedPointer.h"
+#include "Templates/Tuple.h"
+#include "Templates/UnrealTemplate.h"
+#include "TextureResource.h"
+#include "Trace/Detail/Channel.h"
+#include "UObject/Class.h"
+#include "UObject/Package.h"
+#include "UObject/SoftObjectPath.h"
+#include "UObject/SoftObjectPtr.h"
+#include "UObject/UObjectIterator.h"
+#include "UObject/UnrealNames.h"
+#include "UObject/UnrealType.h"
+#include "Widgets/Notifications/SNotificationList.h"
 
-#include "Runtime/Launch/Resources/Version.h"
-
-#include "MuCO/CustomizableObjectMipDataProvider.h"
-#include "PhysicsEngine/PhysicsAsset.h"
-#include "Animation/AnimBlueprint.h"
-#include "MuCO/CustomizableObjectMipDataProvider.h"
-
 #if WITH_EDITOR
-#include "MessageLogModule.h"
-#include "Logging/MessageLog.h"
 #include "Framework/Notifications/NotificationManager.h"
+#include "Logging/MessageLog.h"
+#include "MessageLogModule.h"
 #endif
 
-#include "MuR/Model.h"
-#include "MuR/System.h"
-#include "MuR/Mesh.h"
+// Required for engine branch preprocessor defines.
+#include "MuCO/UnrealPortabilityHelpers.h"
 #include "MuR/Image.h"
-#include "MuR/Parameters.h"
+#include "MuR/Mesh.h"
+#include "MuR/Model.h"
 #include "MuR/MutableTrace.h"
-#include "MuR/ImagePrivate.h"
-#include "MuR/OpImageSwizzle.h"
 #include "MuR/OpImagePixelFormat.h"
-#include "MuR/MeshPrivate.h"
-#include "MuR/ModelPrivate.h"
-
-// Required for engine branch preprocessor defines.
-#include "MuCO/UnrealPortabilityHelpers.h"
-#include "Animation/AnimBlueprintGeneratedClass.h"
-#include "PhysicsEngine/BodyInstance.h"
+#include "MuR/OpImageSwizzle.h"
+#include "MuR/Parameters.h"
+#include "MuR/System.h"
 
 
 UTexture2D* UCustomizableInstancePrivateData::CreateTexture()
diff --git a/Engine/Plugins/Experimental/Mutable/Source/CustomizableObject/Private/MuCO/CustomizableObjectInstanceDescriptor.cpp b/Engine/Plugins/Experimental/Mutable/Source/CustomizableObject/Private/MuCO/CustomizableObjectInstanceDescriptor.cpp
--- a/Engine/Plugins/Experimental/Mutable/Source/CustomizableObject/Private/MuCO/CustomizableObjectInstanceDescriptor.cpp
+++ b/Engine/Plugins/Experimental/Mutable/Source/CustomizableObject/Private/MuCO/CustomizableObjectInstanceDescriptor.cpp
@@ -2,9 +2,21 @@
 
 #include "MuCO/CustomizableObjectInstanceDescriptor.h"
 
+#include "HAL/PlatformCrt.h"
+#include "Logging/LogCategory.h"
+#include "Logging/LogMacros.h"
+#include "Math/Vector.h"
+#include "Misc/AssertionMacros.h"
+#include "MuCO/CustomizableObject.h"
 #include "MuCO/CustomizableObjectPrivate.h"
 #include "MuCO/CustomizableObjectSystem.h"
 #include "MuR/Model.h"
+#include "MuR/Parameters.h"
+#include "MuR/RefCounted.h"
+#include "Serialization/Archive.h"
+#include "Serialization/StructuredArchiveAdapters.h"
+#include "Templates/TypeHash.h"
+#include "Trace/Detail/Channel.h"
 
 
 const FString EmptyString;
diff --git a/Engine/Plugins/Experimental/Mutable/Source/CustomizableObject/Private/MuCO/CustomizableObjectMipDataProvider.cpp b/Engine/Plugins/Experimental/Mutable/Source/CustomizableObject/Private/MuCO/CustomizableObjectMipDataProvider.cpp
--- a/Engine/Plugins/Experimental/Mutable/Source/CustomizableObject/Private/MuCO/CustomizableObjectMipDataProvider.cpp
+++ b/Engine/Plugins/Experimental/Mutable/Source/CustomizableObject/Private/MuCO/CustomizableObjectMipDataProvider.cpp
@@ -1,8 +1,33 @@
 // Copyright Epic Games, Inc. All Rights Reserved.
 
 #include "MuCO/CustomizableObjectMipDataProvider.h"
-#include "MuR/Mesh.h"
+
+#include "Async/Fundamental/Task.h"
+#include "Containers/ArrayView.h"
+#include "Containers/IndirectArray.h"
+#include "Engine/Texture.h"
+#include "Engine/Texture2D.h"
+#include "HAL/ThreadSafeCounter.h"
+#include "HAL/UnrealMemory.h"
+#include "Logging/LogCategory.h"
+#include "Logging/LogMacros.h"
+#include "Math/UnrealMathSSE.h"
+#include "Misc/ScopeLock.h"
+#include "MuCO/CustomizableObject.h"
 #include "MuCO/CustomizableObjectSystemPrivate.h"
+#include "MuR/Instance.h"
+#include "MuR/Mesh.h"
+#include "MuR/Model.h"
+#include "MuR/MutableTrace.h"
+#include "MuR/Ptr.h"
+#include "MuR/Types.h"
+#include "Serialization/BulkData.h"
+#include "Templates/Casts.h"
+#include "Templates/Function.h"
+#include "TextureResource.h"
+#include "Trace/Detail/Channel.h"
+
+#include UE_INLINE_GENERATED_CPP_BY_NAME(CustomizableObjectMipDataProvider)
 
 UMutableTextureMipDataProviderFactory::UMutableTextureMipDataProviderFactory(const FObjectInitializer& ObjectInitializer)
 	: Super(ObjectInitializer)
diff --git a/Engine/Plugins/Experimental/Mutable/Source/CustomizableObject/Private/MuCO/CustomizableObjectModule.cpp b/Engine/Plugins/Experimental/Mutable/Source/CustomizableObject/Private/MuCO/CustomizableObjectModule.cpp
--- a/Engine/Plugins/Experimental/Mutable/Source/CustomizableObject/Private/MuCO/CustomizableObjectModule.cpp
+++ b/Engine/Plugins/Experimental/Mutable/Source/CustomizableObject/Private/MuCO/CustomizableObjectModule.cpp
@@ -1,17 +1,14 @@
 // Copyright Epic Games, Inc. All Rights Reserved.
 
-#include "Modules/ModuleManager.h"
-
-#include "MuCO/ICustomizableObjectModule.h"
-#include "MuCO/CustomizableObjectSystem.h"
-
+#include "Containers/UnrealString.h"
+#include "CoreGlobals.h"
+#include "HAL/Platform.h"
 #include "Interfaces/IPluginManager.h"
 #include "Misc/ConfigCacheIni.h"
-
-#include "Runtime/Launch/Resources/Version.h"
-
-#include "MuR/MutableMemory.h"
-#include "MuR/System.h"
+#include "Modules/ModuleManager.h"
+#include "MuCO/ICustomizableObjectModule.h"
+#include "PluginDescriptor.h"
+#include "Templates/SharedPointer.h"
 
 /**
  * Customizable Object module implementation (private)
diff --git a/Engine/Plugins/Experimental/Mutable/Source/CustomizableObject/Private/MuCO/CustomizableObjectSystem.cpp b/Engine/Plugins/Experimental/Mutable/Source/CustomizableObject/Private/MuCO/CustomizableObjectSystem.cpp
--- a/Engine/Plugins/Experimental/Mutable/Source/CustomizableObject/Private/MuCO/CustomizableObjectSystem.cpp
+++ b/Engine/Plugins/Experimental/Mutable/Source/CustomizableObject/Private/MuCO/CustomizableObjectSystem.cpp
@@ -2,39 +2,112 @@
 
 #include "MuCO/CustomizableObjectSystem.h"
 
-#include "MuCO/CustomizableObjectSystemPrivate.h"
-#include "MuCO/CustomizableObject.h"
-#include "MuCO/CustomizableObjectPrivate.h"
-#include "MuCO/CustomizableInstancePrivateData.h"
-#include "MuCO/UnrealMutableImageProvider.h"
-#include "MuCO/UnrealMutableModelDiskStreamer.h"
-#include "MuCO/LogBenchmarkUtil.h"
-#include "MuCO/ICustomizableObjectModule.h"
-#include "MuCO/CustomizableInstanceLODManagement.h"
-#include "MuCO/UnrealBakeHelpers.h"
-
+#include "Animation/Skeleton.h"
+#include "AssetRegistry/ARFilter.h"
+#include "AssetRegistry/AssetData.h"
 #include "AssetRegistry/AssetRegistryModule.h"
+#include "AssetRegistry/IAssetRegistry.h"
+#include "Async/Fundamental/Task.h"
+#include "Async/TaskGraphInterfaces.h"
+#include "Containers/EnumAsByte.h"
+#include "Containers/Queue.h"
+#include "Containers/SparseArray.h"
+#include "Containers/Ticker.h"
+#include "CoreGlobals.h"
+#include "Engine/Engine.h"
 #include "Engine/SkeletalMesh.h"
-#include "Animation/Skeleton.h"
-#include "MuCO/LogInformationUtil.h"
-#include "TimerManager.h"
-#include "SceneManagement.h"
+#include "Engine/Texture.h"
+#include "Engine/Texture2D.h"
+#include "Engine/World.h"
+#include "GameFramework/Pawn.h"
 #include "GameFramework/PlayerController.h"
+#include "GameplayTagContainer.h"
 #include "HAL/IConsoleManager.h"
+#include "HAL/PlatformCrt.h"
+#include "HAL/PlatformMisc.h"
+#include "HAL/PlatformProperties.h"
+#include "HAL/PlatformTime.h"
+#include "HAL/UnrealMemory.h"
+#include "Interfaces/ITargetPlatform.h"
+#include "Internationalization/Text.h"
 #include "Kismet/GameplayStatics.h"
+#include "Logging/LogCategory.h"
+#include "Logging/LogMacros.h"
+#include "Math/Color.h"
+#include "Math/Matrix.h"
+#include "Math/Transform.h"
+#include "Math/TransformVectorized.h"
+#include "Misc/AssertionMacros.h"
+#include "Misc/CString.h"
+#include "Misc/CommandLine.h"
+#include "Misc/CoreMisc.h"
+#include "Misc/Guid.h"
+#include "Misc/Parse.h"
+#include "Modules/ModuleManager.h"
+#include "MuCO/CustomizableInstanceLODManagement.h"
+#include "MuCO/CustomizableInstancePrivateData.h"
+#include "MuCO/CustomizableObject.h"
+#include "MuCO/CustomizableObjectInstance.h"
+#include "MuCO/CustomizableObjectInstanceDescriptor.h"
+#include "MuCO/CustomizableObjectParameterTypeDefinitions.h"
+#include "MuCO/CustomizableObjectPrivate.h"
+#include "MuCO/CustomizableObjectSystemPrivate.h"
+#include "MuCO/CustomizableObjectUIData.h"
+#include "MuCO/CustomizableSkeletalComponent.h"
+#include "MuCO/ICustomizableObjectModule.h"
+#include "MuCO/LogBenchmarkUtil.h"
+#include "MuCO/LogInformationUtil.h"
+#include "MuCO/UnrealBakeHelpers.h"
+#include "MuCO/UnrealMutableImageProvider.h"
+#include "MuCO/UnrealMutableModelDiskStreamer.h"
+#include "MuR/Image.h"
+#include "MuR/Instance.h"
+#include "MuR/Mesh.h"
+#include "MuR/MeshBufferSet.h"
+#include "MuR/MutableMath.h"
+#include "MuR/Ptr.h"
+#include "MuR/RefCounted.h"
+#include "MuR/Settings.h"
+#include "MuR/Skeleton.h"
+#include "MuR/Types.h"
+#include "ProfilingDebugging/CpuProfilerTrace.h"
+#include "RHI.h"
+#include "ReferenceSkeleton.h"
+#include "RenderCommandFence.h"
+#include "Templates/Casts.h"
+#include "Templates/RefCounting.h"
+#include "Templates/Tuple.h"
+#include "Templates/UnrealTemplate.h"
+#include "TextureResource.h"
+#include "Trace/Detail/Channel.h"
+#include "UObject/Class.h"
+#include "UObject/GarbageCollection.h"
+#include "UObject/NameTypes.h"
+#include "UObject/SoftObjectPtr.h"
+#include "UObject/TopLevelAssetPath.h"
+#include "UObject/UObjectArray.h"
+#include "UObject/UObjectIterator.h"
+#include "UObject/WeakObjectPtr.h"
+#include "UObject/WeakObjectPtrTemplates.h"
+#include "Widgets/Notifications/SNotificationList.h"
 
 #if WITH_EDITOR
 #include "Editor.h"
+#include "Framework/Notifications/NotificationManager.h"
+#include "Logging/MessageLog.h"
 #include "Misc/ConfigCacheIni.h"
 #include "Misc/MessageDialog.h"
-#include "Logging/MessageLog.h"
-#include "Framework/Notifications/NotificationManager.h"
 #endif
 
 #include "MuR/MutableTrace.h"
-#include "MuR/System.h"
-#include "MuR/SystemPrivate.h"
 #include "MuR/Parameters.h"
+#include "MuR/System.h"
+
+#include UE_INLINE_GENERATED_CPP_BY_NAME(CustomizableObjectSystem)
+
+class AActor;
+class UAnimInstance;
+class UMaterialInterface;
 
 
 DEFINE_STAT(STAT_MutableNumSkeletalMeshes);
@@ -929,13 +1002,6 @@ namespace impl
 		// This runs in the mutable thread.
 		check(MutableParameters);
 
-#ifdef WITH_EDITOR
-		// \TODO: Review this to make sure it is really needed.
-		// NOTE: this delegate can be called outside the main thread, make sure operations done in the
-		//       callbacks are safe, avoid using it outside editor functionality
-		//Public->UpdateBeginDelegate.ExecuteIfBound();
-#endif
-
 		OperationData->InstanceUpdateData.Clear();
 
 		mu::System* System = UCustomizableObjectSystem::GetInstance()->GetPrivate()->MutableSystem.get();
diff --git a/Engine/Plugins/Experimental/Mutable/Source/CustomizableObject/Private/MuCO/CustomizableSkeletalComponent.cpp b/Engine/Plugins/Experimental/Mutable/Source/CustomizableObject/Private/MuCO/CustomizableSkeletalComponent.cpp
--- a/Engine/Plugins/Experimental/Mutable/Source/CustomizableObject/Private/MuCO/CustomizableSkeletalComponent.cpp
+++ b/Engine/Plugins/Experimental/Mutable/Source/CustomizableObject/Private/MuCO/CustomizableSkeletalComponent.cpp
@@ -2,20 +2,35 @@
 
 #include "MuCO/CustomizableSkeletalComponent.h"
 
-#include "MuCO/CustomizableObject.h"
-#include "MuCO/CustomizableObjectInstance.h"
-#include "MuCO/CustomizableInstancePrivateData.h"
-#include "MuCO/CustomizableObjectSystem.h"
-
-#include "Engine/Texture.h"
-#include "Engine/SkeletalMesh.h"
-#include "Engine/StaticMesh.h"
+#include "AnimationRuntime.h"
 #include "Components/SkeletalMeshComponent.h"
+#include "Components/SkinnedMeshComponent.h"
+#include "Containers/Array.h"
+#include "Containers/EnumAsByte.h"
+#include "Containers/IndirectArray.h"
+#include "Containers/Map.h"
+#include "Engine/EngineTypes.h"
+#include "Engine/SkeletalMesh.h"
+#include "Engine/World.h"
+#include "GameFramework/Actor.h"
+#include "GameFramework/Pawn.h"
+#include "Logging/LogCategory.h"
+#include "Logging/LogMacros.h"
 #include "Materials/MaterialInstanceDynamic.h"
-
-#include "AnimationRuntime.h"
-
+#include "Materials/MaterialInterface.h"
+#include "Math/Vector.h"
+#include "MuCO/CustomizableInstancePrivateData.h"
+#include "MuCO/CustomizableObject.h"
+#include "MuCO/CustomizableObjectInstance.h"
+#include "Rendering/SkeletalMeshLODRenderData.h"
+#include "Rendering/SkeletalMeshRenderData.h"
+#include "Templates/Casts.h"
+#include "Trace/Detail/Channel.h"
+#include "UObject/Class.h"
+#include "UObject/NameTypes.h"
 #include "UObject/ObjectSaveContext.h"
+#include "UObject/WeakObjectPtr.h"
+#include "UObject/WeakObjectPtrTemplates.h"
 
 
 UCustomizableSkeletalComponent::UCustomizableSkeletalComponent()
diff --git a/Engine/Plugins/Experimental/Mutable/Source/CustomizableObject/Private/MuCO/CustomizableSkeletalMeshActor.cpp b/Engine/Plugins/Experimental/Mutable/Source/CustomizableObject/Private/MuCO/CustomizableSkeletalMeshActor.cpp
--- a/Engine/Plugins/Experimental/Mutable/Source/CustomizableObject/Private/MuCO/CustomizableSkeletalMeshActor.cpp
+++ b/Engine/Plugins/Experimental/Mutable/Source/CustomizableObject/Private/MuCO/CustomizableSkeletalMeshActor.cpp
@@ -1,7 +1,13 @@
 // Copyright Epic Games, Inc. All Rights Reserved.
 
 #include "MuCO/CustomizableSkeletalMeshActor.h"
+
 #include "Components/SkeletalMeshComponent.h"
+#include "Containers/UnrealString.h"
+#include "Engine/EngineTypes.h"
+#include "HAL/PlatformCrt.h"
+#include "MuCO/CustomizableSkeletalComponent.h"
+#include "UObject/NameTypes.h"
 
 #define LOCTEXT_NAMESPACE "CustomizableObject"
 
diff --git a/Engine/Plugins/Experimental/Mutable/Source/CustomizableObject/Private/MuCO/LogBenchmarkUtil.cpp b/Engine/Plugins/Experimental/Mutable/Source/CustomizableObject/Private/MuCO/LogBenchmarkUtil.cpp
--- a/Engine/Plugins/Experimental/Mutable/Source/CustomizableObject/Private/MuCO/LogBenchmarkUtil.cpp
+++ b/Engine/Plugins/Experimental/Mutable/Source/CustomizableObject/Private/MuCO/LogBenchmarkUtil.cpp
@@ -1,10 +1,13 @@
 // Copyright Epic Games, Inc. All Rights Reserved.
 
 #include "MuCO/LogBenchmarkUtil.h"
-#include "Misc/FileHelper.h"
+
 #include "Containers/UnrealString.h"
+#include "GenericPlatform/GenericPlatformFile.h"
+#include "HAL/FileManager.h"
 #include "HAL/PlatformFileManager.h"
 #include "Misc/DateTime.h"
+#include "Misc/FileHelper.h"
 #include "MuCO/UnrealPortabilityHelpers.h"
 
 bool LogBenchmarkUtil::bLoggingActive = false;
diff --git a/Engine/Plugins/Experimental/Mutable/Source/CustomizableObject/Private/MuCO/LogInformationUtil.cpp b/Engine/Plugins/Experimental/Mutable/Source/CustomizableObject/Private/MuCO/LogInformationUtil.cpp
--- a/Engine/Plugins/Experimental/Mutable/Source/CustomizableObject/Private/MuCO/LogInformationUtil.cpp
+++ b/Engine/Plugins/Experimental/Mutable/Source/CustomizableObject/Private/MuCO/LogInformationUtil.cpp
@@ -1,11 +1,27 @@
 // Copyright Epic Games, Inc. All Rights Reserved.
 
 #include "MuCO/LogInformationUtil.h"
+
+#include "Containers/IndirectArray.h"
+#include "Containers/UnrealString.h"
 #include "Engine/SkeletalMesh.h"
+#include "Engine/Texture.h"
 #include "Engine/Texture2D.h"
-#include "MuCO/CustomizableObject.h"
-#include "MuCO/CustomizableInstancePrivateData.h"
+#include "Engine/World.h"
 #include "GameFramework/PlayerController.h"
+#include "Logging/LogCategory.h"
+#include "Logging/LogMacros.h"
+#include "Math/Color.h"
+#include "Math/UnrealMathSSE.h"
+#include "Math/Vector.h"
+#include "MuCO/CustomizableInstancePrivateData.h"
+#include "MuCO/CustomizableObject.h"
+#include "MuCO/CustomizableObjectInstance.h"
+#include "MuCO/CustomizableObjectParameterTypeDefinitions.h"
+#include "TextureResource.h"
+#include "Trace/Detail/Channel.h"
+#include "UObject/Class.h"
+#include "UObject/ObjectPtr.h"
 
 int LogInformationUtil::CountLOD0 = 0;
 int LogInformationUtil::CountLOD1 = 0;
diff --git a/Engine/Plugins/Experimental/Mutable/Source/CustomizableObject/Private/MuCO/LogInformationUtil.h b/Engine/Plugins/Experimental/Mutable/Source/CustomizableObject/Private/MuCO/LogInformationUtil.h
--- a/Engine/Plugins/Experimental/Mutable/Source/CustomizableObject/Private/MuCO/LogInformationUtil.h
+++ b/Engine/Plugins/Experimental/Mutable/Source/CustomizableObject/Private/MuCO/LogInformationUtil.h
@@ -3,9 +3,12 @@
 #pragma once
 
 #include "Containers/Array.h"
-#include "Engine/Texture.h"
+#include "Containers/Map.h"
+#include "HAL/Platform.h"
 
+class FString;
 class UCustomizableObjectInstance;
+struct FTexturePlatformData;
 
 #define STRINGIZE(A) STRINGIZE_NX(A)
 #define STRINGIZE_NX(A) #A
diff --git a/Engine/Plugins/Experimental/Mutable/Source/CustomizableObject/Private/MuCO/MultilayerProjector.cpp b/Engine/Plugins/Experimental/Mutable/Source/CustomizableObject/Private/MuCO/MultilayerProjector.cpp
--- a/Engine/Plugins/Experimental/Mutable/Source/CustomizableObject/Private/MuCO/MultilayerProjector.cpp
+++ b/Engine/Plugins/Experimental/Mutable/Source/CustomizableObject/Private/MuCO/MultilayerProjector.cpp
@@ -2,7 +2,12 @@
 
 #include "MuCO/MultilayerProjector.h"
 
+#include "HAL/PlatformCrt.h"
+#include "Misc/AssertionMacros.h"
+#include "MuCO/CustomizableObject.h"
 #include "MuCO/CustomizableObjectInstance.h"
+#include "MuCO/CustomizableObjectParameterTypeDefinitions.h"
+#include "Templates/Tuple.h"
 
 
 void FMultilayerProjectorLayer::Read(const FMultilayerProjector& MultilayerProjector, const int32 Index)
diff --git a/Engine/Plugins/Experimental/Mutable/Source/CustomizableObject/Private/MuCO/MutableMeshBufferUtils.cpp b/Engine/Plugins/Experimental/Mutable/Source/CustomizableObject/Private/MuCO/MutableMeshBufferUtils.cpp
--- a/Engine/Plugins/Experimental/Mutable/Source/CustomizableObject/Private/MuCO/MutableMeshBufferUtils.cpp
+++ b/Engine/Plugins/Experimental/Mutable/Source/CustomizableObject/Private/MuCO/MutableMeshBufferUtils.cpp
@@ -2,8 +2,13 @@
 
 #include "MuCO/MutableMeshBufferUtils.h"
 
+#include "Math/Color.h"
+#include "Misc/AssertionMacros.h"
+#include "MuR/MeshBufferSet.h"
+#include "MuR/Types.h"
 #include "Rendering/PositionVertexBuffer.h"
 #include "Rendering/StaticMeshVertexBuffer.h"
+#include "Templates/UnrealTemplate.h"
 
 
 void MutableMeshBufferUtils::SetupVertexPositionsBuffer(const int32& InCurrentVertexBuffer, mu::FMeshBufferSet& OutTargetVertexBuffers)
diff --git a/Engine/Plugins/Experimental/Mutable/Source/CustomizableObject/Private/MuCO/UnrealBakeHelpers.cpp b/Engine/Plugins/Experimental/Mutable/Source/CustomizableObject/Private/MuCO/UnrealBakeHelpers.cpp
--- a/Engine/Plugins/Experimental/Mutable/Source/CustomizableObject/Private/MuCO/UnrealBakeHelpers.cpp
+++ b/Engine/Plugins/Experimental/Mutable/Source/CustomizableObject/Private/MuCO/UnrealBakeHelpers.cpp
@@ -2,13 +2,42 @@
 
 #include "MuCO/UnrealBakeHelpers.h"
 
-#include "Runtime/Launch/Resources/Version.h"
+#include "BoneIndices.h"
+#include "ClothLODData.h"
+#include "ClothPhysicalMeshData.h"
+#include "ClothingAsset.h"
+#include "ClothingAssetBase.h"
+#include "Containers/Array.h"
+#include "Containers/IndirectArray.h"
+#include "Containers/Map.h"
+#include "Containers/UnrealString.h"
 #include "Engine/SkeletalMesh.h"
-#include "Rendering/SkeletalMeshRenderData.h"
+#include "Engine/SkeletalMeshLODSettings.h"
+#include "GPUSkinPublicDefs.h"
+#include "HAL/PlatformCrt.h"
+#include "Math/Color.h"
+#include "Math/UnrealMathSSE.h"
+#include "Math/Vector.h"
+#include "Math/Vector2D.h"
+#include "Math/Vector4.h"
+#include "Misc/AssertionMacros.h"
+#include "Misc/Guid.h"
+#include "PackedNormal.h"
+#include "PointWeightMap.h"
+#include "RawIndexBuffer.h"
+#include "Rendering/MultiSizeIndexContainer.h"
+#include "Rendering/PositionVertexBuffer.h"
+#include "Rendering/SkeletalMeshLODModel.h"
 #include "Rendering/SkeletalMeshLODRenderData.h"
 #include "Rendering/SkeletalMeshModel.h"
-#include "Rendering/SkeletalMeshLODModel.h"
-#include "ClothingAsset.h"
+#include "Rendering/SkeletalMeshRenderData.h"
+#include "Rendering/SkinWeightVertexBuffer.h"
+#include "Rendering/StaticMeshVertexBuffer.h"
+#include "SkeletalMeshTypes.h"
+#include "StaticMeshResources.h"
+#include "Templates/Casts.h"
+#include "Templates/Decay.h"
+#include "Templates/UnrealTemplate.h"
 
 
 //-------------------------------------------------------------------------------------------------
diff --git a/Engine/Plugins/Experimental/Mutable/Source/CustomizableObject/Private/MuCO/UnrealConversionUtils.cpp b/Engine/Plugins/Experimental/Mutable/Source/CustomizableObject/Private/MuCO/UnrealConversionUtils.cpp
--- a/Engine/Plugins/Experimental/Mutable/Source/CustomizableObject/Private/MuCO/UnrealConversionUtils.cpp
+++ b/Engine/Plugins/Experimental/Mutable/Source/CustomizableObject/Private/MuCO/UnrealConversionUtils.cpp
@@ -2,15 +2,37 @@
 
 #include "MuCO/UnrealConversionUtils.h"
 
-#include "Animation/Skeleton.h"
+#include "Containers/IndirectArray.h"
+#include "Containers/Map.h"
+#include "CoreTypes.h"
+#include "Engine/SkeletalMesh.h"
+#include "HAL/UnrealMemory.h"
+#include "Logging/LogCategory.h"
+#include "Logging/LogMacros.h"
+#include "Math/Transform.h"
+#include "Misc/AssertionMacros.h"
+#include "MuCO/CustomizableObject.h"
+#include "MuCO/CustomizableObjectInstance.h"
 #include "MuCO/CustomizableObjectSystemPrivate.h"
-#include "Materials/Material.h"
-#include "MuR/MeshPrivate.h"
-#include "MuR/OpMeshFormat.h"
+#include "MuCO/UnrealPortabilityHelpers.h"
 #include "MuR/Mesh.h"
 #include "MuR/MutableTrace.h"
 #include "MuR/Parameters.h"
 #include "MuR/Ptr.h"
+#include "RawIndexBuffer.h"
+#include "Rendering/ColorVertexBuffer.h"
+#include "Rendering/MultiSizeIndexContainer.h"
+#include "Rendering/PositionVertexBuffer.h"
+#include "Rendering/SkeletalMeshDuplicatedVerticesBuffer.h"
+#include "Rendering/SkeletalMeshLODRenderData.h"
+#include "Rendering/SkinWeightVertexBuffer.h"
+#include "Rendering/StaticMeshVertexBuffer.h"
+#include "SkeletalMeshTypes.h"
+#include "StaticMeshResources.h"
+#include "Trace/Detail/Channel.h"
+#include "UObject/NameTypes.h"
+
+class USkeleton;
 
 namespace UnrealConversionUtils
 {
diff --git a/Engine/Plugins/Experimental/Mutable/Source/CustomizableObject/Private/MuCO/UnrealMutableImageProvider.cpp b/Engine/Plugins/Experimental/Mutable/Source/CustomizableObject/Private/MuCO/UnrealMutableImageProvider.cpp
--- a/Engine/Plugins/Experimental/Mutable/Source/CustomizableObject/Private/MuCO/UnrealMutableImageProvider.cpp
+++ b/Engine/Plugins/Experimental/Mutable/Source/CustomizableObject/Private/MuCO/UnrealMutableImageProvider.cpp
@@ -2,8 +2,17 @@
 
 #include "MuCO/UnrealMutableImageProvider.h"
 
-#include "MuCO/CustomizableObject.h"
+#include "CoreGlobals.h"
+#include "Engine/Texture.h"
+#include "Engine/Texture2D.h"
+#include "Engine/TextureDefines.h"
+#include "HAL/UnrealMemory.h"
+#include "Math/IntVector.h"
+#include "Math/UnrealMathSSE.h"
+#include "Misc/AssertionMacros.h"
 #include "MuCO/CustomizableObjectSystem.h"
+#include "MuR/Ptr.h"
+#include "UObject/WeakObjectPtr.h"
 
 
 //-------------------------------------------------------------------------------------------------
diff --git a/Engine/Plugins/Experimental/Mutable/Source/CustomizableObject/Private/MuCO/UnrealMutableImageProvider.h b/Engine/Plugins/Experimental/Mutable/Source/CustomizableObject/Private/MuCO/UnrealMutableImageProvider.h
--- a/Engine/Plugins/Experimental/Mutable/Source/CustomizableObject/Private/MuCO/UnrealMutableImageProvider.h
+++ b/Engine/Plugins/Experimental/Mutable/Source/CustomizableObject/Private/MuCO/UnrealMutableImageProvider.h
@@ -3,10 +3,12 @@
 #pragma once
 
 #include "Containers/Array.h"
-#include "UObject/WeakObjectPtrTemplates.h"
-
-#include "MuR/System.h"
+#include "Containers/Map.h"
+#include "HAL/Platform.h"
 #include "MuR/Image.h"
+#include "MuR/System.h"
+#include "MuR/Types.h"
+#include "UObject/WeakObjectPtrTemplates.h"
 
 
 /** Implementation of a mutable core provider for image parameters that are application-specific. */
diff --git a/Engine/Plugins/Experimental/Mutable/Source/CustomizableObject/Private/MuCO/UnrealMutableModelDiskStreamer.cpp b/Engine/Plugins/Experimental/Mutable/Source/CustomizableObject/Private/MuCO/UnrealMutableModelDiskStreamer.cpp
--- a/Engine/Plugins/Experimental/Mutable/Source/CustomizableObject/Private/MuCO/UnrealMutableModelDiskStreamer.cpp
+++ b/Engine/Plugins/Experimental/Mutable/Source/CustomizableObject/Private/MuCO/UnrealMutableModelDiskStreamer.cpp
@@ -2,23 +2,28 @@
 
 #include "MuCO/UnrealMutableModelDiskStreamer.h"
 
+#include "Async/AsyncFileHandle.h"
+#include "Containers/UnrealString.h"
+#include "CoreGlobals.h"
+#include "GenericPlatform/GenericPlatformFile.h"
+#include "HAL/IConsoleManager.h"
+#include "Logging/LogCategory.h"
+#include "Logging/LogMacros.h"
+#include "Misc/AssertionMacros.h"
 #include "MuCO/CustomizableObject.h"
 #include "MuCO/CustomizableObjectSystem.h"
-
-#include "Serialization/MemoryReader.h"
-#include "Serialization/MemoryWriter.h"
-#include "Serialization/BulkData.h"
-
+#include "MuR/Model.h"
+#include "MuR/MutableTrace.h"
+#include "MuR/Ptr.h"
+#include "Serialization/Archive.h"
+#include "Stats/Stats2.h"
+#include "Templates/UnrealTemplate.h"
+#include "Trace/Detail/Channel.h"
 
 #if WITH_EDITOR
-#include "Misc/Paths.h"
-#include "HAL/FileManager.h"
 #include "HAL/PlatformFileManager.h"
 #endif
 
-#include "MuR/MutableTrace.h"
-#include "Runtime/Launch/Resources/Version.h"
-
 
 //-------------------------------------------------------------------------------------------------
 #if WITH_EDITOR
diff --git a/Engine/Plugins/Experimental/Mutable/Source/CustomizableObject/Public/MuCO/CustomizableObjectCustomVersion.h b/Engine/Plugins/Experimental/Mutable/Source/CustomizableObject/Public/MuCO/CustomizableObjectCustomVersion.h
--- a/Engine/Plugins/Experimental/Mutable/Source/CustomizableObject/Public/MuCO/CustomizableObjectCustomVersion.h
+++ b/Engine/Plugins/Experimental/Mutable/Source/CustomizableObject/Public/MuCO/CustomizableObjectCustomVersion.h
@@ -2,8 +2,9 @@
 
 #pragma once
 
-#include "CoreMinimal.h"
-#include "Misc/Guid.h"
+#include "HAL/Platform.h"
+
+struct FGuid;
 
 // Custom serialization version for assets/classes in the CustomizableObject and CustomizableObjectEditor modules
 struct CUSTOMIZABLEOBJECT_API FCustomizableObjectCustomVersion
diff --git a/Engine/Plugins/Experimental/Mutable/Source/CustomizableObject/Public/MuCO/CustomizableObjectDGGUI.h b/Engine/Plugins/Experimental/Mutable/Source/CustomizableObject/Public/MuCO/CustomizableObjectDGGUI.h
--- a/Engine/Plugins/Experimental/Mutable/Source/CustomizableObject/Public/MuCO/CustomizableObjectDGGUI.h
+++ b/Engine/Plugins/Experimental/Mutable/Source/CustomizableObject/Public/MuCO/CustomizableObjectDGGUI.h
@@ -3,9 +3,17 @@
 #pragma once
 
 #include "Blueprint/UserWidget.h"
-#include "Animation/AnimClassInterface.h"
+#include "HAL/Platform.h"
+#include "Templates/SubclassOf.h"
+#include "UObject/UObjectGlobals.h"
+
 #include "CustomizableObjectDGGUI.generated.h"
 
+class UAnimInstance;
+class UCustomizableSkeletalComponent;
+class UObject;
+class UWorld;
+
 
 UCLASS(Abstract, BlueprintType, Blueprintable)
 class CUSTOMIZABLEOBJECT_API UDGGUI : public UUserWidget
diff --git a/Engine/Plugins/Experimental/Mutable/Source/CustomizableObject/Public/MuCO/CustomizableObjectInstanceDescriptor.h b/Engine/Plugins/Experimental/Mutable/Source/CustomizableObject/Public/MuCO/CustomizableObjectInstanceDescriptor.h
--- a/Engine/Plugins/Experimental/Mutable/Source/CustomizableObject/Public/MuCO/CustomizableObjectInstanceDescriptor.h
+++ b/Engine/Plugins/Experimental/Mutable/Source/CustomizableObject/Public/MuCO/CustomizableObjectInstanceDescriptor.h
@@ -2,12 +2,20 @@
 
 #pragma once
 
-#include "MuCO/CustomizableObject.h"
+#include "Containers/Array.h"
+#include "Containers/Map.h"
+#include "Containers/UnrealString.h"
+#include "Math/Color.h"
+#include "Math/UnrealMathSSE.h"
+#include "MuCO/CustomizableObjectParameterTypeDefinitions.h"
+#include "UObject/ObjectPtr.h"
 
 #include "CustomizableObjectInstanceDescriptor.generated.h"
 
-class UCustomizableObjectInstance;
+class FArchive;
 class UCustomizableInstancePrivateData;
+class UCustomizableObject;
+class UCustomizableObjectInstance;
 
 
 /** Set of parameters + state that defines a CustomizableObjectInstance.
diff --git a/Engine/Plugins/Experimental/Mutable/Source/CustomizableObject/Public/MuCO/CustomizableObjectMipDataProvider.h b/Engine/Plugins/Experimental/Mutable/Source/CustomizableObject/Public/MuCO/CustomizableObjectMipDataProvider.h
--- a/Engine/Plugins/Experimental/Mutable/Source/CustomizableObject/Public/MuCO/CustomizableObjectMipDataProvider.h
+++ b/Engine/Plugins/Experimental/Mutable/Source/CustomizableObject/Public/MuCO/CustomizableObjectMipDataProvider.h
@@ -2,16 +2,28 @@
 
 #pragma once
 
-#include "Streaming/TextureMipDataProvider.h"
+#include "Async/TaskGraphInterfaces.h"
+#include "Containers/Array.h"
 #include "Engine/TextureMipDataProviderFactory.h"
+#include "HAL/CriticalSection.h"
+#include "HAL/Platform.h"
+#include "Misc/AssertionMacros.h"
 #include "MuCO/CustomizableObjectSystem.h"
-#include "MuR/System.h"
-#include "MuR/Model.h"
+#include "MuR/Image.h"
 #include "MuR/Parameters.h"
+#include "MuR/RefCounted.h"
+#include "MuR/System.h"
+#include "Streaming/TextureMipDataProvider.h"
+#include "Templates/SharedPointer.h"
+#include "UObject/ObjectPtr.h"
+#include "UObject/UObjectGlobals.h"
 
 #include "CustomizableObjectMipDataProvider.generated.h"
 
+class FThreadSafeCounter;
 class UCustomizableObjectInstance;
+class UObject;
+class UTexture;
 
 /** This struct stores the data relevant for the construction of a specific texture. 
 * This includes all the data required to rebuild the image (or any of its mips).
diff --git a/Engine/Plugins/Experimental/Mutable/Source/CustomizableObject/Public/MuCO/MultilayerProjector.h b/Engine/Plugins/Experimental/Mutable/Source/CustomizableObject/Public/MuCO/MultilayerProjector.h
--- a/Engine/Plugins/Experimental/Mutable/Source/CustomizableObject/Public/MuCO/MultilayerProjector.h
+++ b/Engine/Plugins/Experimental/Mutable/Source/CustomizableObject/Public/MuCO/MultilayerProjector.h
@@ -2,6 +2,14 @@
 
 #pragma once
 
+#include "Containers/Array.h"
+#include "Containers/Map.h"
+#include "Containers/UnrealString.h"
+#include "Math/UnrealMathSSE.h"
+#include "Math/Vector.h"
+#include "UObject/NameTypes.h"
+#include "UObject/ObjectPtr.h"
+
 #include "MultilayerProjector.generated.h"
 
 class UCustomizableObjectInstance;
diff --git a/Engine/Plugins/Experimental/Mutable/Source/CustomizableObject/Public/MuCO/MutableMeshBufferUtils.h b/Engine/Plugins/Experimental/Mutable/Source/CustomizableObject/Public/MuCO/MutableMeshBufferUtils.h
--- a/Engine/Plugins/Experimental/Mutable/Source/CustomizableObject/Public/MuCO/MutableMeshBufferUtils.h
+++ b/Engine/Plugins/Experimental/Mutable/Source/CustomizableObject/Public/MuCO/MutableMeshBufferUtils.h
@@ -2,7 +2,9 @@
 
 #pragma once
 
-#include "MuR/Mesh.h"
+#include "HAL/Platform.h"
+
+namespace mu { class FMeshBufferSet; }
 
 /** Publicly accessible methods designed to prepare the Mutable Mesh buffers with the internal structure  required by Unreal and
  * our systems to be able to work with the mesh owning those buffers.
diff --git a/Engine/Plugins/Experimental/Mutable/Source/CustomizableObject/Public/MuCO/UnrealBakeHelpers.h b/Engine/Plugins/Experimental/Mutable/Source/CustomizableObject/Public/MuCO/UnrealBakeHelpers.h
--- a/Engine/Plugins/Experimental/Mutable/Source/CustomizableObject/Public/MuCO/UnrealBakeHelpers.h
+++ b/Engine/Plugins/Experimental/Mutable/Source/CustomizableObject/Public/MuCO/UnrealBakeHelpers.h
@@ -2,7 +2,7 @@
 
 #pragma once
 
-#include "MuCO/ICustomizableObjectModule.h"
+#include "HAL/Platform.h"
 
 //-------------------------------------------------------------------------------------------------
 
diff --git a/Engine/Plugins/Experimental/Mutable/Source/CustomizableObject/Public/MuCO/UnrealConversionUtils.h b/Engine/Plugins/Experimental/Mutable/Source/CustomizableObject/Public/MuCO/UnrealConversionUtils.h
--- a/Engine/Plugins/Experimental/Mutable/Source/CustomizableObject/Public/MuCO/UnrealConversionUtils.h
+++ b/Engine/Plugins/Experimental/Mutable/Source/CustomizableObject/Public/MuCO/UnrealConversionUtils.h
@@ -3,7 +3,15 @@
 #pragma once
 
 #include "../Private/MuCO/CustomizableObjectSystemPrivate.h"
-#include "MuR/Mesh.h"
+#include "Containers/Array.h"
+#include "HAL/Platform.h"
+#include "MuR/MeshBufferSet.h"
+#include "MuR/System.h"
+#include "ReferenceSkeleton.h"
+
+class FSkeletalMeshLODRenderData;
+class USkeletalMesh;
+class USkeleton;
 
 
 namespace UnrealConversionUtils
diff --git a/Engine/Plugins/Experimental/Mutable/Source/CustomizableObject/Public/MuCO/UnrealMutableModelDiskStreamer.h b/Engine/Plugins/Experimental/Mutable/Source/CustomizableObject/Public/MuCO/UnrealMutableModelDiskStreamer.h
--- a/Engine/Plugins/Experimental/Mutable/Source/CustomizableObject/Public/MuCO/UnrealMutableModelDiskStreamer.h
+++ b/Engine/Plugins/Experimental/Mutable/Source/CustomizableObject/Public/MuCO/UnrealMutableModelDiskStreamer.h
@@ -2,15 +2,22 @@
 
 #pragma once
 
-#include "Serialization/Archive.h"
-#include "Containers/UnrealString.h"
-#include "Async/AsyncFileHandle.h"
-
+#include "Containers/Array.h"
+#include "Containers/Map.h"
+#include "Containers/SparseArray.h"
+#include "HAL/Platform.h"
+#include "HAL/PlatformCrt.h"
 #include "MuCO/CustomizableObject.h"
-
-#include "MuR/Model.h"
+#include "MuR/RefCounted.h"
 #include "MuR/Serialisation.h"
 
+class FArchive;
+class IAsyncReadFileHandle;
+class IAsyncReadRequest;
+class UCustomizableObject;
+namespace mu { class Model; }
+struct FMutableStreamableBlock;
+
 
 // Implementation fo the basic mutable streams.
 #if WITH_EDITOR
diff --git a/Engine/Plugins/Experimental/Mutable/Source/CustomizableObjectEditor/Classes/MuCOE/CustomizableObjectEditorLevelProfile.h b/Engine/Plugins/Experimental/Mutable/Source/CustomizableObjectEditor/Classes/MuCOE/CustomizableObjectEditorLevelProfile.h
--- a/Engine/Plugins/Experimental/Mutable/Source/CustomizableObjectEditor/Classes/MuCOE/CustomizableObjectEditorLevelProfile.h
+++ b/Engine/Plugins/Experimental/Mutable/Source/CustomizableObjectEditor/Classes/MuCOE/CustomizableObjectEditorLevelProfile.h
@@ -2,10 +2,11 @@
 
 #pragma once
 
-#include "CoreMinimal.h"
-//#include "Classes/Engine/DataTable.h"
-#include "AssetViewerSettings.h"
 
+#include "Containers/Array.h"
+#include "Containers/UnrealString.h"
+
+struct FPreviewSceneProfile;
 
 class LevelProfileManage
 {
diff --git a/Engine/Plugins/Experimental/Mutable/Source/CustomizableObjectEditor/Classes/MuCOE/CustomizableObjectEditorSettings.h b/Engine/Plugins/Experimental/Mutable/Source/CustomizableObjectEditor/Classes/MuCOE/CustomizableObjectEditorSettings.h
--- a/Engine/Plugins/Experimental/Mutable/Source/CustomizableObjectEditor/Classes/MuCOE/CustomizableObjectEditorSettings.h
+++ b/Engine/Plugins/Experimental/Mutable/Source/CustomizableObjectEditor/Classes/MuCOE/CustomizableObjectEditorSettings.h
@@ -2,9 +2,10 @@
 
 #pragma once
 
-#include "CoreMinimal.h"
-#include "UObject/ObjectMacros.h"
 #include "UObject/Object.h"
+#include "UObject/ObjectMacros.h"
+#include "UObject/UObjectGlobals.h"
+
 #include "CustomizableObjectEditorSettings.generated.h"
 
 UCLASS(config = Engine)
diff --git a/Engine/Plugins/Experimental/Mutable/Source/CustomizableObjectEditor/Classes/MuCOE/CustomizableObjectFactory.h b/Engine/Plugins/Experimental/Mutable/Source/CustomizableObjectEditor/Classes/MuCOE/CustomizableObjectFactory.h
--- a/Engine/Plugins/Experimental/Mutable/Source/CustomizableObjectEditor/Classes/MuCOE/CustomizableObjectFactory.h
+++ b/Engine/Plugins/Experimental/Mutable/Source/CustomizableObjectEditor/Classes/MuCOE/CustomizableObjectFactory.h
@@ -3,9 +3,15 @@
 #pragma once
 
 #include "Factories/Factory.h"
+#include "UObject/NameTypes.h"
+#include "UObject/UObjectGlobals.h"
 
 #include "CustomizableObjectFactory.generated.h"
 
+class FFeedbackContext;
+class UClass;
+class UObject;
+
 
 UCLASS(MinimalAPI)
 class UCustomizableObjectFactory : public UFactory
diff --git a/Engine/Plugins/Experimental/Mutable/Source/CustomizableObjectEditor/Classes/MuCOE/CustomizableObjectGraph.h b/Engine/Plugins/Experimental/Mutable/Source/CustomizableObjectEditor/Classes/MuCOE/CustomizableObjectGraph.h
--- a/Engine/Plugins/Experimental/Mutable/Source/CustomizableObjectEditor/Classes/MuCOE/CustomizableObjectGraph.h
+++ b/Engine/Plugins/Experimental/Mutable/Source/CustomizableObjectEditor/Classes/MuCOE/CustomizableObjectGraph.h
@@ -2,14 +2,16 @@
 
 #pragma once
 
+#include "Containers/Map.h"
+#include "Containers/Set.h"
 #include "EdGraph/EdGraph.h"
-#include "Logging/TokenizedMessage.h"
-#include "Misc/UObjectToken.h"
+#include "Misc/Guid.h"
+#include "UObject/NameTypes.h"
+#include "UObject/UObjectGlobals.h"
 
 #include "CustomizableObjectGraph.generated.h"
 
-
-class UCustomizableObjectNode;
+class UObject;
 
 
 UCLASS()
diff --git a/Engine/Plugins/Experimental/Mutable/Source/CustomizableObjectEditor/Classes/MuCOE/CustomizableObjectIdentifierCustomization.h b/Engine/Plugins/Experimental/Mutable/Source/CustomizableObjectEditor/Classes/MuCOE/CustomizableObjectIdentifierCustomization.h
--- a/Engine/Plugins/Experimental/Mutable/Source/CustomizableObjectEditor/Classes/MuCOE/CustomizableObjectIdentifierCustomization.h
+++ b/Engine/Plugins/Experimental/Mutable/Source/CustomizableObjectEditor/Classes/MuCOE/CustomizableObjectIdentifierCustomization.h
@@ -2,8 +2,12 @@
 
 #pragma once
 
+#include "Containers/UnrealString.h"
 #include "IPropertyTypeCustomization.h"
-#include "AssetRegistry/AssetData.h"
+#include "Templates/SharedPointer.h"
+
+class IPropertyHandle;
+struct FAssetData;
 
 class FCustomizableObjectIdentifierCustomization : public IPropertyTypeCustomization
 {
diff --git a/Engine/Plugins/Experimental/Mutable/Source/CustomizableObjectEditor/Classes/MuCOE/CustomizableObjectInstanceFactory.h b/Engine/Plugins/Experimental/Mutable/Source/CustomizableObjectEditor/Classes/MuCOE/CustomizableObjectInstanceFactory.h
--- a/Engine/Plugins/Experimental/Mutable/Source/CustomizableObjectEditor/Classes/MuCOE/CustomizableObjectInstanceFactory.h
+++ b/Engine/Plugins/Experimental/Mutable/Source/CustomizableObjectEditor/Classes/MuCOE/CustomizableObjectInstanceFactory.h
@@ -2,14 +2,19 @@
 
 #pragma once
 
-#include "CoreMinimal.h"
-#include "UObject/ObjectMacros.h"
 #include "ActorFactories/ActorFactory.h"
+#include "Math/Quat.h"
+#include "Math/UnrealMathSSE.h"
+#include "UObject/ObjectMacros.h"
+#include "UObject/UObjectGlobals.h"
+
 #include "CustomizableObjectInstanceFactory.generated.h"
 
 class AActor;
-struct FAssetData;
+class FText;
+class UObject;
 class USkeletalMesh;
+struct FAssetData;
 
 UCLASS(MinimalAPI, config=Editor)
 class UCustomizableObjectInstanceFactory : public UActorFactory
diff --git a/Engine/Plugins/Experimental/Mutable/Source/CustomizableObjectEditor/Classes/MuCOE/CustomizableObjectLayout.h b/Engine/Plugins/Experimental/Mutable/Source/CustomizableObjectEditor/Classes/MuCOE/CustomizableObjectLayout.h
--- a/Engine/Plugins/Experimental/Mutable/Source/CustomizableObjectEditor/Classes/MuCOE/CustomizableObjectLayout.h
+++ b/Engine/Plugins/Experimental/Mutable/Source/CustomizableObjectEditor/Classes/MuCOE/CustomizableObjectLayout.h
@@ -2,6 +2,16 @@
 
 #pragma once
 
+#include "Containers/Array.h"
+#include "Containers/UnrealString.h"
+#include "Math/IntPoint.h"
+#include "Math/UnrealMathSSE.h"
+#include "Math/Vector2D.h"
+#include "Misc/Guid.h"
+#include "UObject/Object.h"
+#include "UObject/ObjectPtr.h"
+#include "UObject/UObjectGlobals.h"
+
 #include "CustomizableObjectLayout.generated.h"
 
 UENUM()
diff --git a/Engine/Plugins/Experimental/Mutable/Source/CustomizableObjectEditor/Classes/MuCOE/EdGraphSchema_CustomizableObject.h b/Engine/Plugins/Experimental/Mutable/Source/CustomizableObjectEditor/Classes/MuCOE/EdGraphSchema_CustomizableObject.h
--- a/Engine/Plugins/Experimental/Mutable/Source/CustomizableObjectEditor/Classes/MuCOE/EdGraphSchema_CustomizableObject.h
+++ b/Engine/Plugins/Experimental/Mutable/Source/CustomizableObjectEditor/Classes/MuCOE/EdGraphSchema_CustomizableObject.h
@@ -2,13 +2,29 @@
 
 #pragma once
 
-#include "AssetRegistry/AssetData.h"
+#include "Containers/Array.h"
+#include "Containers/UnrealString.h"
+#include "EdGraph/EdGraphSchema.h"
 #include "EdGraphSchema_K2_Actions.h"
-#include "Runtime/Launch/Resources/Version.h"
+#include "HAL/PlatformMath.h"
+#include "Internationalization/Text.h"
 #include "MaterialTypes.h"
+#include "Math/Color.h"
+#include "Math/Vector2D.h"
+#include "Templates/SharedPointer.h"
+#include "UObject/NameTypes.h"
+#include "UObject/ObjectPtr.h"
+#include "UObject/UObjectGlobals.h"
 
 #include "EdGraphSchema_CustomizableObject.generated.h"
 
+class UEdGraph;
+class UEdGraphNode;
+class UEdGraphPin;
+class UObject;
+struct FAssetData;
+struct FEdGraphPinType;
+
 /** Action to add a node to the graph */
 USTRUCT()
 struct CUSTOMIZABLEOBJECTEDITOR_API FCustomizableObjectSchemaAction_NewNode : public FEdGraphSchemaAction
diff --git a/Engine/Plugins/Experimental/Mutable/Source/CustomizableObjectEditor/Classes/MuCOE/Nodes/CustomizableObjectNode.h b/Engine/Plugins/Experimental/Mutable/Source/CustomizableObjectEditor/Classes/MuCOE/Nodes/CustomizableObjectNode.h
--- a/Engine/Plugins/Experimental/Mutable/Source/CustomizableObjectEditor/Classes/MuCOE/Nodes/CustomizableObjectNode.h
+++ b/Engine/Plugins/Experimental/Mutable/Source/CustomizableObjectEditor/Classes/MuCOE/Nodes/CustomizableObjectNode.h
@@ -2,17 +2,32 @@
 
 #pragma once
 
+#include "Containers/Array.h"
+#include "Containers/ArrayView.h"
+#include "Containers/Map.h"
+#include "Containers/UnrealString.h"
+#include "Delegates/Delegate.h"
+#include "EdGraph/EdGraph.h"
 #include "EdGraph/EdGraphNode.h"
+#include "EdGraph/EdGraphPin.h"
+#include "HAL/Platform.h"
+#include "Misc/Guid.h"
 #include "MuCO/CustomizableObject.h"
-#include "MuCOE/RemapPins/CustomizableObjectNodeRemapPins.h"
-#include "MuCOE/RemapPins/CustomizableObjectNodeRemapPinsByName.h"
-#include "MuCOE/RemapPins/CustomizableObjectNodeRemapPinsByPosition.h"
+#include "Templates/SharedPointer.h"
+#include "UObject/NameTypes.h"
+#include "UObject/Object.h"
+#include "UObject/ObjectPtr.h"
+#include "UObject/UObjectGlobals.h"
 
 #include "CustomizableObjectNode.generated.h"
 
-// Cyclical references
+class FArchive;
+class SWidget;
 class UCustomizableObjectNodeObject;
-
+class UCustomizableObjectNodeRemapPins;
+class UCustomizableObjectNodeRemapPinsByName;
+class UCustomizableObjectNodeRemapPinsByPosition;
+struct FPropertyChangedEvent;
 
 /** Abstract base class for all pin data. */
 UCLASS(Abstract)
diff --git a/Engine/Plugins/Experimental/Mutable/Source/CustomizableObjectEditor/Classes/MuCOE/Nodes/CustomizableObjectNodeAnimationPose.h b/Engine/Plugins/Experimental/Mutable/Source/CustomizableObjectEditor/Classes/MuCOE/Nodes/CustomizableObjectNodeAnimationPose.h
--- a/Engine/Plugins/Experimental/Mutable/Source/CustomizableObjectEditor/Classes/MuCOE/Nodes/CustomizableObjectNodeAnimationPose.h
+++ b/Engine/Plugins/Experimental/Mutable/Source/CustomizableObjectEditor/Classes/MuCOE/Nodes/CustomizableObjectNodeAnimationPose.h
@@ -2,10 +2,24 @@
 
 #pragma once
 
+#include "Containers/Array.h"
+#include "EdGraph/EdGraphNode.h"
+#include "Internationalization/Text.h"
+#include "Math/Color.h"
+#include "Math/Transform.h"
 #include "MuCOE/Nodes/CustomizableObjectNode.h"
-#include "Animation/PoseAsset.h"
+#include "UObject/ObjectPtr.h"
+#include "UObject/UObjectGlobals.h"
+
 #include "CustomizableObjectNodeAnimationPose.generated.h"
 
+class FString;
+class UCustomizableObjectNodeRemapPins;
+class UEdGraphPin;
+class UObject;
+class UPoseAsset;
+struct FPropertyChangedEvent;
+
 
 UCLASS()
 class CUSTOMIZABLEOBJECTEDITOR_API UCustomizableObjectNodeAnimationPose : public UCustomizableObjectNode
diff --git a/Engine/Plugins/Experimental/Mutable/Source/CustomizableObjectEditor/Classes/MuCOE/Nodes/CustomizableObjectNodeColorArithmeticOp.h b/Engine/Plugins/Experimental/Mutable/Source/CustomizableObjectEditor/Classes/MuCOE/Nodes/CustomizableObjectNodeColorArithmeticOp.h
--- a/Engine/Plugins/Experimental/Mutable/Source/CustomizableObjectEditor/Classes/MuCOE/Nodes/CustomizableObjectNodeColorArithmeticOp.h
+++ b/Engine/Plugins/Experimental/Mutable/Source/CustomizableObjectEditor/Classes/MuCOE/Nodes/CustomizableObjectNodeColorArithmeticOp.h
@@ -2,10 +2,20 @@
 
 #pragma once
 
+#include "EdGraph/EdGraphNode.h"
+#include "HAL/Platform.h"
+#include "Internationalization/Text.h"
+#include "Math/Color.h"
 #include "MuCOE/Nodes/CustomizableObjectNode.h"
+#include "UObject/UObjectGlobals.h"
 
 #include "CustomizableObjectNodeColorArithmeticOp.generated.h"
 
+class UCustomizableObjectNodeRemapPins;
+class UEdGraphPin;
+class UObject;
+struct FPropertyChangedEvent;
+
 
 UENUM(BlueprintType)
 enum class EColorArithmeticOperation : uint8
diff --git a/Engine/Plugins/Experimental/Mutable/Source/CustomizableObjectEditor/Classes/MuCOE/Nodes/CustomizableObjectNodeColorConstant.h b/Engine/Plugins/Experimental/Mutable/Source/CustomizableObjectEditor/Classes/MuCOE/Nodes/CustomizableObjectNodeColorConstant.h
--- a/Engine/Plugins/Experimental/Mutable/Source/CustomizableObjectEditor/Classes/MuCOE/Nodes/CustomizableObjectNodeColorConstant.h
+++ b/Engine/Plugins/Experimental/Mutable/Source/CustomizableObjectEditor/Classes/MuCOE/Nodes/CustomizableObjectNodeColorConstant.h
@@ -2,11 +2,30 @@
 
 #pragma once
 
+#include "EdGraph/EdGraphNode.h"
+#include "Input/Reply.h"
+#include "Internationalization/Text.h"
+#include "Math/Color.h"
 #include "MuCOE/Nodes/CustomizableObjectNode.h"
 #include "SGraphNode.h"
+#include "Styling/SlateTypes.h"
+#include "Templates/SharedPointer.h"
+#include "Types/SlateEnums.h"
+#include "UObject/UObjectGlobals.h"
+#include "Widgets/DeclarativeSyntaxSupport.h"
 
 #include "CustomizableObjectNodeColorConstant.generated.h"
 
+class SOverlay;
+class SVerticalBox;
+class UCustomizableObjectNodeColorConstant;
+class UCustomizableObjectNodeRemapPins;
+class UObject;
+struct FGeometry;
+struct FPointerEvent;
+struct FPropertyChangedEvent;
+struct FSlateBrush;
+
 
 enum class ColorChannel
 {
diff --git a/Engine/Plugins/Experimental/Mutable/Source/CustomizableObjectEditor/Classes/MuCOE/Nodes/CustomizableObjectNodeColorFromFloats.h b/Engine/Plugins/Experimental/Mutable/Source/CustomizableObjectEditor/Classes/MuCOE/Nodes/CustomizableObjectNodeColorFromFloats.h
--- a/Engine/Plugins/Experimental/Mutable/Source/CustomizableObjectEditor/Classes/MuCOE/Nodes/CustomizableObjectNodeColorFromFloats.h
+++ b/Engine/Plugins/Experimental/Mutable/Source/CustomizableObjectEditor/Classes/MuCOE/Nodes/CustomizableObjectNodeColorFromFloats.h
@@ -2,10 +2,20 @@
 
 #pragma once
 
+#include "EdGraph/EdGraphNode.h"
+#include "HAL/Platform.h"
+#include "Internationalization/Text.h"
+#include "Math/Color.h"
 #include "MuCOE/Nodes/CustomizableObjectNode.h"
+#include "UObject/UObjectGlobals.h"
 
 #include "CustomizableObjectNodeColorFromFloats.generated.h"
 
+class UCustomizableObjectNodeRemapPins;
+class UEdGraphPin;
+class UObject;
+struct FPropertyChangedEvent;
+
 
 UCLASS()
 class CUSTOMIZABLEOBJECTEDITOR_API UCustomizableObjectNodeColorFromFloats : public UCustomizableObjectNode
diff --git a/Engine/Plugins/Experimental/Mutable/Source/CustomizableObjectEditor/Classes/MuCOE/Nodes/CustomizableObjectNodeColorParameter.h b/Engine/Plugins/Experimental/Mutable/Source/CustomizableObjectEditor/Classes/MuCOE/Nodes/CustomizableObjectNodeColorParameter.h
--- a/Engine/Plugins/Experimental/Mutable/Source/CustomizableObjectEditor/Classes/MuCOE/Nodes/CustomizableObjectNodeColorParameter.h
+++ b/Engine/Plugins/Experimental/Mutable/Source/CustomizableObjectEditor/Classes/MuCOE/Nodes/CustomizableObjectNodeColorParameter.h
@@ -2,10 +2,19 @@
 
 #pragma once
 
+#include "Containers/UnrealString.h"
+#include "EdGraph/EdGraphNode.h"
+#include "Internationalization/Text.h"
+#include "Math/Color.h"
+#include "MuCO/CustomizableObjectUIData.h"
 #include "MuCOE/Nodes/CustomizableObjectNode.h"
+#include "UObject/UObjectGlobals.h"
 
 #include "CustomizableObjectNodeColorParameter.generated.h"
 
+class UCustomizableObjectNodeRemapPins;
+class UObject;
+
 
 UCLASS()
 class CUSTOMIZABLEOBJECTEDITOR_API UCustomizableObjectNodeColorParameter : public UCustomizableObjectNode
diff --git a/Engine/Plugins/Experimental/Mutable/Source/CustomizableObjectEditor/Classes/MuCOE/Nodes/CustomizableObjectNodeColorSwitch.h b/Engine/Plugins/Experimental/Mutable/Source/CustomizableObjectEditor/Classes/MuCOE/Nodes/CustomizableObjectNodeColorSwitch.h
--- a/Engine/Plugins/Experimental/Mutable/Source/CustomizableObjectEditor/Classes/MuCOE/Nodes/CustomizableObjectNodeColorSwitch.h
+++ b/Engine/Plugins/Experimental/Mutable/Source/CustomizableObjectEditor/Classes/MuCOE/Nodes/CustomizableObjectNodeColorSwitch.h
@@ -2,12 +2,16 @@
 
 #pragma once
 
+#include "Containers/UnrealString.h"
 #include "MuCOE/Nodes/CustomizableObjectNodeSwitchBase.h"
-
-#include "EdGraph/EdGraphPin.h"
+#include "UObject/NameTypes.h"
+#include "UObject/UObjectGlobals.h"
 
 #include "CustomizableObjectNodeColorSwitch.generated.h"
 
+class FArchive;
+class UObject;
+
 
 UCLASS()
 class CUSTOMIZABLEOBJECTEDITOR_API UCustomizableObjectNodeColorSwitch : public UCustomizableObjectNodeSwitchBase
diff --git a/Engine/Plugins/Experimental/Mutable/Source/CustomizableObjectEditor/Classes/MuCOE/Nodes/CustomizableObjectNodeColorVariation.h b/Engine/Plugins/Experimental/Mutable/Source/CustomizableObjectEditor/Classes/MuCOE/Nodes/CustomizableObjectNodeColorVariation.h
--- a/Engine/Plugins/Experimental/Mutable/Source/CustomizableObjectEditor/Classes/MuCOE/Nodes/CustomizableObjectNodeColorVariation.h
+++ b/Engine/Plugins/Experimental/Mutable/Source/CustomizableObjectEditor/Classes/MuCOE/Nodes/CustomizableObjectNodeColorVariation.h
@@ -2,10 +2,23 @@
 
 #pragma once
 
+#include "Containers/Array.h"
+#include "Containers/Map.h"
+#include "Containers/UnrealString.h"
+#include "EdGraph/EdGraphNode.h"
+#include "HAL/Platform.h"
+#include "Internationalization/Text.h"
+#include "Math/Color.h"
 #include "MuCOE/Nodes/CustomizableObjectNode.h"
+#include "UObject/UObjectGlobals.h"
 
 #include "CustomizableObjectNodeColorVariation.generated.h"
 
+class UCustomizableObjectNodeRemapPins;
+class UEdGraphPin;
+class UObject;
+struct FPropertyChangedEvent;
+
 
 USTRUCT()
 struct CUSTOMIZABLEOBJECTEDITOR_API FCustomizableObjectColorVariation
diff --git a/Engine/Plugins/Experimental/Mutable/Source/CustomizableObjectEditor/Classes/MuCOE/Nodes/CustomizableObjectNodeCopyMaterial.h b/Engine/Plugins/Experimental/Mutable/Source/CustomizableObjectEditor/Classes/MuCOE/Nodes/CustomizableObjectNodeCopyMaterial.h
--- a/Engine/Plugins/Experimental/Mutable/Source/CustomizableObjectEditor/Classes/MuCOE/Nodes/CustomizableObjectNodeCopyMaterial.h
+++ b/Engine/Plugins/Experimental/Mutable/Source/CustomizableObjectEditor/Classes/MuCOE/Nodes/CustomizableObjectNodeCopyMaterial.h
@@ -2,11 +2,18 @@
 
 #pragma once
 
+#include "EdGraph/EdGraphNode.h"
+#include "Internationalization/Text.h"
 #include "MuCOE/Nodes/CustomizableObjectNodeMaterial.h"
-#include "MuCOE/Nodes/CustomizableObjectNodeSkeletalMesh.h"
+#include "UObject/UObjectGlobals.h"
 
 #include "CustomizableObjectNodeCopyMaterial.generated.h"
 
+class UCustomizableObjectNodeRemapPins;
+class UCustomizableObjectNodeSkeletalMesh;
+class UEdGraphPin;
+class UObject;
+
 /**
  * Copy Material node. Duplicates a Material Node. Duplicates all Material node input pins and properties except for the Mesh input pin.
  * A new Mesh has to be defined through the new Mesh input pin.
diff --git a/Engine/Plugins/Experimental/Mutable/Source/CustomizableObjectEditor/Classes/MuCOE/Nodes/CustomizableObjectNodeCurve.h b/Engine/Plugins/Experimental/Mutable/Source/CustomizableObjectEditor/Classes/MuCOE/Nodes/CustomizableObjectNodeCurve.h
--- a/Engine/Plugins/Experimental/Mutable/Source/CustomizableObjectEditor/Classes/MuCOE/Nodes/CustomizableObjectNodeCurve.h
+++ b/Engine/Plugins/Experimental/Mutable/Source/CustomizableObjectEditor/Classes/MuCOE/Nodes/CustomizableObjectNodeCurve.h
@@ -2,11 +2,23 @@
 
 #pragma once
 
+#include "EdGraph/EdGraphNode.h"
+#include "HAL/Platform.h"
+#include "Internationalization/Text.h"
+#include "Math/Color.h"
 #include "MuCOE/Nodes/CustomizableObjectNode.h"
-#include "Curves/CurveBase.h"
+#include "UObject/ObjectPtr.h"
+#include "UObject/UObjectGlobals.h"
 
 #include "CustomizableObjectNodeCurve.generated.h"
 
+class FArchive;
+class UCurveBase;
+class UCustomizableObjectNodeRemapPins;
+class UEdGraphPin;
+class UObject;
+struct FPropertyChangedEvent;
+
 
 UCLASS()
 class CUSTOMIZABLEOBJECTEDITOR_API UCustomizableObjectNodeCurve : public UCustomizableObjectNode
diff --git a/Engine/Plugins/Experimental/Mutable/Source/CustomizableObjectEditor/Classes/MuCOE/Nodes/CustomizableObjectNodeEditMaterial.h b/Engine/Plugins/Experimental/Mutable/Source/CustomizableObjectEditor/Classes/MuCOE/Nodes/CustomizableObjectNodeEditMaterial.h
--- a/Engine/Plugins/Experimental/Mutable/Source/CustomizableObjectEditor/Classes/MuCOE/Nodes/CustomizableObjectNodeEditMaterial.h
+++ b/Engine/Plugins/Experimental/Mutable/Source/CustomizableObjectEditor/Classes/MuCOE/Nodes/CustomizableObjectNodeEditMaterial.h
@@ -2,12 +2,29 @@
 
 #pragma once
 
+#include "Containers/Array.h"
+#include "Containers/Map.h"
+#include "Containers/UnrealString.h"
+#include "EdGraph/EdGraphNode.h"
+#include "EdGraph/EdGraphPin.h"
+#include "HAL/Platform.h"
+#include "Internationalization/Text.h"
+#include "Math/Color.h"
+#include "Misc/Guid.h"
+#include "MuCOE/Nodes/CustomizableObjectNode.h"
 #include "MuCOE/Nodes/CustomizableObjectNodeEditLayoutBlocks.h"
 #include "MuCOE/Nodes/CustomizableObjectNodeUseMaterial.h"
-#include "MuCOE/CustomizableObjectEditor_Deprecated.h"
+#include "UObject/UObjectGlobals.h"
 
 #include "CustomizableObjectNodeEditMaterial.generated.h"
 
+class FArchive;
+class FCustomizableObjectNodeParentedMaterial;
+class UCustomizableObject;
+class UCustomizableObjectNodeRemapPins;
+class UObject;
+struct FCustomizableObjectNodeEditMaterialImage;
+
 
 /** Additional data for the Parent Texture Parameter to edit pin. */
 UCLASS()
diff --git a/Engine/Plugins/Experimental/Mutable/Source/CustomizableObjectEditor/Classes/MuCOE/Nodes/CustomizableObjectNodeEditMaterialBase.h b/Engine/Plugins/Experimental/Mutable/Source/CustomizableObjectEditor/Classes/MuCOE/Nodes/CustomizableObjectNodeEditMaterialBase.h
--- a/Engine/Plugins/Experimental/Mutable/Source/CustomizableObjectEditor/Classes/MuCOE/Nodes/CustomizableObjectNodeEditMaterialBase.h
+++ b/Engine/Plugins/Experimental/Mutable/Source/CustomizableObjectEditor/Classes/MuCOE/Nodes/CustomizableObjectNodeEditMaterialBase.h
@@ -2,11 +2,19 @@
 
 #pragma once
 
+#include "HAL/Platform.h"
+#include "Misc/Guid.h"
 #include "MuCOE/Nodes/CustomizableObjectNode.h"
 #include "MuCOE/Nodes/CustomizableObjectNodeParentedMaterial.h"
+#include "UObject/ObjectPtr.h"
+#include "UObject/UObjectGlobals.h"
 
 #include "CustomizableObjectNodeEditMaterialBase.generated.h"
 
+class UCustomizableObject;
+class UEdGraphPin;
+class UObject;
+
 
 UCLASS(Abstract)
 class CUSTOMIZABLEOBJECTEDITOR_API UCustomizableObjectNodeEditMaterialBase : public UCustomizableObjectNode, public FCustomizableObjectNodeParentedMaterial
diff --git a/Engine/Plugins/Experimental/Mutable/Source/CustomizableObjectEditor/Classes/MuCOE/Nodes/CustomizableObjectNodeEnumParameter.h b/Engine/Plugins/Experimental/Mutable/Source/CustomizableObjectEditor/Classes/MuCOE/Nodes/CustomizableObjectNodeEnumParameter.h
--- a/Engine/Plugins/Experimental/Mutable/Source/CustomizableObjectEditor/Classes/MuCOE/Nodes/CustomizableObjectNodeEnumParameter.h
+++ b/Engine/Plugins/Experimental/Mutable/Source/CustomizableObjectEditor/Classes/MuCOE/Nodes/CustomizableObjectNodeEnumParameter.h
@@ -2,10 +2,22 @@
 
 #pragma once
 
+#include "Containers/Array.h"
+#include "Containers/UnrealString.h"
+#include "EdGraph/EdGraphNode.h"
+#include "HAL/Platform.h"
+#include "Internationalization/Text.h"
+#include "Math/Color.h"
+#include "MuCO/CustomizableObjectUIData.h"
 #include "MuCOE/Nodes/CustomizableObjectNode.h"
+#include "UObject/UObjectGlobals.h"
 
 #include "CustomizableObjectNodeEnumParameter.generated.h"
 
+class UCustomizableObjectNodeRemapPins;
+class UObject;
+struct FPropertyChangedEvent;
+
 
 USTRUCT()
 struct CUSTOMIZABLEOBJECTEDITOR_API FCustomizableObjectNodeEnumValue
diff --git a/Engine/Plugins/Experimental/Mutable/Source/CustomizableObjectEditor/Classes/MuCOE/Nodes/CustomizableObjectNodeExposePin.h b/Engine/Plugins/Experimental/Mutable/Source/CustomizableObjectEditor/Classes/MuCOE/Nodes/CustomizableObjectNodeExposePin.h
--- a/Engine/Plugins/Experimental/Mutable/Source/CustomizableObjectEditor/Classes/MuCOE/Nodes/CustomizableObjectNodeExposePin.h
+++ b/Engine/Plugins/Experimental/Mutable/Source/CustomizableObjectEditor/Classes/MuCOE/Nodes/CustomizableObjectNodeExposePin.h
@@ -2,10 +2,23 @@
 
 #pragma once
 
+#include "Containers/UnrealString.h"
+#include "Delegates/Delegate.h"
+#include "EdGraph/EdGraphNode.h"
+#include "HAL/Platform.h"
+#include "Internationalization/Text.h"
+#include "Math/Color.h"
 #include "MuCOE/Nodes/CustomizableObjectNode.h"
+#include "UObject/NameTypes.h"
+#include "UObject/UObjectGlobals.h"
 
 #include "CustomizableObjectNodeExposePin.generated.h"
 
+class UCustomizableObjectNodeRemapPins;
+class UEdGraphPin;
+class UObject;
+struct FPropertyChangedEvent;
+
 
 DECLARE_MULTICAST_DELEGATE(FOnNameChangedDelegate);
 
diff --git a/Engine/Plugins/Experimental/Mutable/Source/CustomizableObjectEditor/Classes/MuCOE/Nodes/CustomizableObjectNodeExtendMaterial.h b/Engine/Plugins/Experimental/Mutable/Source/CustomizableObjectEditor/Classes/MuCOE/Nodes/CustomizableObjectNodeExtendMaterial.h
--- a/Engine/Plugins/Experimental/Mutable/Source/CustomizableObjectEditor/Classes/MuCOE/Nodes/CustomizableObjectNodeExtendMaterial.h
+++ b/Engine/Plugins/Experimental/Mutable/Source/CustomizableObjectEditor/Classes/MuCOE/Nodes/CustomizableObjectNodeExtendMaterial.h
@@ -2,13 +2,30 @@
 
 #pragma once
 
+#include "Containers/Array.h"
+#include "Containers/Map.h"
+#include "Containers/UnrealString.h"
+#include "EdGraph/EdGraphNode.h"
+#include "Internationalization/Text.h"
+#include "Math/Color.h"
+#include "Misc/Guid.h"
 #include "MuCOE/Nodes/CustomizableObjectNodeMaterialBase.h"
 #include "MuCOE/Nodes/CustomizableObjectNodeParentedMaterial.h"
 #include "MuCOE/Nodes/CustomizableObjectNodeUseMaterial.h"
-#include "MuCOE/CustomizableObjectEditor_Deprecated.h"
+#include "UObject/ObjectPtr.h"
+#include "UObject/UObjectGlobals.h"
 
 #include "CustomizableObjectNodeExtendMaterial.generated.h"
 
+class FArchive;
+class UCustomizableObject;
+class UCustomizableObjectNode;
+class UCustomizableObjectNodeRemapPins;
+class UEdGraphPin;
+class UObject;
+struct FCustomizableObjectNodeExtendMaterialImage;
+struct FEdGraphPinReference;
+
 
 UCLASS()
 class CUSTOMIZABLEOBJECTEDITOR_API UCustomizableObjectNodeExtendMaterial :
diff --git a/Engine/Plugins/Experimental/Mutable/Source/CustomizableObjectEditor/Classes/MuCOE/Nodes/CustomizableObjectNodeExternalPin.h b/Engine/Plugins/Experimental/Mutable/Source/CustomizableObjectEditor/Classes/MuCOE/Nodes/CustomizableObjectNodeExternalPin.h
--- a/Engine/Plugins/Experimental/Mutable/Source/CustomizableObjectEditor/Classes/MuCOE/Nodes/CustomizableObjectNodeExternalPin.h
+++ b/Engine/Plugins/Experimental/Mutable/Source/CustomizableObjectEditor/Classes/MuCOE/Nodes/CustomizableObjectNodeExternalPin.h
@@ -2,11 +2,28 @@
 
 #pragma once
 
+#include "Containers/Array.h"
+#include "Containers/Map.h"
+#include "Delegates/IDelegateInstance.h"
+#include "EdGraph/EdGraphNode.h"
+#include "Internationalization/Text.h"
+#include "Math/Color.h"
+#include "Misc/Guid.h"
 #include "MuCOE/Nodes/CustomizableObjectNode.h"
-#include "MuCOE/Nodes/CustomizableObjectNodeExposePin.h"
+#include "MuCOE/RemapPins/CustomizableObjectNodeRemapPinsByName.h"
+#include "UObject/NameTypes.h"
+#include "UObject/ObjectPtr.h"
+#include "UObject/UObjectGlobals.h"
 
 #include "CustomizableObjectNodeExternalPin.generated.h"
 
+class FArchive;
+class UCustomizableObject;
+class UCustomizableObjectNodeExposePin;
+class UCustomizableObjectNodeRemapPins;
+class UEdGraphPin;
+class UObject;
+
 
 /**
  * Custom remap pins by name action.
diff --git a/Engine/Plugins/Experimental/Mutable/Source/CustomizableObjectEditor/Classes/MuCOE/Nodes/CustomizableObjectNodeFloatConstant.h b/Engine/Plugins/Experimental/Mutable/Source/CustomizableObjectEditor/Classes/MuCOE/Nodes/CustomizableObjectNodeFloatConstant.h
--- a/Engine/Plugins/Experimental/Mutable/Source/CustomizableObjectEditor/Classes/MuCOE/Nodes/CustomizableObjectNodeFloatConstant.h
+++ b/Engine/Plugins/Experimental/Mutable/Source/CustomizableObjectEditor/Classes/MuCOE/Nodes/CustomizableObjectNodeFloatConstant.h
@@ -2,11 +2,27 @@
 
 #pragma once
 
+#include "EdGraph/EdGraphNode.h"
+#include "Internationalization/Text.h"
+#include "Math/Color.h"
 #include "MuCOE/Nodes/CustomizableObjectNode.h"
 #include "SGraphNode.h"
+#include "Styling/SlateTypes.h"
+#include "Templates/SharedPointer.h"
+#include "Types/SlateEnums.h"
+#include "UObject/UObjectGlobals.h"
+#include "Widgets/DeclarativeSyntaxSupport.h"
 
 #include "CustomizableObjectNodeFloatConstant.generated.h"
 
+class SOverlay;
+class SVerticalBox;
+class UCustomizableObjectNodeFloatConstant;
+class UCustomizableObjectNodeRemapPins;
+class UObject;
+struct FPropertyChangedEvent;
+struct FSlateBrush;
+
 // Class create an input float in the NodeFloatConstant
 class SGraphNodeFloatConstant : public SGraphNode
 {
diff --git a/Engine/Plugins/Experimental/Mutable/Source/CustomizableObjectEditor/Classes/MuCOE/Nodes/CustomizableObjectNodeFloatParameter.h b/Engine/Plugins/Experimental/Mutable/Source/CustomizableObjectEditor/Classes/MuCOE/Nodes/CustomizableObjectNodeFloatParameter.h
--- a/Engine/Plugins/Experimental/Mutable/Source/CustomizableObjectEditor/Classes/MuCOE/Nodes/CustomizableObjectNodeFloatParameter.h
+++ b/Engine/Plugins/Experimental/Mutable/Source/CustomizableObjectEditor/Classes/MuCOE/Nodes/CustomizableObjectNodeFloatParameter.h
@@ -2,9 +2,25 @@
 
 #pragma once
 
+#include "Containers/Array.h"
+#include "Containers/UnrealString.h"
+#include "EdGraph/EdGraphNode.h"
+#include "HAL/Platform.h"
+#include "Internationalization/Text.h"
+#include "Math/Color.h"
+#include "MuCO/CustomizableObjectUIData.h"
 #include "MuCOE/Nodes/CustomizableObjectNode.h"
+#include "UObject/UObjectGlobals.h"
+
 #include "CustomizableObjectNodeFloatParameter.generated.h"
 
+class FArchive;
+class UCustomizableObjectNodeRemapPins;
+class UEdGraphPin;
+class UObject;
+struct FEdGraphPinReference;
+struct FPropertyChangedEvent;
+
 
 USTRUCT()
 struct CUSTOMIZABLEOBJECTEDITOR_API FCustomizableObjectNodeFloatDescription
diff --git a/Engine/Plugins/Experimental/Mutable/Source/CustomizableObjectEditor/Classes/MuCOE/Nodes/CustomizableObjectNodeFloatSwitch.h b/Engine/Plugins/Experimental/Mutable/Source/CustomizableObjectEditor/Classes/MuCOE/Nodes/CustomizableObjectNodeFloatSwitch.h
--- a/Engine/Plugins/Experimental/Mutable/Source/CustomizableObjectEditor/Classes/MuCOE/Nodes/CustomizableObjectNodeFloatSwitch.h
+++ b/Engine/Plugins/Experimental/Mutable/Source/CustomizableObjectEditor/Classes/MuCOE/Nodes/CustomizableObjectNodeFloatSwitch.h
@@ -2,12 +2,16 @@
 
 #pragma once
 
+#include "Containers/UnrealString.h"
 #include "MuCOE/Nodes/CustomizableObjectNodeSwitchBase.h"
-
-#include "EdGraph/EdGraphPin.h"
+#include "UObject/NameTypes.h"
+#include "UObject/UObjectGlobals.h"
 
 #include "CustomizableObjectNodeFloatSwitch.generated.h"
 
+class FArchive;
+class UObject;
+
 
 UCLASS()
 class CUSTOMIZABLEOBJECTEDITOR_API UCustomizableObjectNodeFloatSwitch : public UCustomizableObjectNodeSwitchBase
diff --git a/Engine/Plugins/Experimental/Mutable/Source/CustomizableObjectEditor/Classes/MuCOE/Nodes/CustomizableObjectNodeFloatVariation.h b/Engine/Plugins/Experimental/Mutable/Source/CustomizableObjectEditor/Classes/MuCOE/Nodes/CustomizableObjectNodeFloatVariation.h
--- a/Engine/Plugins/Experimental/Mutable/Source/CustomizableObjectEditor/Classes/MuCOE/Nodes/CustomizableObjectNodeFloatVariation.h
+++ b/Engine/Plugins/Experimental/Mutable/Source/CustomizableObjectEditor/Classes/MuCOE/Nodes/CustomizableObjectNodeFloatVariation.h
@@ -2,10 +2,22 @@
 
 #pragma once
 
+#include "Containers/Array.h"
+#include "Containers/UnrealString.h"
+#include "EdGraph/EdGraphNode.h"
+#include "HAL/Platform.h"
+#include "Internationalization/Text.h"
+#include "Math/Color.h"
 #include "MuCOE/Nodes/CustomizableObjectNode.h"
+#include "UObject/UObjectGlobals.h"
 
 #include "CustomizableObjectNodeFloatVariation.generated.h"
 
+class UCustomizableObjectNodeRemapPins;
+class UEdGraphPin;
+class UObject;
+struct FPropertyChangedEvent;
+
 
 USTRUCT()
 struct CUSTOMIZABLEOBJECTEDITOR_API FCustomizableObjectFloatVariation
diff --git a/Engine/Plugins/Experimental/Mutable/Source/CustomizableObjectEditor/Classes/MuCOE/Nodes/CustomizableObjectNodeGroupProjectorParameter.h b/Engine/Plugins/Experimental/Mutable/Source/CustomizableObjectEditor/Classes/MuCOE/Nodes/CustomizableObjectNodeGroupProjectorParameter.h
--- a/Engine/Plugins/Experimental/Mutable/Source/CustomizableObjectEditor/Classes/MuCOE/Nodes/CustomizableObjectNodeGroupProjectorParameter.h
+++ b/Engine/Plugins/Experimental/Mutable/Source/CustomizableObjectEditor/Classes/MuCOE/Nodes/CustomizableObjectNodeGroupProjectorParameter.h
@@ -2,13 +2,26 @@
 
 #pragma once
 
-#include "MuCO/CustomizableObjectInstance.h"
+#include "Containers/Array.h"
+#include "Containers/UnrealString.h"
+#include "EdGraph/EdGraphNode.h"
+#include "HAL/Platform.h"
+#include "Internationalization/Text.h"
+#include "Math/Color.h"
 #include "MuCOE/Nodes/CustomizableObjectNodeProjectorParameter.h"
-#include "Engine/DataTable.h"
-#include "Animation/PoseAsset.h"
+#include "UObject/NameTypes.h"
+#include "UObject/ObjectPtr.h"
+#include "UObject/UObjectGlobals.h"
 
 #include "CustomizableObjectNodeGroupProjectorParameter.generated.h"
 
+class FArchive;
+class UCustomizableObjectNodeRemapPins;
+class UDataTable;
+class UObject;
+class UPoseAsset;
+class UTexture2D;
+
 USTRUCT()
 struct FGroupProjectorParameterImage
 {
diff --git a/Engine/Plugins/Experimental/Mutable/Source/CustomizableObjectEditor/Classes/MuCOE/Nodes/CustomizableObjectNodeLayoutBlocks.h b/Engine/Plugins/Experimental/Mutable/Source/CustomizableObjectEditor/Classes/MuCOE/Nodes/CustomizableObjectNodeLayoutBlocks.h
--- a/Engine/Plugins/Experimental/Mutable/Source/CustomizableObjectEditor/Classes/MuCOE/Nodes/CustomizableObjectNodeLayoutBlocks.h
+++ b/Engine/Plugins/Experimental/Mutable/Source/CustomizableObjectEditor/Classes/MuCOE/Nodes/CustomizableObjectNodeLayoutBlocks.h
@@ -2,11 +2,26 @@
 
 #pragma once
 
-#include "MuCOE/Nodes/CustomizableObjectNode.h"
+#include "Containers/Array.h"
+#include "EdGraph/EdGraphNode.h"
+#include "Internationalization/Text.h"
+#include "Math/Color.h"
+#include "Math/IntPoint.h"
+#include "Math/UnrealMathSSE.h"
 #include "MuCOE/CustomizableObjectLayout.h"
+#include "MuCOE/Nodes/CustomizableObjectNode.h"
+#include "UObject/ObjectPtr.h"
+#include "UObject/UObjectGlobals.h"
+#include "UObject/WeakObjectPtrTemplates.h"
 
 #include "CustomizableObjectNodeLayoutBlocks.generated.h"
 
+class FArchive;
+class UCustomizableObjectNodeRemapPins;
+class UEdGraphPin;
+class UObject;
+struct FPropertyChangedEvent;
+
 
 UCLASS()
 class CUSTOMIZABLEOBJECTEDITOR_API UCustomizableObjectNodeLayoutBlocks : public UCustomizableObjectNode
diff --git a/Engine/Plugins/Experimental/Mutable/Source/CustomizableObjectEditor/Classes/MuCOE/Nodes/CustomizableObjectNodeMaterial.h b/Engine/Plugins/Experimental/Mutable/Source/CustomizableObjectEditor/Classes/MuCOE/Nodes/CustomizableObjectNodeMaterial.h
--- a/Engine/Plugins/Experimental/Mutable/Source/CustomizableObjectEditor/Classes/MuCOE/Nodes/CustomizableObjectNodeMaterial.h
+++ b/Engine/Plugins/Experimental/Mutable/Source/CustomizableObjectEditor/Classes/MuCOE/Nodes/CustomizableObjectNodeMaterial.h
@@ -2,13 +2,44 @@
 
 #pragma once
 
-#include "Engine/Texture2D.h"
-#include "MuCOE/Nodes/CustomizableObjectNodeMaterialBase.h"
-#include "MuCOE/EdGraphSchema_CustomizableObject.h"
+#include "Containers/Array.h"
+#include "Containers/Map.h"
+#include "Containers/UnrealString.h"
+#include "Delegates/Delegate.h"
+#include "EdGraph/EdGraphNode.h"
+#include "HAL/Platform.h"
+#include "Internationalization/Text.h"
+#include "MaterialTypes.h"
+#include "Math/Color.h"
+#include "Misc/Guid.h"
 #include "MuCOE/CustomizableObjectEditor_Deprecated.h"
+#include "MuCOE/Nodes/CustomizableObjectNode.h"
+#include "MuCOE/Nodes/CustomizableObjectNodeMaterialBase.h"
+#include "MuCOE/RemapPins/CustomizableObjectNodeRemapPinsByName.h"
+#include "Templates/SharedPointer.h"
+#include "UObject/NameTypes.h"
+#include "UObject/ObjectPtr.h"
+#include "UObject/UObjectGlobals.h"
+#include "UObject/WeakObjectPtrTemplates.h"
 
 #include "CustomizableObjectNodeMaterial.generated.h"
 
+class FArchive;
+class SGraphNode;
+class SWidget;
+class UCustomizableObjectNodeRemapPins;
+class UEdGraphPin;
+class UMaterial;
+class UMaterialInterface;
+class UObject;
+class UTexture2D;
+struct FCustomizableObjectNodeMaterialImage;
+struct FCustomizableObjectNodeMaterialScalar;
+struct FCustomizableObjectNodeMaterialVector;
+struct FEdGraphPinReference;
+struct FFrame;
+struct FPropertyChangedEvent;
+
 
 DECLARE_MULTICAST_DELEGATE(FPostImagePinModeChangedDelegate)
 
diff --git a/Engine/Plugins/Experimental/Mutable/Source/CustomizableObjectEditor/Classes/MuCOE/Nodes/CustomizableObjectNodeMaterialBase.h b/Engine/Plugins/Experimental/Mutable/Source/CustomizableObjectEditor/Classes/MuCOE/Nodes/CustomizableObjectNodeMaterialBase.h
--- a/Engine/Plugins/Experimental/Mutable/Source/CustomizableObjectEditor/Classes/MuCOE/Nodes/CustomizableObjectNodeMaterialBase.h
+++ b/Engine/Plugins/Experimental/Mutable/Source/CustomizableObjectEditor/Classes/MuCOE/Nodes/CustomizableObjectNodeMaterialBase.h
@@ -2,9 +2,16 @@
 
 #pragma once
 
+#include "Containers/Array.h"
+#include "Misc/AssertionMacros.h"
 #include "MuCOE/Nodes/CustomizableObjectNode.h"
+#include "UObject/UObjectGlobals.h"
+
 #include "CustomizableObjectNodeMaterialBase.generated.h"
 
+class UEdGraphPin;
+class UObject;
+
 
 
 UCLASS(abstract)
diff --git a/Engine/Plugins/Experimental/Mutable/Source/CustomizableObjectEditor/Classes/MuCOE/Nodes/CustomizableObjectNodeMaterialVariation.h b/Engine/Plugins/Experimental/Mutable/Source/CustomizableObjectEditor/Classes/MuCOE/Nodes/CustomizableObjectNodeMaterialVariation.h
--- a/Engine/Plugins/Experimental/Mutable/Source/CustomizableObjectEditor/Classes/MuCOE/Nodes/CustomizableObjectNodeMaterialVariation.h
+++ b/Engine/Plugins/Experimental/Mutable/Source/CustomizableObjectEditor/Classes/MuCOE/Nodes/CustomizableObjectNodeMaterialVariation.h
@@ -2,10 +2,22 @@
 
 #pragma once
 
+#include "Containers/Array.h"
+#include "Containers/UnrealString.h"
+#include "EdGraph/EdGraphNode.h"
+#include "HAL/Platform.h"
+#include "Internationalization/Text.h"
+#include "Math/Color.h"
 #include "MuCOE/Nodes/CustomizableObjectNodeMaterialBase.h"
+#include "UObject/UObjectGlobals.h"
 
 #include "CustomizableObjectNodeMaterialVariation.generated.h"
 
+class UCustomizableObjectNodeRemapPins;
+class UEdGraphPin;
+class UObject;
+struct FPropertyChangedEvent;
+
 
 USTRUCT()
 struct CUSTOMIZABLEOBJECTEDITOR_API FCustomizableObjectMaterialVariation
diff --git a/Engine/Plugins/Experimental/Mutable/Source/CustomizableObjectEditor/Classes/MuCOE/Nodes/CustomizableObjectNodeMesh.h b/Engine/Plugins/Experimental/Mutable/Source/CustomizableObjectEditor/Classes/MuCOE/Nodes/CustomizableObjectNodeMesh.h
--- a/Engine/Plugins/Experimental/Mutable/Source/CustomizableObjectEditor/Classes/MuCOE/Nodes/CustomizableObjectNodeMesh.h
+++ b/Engine/Plugins/Experimental/Mutable/Source/CustomizableObjectEditor/Classes/MuCOE/Nodes/CustomizableObjectNodeMesh.h
@@ -2,11 +2,20 @@
 
 #pragma once
 
-#include "MuCOE/CustomizableObjectLayout.h"
+#include "Containers/Array.h"
+#include "CoreTypes.h"
+#include "Math/Vector2D.h"
+#include "Misc/AssertionMacros.h"
 #include "MuCOE/Nodes/CustomizableObjectNode.h"
+#include "UObject/UObjectGlobals.h"
 
 #include "CustomizableObjectNodeMesh.generated.h"
 
+class UCustomizableObjectLayout;
+class UEdGraphPin;
+class UObject;
+class UTexture2D;
+
 
 UCLASS(abstract)
 class CUSTOMIZABLEOBJECTEDITOR_API UCustomizableObjectNodeMesh : public UCustomizableObjectNode
diff --git a/Engine/Plugins/Experimental/Mutable/Source/CustomizableObjectEditor/Classes/MuCOE/Nodes/CustomizableObjectNodeMeshClipDeform.h b/Engine/Plugins/Experimental/Mutable/Source/CustomizableObjectEditor/Classes/MuCOE/Nodes/CustomizableObjectNodeMeshClipDeform.h
--- a/Engine/Plugins/Experimental/Mutable/Source/CustomizableObjectEditor/Classes/MuCOE/Nodes/CustomizableObjectNodeMeshClipDeform.h
+++ b/Engine/Plugins/Experimental/Mutable/Source/CustomizableObjectEditor/Classes/MuCOE/Nodes/CustomizableObjectNodeMeshClipDeform.h
@@ -2,10 +2,21 @@
 
 #pragma once
 
+#include "Containers/Array.h"
+#include "Containers/UnrealString.h"
+#include "EdGraph/EdGraphNode.h"
+#include "HAL/Platform.h"
+#include "Internationalization/Text.h"
+#include "Math/Color.h"
 #include "MuCOE/Nodes/CustomizableObjectNodeModifierBase.h"
+#include "UObject/UObjectGlobals.h"
 
 #include "CustomizableObjectNodeMeshClipDeform.generated.h"
 
+class UCustomizableObjectNodeRemapPins;
+class UEdGraphPin;
+class UObject;
+
 UENUM()
 enum class EShapeBindingMethod : uint32
 {
diff --git a/Engine/Plugins/Experimental/Mutable/Source/CustomizableObjectEditor/Classes/MuCOE/Nodes/CustomizableObjectNodeMeshClipMorph.h b/Engine/Plugins/Experimental/Mutable/Source/CustomizableObjectEditor/Classes/MuCOE/Nodes/CustomizableObjectNodeMeshClipMorph.h
--- a/Engine/Plugins/Experimental/Mutable/Source/CustomizableObjectEditor/Classes/MuCOE/Nodes/CustomizableObjectNodeMeshClipMorph.h
+++ b/Engine/Plugins/Experimental/Mutable/Source/CustomizableObjectEditor/Classes/MuCOE/Nodes/CustomizableObjectNodeMeshClipMorph.h
@@ -2,10 +2,25 @@
 
 #pragma once
 
+#include "Containers/Array.h"
+#include "Containers/UnrealString.h"
+#include "EdGraph/EdGraphNode.h"
+#include "Internationalization/Text.h"
+#include "Math/Color.h"
+#include "Math/UnrealMathSSE.h"
+#include "Math/Vector.h"
 #include "MuCOE/Nodes/CustomizableObjectNodeModifierBase.h"
+#include "UObject/NameTypes.h"
+#include "UObject/UObjectGlobals.h"
 
 #include "CustomizableObjectNodeMeshClipMorph.generated.h"
 
+class FArchive;
+class UCustomizableObjectNodeRemapPins;
+class UEdGraphPin;
+class UObject;
+struct FPropertyChangedEvent;
+
 UCLASS()
 class CUSTOMIZABLEOBJECTEDITOR_API UCustomizableObjectNodeMeshClipMorph : public UCustomizableObjectNodeModifierBase
 {
diff --git a/Engine/Plugins/Experimental/Mutable/Source/CustomizableObjectEditor/Classes/MuCOE/Nodes/CustomizableObjectNodeMeshClipWithMesh.h b/Engine/Plugins/Experimental/Mutable/Source/CustomizableObjectEditor/Classes/MuCOE/Nodes/CustomizableObjectNodeMeshClipWithMesh.h
--- a/Engine/Plugins/Experimental/Mutable/Source/CustomizableObjectEditor/Classes/MuCOE/Nodes/CustomizableObjectNodeMeshClipWithMesh.h
+++ b/Engine/Plugins/Experimental/Mutable/Source/CustomizableObjectEditor/Classes/MuCOE/Nodes/CustomizableObjectNodeMeshClipWithMesh.h
@@ -2,11 +2,25 @@
 
 #pragma once
 
+#include "Containers/Array.h"
+#include "Containers/UnrealString.h"
+#include "EdGraph/EdGraphNode.h"
+#include "HAL/Platform.h"
+#include "Internationalization/Text.h"
+#include "Math/Color.h"
+#include "Math/Transform.h"
 #include "MuCOE/Nodes/CustomizableObjectNodeModifierBase.h"
-#include "MuCOE/Nodes/CustomizableObjectNodeMaterial.h"
+#include "UObject/ObjectPtr.h"
+#include "UObject/UObjectGlobals.h"
 
 #include "CustomizableObjectNodeMeshClipWithMesh.generated.h"
 
+class UCustomizableObject;
+class UCustomizableObjectNodeRemapPins;
+class UEdGraphPin;
+class UObject;
+struct FGuid;
+
 UCLASS()
 class CUSTOMIZABLEOBJECTEDITOR_API UCustomizableObjectNodeMeshClipWithMesh : public UCustomizableObjectNodeModifierBase
 {
diff --git a/Engine/Plugins/Experimental/Mutable/Source/CustomizableObjectEditor/Classes/MuCOE/Nodes/CustomizableObjectNodeMeshGeometryOperation.h b/Engine/Plugins/Experimental/Mutable/Source/CustomizableObjectEditor/Classes/MuCOE/Nodes/CustomizableObjectNodeMeshGeometryOperation.h
--- a/Engine/Plugins/Experimental/Mutable/Source/CustomizableObjectEditor/Classes/MuCOE/Nodes/CustomizableObjectNodeMeshGeometryOperation.h
+++ b/Engine/Plugins/Experimental/Mutable/Source/CustomizableObjectEditor/Classes/MuCOE/Nodes/CustomizableObjectNodeMeshGeometryOperation.h
@@ -2,10 +2,19 @@
 
 #pragma once
 
+#include "EdGraph/EdGraphNode.h"
+#include "HAL/Platform.h"
+#include "Internationalization/Text.h"
+#include "Math/Color.h"
 #include "MuCOE/Nodes/CustomizableObjectNode.h"
+#include "UObject/UObjectGlobals.h"
 
 #include "CustomizableObjectNodeMeshGeometryOperation.generated.h"
 
+class UCustomizableObjectNodeRemapPins;
+class UEdGraphPin;
+class UObject;
+
 
 UCLASS()
 class CUSTOMIZABLEOBJECTEDITOR_API UCustomizableObjectNodeMeshGeometryOperation : public UCustomizableObjectNode
