Ran IWYU on Mutable plugin

diff --git a/Engine/Plugins/Experimental/Mutable/Source/CustomizableObjectEditor/Classes/MuCOE/Nodes/CustomizableObjectNodeMeshMorph.h b/Engine/Plugins/Experimental/Mutable/Source/CustomizableObjectEditor/Classes/MuCOE/Nodes/CustomizableObjectNodeMeshMorph.h
--- a/Engine/Plugins/Experimental/Mutable/Source/CustomizableObjectEditor/Classes/MuCOE/Nodes/CustomizableObjectNodeMeshMorph.h
+++ b/Engine/Plugins/Experimental/Mutable/Source/CustomizableObjectEditor/Classes/MuCOE/Nodes/CustomizableObjectNodeMeshMorph.h
@@ -2,11 +2,22 @@
 
 #pragma once
 
+#include "Containers/Array.h"
+#include "Containers/UnrealString.h"
+#include "EdGraph/EdGraphNode.h"
+#include "HAL/Platform.h"
+#include "Internationalization/Text.h"
+#include "Math/Color.h"
 #include "MuCOE/Nodes/CustomizableObjectNode.h"
-#include "MuCOE/Nodes/CustomizableObjectNodeMeshReshapeCommon.h"
+#include "UObject/UObjectGlobals.h"
 
 #include "CustomizableObjectNodeMeshMorph.generated.h"
 
+class UCustomizableObjectNodeRemapPins;
+class UEdGraphPin;
+class UObject;
+struct FMeshReshapeBoneReference;
+
 UCLASS()
 class CUSTOMIZABLEOBJECTEDITOR_API UCustomizableObjectNodeMeshMorph : public UCustomizableObjectNode
 {
diff --git a/Engine/Plugins/Experimental/Mutable/Source/CustomizableObjectEditor/Classes/MuCOE/Nodes/CustomizableObjectNodeMeshMorphStackApplication.h b/Engine/Plugins/Experimental/Mutable/Source/CustomizableObjectEditor/Classes/MuCOE/Nodes/CustomizableObjectNodeMeshMorphStackApplication.h
--- a/Engine/Plugins/Experimental/Mutable/Source/CustomizableObjectEditor/Classes/MuCOE/Nodes/CustomizableObjectNodeMeshMorphStackApplication.h
+++ b/Engine/Plugins/Experimental/Mutable/Source/CustomizableObjectEditor/Classes/MuCOE/Nodes/CustomizableObjectNodeMeshMorphStackApplication.h
@@ -2,9 +2,21 @@
 
 #pragma once
 
+#include "Containers/Array.h"
+#include "Containers/UnrealString.h"
+#include "EdGraph/EdGraphNode.h"
+#include "HAL/Platform.h"
+#include "Internationalization/Text.h"
+#include "Math/Color.h"
 #include "MuCOE/Nodes/CustomizableObjectNode.h"
+#include "UObject/UObjectGlobals.h"
+
 #include "CustomizableObjectNodeMeshMorphStackApplication.generated.h"
 
+class UCustomizableObjectNodeRemapPins;
+class UEdGraphPin;
+class UObject;
+
 
 UCLASS()
 class CUSTOMIZABLEOBJECTEDITOR_API UCustomizableObjectNodeMeshMorphStackApplication : public UCustomizableObjectNode
diff --git a/Engine/Plugins/Experimental/Mutable/Source/CustomizableObjectEditor/Classes/MuCOE/Nodes/CustomizableObjectNodeMeshMorphStackDefinition.h b/Engine/Plugins/Experimental/Mutable/Source/CustomizableObjectEditor/Classes/MuCOE/Nodes/CustomizableObjectNodeMeshMorphStackDefinition.h
--- a/Engine/Plugins/Experimental/Mutable/Source/CustomizableObjectEditor/Classes/MuCOE/Nodes/CustomizableObjectNodeMeshMorphStackDefinition.h
+++ b/Engine/Plugins/Experimental/Mutable/Source/CustomizableObjectEditor/Classes/MuCOE/Nodes/CustomizableObjectNodeMeshMorphStackDefinition.h
@@ -2,9 +2,21 @@
 
 #pragma once
 
+#include "Containers/Array.h"
+#include "Containers/UnrealString.h"
+#include "EdGraph/EdGraphNode.h"
+#include "HAL/Platform.h"
+#include "Internationalization/Text.h"
+#include "Math/Color.h"
 #include "MuCOE/Nodes/CustomizableObjectNode.h"
+#include "UObject/UObjectGlobals.h"
+
 #include "CustomizableObjectNodeMeshMorphStackDefinition.generated.h"
 
+class UCustomizableObjectNodeRemapPins;
+class UEdGraphPin;
+class UObject;
+
 
 UCLASS()
 class CUSTOMIZABLEOBJECTEDITOR_API UCustomizableObjectNodeMeshMorphStackDefinition : public UCustomizableObjectNode
diff --git a/Engine/Plugins/Experimental/Mutable/Source/CustomizableObjectEditor/Classes/MuCOE/Nodes/CustomizableObjectNodeMeshReshape.h b/Engine/Plugins/Experimental/Mutable/Source/CustomizableObjectEditor/Classes/MuCOE/Nodes/CustomizableObjectNodeMeshReshape.h
--- a/Engine/Plugins/Experimental/Mutable/Source/CustomizableObjectEditor/Classes/MuCOE/Nodes/CustomizableObjectNodeMeshReshape.h
+++ b/Engine/Plugins/Experimental/Mutable/Source/CustomizableObjectEditor/Classes/MuCOE/Nodes/CustomizableObjectNodeMeshReshape.h
@@ -2,11 +2,22 @@
 
 #pragma once
 
+#include "Containers/Array.h"
+#include "Containers/UnrealString.h"
+#include "EdGraph/EdGraphNode.h"
+#include "HAL/Platform.h"
+#include "Internationalization/Text.h"
+#include "Math/Color.h"
 #include "MuCOE/Nodes/CustomizableObjectNode.h"
-#include "MuCOE/Nodes/CustomizableObjectNodeMeshReshapeCommon.h"
+#include "UObject/UObjectGlobals.h"
 
 #include "CustomizableObjectNodeMeshReshape.generated.h"
 
+class UCustomizableObjectNodeRemapPins;
+class UEdGraphPin;
+class UObject;
+struct FMeshReshapeBoneReference;
+
 UCLASS()
 class CUSTOMIZABLEOBJECTEDITOR_API UCustomizableObjectNodeMeshReshape : public UCustomizableObjectNode
 {
diff --git a/Engine/Plugins/Experimental/Mutable/Source/CustomizableObjectEditor/Classes/MuCOE/Nodes/CustomizableObjectNodeMeshReshapeCommon.h b/Engine/Plugins/Experimental/Mutable/Source/CustomizableObjectEditor/Classes/MuCOE/Nodes/CustomizableObjectNodeMeshReshapeCommon.h
--- a/Engine/Plugins/Experimental/Mutable/Source/CustomizableObjectEditor/Classes/MuCOE/Nodes/CustomizableObjectNodeMeshReshapeCommon.h
+++ b/Engine/Plugins/Experimental/Mutable/Source/CustomizableObjectEditor/Classes/MuCOE/Nodes/CustomizableObjectNodeMeshReshapeCommon.h
@@ -2,6 +2,8 @@
 
 #pragma once
 
+#include "UObject/NameTypes.h"
+
 #include "CustomizableObjectNodeMeshReshapeCommon.generated.h"
 
 USTRUCT()
diff --git a/Engine/Plugins/Experimental/Mutable/Source/CustomizableObjectEditor/Classes/MuCOE/Nodes/CustomizableObjectNodeModifierBase.h b/Engine/Plugins/Experimental/Mutable/Source/CustomizableObjectEditor/Classes/MuCOE/Nodes/CustomizableObjectNodeModifierBase.h
--- a/Engine/Plugins/Experimental/Mutable/Source/CustomizableObjectEditor/Classes/MuCOE/Nodes/CustomizableObjectNodeModifierBase.h
+++ b/Engine/Plugins/Experimental/Mutable/Source/CustomizableObjectEditor/Classes/MuCOE/Nodes/CustomizableObjectNodeModifierBase.h
@@ -3,9 +3,12 @@
 #pragma once
 
 #include "MuCOE/Nodes/CustomizableObjectNodeMaterialBase.h"
+#include "UObject/UObjectGlobals.h"
 
 #include "CustomizableObjectNodeModifierBase.generated.h"
 
+class UObject;
+
 
 UCLASS(abstract)
 class CUSTOMIZABLEOBJECTEDITOR_API UCustomizableObjectNodeModifierBase : public UCustomizableObjectNodeMaterialBase
diff --git a/Engine/Plugins/Experimental/Mutable/Source/CustomizableObjectEditor/Classes/MuCOE/Nodes/CustomizableObjectNodeMorphMaterial.h b/Engine/Plugins/Experimental/Mutable/Source/CustomizableObjectEditor/Classes/MuCOE/Nodes/CustomizableObjectNodeMorphMaterial.h
--- a/Engine/Plugins/Experimental/Mutable/Source/CustomizableObjectEditor/Classes/MuCOE/Nodes/CustomizableObjectNodeMorphMaterial.h
+++ b/Engine/Plugins/Experimental/Mutable/Source/CustomizableObjectEditor/Classes/MuCOE/Nodes/CustomizableObjectNodeMorphMaterial.h
@@ -2,10 +2,21 @@
 
 #pragma once
 
+#include "Containers/UnrealString.h"
+#include "EdGraph/EdGraphNode.h"
+#include "HAL/Platform.h"
+#include "Internationalization/Text.h"
+#include "Math/Color.h"
 #include "MuCOE/Nodes/CustomizableObjectNodeEditMaterialBase.h"
+#include "UObject/UObjectGlobals.h"
 
 #include "CustomizableObjectNodeMorphMaterial.generated.h"
 
+class UCustomizableObjectNodeRemapPins;
+class UEdGraphPin;
+class UObject;
+struct FPropertyChangedEvent;
+
 
 UCLASS()
 class CUSTOMIZABLEOBJECTEDITOR_API UCustomizableObjectNodeMorphMaterial : public UCustomizableObjectNodeEditMaterialBase
diff --git a/Engine/Plugins/Experimental/Mutable/Source/CustomizableObjectEditor/Classes/MuCOE/Nodes/CustomizableObjectNodeObject.h b/Engine/Plugins/Experimental/Mutable/Source/CustomizableObjectEditor/Classes/MuCOE/Nodes/CustomizableObjectNodeObject.h
--- a/Engine/Plugins/Experimental/Mutable/Source/CustomizableObjectEditor/Classes/MuCOE/Nodes/CustomizableObjectNodeObject.h
+++ b/Engine/Plugins/Experimental/Mutable/Source/CustomizableObjectEditor/Classes/MuCOE/Nodes/CustomizableObjectNodeObject.h
@@ -2,11 +2,30 @@
 
 #pragma once
 
-#include "MuCOE/Nodes/CustomizableObjectNodeObjectGroup.h"
-#include "MuCOE/Nodes/CustomizableObjectNodeMaterial.h"
+#include "Containers/Array.h"
+#include "Containers/Map.h"
+#include "Containers/UnrealString.h"
+#include "EdGraph/EdGraphNode.h"
+#include "EdGraph/EdGraphPin.h"
+#include "HAL/Platform.h"
+#include "Internationalization/Text.h"
+#include "Math/Color.h"
+#include "Misc/Guid.h"
+#include "MuCO/CustomizableObjectUIData.h"
+#include "MuCOE/Nodes/CustomizableObjectNode.h"
+#include "UObject/NameTypes.h"
+#include "UObject/ObjectPtr.h"
+#include "UObject/UObjectGlobals.h"
 
 #include "CustomizableObjectNodeObject.generated.h"
 
+class UCustomizableObject;
+class UCustomizableObjectNodeMaterial;
+class UCustomizableObjectNodeRemapPins;
+class UObject;
+struct FPropertyChangedEvent;
+struct FSoftObjectPath;
+
 
 USTRUCT()
 struct CUSTOMIZABLEOBJECTEDITOR_API FCustomizableObjectState
diff --git a/Engine/Plugins/Experimental/Mutable/Source/CustomizableObjectEditor/Classes/MuCOE/Nodes/CustomizableObjectNodeObjectChild.h b/Engine/Plugins/Experimental/Mutable/Source/CustomizableObjectEditor/Classes/MuCOE/Nodes/CustomizableObjectNodeObjectChild.h
--- a/Engine/Plugins/Experimental/Mutable/Source/CustomizableObjectEditor/Classes/MuCOE/Nodes/CustomizableObjectNodeObjectChild.h
+++ b/Engine/Plugins/Experimental/Mutable/Source/CustomizableObjectEditor/Classes/MuCOE/Nodes/CustomizableObjectNodeObjectChild.h
@@ -2,10 +2,15 @@
 
 #pragma once
 
+#include "EdGraph/EdGraphNode.h"
+#include "Internationalization/Text.h"
 #include "MuCOE/Nodes/CustomizableObjectNodeObject.h"
+#include "UObject/UObjectGlobals.h"
 
 #include "CustomizableObjectNodeObjectChild.generated.h"
 
+class UObject;
+
 
 UCLASS()
 class CUSTOMIZABLEOBJECTEDITOR_API UCustomizableObjectNodeObjectChild : public UCustomizableObjectNodeObject
diff --git a/Engine/Plugins/Experimental/Mutable/Source/CustomizableObjectEditor/Classes/MuCOE/Nodes/CustomizableObjectNodeObjectGroup.h b/Engine/Plugins/Experimental/Mutable/Source/CustomizableObjectEditor/Classes/MuCOE/Nodes/CustomizableObjectNodeObjectGroup.h
--- a/Engine/Plugins/Experimental/Mutable/Source/CustomizableObjectEditor/Classes/MuCOE/Nodes/CustomizableObjectNodeObjectGroup.h
+++ b/Engine/Plugins/Experimental/Mutable/Source/CustomizableObjectEditor/Classes/MuCOE/Nodes/CustomizableObjectNodeObjectGroup.h
@@ -2,9 +2,24 @@
 
 #pragma once
 
+#include "Containers/UnrealString.h"
+#include "EdGraph/EdGraphNode.h"
+#include "HAL/Platform.h"
+#include "Internationalization/Text.h"
+#include "Math/Color.h"
+#include "MuCO/CustomizableObjectParameterTypeDefinitions.h"
+#include "MuCO/CustomizableObjectUIData.h"
 #include "MuCOE/Nodes/CustomizableObjectNode.h"
+#include "UObject/UObjectGlobals.h"
+
 #include "CustomizableObjectNodeObjectGroup.generated.h"
 
+class FArchive;
+class UCustomizableObjectNodeRemapPins;
+class UEdGraphPin;
+class UObject;
+struct FPropertyChangedEvent;
+
 
 UCLASS()
 class CUSTOMIZABLEOBJECTEDITOR_API UCustomizableObjectNodeObjectGroup : public UCustomizableObjectNode
diff --git a/Engine/Plugins/Experimental/Mutable/Source/CustomizableObjectEditor/Classes/MuCOE/Nodes/CustomizableObjectNodeParentedMaterial.h b/Engine/Plugins/Experimental/Mutable/Source/CustomizableObjectEditor/Classes/MuCOE/Nodes/CustomizableObjectNodeParentedMaterial.h
--- a/Engine/Plugins/Experimental/Mutable/Source/CustomizableObjectEditor/Classes/MuCOE/Nodes/CustomizableObjectNodeParentedMaterial.h
+++ b/Engine/Plugins/Experimental/Mutable/Source/CustomizableObjectEditor/Classes/MuCOE/Nodes/CustomizableObjectNodeParentedMaterial.h
@@ -2,9 +2,11 @@
 
 #pragma once
 
+#include "Containers/Array.h"
 #include "MuCOE/Nodes/CustomizableObjectNodeParentedNode.h"
-#include "MuCOE/Nodes/CustomizableObjectNode.h"
-#include "MuCOE/Nodes/CustomizableObjectNodeMaterial.h"
+
+class UCustomizableObjectNode;
+class UCustomizableObjectNodeMaterial;
 
 
 /** Node Material specialization of ICustomizableObjectNodeParentedNode. */
diff --git a/Engine/Plugins/Experimental/Mutable/Source/CustomizableObjectEditor/Classes/MuCOE/Nodes/CustomizableObjectNodeParentedNode.h b/Engine/Plugins/Experimental/Mutable/Source/CustomizableObjectEditor/Classes/MuCOE/Nodes/CustomizableObjectNodeParentedNode.h
--- a/Engine/Plugins/Experimental/Mutable/Source/CustomizableObjectEditor/Classes/MuCOE/Nodes/CustomizableObjectNodeParentedNode.h
+++ b/Engine/Plugins/Experimental/Mutable/Source/CustomizableObjectEditor/Classes/MuCOE/Nodes/CustomizableObjectNodeParentedNode.h
@@ -2,9 +2,11 @@
 
 #pragma once
 
-#include "MuCO/CustomizableObject.h"
-#include "MuCOE/Nodes/CustomizableObjectNode.h"
 
+#include "Misc/Guid.h"
+
+class UCustomizableObject;
+class UCustomizableObjectNode;
 
 /** Interface for a node that saves a parent node. */
 class ICustomizableObjectNodeParentedNode
diff --git a/Engine/Plugins/Experimental/Mutable/Source/CustomizableObjectEditor/Classes/MuCOE/Nodes/CustomizableObjectNodeProjectorConstant.h b/Engine/Plugins/Experimental/Mutable/Source/CustomizableObjectEditor/Classes/MuCOE/Nodes/CustomizableObjectNodeProjectorConstant.h
--- a/Engine/Plugins/Experimental/Mutable/Source/CustomizableObjectEditor/Classes/MuCOE/Nodes/CustomizableObjectNodeProjectorConstant.h
+++ b/Engine/Plugins/Experimental/Mutable/Source/CustomizableObjectEditor/Classes/MuCOE/Nodes/CustomizableObjectNodeProjectorConstant.h
@@ -2,11 +2,24 @@
 
 #pragma once
 
-#include "MuCO/CustomizableObjectInstance.h"
+#include "EdGraph/EdGraphNode.h"
+#include "Internationalization/Text.h"
+#include "Math/Color.h"
+#include "Math/UnrealMathSSE.h"
+#include "Math/Vector.h"
+#include "MuCO/CustomizableObjectParameterTypeDefinitions.h"
+#include "MuCO/CustomizableObjectUIData.h"
 #include "MuCOE/Nodes/CustomizableObjectNode.h"
+#include "UObject/NameTypes.h"
+#include "UObject/UObjectGlobals.h"
 
 #include "CustomizableObjectNodeProjectorConstant.generated.h"
 
+class UCustomizableObjectNodeRemapPins;
+class UEdGraphPin;
+class UObject;
+struct FPropertyChangedEvent;
+
 
 UCLASS(hideCategories = (CustomizableObjectHide))
 class CUSTOMIZABLEOBJECTEDITOR_API UCustomizableObjectNodeProjectorConstant : public UCustomizableObjectNode
diff --git a/Engine/Plugins/Experimental/Mutable/Source/CustomizableObjectEditor/Classes/MuCOE/Nodes/CustomizableObjectNodeProjectorParameter.h b/Engine/Plugins/Experimental/Mutable/Source/CustomizableObjectEditor/Classes/MuCOE/Nodes/CustomizableObjectNodeProjectorParameter.h
--- a/Engine/Plugins/Experimental/Mutable/Source/CustomizableObjectEditor/Classes/MuCOE/Nodes/CustomizableObjectNodeProjectorParameter.h
+++ b/Engine/Plugins/Experimental/Mutable/Source/CustomizableObjectEditor/Classes/MuCOE/Nodes/CustomizableObjectNodeProjectorParameter.h
@@ -2,11 +2,25 @@
 
 #pragma once
 
-#include "MuCO/CustomizableObjectInstance.h"
+#include "Containers/UnrealString.h"
+#include "EdGraph/EdGraphNode.h"
+#include "Internationalization/Text.h"
+#include "Math/Color.h"
+#include "Math/UnrealMathSSE.h"
+#include "Math/Vector.h"
+#include "MuCO/CustomizableObjectParameterTypeDefinitions.h"
+#include "MuCO/CustomizableObjectUIData.h"
 #include "MuCOE/Nodes/CustomizableObjectNode.h"
+#include "UObject/NameTypes.h"
+#include "UObject/UObjectGlobals.h"
 
 #include "CustomizableObjectNodeProjectorParameter.generated.h"
 
+class UCustomizableObjectNodeRemapPins;
+class UEdGraphPin;
+class UObject;
+struct FPropertyChangedEvent;
+
 
 UCLASS(hideCategories = (CustomizableObjectHide))
 class CUSTOMIZABLEOBJECTEDITOR_API UCustomizableObjectNodeProjectorParameter : public UCustomizableObjectNode
diff --git a/Engine/Plugins/Experimental/Mutable/Source/CustomizableObjectEditor/Classes/MuCOE/Nodes/CustomizableObjectNodeRemoveMesh.h b/Engine/Plugins/Experimental/Mutable/Source/CustomizableObjectEditor/Classes/MuCOE/Nodes/CustomizableObjectNodeRemoveMesh.h
--- a/Engine/Plugins/Experimental/Mutable/Source/CustomizableObjectEditor/Classes/MuCOE/Nodes/CustomizableObjectNodeRemoveMesh.h
+++ b/Engine/Plugins/Experimental/Mutable/Source/CustomizableObjectEditor/Classes/MuCOE/Nodes/CustomizableObjectNodeRemoveMesh.h
@@ -2,10 +2,19 @@
 
 #pragma once
 
+#include "EdGraph/EdGraphNode.h"
+#include "HAL/Platform.h"
+#include "Internationalization/Text.h"
+#include "Math/Color.h"
 #include "MuCOE/Nodes/CustomizableObjectNodeEditMaterialBase.h"
+#include "UObject/UObjectGlobals.h"
 
 #include "CustomizableObjectNodeRemoveMesh.generated.h"
 
+class UCustomizableObjectNodeRemapPins;
+class UEdGraphPin;
+class UObject;
+
 UCLASS()
 class CUSTOMIZABLEOBJECTEDITOR_API UCustomizableObjectNodeRemoveMesh : public UCustomizableObjectNodeEditMaterialBase
 {
diff --git a/Engine/Plugins/Experimental/Mutable/Source/CustomizableObjectEditor/Classes/MuCOE/Nodes/CustomizableObjectNodeRemoveMeshBlocks.h b/Engine/Plugins/Experimental/Mutable/Source/CustomizableObjectEditor/Classes/MuCOE/Nodes/CustomizableObjectNodeRemoveMeshBlocks.h
--- a/Engine/Plugins/Experimental/Mutable/Source/CustomizableObjectEditor/Classes/MuCOE/Nodes/CustomizableObjectNodeRemoveMeshBlocks.h
+++ b/Engine/Plugins/Experimental/Mutable/Source/CustomizableObjectEditor/Classes/MuCOE/Nodes/CustomizableObjectNodeRemoveMeshBlocks.h
@@ -2,10 +2,23 @@
 
 #pragma once
 
+#include "Containers/Array.h"
+#include "Containers/UnrealString.h"
+#include "EdGraph/EdGraphNode.h"
+#include "HAL/Platform.h"
+#include "Internationalization/Text.h"
+#include "Math/Color.h"
 #include "MuCOE/Nodes/CustomizableObjectNodeEditLayoutBlocks.h"
+#include "UObject/UObjectGlobals.h"
 
 #include "CustomizableObjectNodeRemoveMeshBlocks.generated.h"
 
+class FArchive;
+class UCustomizableObjectNodeRemapPins;
+class UEdGraphPin;
+class UObject;
+struct FPropertyChangedEvent;
+
 UCLASS()
 class CUSTOMIZABLEOBJECTEDITOR_API UCustomizableObjectNodeRemoveMeshBlocks : public UCustomizableObjectNodeEditLayoutBlocks
 {
diff --git a/Engine/Plugins/Experimental/Mutable/Source/CustomizableObjectEditor/Classes/MuCOE/Nodes/CustomizableObjectNodeSkeletalMesh.h b/Engine/Plugins/Experimental/Mutable/Source/CustomizableObjectEditor/Classes/MuCOE/Nodes/CustomizableObjectNodeSkeletalMesh.h
--- a/Engine/Plugins/Experimental/Mutable/Source/CustomizableObjectEditor/Classes/MuCOE/Nodes/CustomizableObjectNodeSkeletalMesh.h
+++ b/Engine/Plugins/Experimental/Mutable/Source/CustomizableObjectEditor/Classes/MuCOE/Nodes/CustomizableObjectNodeSkeletalMesh.h
@@ -2,14 +2,47 @@
 
 #pragma once
 
-#include "AssetThumbnail.h"
+#include "Containers/Array.h"
+#include "Containers/UnrealString.h"
+#include "EdGraph/EdGraphNode.h"
 #include "EdGraph/EdGraphPin.h"
 #include "GameplayTagContainer.h"
+#include "HAL/Platform.h"
+#include "HAL/PlatformCrt.h"
+#include "Internationalization/Text.h"
+#include "Layout/Visibility.h"
+#include "Math/Color.h"
+#include "Math/Vector2D.h"
 #include "MuCOE/Nodes/CustomizableObjectNodeMesh.h"
 #include "SGraphNode.h"
+#include "Styling/SlateTypes.h"
+#include "Templates/SharedPointer.h"
+#include "UObject/ObjectPtr.h"
+#include "UObject/SoftObjectPtr.h"
+#include "UObject/UObjectGlobals.h"
+#include "Widgets/DeclarativeSyntaxSupport.h"
 
 #include "CustomizableObjectNodeSkeletalMesh.generated.h"
 
+class FArchive;
+class FAssetThumbnail;
+class FAssetThumbnailPool;
+class FSkeletalMeshModel;
+class ISinglePropertyView;
+class SOverlay;
+class SVerticalBox;
+class UAnimInstance;
+class UCustomizableObjectLayout;
+class UCustomizableObjectNodeRemapPins;
+class UCustomizableObjectNodeRemapPinsByName;
+class UMaterialInterface;
+class UObject;
+class USkeletalMesh;
+class UTexture2D;
+struct FPropertyChangedEvent;
+struct FSkeletalMaterial;
+struct FSlateBrush;
+
 
 // Class to render the Skeletal Mesh thumbnail of a CustomizableObjectNodeSkeletalMesh
 class SGraphNodeSkeletalMesh : public SGraphNode
diff --git a/Engine/Plugins/Experimental/Mutable/Source/CustomizableObjectEditor/Classes/MuCOE/Nodes/CustomizableObjectNodeStaticMesh.h b/Engine/Plugins/Experimental/Mutable/Source/CustomizableObjectEditor/Classes/MuCOE/Nodes/CustomizableObjectNodeStaticMesh.h
--- a/Engine/Plugins/Experimental/Mutable/Source/CustomizableObjectEditor/Classes/MuCOE/Nodes/CustomizableObjectNodeStaticMesh.h
+++ b/Engine/Plugins/Experimental/Mutable/Source/CustomizableObjectEditor/Classes/MuCOE/Nodes/CustomizableObjectNodeStaticMesh.h
@@ -2,12 +2,40 @@
 
 #pragma once
 
+#include "Containers/Array.h"
+#include "Containers/UnrealString.h"
+#include "EdGraph/EdGraphNode.h"
+#include "EdGraph/EdGraphPin.h"
+#include "HAL/Platform.h"
+#include "Internationalization/Text.h"
+#include "Layout/Visibility.h"
+#include "Math/Color.h"
+#include "Math/Vector2D.h"
 #include "MuCOE/Nodes/CustomizableObjectNodeMesh.h"
 #include "SGraphNode.h"
-#include "AssetThumbnail.h"
+#include "Styling/SlateTypes.h"
+#include "Templates/SharedPointer.h"
+#include "UObject/ObjectPtr.h"
+#include "UObject/UObjectGlobals.h"
+#include "Widgets/DeclarativeSyntaxSupport.h"
 
 #include "CustomizableObjectNodeStaticMesh.generated.h"
 
+class FArchive;
+class FAssetThumbnail;
+class FAssetThumbnailPool;
+class SOverlay;
+class SVerticalBox;
+class UCustomizableObjectLayout;
+class UCustomizableObjectNodeRemapPins;
+class UCustomizableObjectNodeRemapPinsByName;
+class UMaterialInterface;
+class UObject;
+class UStaticMesh;
+class UTexture2D;
+struct FPropertyChangedEvent;
+struct FSlateBrush;
+
 
 // Class to render the Static Mesh thumbnail of a CustomizableObjectNodeStaticMesh
 class SGraphNodeStaticMesh : public SGraphNode
diff --git a/Engine/Plugins/Experimental/Mutable/Source/CustomizableObjectEditor/Classes/MuCOE/Nodes/CustomizableObjectNodeSwitchBase.cpp b/Engine/Plugins/Experimental/Mutable/Source/CustomizableObjectEditor/Classes/MuCOE/Nodes/CustomizableObjectNodeSwitchBase.cpp
--- a/Engine/Plugins/Experimental/Mutable/Source/CustomizableObjectEditor/Classes/MuCOE/Nodes/CustomizableObjectNodeSwitchBase.cpp
+++ b/Engine/Plugins/Experimental/Mutable/Source/CustomizableObjectEditor/Classes/MuCOE/Nodes/CustomizableObjectNodeSwitchBase.cpp
@@ -2,12 +2,18 @@
 
 #include "MuCOE/Nodes/CustomizableObjectNodeSwitchBase.h"
 
-#include "MuCOE/Nodes/CustomizableObjectNodeExternalPin.h"
-
-#include "MuCOE/CustomizableObjectEditorModule.h"
-#include "MuCOE/UnrealEditorPortabilityHelpers.h"
-#include "MuCOE/GenerateMutableSource/GenerateMutableSource.h"
+#include "Delegates/Delegate.h"
+#include "HAL/PlatformCrt.h"
+#include "Internationalization/Internationalization.h"
+#include "MuCO/CustomizableObjectCustomVersion.h"
+#include "MuCOE/EdGraphSchema_CustomizableObject.h"
 #include "MuCOE/GraphTraversal.h"
+#include "MuCOE/Nodes/CustomizableObjectNodeEnumParameter.h"
+#include "Serialization/Archive.h"
+#include "Templates/Casts.h"
+#include "UObject/WeakObjectPtr.h"
+
+class UCustomizableObjectNodeRemapPins;
 
 #define LOCTEXT_NAMESPACE "CustomizableObjectEditor"
 
diff --git a/Engine/Plugins/Experimental/Mutable/Source/CustomizableObjectEditor/Classes/MuCOE/Nodes/CustomizableObjectNodeSwitchBase.h b/Engine/Plugins/Experimental/Mutable/Source/CustomizableObjectEditor/Classes/MuCOE/Nodes/CustomizableObjectNodeSwitchBase.h
--- a/Engine/Plugins/Experimental/Mutable/Source/CustomizableObjectEditor/Classes/MuCOE/Nodes/CustomizableObjectNodeSwitchBase.h
+++ b/Engine/Plugins/Experimental/Mutable/Source/CustomizableObjectEditor/Classes/MuCOE/Nodes/CustomizableObjectNodeSwitchBase.h
@@ -2,12 +2,25 @@
 
 #pragma once
 
-#include "MuCOE/Nodes/CustomizableObjectNode.h"
-
+#include "Containers/Array.h"
+#include "Containers/UnrealString.h"
+#include "EdGraph/EdGraphNode.h"
 #include "EdGraph/EdGraphPin.h"
+#include "HAL/Platform.h"
+#include "Internationalization/Text.h"
+#include "Math/Color.h"
+#include "MuCOE/Nodes/CustomizableObjectNode.h"
+#include "UObject/NameTypes.h"
+#include "UObject/UObjectGlobals.h"
+#include "UObject/WeakObjectPtrTemplates.h"
 
 #include "CustomizableObjectNodeSwitchBase.generated.h"
 
+class FArchive;
+class UCustomizableObjectNodeRemapPins;
+class UObject;
+struct FFrame;
+
 
 UCLASS()
 class CUSTOMIZABLEOBJECTEDITOR_API UCustomizableObjectNodeSwitchBase : public UCustomizableObjectNode
diff --git a/Engine/Plugins/Experimental/Mutable/Source/CustomizableObjectEditor/Classes/MuCOE/Nodes/CustomizableObjectNodeTable.h b/Engine/Plugins/Experimental/Mutable/Source/CustomizableObjectEditor/Classes/MuCOE/Nodes/CustomizableObjectNodeTable.h
--- a/Engine/Plugins/Experimental/Mutable/Source/CustomizableObjectEditor/Classes/MuCOE/Nodes/CustomizableObjectNodeTable.h
+++ b/Engine/Plugins/Experimental/Mutable/Source/CustomizableObjectEditor/Classes/MuCOE/Nodes/CustomizableObjectNodeTable.h
@@ -2,14 +2,37 @@
 
 #pragma once
 
-#include "MuCOE/Nodes/CustomizableObjectNodeMaterial.h"
+#include "Containers/Array.h"
+#include "Containers/Map.h"
+#include "Containers/UnrealString.h"
+#include "Delegates/IDelegateInstance.h"
+#include "EdGraph/EdGraphNode.h"
 #include "Engine/DataTable.h"
-
+#include "HAL/Platform.h"
+#include "HAL/UnrealMemory.h"
+#include "Internationalization/Text.h"
+#include "Math/Color.h"
+#include "Math/Vector2D.h"
+#include "MuCO/CustomizableObjectUIData.h"
 #include "MuCOE/Nodes/CustomizableObjectNode.h"
-#include "MuCOE/Nodes/CustomizableObjectNodeLayoutBlocks.h"
+#include "UObject/Class.h"
+#include "UObject/Field.h"
+#include "UObject/NameTypes.h"
+#include "UObject/ObjectPtr.h"
+#include "UObject/SoftObjectPtr.h"
+#include "UObject/UObjectGlobals.h"
+#include "UObject/UnrealType.h"
 
 #include "CustomizableObjectNodeTable.generated.h"
 
+class UCustomizableObjectLayout;
+class UCustomizableObjectNodeRemapPins;
+class UEdGraphPin;
+class UObject;
+class USkeletalMesh;
+class UTexture2D;
+struct FGuid;
+
 
 /** Base class for all Table Pins. */
 UCLASS()
diff --git a/Engine/Plugins/Experimental/Mutable/Source/CustomizableObjectEditor/Classes/MuCOE/Nodes/CustomizableObjectNodeTexture.h b/Engine/Plugins/Experimental/Mutable/Source/CustomizableObjectEditor/Classes/MuCOE/Nodes/CustomizableObjectNodeTexture.h
--- a/Engine/Plugins/Experimental/Mutable/Source/CustomizableObjectEditor/Classes/MuCOE/Nodes/CustomizableObjectNodeTexture.h
+++ b/Engine/Plugins/Experimental/Mutable/Source/CustomizableObjectEditor/Classes/MuCOE/Nodes/CustomizableObjectNodeTexture.h
@@ -2,11 +2,29 @@
 
 #pragma once
 
+#include "EdGraph/EdGraphNode.h"
+#include "Internationalization/Text.h"
+#include "Layout/Visibility.h"
+#include "Math/Color.h"
 #include "MuCOE/Nodes/CustomizableObjectNode.h"
 #include "SGraphNode.h"
+#include "Styling/SlateBrush.h"
+#include "Styling/SlateTypes.h"
+#include "Templates/SharedPointer.h"
+#include "UObject/ObjectPtr.h"
+#include "UObject/UObjectGlobals.h"
+#include "Widgets/DeclarativeSyntaxSupport.h"
 
 #include "CustomizableObjectNodeTexture.generated.h"
 
+class SOverlay;
+class SVerticalBox;
+class UCustomizableObjectNodeRemapPins;
+class UCustomizableObjectNodeTexture;
+class UObject;
+class UTexture2D;
+struct FGeometry;
+
 // Class to render the Texture thumbnail of a CustomizableObjectNodeTexture
 class SGraphNodeTexture : public SGraphNode
 {
diff --git a/Engine/Plugins/Experimental/Mutable/Source/CustomizableObjectEditor/Classes/MuCOE/Nodes/CustomizableObjectNodeTextureBinarise.h b/Engine/Plugins/Experimental/Mutable/Source/CustomizableObjectEditor/Classes/MuCOE/Nodes/CustomizableObjectNodeTextureBinarise.h
--- a/Engine/Plugins/Experimental/Mutable/Source/CustomizableObjectEditor/Classes/MuCOE/Nodes/CustomizableObjectNodeTextureBinarise.h
+++ b/Engine/Plugins/Experimental/Mutable/Source/CustomizableObjectEditor/Classes/MuCOE/Nodes/CustomizableObjectNodeTextureBinarise.h
@@ -1,10 +1,20 @@
 // Copyright Epic Games, Inc. All Rights Reserved.
 #pragma once
 
+#include "EdGraph/EdGraphNode.h"
+#include "EdGraph/EdGraphPin.h"
+#include "HAL/Platform.h"
+#include "Internationalization/Text.h"
+#include "Math/Color.h"
 #include "MuCOE/Nodes/CustomizableObjectNode.h"
+#include "UObject/UObjectGlobals.h"
 
 #include "CustomizableObjectNodeTextureBinarise.generated.h"
 
+class FArchive;
+class UCustomizableObjectNodeRemapPins;
+class UObject;
+
 UCLASS()
 class CUSTOMIZABLEOBJECTEDITOR_API UCustomizableObjectNodeTextureBinarise : public UCustomizableObjectNode
 {
diff --git a/Engine/Plugins/Experimental/Mutable/Source/CustomizableObjectEditor/Classes/MuCOE/Nodes/CustomizableObjectNodeTextureInterpolate.h b/Engine/Plugins/Experimental/Mutable/Source/CustomizableObjectEditor/Classes/MuCOE/Nodes/CustomizableObjectNodeTextureInterpolate.h
--- a/Engine/Plugins/Experimental/Mutable/Source/CustomizableObjectEditor/Classes/MuCOE/Nodes/CustomizableObjectNodeTextureInterpolate.h
+++ b/Engine/Plugins/Experimental/Mutable/Source/CustomizableObjectEditor/Classes/MuCOE/Nodes/CustomizableObjectNodeTextureInterpolate.h
@@ -2,10 +2,22 @@
 
 #pragma once
 
+#include "Containers/Array.h"
+#include "Containers/UnrealString.h"
+#include "EdGraph/EdGraphNode.h"
+#include "EdGraph/EdGraphPin.h"
+#include "HAL/Platform.h"
+#include "Internationalization/Text.h"
+#include "Math/Color.h"
 #include "MuCOE/Nodes/CustomizableObjectNode.h"
+#include "UObject/UObjectGlobals.h"
 
 #include "CustomizableObjectNodeTextureInterpolate.generated.h"
 
+class UCustomizableObjectNodeRemapPins;
+class UObject;
+struct FPropertyChangedEvent;
+
 
 UCLASS()
 class CUSTOMIZABLEOBJECTEDITOR_API UCustomizableObjectNodeTextureInterpolate : public UCustomizableObjectNode
diff --git a/Engine/Plugins/Experimental/Mutable/Source/CustomizableObjectEditor/Classes/MuCOE/Nodes/CustomizableObjectNodeTextureInvert.h b/Engine/Plugins/Experimental/Mutable/Source/CustomizableObjectEditor/Classes/MuCOE/Nodes/CustomizableObjectNodeTextureInvert.h
--- a/Engine/Plugins/Experimental/Mutable/Source/CustomizableObjectEditor/Classes/MuCOE/Nodes/CustomizableObjectNodeTextureInvert.h
+++ b/Engine/Plugins/Experimental/Mutable/Source/CustomizableObjectEditor/Classes/MuCOE/Nodes/CustomizableObjectNodeTextureInvert.h
@@ -1,10 +1,19 @@
 // Copyright Epic Games, Inc. All Rights Reserved.
 #pragma once
 
+#include "EdGraph/EdGraphNode.h"
+#include "EdGraph/EdGraphPin.h"
+#include "Internationalization/Text.h"
+#include "Math/Color.h"
 #include "MuCOE/Nodes/CustomizableObjectNode.h"
+#include "UObject/UObjectGlobals.h"
 
 #include "CustomizableObjectNodeTextureInvert.generated.h"
 
+class FArchive;
+class UCustomizableObjectNodeRemapPins;
+class UObject;
+
 UCLASS()
 class CUSTOMIZABLEOBJECTEDITOR_API UCustomizableObjectNodeTextureInvert : public UCustomizableObjectNode
 {
diff --git a/Engine/Plugins/Experimental/Mutable/Source/CustomizableObjectEditor/Classes/MuCOE/Nodes/CustomizableObjectNodeTextureLayer.h b/Engine/Plugins/Experimental/Mutable/Source/CustomizableObjectEditor/Classes/MuCOE/Nodes/CustomizableObjectNodeTextureLayer.h
--- a/Engine/Plugins/Experimental/Mutable/Source/CustomizableObjectEditor/Classes/MuCOE/Nodes/CustomizableObjectNodeTextureLayer.h
+++ b/Engine/Plugins/Experimental/Mutable/Source/CustomizableObjectEditor/Classes/MuCOE/Nodes/CustomizableObjectNodeTextureLayer.h
@@ -2,9 +2,25 @@
 
 #pragma once
 
+#include "Containers/Array.h"
+#include "Containers/EnumAsByte.h"
+#include "Containers/UnrealString.h"
+#include "EdGraph/EdGraphNode.h"
+#include "EdGraph/EdGraphPin.h"
+#include "HAL/Platform.h"
+#include "Internationalization/Text.h"
+#include "Math/Color.h"
 #include "MuCOE/Nodes/CustomizableObjectNode.h"
+#include "UObject/UObjectGlobals.h"
+
 #include "CustomizableObjectNodeTextureLayer.generated.h"
 
+class FArchive;
+class UCustomizableObjectNodeRemapPins;
+class UEdGraph;
+class UObject;
+struct FPropertyChangedEvent;
+
 
 UENUM()
 enum ECustomizableObjectTextureLayerEffect
diff --git a/Engine/Plugins/Experimental/Mutable/Source/CustomizableObjectEditor/Classes/MuCOE/Nodes/CustomizableObjectNodeTextureParameter.h b/Engine/Plugins/Experimental/Mutable/Source/CustomizableObjectEditor/Classes/MuCOE/Nodes/CustomizableObjectNodeTextureParameter.h
--- a/Engine/Plugins/Experimental/Mutable/Source/CustomizableObjectEditor/Classes/MuCOE/Nodes/CustomizableObjectNodeTextureParameter.h
+++ b/Engine/Plugins/Experimental/Mutable/Source/CustomizableObjectEditor/Classes/MuCOE/Nodes/CustomizableObjectNodeTextureParameter.h
@@ -2,12 +2,21 @@
 
 #pragma once
 
+#include "Containers/UnrealString.h"
+#include "EdGraph/EdGraphNode.h"
+#include "Internationalization/Text.h"
+#include "Math/Color.h"
+#include "MuCO/CustomizableObjectUIData.h"
 #include "MuCOE/Nodes/CustomizableObjectNode.h"
-
-#include "Engine/Texture2D.h"
+#include "UObject/ObjectPtr.h"
+#include "UObject/UObjectGlobals.h"
 
 #include "CustomizableObjectNodeTextureParameter.generated.h"
 
+class UCustomizableObjectNodeRemapPins;
+class UObject;
+class UTexture2D;
+
 
 UCLASS()
 class CUSTOMIZABLEOBJECTEDITOR_API UCustomizableObjectNodeTextureParameter : public UCustomizableObjectNode
diff --git a/Engine/Plugins/Experimental/Mutable/Source/CustomizableObjectEditor/Classes/MuCOE/Nodes/CustomizableObjectNodeTextureProject.h b/Engine/Plugins/Experimental/Mutable/Source/CustomizableObjectEditor/Classes/MuCOE/Nodes/CustomizableObjectNodeTextureProject.h
--- a/Engine/Plugins/Experimental/Mutable/Source/CustomizableObjectEditor/Classes/MuCOE/Nodes/CustomizableObjectNodeTextureProject.h
+++ b/Engine/Plugins/Experimental/Mutable/Source/CustomizableObjectEditor/Classes/MuCOE/Nodes/CustomizableObjectNodeTextureProject.h
@@ -2,10 +2,23 @@
 
 #pragma once
 
+#include "Containers/Array.h"
+#include "EdGraph/EdGraphNode.h"
+#include "HAL/Platform.h"
+#include "Internationalization/Text.h"
+#include "Math/Color.h"
 #include "MuCOE/Nodes/CustomizableObjectNode.h"
+#include "UObject/UObjectGlobals.h"
 
 #include "CustomizableObjectNodeTextureProject.generated.h"
 
+class FArchive;
+class UCustomizableObjectNodeRemapPins;
+class UEdGraphPin;
+class UObject;
+struct FEdGraphPinReference;
+struct FPropertyChangedEvent;
+
 
 UCLASS()
 class CUSTOMIZABLEOBJECTEDITOR_API UCustomizableObjectNodeTextureProject : public UCustomizableObjectNode
diff --git a/Engine/Plugins/Experimental/Mutable/Source/CustomizableObjectEditor/Classes/MuCOE/Nodes/CustomizableObjectNodeTextureToChannels.h b/Engine/Plugins/Experimental/Mutable/Source/CustomizableObjectEditor/Classes/MuCOE/Nodes/CustomizableObjectNodeTextureToChannels.h
--- a/Engine/Plugins/Experimental/Mutable/Source/CustomizableObjectEditor/Classes/MuCOE/Nodes/CustomizableObjectNodeTextureToChannels.h
+++ b/Engine/Plugins/Experimental/Mutable/Source/CustomizableObjectEditor/Classes/MuCOE/Nodes/CustomizableObjectNodeTextureToChannels.h
@@ -2,10 +2,21 @@
 
 #pragma once
 
+#include "EdGraph/EdGraphNode.h"
+#include "EdGraph/EdGraphPin.h"
+#include "HAL/Platform.h"
+#include "Internationalization/Text.h"
+#include "Math/Color.h"
 #include "MuCOE/Nodes/CustomizableObjectNode.h"
+#include "UObject/UObjectGlobals.h"
 
 #include "CustomizableObjectNodeTextureToChannels.generated.h"
 
+class FArchive;
+class UCustomizableObjectNodeRemapPins;
+class UObject;
+struct FPropertyChangedEvent;
+
 
 UCLASS()
 class CUSTOMIZABLEOBJECTEDITOR_API UCustomizableObjectNodeTextureToChannels : public UCustomizableObjectNode
diff --git a/Engine/Plugins/Experimental/Mutable/Source/CustomizableObjectEditor/Classes/MuCOE/Nodes/CustomizableObjectNodeTextureTransform.h b/Engine/Plugins/Experimental/Mutable/Source/CustomizableObjectEditor/Classes/MuCOE/Nodes/CustomizableObjectNodeTextureTransform.h
--- a/Engine/Plugins/Experimental/Mutable/Source/CustomizableObjectEditor/Classes/MuCOE/Nodes/CustomizableObjectNodeTextureTransform.h
+++ b/Engine/Plugins/Experimental/Mutable/Source/CustomizableObjectEditor/Classes/MuCOE/Nodes/CustomizableObjectNodeTextureTransform.h
@@ -1,10 +1,20 @@
 // Copyright Epic Games, Inc. All Rights Reserved.
 #pragma once
 
+#include "EdGraph/EdGraphNode.h"
+#include "EdGraph/EdGraphPin.h"
+#include "HAL/Platform.h"
+#include "Internationalization/Text.h"
+#include "Math/Color.h"
 #include "MuCOE/Nodes/CustomizableObjectNode.h"
+#include "UObject/UObjectGlobals.h"
 
 #include "CustomizableObjectNodeTextureTransform.generated.h"
 
+class FArchive;
+class UCustomizableObjectNodeRemapPins;
+class UObject;
+
 UCLASS()
 class CUSTOMIZABLEOBJECTEDITOR_API UCustomizableObjectNodeTextureTransform : public UCustomizableObjectNode
 {
diff --git a/Engine/Plugins/Experimental/Mutable/Source/CustomizableObjectEditor/Classes/MuCOE/Nodes/CustomizableObjectNodeUseMaterial.h b/Engine/Plugins/Experimental/Mutable/Source/CustomizableObjectEditor/Classes/MuCOE/Nodes/CustomizableObjectNodeUseMaterial.h
--- a/Engine/Plugins/Experimental/Mutable/Source/CustomizableObjectEditor/Classes/MuCOE/Nodes/CustomizableObjectNodeUseMaterial.h
+++ b/Engine/Plugins/Experimental/Mutable/Source/CustomizableObjectEditor/Classes/MuCOE/Nodes/CustomizableObjectNodeUseMaterial.h
@@ -2,9 +2,16 @@
 
 #pragma once
 
-#include "MuCOE/Nodes/CustomizableObjectNodeParentedMaterial.h"
-#include "MuCOE/Nodes/CustomizableObjectNode.h"
 
+#include "Containers/Map.h"
+#include "Delegates/IDelegateInstance.h"
+
+class FCustomizableObjectNodeParentedMaterial;
+class UCustomizableObject;
+class UCustomizableObjectNode;
+class UEdGraphPin;
+struct FEdGraphPinReference;
+struct FGuid;
 
 /** Common functionalities from EditMaterial and ExtendMaterial.
  *
diff --git a/Engine/Plugins/Experimental/Mutable/Source/CustomizableObjectEditor/Classes/MuCOE/Thumbnails/CustomizableObjectInstanceThumbnailRenderer.h b/Engine/Plugins/Experimental/Mutable/Source/CustomizableObjectEditor/Classes/MuCOE/Thumbnails/CustomizableObjectInstanceThumbnailRenderer.h
--- a/Engine/Plugins/Experimental/Mutable/Source/CustomizableObjectEditor/Classes/MuCOE/Thumbnails/CustomizableObjectInstanceThumbnailRenderer.h
+++ b/Engine/Plugins/Experimental/Mutable/Source/CustomizableObjectEditor/Classes/MuCOE/Thumbnails/CustomizableObjectInstanceThumbnailRenderer.h
@@ -6,10 +6,17 @@
  */
 
 #pragma once
-#include "ThumbnailHelpers.h"
+#include "HAL/Platform.h"
 #include "ThumbnailRendering/DefaultSizedThumbnailRenderer.h"
+#include "UObject/ObjectPtr.h"
+#include "UObject/UObjectGlobals.h"
+
 #include "CustomizableObjectInstanceThumbnailRenderer.generated.h"
 
+class FCanvas;
+class FRenderTarget;
+class UObject;
+
 
 UCLASS(config=Editor)
 class CUSTOMIZABLEOBJECTEDITOR_API UCustomizableObjectInstanceThumbnailRenderer : public UDefaultSizedThumbnailRenderer
diff --git a/Engine/Plugins/Experimental/Mutable/Source/CustomizableObjectEditor/Classes/MuCOE/Thumbnails/CustomizableObjectThumbnailRenderer.h b/Engine/Plugins/Experimental/Mutable/Source/CustomizableObjectEditor/Classes/MuCOE/Thumbnails/CustomizableObjectThumbnailRenderer.h
--- a/Engine/Plugins/Experimental/Mutable/Source/CustomizableObjectEditor/Classes/MuCOE/Thumbnails/CustomizableObjectThumbnailRenderer.h
+++ b/Engine/Plugins/Experimental/Mutable/Source/CustomizableObjectEditor/Classes/MuCOE/Thumbnails/CustomizableObjectThumbnailRenderer.h
@@ -6,11 +6,17 @@
  */
 
 #pragma once
-#include "ThumbnailHelpers.h"
+#include "HAL/Platform.h"
 #include "ThumbnailRendering/DefaultSizedThumbnailRenderer.h"
+#include "UObject/ObjectPtr.h"
+#include "UObject/UObjectGlobals.h"
 
 #include "CustomizableObjectThumbnailRenderer.generated.h"
 
+class FCanvas;
+class FRenderTarget;
+class UObject;
+
 
 UCLASS(config=Editor)
 class CUSTOMIZABLEOBJECTEDITOR_API UCustomizableObjectThumbnailRenderer : public UDefaultSizedThumbnailRenderer
diff --git a/Engine/Plugins/Experimental/Mutable/Source/CustomizableObjectEditor/Private/MuCOE/AssetTypeActions_CustomizableObject.cpp b/Engine/Plugins/Experimental/Mutable/Source/CustomizableObjectEditor/Private/MuCOE/AssetTypeActions_CustomizableObject.cpp
--- a/Engine/Plugins/Experimental/Mutable/Source/CustomizableObjectEditor/Private/MuCOE/AssetTypeActions_CustomizableObject.cpp
+++ b/Engine/Plugins/Experimental/Mutable/Source/CustomizableObjectEditor/Private/MuCOE/AssetTypeActions_CustomizableObject.cpp
@@ -1,14 +1,40 @@
 // Copyright Epic Games, Inc. All Rights Reserved.
 
 #include "MuCOE/AssetTypeActions_CustomizableObject.h"
-#include "MuCOE/CustomizableObjectEditorModule.h"
-#include "ContentBrowserModule.h"
+
+#include "AssetRegistry/ARFilter.h"
+#include "AssetRegistry/AssetData.h"
 #include "AssetRegistry/AssetRegistryModule.h"
-#include "MuCOE/CustomizableObjectCompiler.h"
-#include "IContentBrowserSingleton.h"
+#include "AssetRegistry/IAssetRegistry.h"
+#include "Containers/Map.h"
+#include "Containers/UnrealString.h"
+#include "ContentBrowserModule.h"
+#include "Delegates/Delegate.h"
+#include "Editor.h"
+#include "Editor/EditorEngine.h"
+#include "Framework/Commands/UIAction.h"
 #include "Framework/MultiBox/MultiBoxBuilder.h"
+#include "HAL/PlatformCrt.h"
+#include "HAL/PlatformMisc.h"
+#include "IContentBrowserSingleton.h"
+#include "Misc/AssertionMacros.h"
 #include "Misc/MessageDialog.h"
-#include "MuCOE/UnrealEditorPortabilityHelpers.h"
+#include "Modules/ModuleManager.h"
+#include "MuCO/CustomizableObjectInstance.h"
+#include "MuCO/CustomizableObjectSystem.h"
+#include "MuCOE/CustomizableObjectEditorModule.h"
+#include "Subsystems/AssetEditorSubsystem.h"
+#include "Templates/Casts.h"
+#include "Textures/SlateIcon.h"
+#include "Toolkits/IToolkit.h"
+#include "UObject/NameTypes.h"
+#include "UObject/Object.h"
+#include "UObject/Package.h"
+#include "UObject/TopLevelAssetPath.h"
+#include "UObject/UObjectGlobals.h"
+#include "UObject/WeakObjectPtr.h"
+
+class IToolkitHost;
 
 
 #define LOCTEXT_NAMESPACE "CustomizableObjectEditor"
diff --git a/Engine/Plugins/Experimental/Mutable/Source/CustomizableObjectEditor/Private/MuCOE/AssetTypeActions_CustomizableObject.h b/Engine/Plugins/Experimental/Mutable/Source/CustomizableObjectEditor/Private/MuCOE/AssetTypeActions_CustomizableObject.h
--- a/Engine/Plugins/Experimental/Mutable/Source/CustomizableObjectEditor/Private/MuCOE/AssetTypeActions_CustomizableObject.h
+++ b/Engine/Plugins/Experimental/Mutable/Source/CustomizableObjectEditor/Private/MuCOE/AssetTypeActions_CustomizableObject.h
@@ -3,8 +3,20 @@
 #pragma once
 
 #include "AssetTypeActions_Base.h"
-#include "Runtime/Launch/Resources/Version.h"
-#include "CustomizableObjectClasses.h"
+#include "Containers/Array.h"
+#include "HAL/Platform.h"
+#include "IAssetTypeActions.h"
+#include "Internationalization/Internationalization.h"
+#include "Internationalization/Text.h"
+#include "Math/Color.h"
+#include "MuCO/CustomizableObject.h"
+#include "Templates/SharedPointer.h"
+#include "UObject/WeakObjectPtrTemplates.h"
+
+class FMenuBuilder;
+class IToolkitHost;
+class UClass;
+class UObject;
 
 enum class ERecompileCO
 {
diff --git a/Engine/Plugins/Experimental/Mutable/Source/CustomizableObjectEditor/Private/MuCOE/AssetTypeActions_CustomizableObjectInstance.cpp b/Engine/Plugins/Experimental/Mutable/Source/CustomizableObjectEditor/Private/MuCOE/AssetTypeActions_CustomizableObjectInstance.cpp
--- a/Engine/Plugins/Experimental/Mutable/Source/CustomizableObjectEditor/Private/MuCOE/AssetTypeActions_CustomizableObjectInstance.cpp
+++ b/Engine/Plugins/Experimental/Mutable/Source/CustomizableObjectEditor/Private/MuCOE/AssetTypeActions_CustomizableObjectInstance.cpp
@@ -1,9 +1,22 @@
 // Copyright Epic Games, Inc. All Rights Reserved.
 
 #include "MuCOE/AssetTypeActions_CustomizableObjectInstance.h"
-#include "MuCOE/CustomizableObjectEditorModule.h"
+
+#include "Delegates/Delegate.h"
+#include "Editor.h"
+#include "Editor/EditorEngine.h"
+#include "Framework/Commands/UIAction.h"
 #include "Framework/MultiBox/MultiBoxBuilder.h"
-#include "MuCOE/UnrealEditorPortabilityHelpers.h"
+#include "HAL/PlatformCrt.h"
+#include "Modules/ModuleManager.h"
+#include "MuCOE/CustomizableObjectEditorModule.h"
+#include "Subsystems/AssetEditorSubsystem.h"
+#include "Templates/Casts.h"
+#include "Textures/SlateIcon.h"
+#include "Toolkits/IToolkit.h"
+#include "UObject/Object.h"
+
+class IToolkitHost;
 
 #define LOCTEXT_NAMESPACE "AssetTypeActions"
 
diff --git a/Engine/Plugins/Experimental/Mutable/Source/CustomizableObjectEditor/Private/MuCOE/AssetTypeActions_CustomizableObjectInstance.h b/Engine/Plugins/Experimental/Mutable/Source/CustomizableObjectEditor/Private/MuCOE/AssetTypeActions_CustomizableObjectInstance.h
--- a/Engine/Plugins/Experimental/Mutable/Source/CustomizableObjectEditor/Private/MuCOE/AssetTypeActions_CustomizableObjectInstance.h
+++ b/Engine/Plugins/Experimental/Mutable/Source/CustomizableObjectEditor/Private/MuCOE/AssetTypeActions_CustomizableObjectInstance.h
@@ -3,8 +3,20 @@
 #pragma once
 
 #include "AssetTypeActions_Base.h"
-#include "Runtime/Launch/Resources/Version.h"
-#include "CustomizableObjectClasses.h"
+#include "Containers/Array.h"
+#include "HAL/Platform.h"
+#include "IAssetTypeActions.h"
+#include "Internationalization/Internationalization.h"
+#include "Internationalization/Text.h"
+#include "Math/Color.h"
+#include "MuCO/CustomizableObjectInstance.h"
+#include "Templates/SharedPointer.h"
+#include "UObject/WeakObjectPtrTemplates.h"
+
+class FMenuBuilder;
+class IToolkitHost;
+class UClass;
+class UObject;
 
 
 class FAssetTypeActions_CustomizableObjectInstance : public FAssetTypeActions_Base
diff --git a/Engine/Plugins/Experimental/Mutable/Source/CustomizableObjectEditor/Private/MuCOE/CompilationMessageCache.cpp b/Engine/Plugins/Experimental/Mutable/Source/CustomizableObjectEditor/Private/MuCOE/CompilationMessageCache.cpp
--- a/Engine/Plugins/Experimental/Mutable/Source/CustomizableObjectEditor/Private/MuCOE/CompilationMessageCache.cpp
+++ b/Engine/Plugins/Experimental/Mutable/Source/CustomizableObjectEditor/Private/MuCOE/CompilationMessageCache.cpp
@@ -3,8 +3,9 @@
 #include "MuCOE/CompilationMessageCache.h"
 
 #include "Containers/Array.h"
-#include "Internationalization/Text.h"
+#include "HAL/PlatformCrt.h"
 #include "Logging/TokenizedMessage.h"
+#include "Misc/AssertionMacros.h"
 
 
 void FCompilationMessageCache::AddMessage(const FText& InMessage, const TArray<const UCustomizableObjectNode*>& InArrayNode, EMessageSeverity::Type MessageSeverity /* = EMessageSeverity::Warning */)
diff --git a/Engine/Plugins/Experimental/Mutable/Source/CustomizableObjectEditor/Private/MuCOE/CustomizableInstanceDetails.cpp b/Engine/Plugins/Experimental/Mutable/Source/CustomizableObjectEditor/Private/MuCOE/CustomizableInstanceDetails.cpp
--- a/Engine/Plugins/Experimental/Mutable/Source/CustomizableObjectEditor/Private/MuCOE/CustomizableInstanceDetails.cpp
+++ b/Engine/Plugins/Experimental/Mutable/Source/CustomizableObjectEditor/Private/MuCOE/CustomizableInstanceDetails.cpp
@@ -1,13 +1,24 @@
 // Copyright Epic Games, Inc. All Rights Reserved.
 
 #include "MuCOE/CustomizableInstanceDetails.h"
-#include "MuCOE/CustomizableObjectEditorModule.h"
-#include "MuCO/CustomizableObjectSystem.h"
-#include "DetailLayoutBuilder.h"
+
+#include "Containers/Array.h"
+#include "Delegates/Delegate.h"
 #include "DetailCategoryBuilder.h"
+#include "DetailLayoutBuilder.h"
 #include "DetailWidgetRow.h"
+#include "HAL/PlatformCrt.h"
+#include "IDetailsView.h"
+#include "Internationalization/Internationalization.h"
+#include "Misc/AssertionMacros.h"
+#include "MuCO/CustomizableObjectInstance.h"
+#include "MuCO/CustomizableObjectSystem.h"
 #include "MuCOE/SCustomizableInstanceProperties.h"
-#include "LevelEditorActions.h"
+#include "Templates/Casts.h"
+#include "UObject/NameTypes.h"
+#include "Widgets/DeclarativeSyntaxSupport.h"
+
+class UObject;
 
 #define LOCTEXT_NAMESPACE "CustomizableInstanceDetails"
 
diff --git a/Engine/Plugins/Experimental/Mutable/Source/CustomizableObjectEditor/Private/MuCOE/CustomizableInstanceDetails.h b/Engine/Plugins/Experimental/Mutable/Source/CustomizableObjectEditor/Private/MuCOE/CustomizableInstanceDetails.h
--- a/Engine/Plugins/Experimental/Mutable/Source/CustomizableObjectEditor/Private/MuCOE/CustomizableInstanceDetails.h
+++ b/Engine/Plugins/Experimental/Mutable/Source/CustomizableObjectEditor/Private/MuCOE/CustomizableInstanceDetails.h
@@ -2,13 +2,15 @@
 
 #pragma once
 
-#include "MuCO/CustomizableObjectInstance.h"
-#include "MuCO/CustomizableSkeletalComponent.h"
 #include "IDetailCustomization.h"
 #include "Templates/SharedPointer.h"
+#include "UObject/WeakObjectPtr.h"
+#include "UObject/WeakObjectPtrTemplates.h"
 
-
+class IDetailLayoutBuilder;
 class SCustomizableInstanceProperties;
+class UCustomizableObjectInstance;
+class UCustomizableSkeletalComponent;
 
 
 class FCustomizableInstanceDetails : public IDetailCustomization
diff --git a/Engine/Plugins/Experimental/Mutable/Source/CustomizableObjectEditor/Private/MuCOE/CustomizableObjectBakeHelpers.cpp b/Engine/Plugins/Experimental/Mutable/Source/CustomizableObjectEditor/Private/MuCOE/CustomizableObjectBakeHelpers.cpp
--- a/Engine/Plugins/Experimental/Mutable/Source/CustomizableObjectEditor/Private/MuCOE/CustomizableObjectBakeHelpers.cpp
+++ b/Engine/Plugins/Experimental/Mutable/Source/CustomizableObjectEditor/Private/MuCOE/CustomizableObjectBakeHelpers.cpp
@@ -2,22 +2,55 @@
 
 #include "MuCOE/CustomizableObjectBakeHelpers.h"
 
-#include "MuCOE/CustomizableObjectEditor.h"
-#include "MuCO/CustomizableObjectMipDataProvider.h"
-#include "BusyCursor.h"
-#include "ObjectTools.h"
 #include "AssetRegistry/AssetRegistryModule.h"
 #include "AssetToolsModule.h"
-#include "Serialization/ArchiveReplaceObjectRef.h"
-#include "Rendering/SkeletalMeshRenderData.h"
+#include "BoneIndices.h"
+#include "Containers/Array.h"
+#include "Containers/EnumAsByte.h"
+#include "Containers/IndirectArray.h"
+#include "Engine/SkeletalMesh.h"
+#include "Engine/Texture.h"
+#include "Engine/Texture2D.h"
+#include "Engine/TextureDefines.h"
+#include "GPUSkinPublicDefs.h"
+#include "HAL/PlatformCrt.h"
+#include "HAL/UnrealMemory.h"
+#include "IAssetTools.h"
+#include "Logging/LogCategory.h"
+#include "Logging/LogMacros.h"
+#include "Math/Color.h"
+#include "Math/UnrealMathSSE.h"
+#include "Math/Vector.h"
+#include "Math/Vector2D.h"
+#include "Math/Vector4.h"
+#include "Misc/AssertionMacros.h"
+#include "Misc/Guid.h"
+#include "Modules/ModuleManager.h"
+#include "MuCO/CustomizableObject.h"
+#include "MuCO/CustomizableObjectInstance.h"
+#include "MuCO/CustomizableObjectMipDataProvider.h"
+#include "PackedNormal.h"
+#include "PixelFormat.h"
+#include "RawIndexBuffer.h"
+#include "Rendering/MultiSizeIndexContainer.h"
+#include "Rendering/PositionVertexBuffer.h"
+#include "Rendering/SkeletalMeshLODModel.h"
 #include "Rendering/SkeletalMeshLODRenderData.h"
 #include "Rendering/SkeletalMeshModel.h"
-#include "Rendering/SkeletalMeshLODModel.h"
-#include "IAssetTools.h"
-
-#include "Misc/ScopeExit.h"
-
+#include "Rendering/SkeletalMeshRenderData.h"
+#include "Rendering/SkinWeightVertexBuffer.h"
+#include "Rendering/StaticMeshVertexBuffer.h"
+#include "Serialization/ArchiveReplaceObjectRef.h"
+#include "Serialization/BulkData.h"
+#include "StaticMeshResources.h"
+#include "Templates/Casts.h"
 #include "Templates/Decay.h"
+#include "Templates/UnrealTemplate.h"
+#include "TextureResource.h"
+#include "Trace/Detail/Channel.h"
+#include "UObject/Object.h"
+#include "UObject/Package.h"
+#include "UObject/UObjectGlobals.h"
 
 namespace
 {
diff --git a/Engine/Plugins/Experimental/Mutable/Source/CustomizableObjectEditor/Private/MuCOE/CustomizableObjectBakeHelpers.h b/Engine/Plugins/Experimental/Mutable/Source/CustomizableObjectEditor/Private/MuCOE/CustomizableObjectBakeHelpers.h
--- a/Engine/Plugins/Experimental/Mutable/Source/CustomizableObjectEditor/Private/MuCOE/CustomizableObjectBakeHelpers.h
+++ b/Engine/Plugins/Experimental/Mutable/Source/CustomizableObjectEditor/Private/MuCOE/CustomizableObjectBakeHelpers.h
@@ -2,8 +2,12 @@
 
 #pragma once
 
-#include "Engine/Texture2D.h"
-#include "Engine/SkeletalMesh.h"
+#include "Containers/Map.h"
+#include "Containers/UnrealString.h"
+
+class UObject;
+class UTexture2D;
+class UTexture;
 
 extern UObject* BakeHelper_DuplicateAsset(UObject* Object, const FString& ObjName, const FString& PkgName, bool ResetDuplicatedFlags, TMap<UObject*, UObject*>& ReplacementMap, bool OverwritePackage);
 
diff --git a/Engine/Plugins/Experimental/Mutable/Source/CustomizableObjectEditor/Private/MuCOE/CustomizableObjectCompileRunnable.cpp b/Engine/Plugins/Experimental/Mutable/Source/CustomizableObjectEditor/Private/MuCOE/CustomizableObjectCompileRunnable.cpp
--- a/Engine/Plugins/Experimental/Mutable/Source/CustomizableObjectEditor/Private/MuCOE/CustomizableObjectCompileRunnable.cpp
+++ b/Engine/Plugins/Experimental/Mutable/Source/CustomizableObjectEditor/Private/MuCOE/CustomizableObjectCompileRunnable.cpp
@@ -1,40 +1,24 @@
 // Copyright Epic Games, Inc. All Rights Reserved.
 
 #include "MuCOE/CustomizableObjectCompileRunnable.h"
-#include "MuCOE/CustomizableObjectCompiler.h"
-#include "MuCOE/CustomizableObjectEditorModule.h"
-#include "MuCO/CustomizableObjectSystem.h"
-#include "MuCO/CustomizableObject.h"
-#include "MuCO/UnrealMutableModelDiskStreamer.h"
 
-#include "StaticMeshResources.h"
-#include "Framework/Notifications/NotificationManager.h"
-#include "Widgets/Notifications/SNotificationList.h"
-#include "AssetRegistry/AssetRegistryModule.h"
-#include "MessageLogModule.h"
-#include "Logging/MessageLog.h"
-#include "Application/ThrottleManager.h"
-#include "HAL/Runnable.h"
-#include "Interfaces/ITargetPlatform.h"
-#include "Interfaces/ITargetPlatformManagerModule.h"
-#include "HAL/RunnableThread.h"
-#include "MeshUtilities.h"
-#include "Misc/ConfigCacheIni.h"
-
-#include "Misc/Paths.h"
+#include "Engine/Texture.h"
 #include "HAL/FileManager.h"
-#include "HAL/PlatformFileManager.h"
-
-#include "MuCOE/UnrealEditorPortabilityHelpers.h"
-#include "MuCO/UnrealPortabilityHelpers.h"
-#include "PlatformInfo.h"
-#include "Engine/AssetManager.h"
-#include "AssetRegistry/ARFilter.h"
-
+#include "HAL/PlatformCrt.h"
+#include "HAL/PlatformTime.h"
+#include "Logging/LogCategory.h"
+#include "Logging/LogMacros.h"
+#include "Misc/Guid.h"
+#include "MuCO/CustomizableObject.h"
+#include "MuCO/UnrealMutableModelDiskStreamer.h"
+#include "MuR/Ptr.h"
 #include "MuT/Compiler.h"
 #include "MuT/ErrorLog.h"
-#include "MuR/System.h"
-#include "MuR/Ptr.h"
+#include "Serialization/Archive.h"
+#include "Serialization/MemoryWriter.h"
+#include "Trace/Detail/Channel.h"
+
+class ITargetPlatform;
 
 
 #define LOCTEXT_NAMESPACE "CustomizableObjectEditor"
diff --git a/Engine/Plugins/Experimental/Mutable/Source/CustomizableObjectEditor/Private/MuCOE/CustomizableObjectCompileRunnable.h b/Engine/Plugins/Experimental/Mutable/Source/CustomizableObjectEditor/Private/MuCOE/CustomizableObjectCompileRunnable.h
--- a/Engine/Plugins/Experimental/Mutable/Source/CustomizableObjectEditor/Private/MuCOE/CustomizableObjectCompileRunnable.h
+++ b/Engine/Plugins/Experimental/Mutable/Source/CustomizableObjectEditor/Private/MuCOE/CustomizableObjectCompileRunnable.h
@@ -2,19 +2,18 @@
 
 #pragma once
 
+#include "Containers/Array.h"
+#include "Containers/ContainersFwd.h"
+#include "Containers/UnrealString.h"
+#include "HAL/Platform.h"
 #include "HAL/Runnable.h"
-#include "HAL/RunnableThread.h"
-#include "Logging/TokenizedMessage.h"
+#include "Internationalization/Text.h"
 #include "MuCO/CustomizableObject.h"
-#include "Widgets/Notifications/SNotificationList.h"
-
-#include "MuT/Node.h"
 #include "MuR/Model.h"
+#include "MuT/Node.h"
+#include "Templates/SharedPointer.h"
 
-
-class UCustomizableObject;
-class UCustomizableObjectNodeObject;
-class CustomizableObjectCompileAsyncTask;
+class ITargetPlatform;
 
 
 class FCustomizableObjectCompileRunnable : public FRunnable
diff --git a/Engine/Plugins/Experimental/Mutable/Source/CustomizableObjectEditor/Private/MuCOE/CustomizableObjectCompiler.cpp b/Engine/Plugins/Experimental/Mutable/Source/CustomizableObjectEditor/Private/MuCOE/CustomizableObjectCompiler.cpp
--- a/Engine/Plugins/Experimental/Mutable/Source/CustomizableObjectEditor/Private/MuCOE/CustomizableObjectCompiler.cpp
+++ b/Engine/Plugins/Experimental/Mutable/Source/CustomizableObjectEditor/Private/MuCOE/CustomizableObjectCompiler.cpp
@@ -1,52 +1,81 @@
 // Copyright Epic Games, Inc. All Rights Reserved.
 
 #include "MuCOE/CustomizableObjectCompiler.h"
-#include "MuCOE/CustomizableObjectEditor.h"
-#include "MuCOE/CustomizableObjectEditorModule.h"
-#include "MuCO/CustomizableObjectSystem.h"
-#include "MuCO/CustomizableObject.h"
-#include "MuCO/ICustomizableObjectModule.h"
-#include "MuCOE/GraphTraversal.h"
 
-#include "StaticMeshResources.h"
-#include "Framework/Notifications/NotificationManager.h"
-#include "Widgets/Notifications/SNotificationList.h"
-#include "AssetRegistry/AssetRegistryModule.h"
-#include "MessageLogModule.h"
-#include "Logging/MessageLog.h"
-#include "Misc/UObjectToken.h"
-#include "EdGraphToken.h"
-#include "Application/ThrottleManager.h"
-#include "Animation/MorphTarget.h"
-#include "Interfaces/ITargetPlatform.h"
-#include "Interfaces/ITargetPlatformManagerModule.h"
-#include "Engine/SkeletalMesh.h"
-#include "Engine/StaticMesh.h"
 #include "Animation/Skeleton.h"
-#include "Engine/TextureLODSettings.h"
-#include "Rendering/SkinWeightVertexBuffer.h"
-#include "MeshUtilities.h"
-#include "Curves/CurveLinearColor.h"
-#include "PhysicsEngine/PhysicsAsset.h"
-#include "ClothingAsset.h"
+#include "AssetRegistry/ARFilter.h"
+#include "AssetRegistry/AssetRegistryModule.h"
+#include "AssetRegistry/IAssetRegistry.h"
 #include "ClothConfig.h"
-
-#include "MuCOE/UnrealEditorPortabilityHelpers.h"
-#include "MuCO/UnrealPortabilityHelpers.h"
-#include "PlatformInfo.h"
-#include "Engine/AssetManager.h"
+#include "Containers/ArrayView.h"
+#include "Containers/EnumAsByte.h"
+#include "Containers/Set.h"
+#include "Containers/SparseArray.h"
+#include "CoreGlobals.h"
+#include "Delegates/Delegate.h"
+#include "EdGraph/EdGraph.h"
+#include "Engine/SkeletalMesh.h"
 #include "Engine/StreamableManager.h"
-#include "AssetRegistry/ARFilter.h"
-
-#include "MuT/Node.h"
-#include "MuT/Streams.h"
-#include "MuR/Serialisation.h"
-#include "MuR/MutableTrace.h"
-#include "MuR/ImagePrivate.h"
-
+#include "Engine/Texture.h"
+#include "Framework/Notifications/NotificationManager.h"
+#include "GenericPlatform/GenericPlatformAffinity.h"
+#include "HAL/IConsoleManager.h"
+#include "HAL/PlatformCrt.h"
+#include "HAL/PlatformTime.h"
+#include "HAL/RunnableThread.h"
+#include "Interfaces/ITargetPlatformManagerModule.h"
+#include "Internationalization/Internationalization.h"
+#include "Internationalization/Text.h"
+#include "Logging/LogCategory.h"
+#include "Logging/LogMacros.h"
+#include "Materials/MaterialInterface.h"
+#include "Math/UnrealMathSSE.h"
+#include "MessageLogModule.h"
+#include "Misc/App.h"
+#include "Misc/AssertionMacros.h"
+#include "Misc/AssetRegistryInterface.h"
+#include "Misc/CoreMisc.h"
+#include "Misc/Guid.h"
+#include "Modules/ModuleManager.h"
+#include "MuCO/CustomizableObject.h"
+#include "MuCO/CustomizableObjectClothingTypes.h"
+#include "MuCO/CustomizableObjectIdentifier.h"
+#include "MuCO/CustomizableObjectInstance.h"
+#include "MuCO/CustomizableObjectParameterTypeDefinitions.h"
+#include "MuCO/CustomizableObjectSystem.h"
+#include "MuCO/CustomizableObjectUIData.h"
+#include "MuCO/ICustomizableObjectModule.h"
+#include "MuCOE/CustomizableObjectCompileRunnable.h"
+#include "MuCOE/CustomizableObjectEditorLogger.h"
+#include "MuCOE/GenerateMutableSource/GenerateMutableSourceImage.h"
+#include "MuCOE/GraphTraversal.h"
 #include "MuCOE/ICustomizableObjectPopulationModule.h"
+#include "MuCOE/Nodes/CustomizableObjectNode.h"
+#include "MuCOE/Nodes/CustomizableObjectNodeObject.h"
+#include "MuCOE/Nodes/CustomizableObjectNodeObjectGroup.h"
+#include "MuR/Image.h"
+#include "MuR/Model.h"
+#include "MuR/MutableTrace.h"
+#include "MuR/Ptr.h"
+#include "MuR/RefCounted.h"
+#include "MuT/Node.h"
+#include "MuT/NodeImageConstant.h"
+#include "MuT/Table.h"
+#include "PhysicsEngine/PhysicsAsset.h"
+#include "Templates/Casts.h"
+#include "Templates/UnrealTemplate.h"
+#include "Trace/Detail/Channel.h"
+#include "UObject/Class.h"
+#include "UObject/Object.h"
+#include "UObject/ObjectPtr.h"
+#include "UObject/Package.h"
+#include "UObject/SoftObjectPath.h"
+#include "UObject/SoftObjectPtr.h"
+#include "UObject/UObjectGlobals.h"
+#include "UObject/UObjectIterator.h"
+#include "Widgets/Notifications/SNotificationList.h"
 
-#include "MuCOE/GenerateMutableSource/GenerateMutableSourceImage.h"
+class UTexture2D;
 
 #define LOCTEXT_NAMESPACE "CustomizableObjectEditor"
 
diff --git a/Engine/Plugins/Experimental/Mutable/Source/CustomizableObjectEditor/Private/MuCOE/CustomizableObjectCompiler.h b/Engine/Plugins/Experimental/Mutable/Source/CustomizableObjectEditor/Private/MuCOE/CustomizableObjectCompiler.h
--- a/Engine/Plugins/Experimental/Mutable/Source/CustomizableObjectEditor/Private/MuCOE/CustomizableObjectCompiler.h
+++ b/Engine/Plugins/Experimental/Mutable/Source/CustomizableObjectEditor/Private/MuCOE/CustomizableObjectCompiler.h
@@ -2,23 +2,27 @@
 
 #pragma once
 
-#include "MuCOE/CustomizableObjectCompileRunnable.h"
-#include "MuCOE/GenerateMutableSource/GenerateMutableSource.h"
-
+#include "AssetRegistry/AssetData.h"
+#include "Containers/Array.h"
+#include "Containers/Map.h"
+#include "Containers/UnrealString.h"
+#include "HAL/Platform.h"
 #include "Logging/TokenizedMessage.h"
-#include "MuCOE/CompilationMessageCache.h"
 #include "MuCO/CustomizableObject.h"
-#include "Widgets/Notifications/SNotificationList.h"
-#include "Engine/StreamableManager.h"
-
+#include "MuCOE/CompilationMessageCache.h"
+#include "MuCOE/GenerateMutableSource/GenerateMutableSource.h"
+#include "MuT/Node.h"
 #include "MuT/NodeObject.h"
-#include "MuT/NodeImageConstant.h"
-
-
-
-class UCustomizableObject;
-class UCustomizableObjectNodeObject;
-class CustomizableObjectCompileAsyncTask;
+#include "Templates/SharedPointer.h"
+#include "UObject/GCObject.h"
+#include "UObject/NameTypes.h"
+
+class FCustomizableObjectCompileRunnable;
+class FCustomizableObjectSaveDDRunnable;
+class FReferenceCollector;
+class FRunnableThread;
+class FText;
+class UCustomizableObjectNode;
 
 
 class FCustomizableObjectCompiler : public FCustomizableObjectCompilerBase, public FGCObject
diff --git a/Engine/Plugins/Experimental/Mutable/Source/CustomizableObjectEditor/Private/MuCOE/CustomizableObjectCustomSettings.cpp b/Engine/Plugins/Experimental/Mutable/Source/CustomizableObjectEditor/Private/MuCOE/CustomizableObjectCustomSettings.cpp
--- a/Engine/Plugins/Experimental/Mutable/Source/CustomizableObjectEditor/Private/MuCOE/CustomizableObjectCustomSettings.cpp
+++ b/Engine/Plugins/Experimental/Mutable/Source/CustomizableObjectEditor/Private/MuCOE/CustomizableObjectCustomSettings.cpp
@@ -1,25 +1,53 @@
 // Copyright Epic Games, Inc. All Rights Reserved.
 
 #include "MuCOE/CustomizableObjectCustomSettings.h"
-#include "MuCOE/SCustomizableObjectEditorViewport.h"
+
+#include "AssetRegistry/AssetData.h"
+#include "AssetRegistry/AssetRegistryModule.h"
+#include "Components/LightComponent.h"
+#include "Components/PointLightComponent.h"
+#include "Components/SpotLightComponent.h"
+#include "Fonts/SlateFontInfo.h"
+#include "Framework/SlateDelegates.h"
+#include "Framework/Views/ITypedTableView.h"
+#include "HAL/Platform.h"
+#include "HAL/PlatformCrt.h"
+#include "Input/Events.h"
+#include "InputCoreTypes.h"
+#include "Internationalization/Text.h"
+#include "Layout/Margin.h"
+#include "Layout/Visibility.h"
+#include "Math/Transform.h"
+#include "Misc/Attribute.h"
+#include "Misc/Paths.h"
 #include "MuCOE/CustomizableObjectEditorViewportClient.h"
+#include "MuCOE/CustomizableObjectEditorViewportLights.h"
+#include "MuCOE/SCustomizableObjectEditorViewport.h"
+#include "PropertyCustomizationHelpers.h"
+#include "SlotBase.h"
+#include "Styling/AppStyle.h"
+#include "Styling/SlateColor.h"
+#include "Templates/Casts.h"
+#include "Templates/UnrealTemplate.h"
+#include "UObject/Package.h"
+#include "Widgets/Colors/SColorBlock.h"
+#include "Widgets/Colors/SColorPicker.h"
+#include "Widgets/DeclarativeSyntaxSupport.h"
+#include "Widgets/Input/SButton.h"
+#include "Widgets/Input/SComboButton.h"
 #include "Widgets/Input/SEditableTextBox.h"
-#include "Widgets/Text/STextBlock.h"
+#include "Widgets/Input/SNumericEntryBox.h"
 #include "Widgets/Layout/SBorder.h"
-#include "Widgets/Views/STableViewBase.h"
-#include "Widgets/Views/STableRow.h"
 #include "Widgets/Layout/SExpandableArea.h"
-#include "Widgets/Input/SNumericEntryBox.h"
-#include "Widgets/Input/SButton.h"
-#include "Widgets/Colors/SColorPicker.h"
-#include "Widgets/Colors/SColorBlock.h"
-#include "AssetRegistry/AssetRegistryModule.h"
-#include "MuCOE/CustomizableObjectEditorSettings.h"
-#include "MuCOE/CustomizableObjectPreviewScene.h"
-#include "Styling/AppStyle.h"
-#include "PropertyCustomizationHelpers.h"
-#include "Components/PointLightComponent.h"
-#include "Components/SpotLightComponent.h"
+#include "Widgets/Text/STextBlock.h"
+#include "Widgets/Views/STableRow.h"
+
+class ITableRow;
+class STableViewBase;
+class SWidget;
+class UClass;
+class UFactory;
+struct FGeometry;
 
 const FString SCustomizableObjectCustomSettings::LightPackagePath = "/Game/EditorAssets/Mutable/";
 
diff --git a/Engine/Plugins/Experimental/Mutable/Source/CustomizableObjectEditor/Private/MuCOE/CustomizableObjectCustomSettings.h b/Engine/Plugins/Experimental/Mutable/Source/CustomizableObjectEditor/Private/MuCOE/CustomizableObjectCustomSettings.h
--- a/Engine/Plugins/Experimental/Mutable/Source/CustomizableObjectEditor/Private/MuCOE/CustomizableObjectCustomSettings.h
+++ b/Engine/Plugins/Experimental/Mutable/Source/CustomizableObjectEditor/Private/MuCOE/CustomizableObjectCustomSettings.h
@@ -2,14 +2,29 @@
 
 #pragma once
 
-// include "SComboBox.h"
-//#include "SBoxPanel.h"
+#include "Containers/Array.h"
+#include "Containers/BitArray.h"
+#include "Containers/Set.h"
+#include "Containers/SparseArray.h"
+#include "Containers/UnrealString.h"
+#include "Delegates/Delegate.h"
+#include "Input/Reply.h"
+#include "Math/Color.h"
+#include "Misc/Optional.h"
+#include "Templates/SharedPointer.h"
+#include "Types/SlateEnums.h"
+#include "UObject/Object.h"
+#include "UObject/UObjectGlobals.h"
+#include "Widgets/SBoxPanel.h"
 #include "Widgets/Views/SListView.h"
-#include "AssetRegistry/AssetData.h"
-//#include "Engine/DataTable.h"
-#include "MuCOE/CustomizableObjectEditorViewportLights.h"
+
 #include "CustomizableObjectCustomSettings.generated.h"
 
+class SWidget;
+struct FAssetData;
+struct FGeometry;
+struct FPointerEvent;
+
 UCLASS()
 class UCustomizableObjectEmptyClassForSettings : public UObject
 {
diff --git a/Engine/Plugins/Experimental/Mutable/Source/CustomizableObjectEditor/Private/MuCOE/CustomizableObjectCustomSettingsDetails.cpp b/Engine/Plugins/Experimental/Mutable/Source/CustomizableObjectEditor/Private/MuCOE/CustomizableObjectCustomSettingsDetails.cpp
--- a/Engine/Plugins/Experimental/Mutable/Source/CustomizableObjectEditor/Private/MuCOE/CustomizableObjectCustomSettingsDetails.cpp
+++ b/Engine/Plugins/Experimental/Mutable/Source/CustomizableObjectEditor/Private/MuCOE/CustomizableObjectCustomSettingsDetails.cpp
@@ -1,10 +1,17 @@
 // Copyright Epic Games, Inc. All Rights Reserved.
 
 #include "MuCOE/CustomizableObjectCustomSettingsDetails.h"
+
+#include "Containers/Array.h"
 #include "DetailCategoryBuilder.h"
 #include "DetailLayoutBuilder.h"
 #include "DetailWidgetRow.h"
+#include "IDetailsView.h"
+#include "Internationalization/Text.h"
 #include "MuCOE/CustomizableObjectCustomSettings.h"
+#include "Templates/Casts.h"
+#include "UObject/WeakObjectPtrTemplates.h"
+#include "Widgets/DeclarativeSyntaxSupport.h"
 
 TSharedRef<IDetailCustomization> FCustomizableObjectCustomSettingsDetails::MakeInstance()
 {
diff --git a/Engine/Plugins/Experimental/Mutable/Source/CustomizableObjectEditor/Private/MuCOE/CustomizableObjectCustomSettingsDetails.h b/Engine/Plugins/Experimental/Mutable/Source/CustomizableObjectEditor/Private/MuCOE/CustomizableObjectCustomSettingsDetails.h
--- a/Engine/Plugins/Experimental/Mutable/Source/CustomizableObjectEditor/Private/MuCOE/CustomizableObjectCustomSettingsDetails.h
+++ b/Engine/Plugins/Experimental/Mutable/Source/CustomizableObjectEditor/Private/MuCOE/CustomizableObjectCustomSettingsDetails.h
@@ -3,6 +3,9 @@
 #pragma once
 
 #include "IDetailCustomization.h"
+#include "Templates/SharedPointer.h"
+
+class IDetailLayoutBuilder;
 
 class FCustomizableObjectCustomSettingsDetails : public IDetailCustomization
 {
diff --git a/Engine/Plugins/Experimental/Mutable/Source/CustomizableObjectEditor/Private/MuCOE/CustomizableObjectDebugger.cpp b/Engine/Plugins/Experimental/Mutable/Source/CustomizableObjectEditor/Private/MuCOE/CustomizableObjectDebugger.cpp
--- a/Engine/Plugins/Experimental/Mutable/Source/CustomizableObjectEditor/Private/MuCOE/CustomizableObjectDebugger.cpp
+++ b/Engine/Plugins/Experimental/Mutable/Source/CustomizableObjectEditor/Private/MuCOE/CustomizableObjectDebugger.cpp
@@ -1,47 +1,24 @@
 // Copyright Epic Games, Inc. All Rights Reserved.
 
 #include "MuCOE/CustomizableObjectDebugger.h"
-#include "MuCOE/CustomizableObjectDebuggerActions.h"
+
+#include "Framework/Docking/TabManager.h"
+#include "Internationalization/Internationalization.h"
+#include "Logging/LogMacros.h"
+#include "Misc/Attribute.h"
+#include "MuCO/CustomizableObject.h"
 #include "MuCOE/CustomizableObjectEditorModule.h"
 #include "MuCOE/CustomizableObjectEditorStyle.h"
-#include "MuCOE/CustomizableObjectEditor.h"
-
-#include "MuCOE/SCustomizableObjectProperties.h"
-#include "MuCOE/SMutableGraphViewer.h"
-#include "MuCOE/SMutableCodeViewer.h"
 #include "MuCOE/SMutableObjectViewer.h"
-
-#include "Editor.h"
-#include "AssetRegistry/AssetRegistryModule.h"
-#include "DesktopPlatformModule.h"
-#include "ContentBrowserModule.h"
-#include "PropertyEditorModule.h"
-
-#include "Widgets/Notifications/SNotificationList.h"
-#include "Widgets/Docking/SDockTab.h"
-#include "Widgets/Input/STextComboBox.h"
-#include "Widgets/Layout/SSplitter.h"
+#include "Styling/ISlateStyle.h"
+#include "Toolkits/AssetEditorToolkit.h"
 #include "Types/SlateEnums.h"
+#include "UObject/Object.h"
+#include "UObject/UObjectGlobals.h"
+#include "Widgets/DeclarativeSyntaxSupport.h"
+#include "Widgets/Docking/SDockTab.h"
 
-#include "PreviewScene.h"
-#include "ScopedTransaction.h"
-#include "FileHelpers.h"
-#include "BusyCursor.h"
-#include "EditorDirectories.h"
-
-#include "WorkspaceMenuStructureModule.h"
-#include "WorkspaceMenuStructure.h"
-#include "Engine/Selection.h"
-#include "Components/SkeletalMeshComponent.h"
-
-#include "IDetailsView.h"
-#include "WorkspaceMenuStructure.h"
-#include "Framework/Commands/GenericCommands.h"
-#include "Framework/Notifications/NotificationManager.h"
-
-#include "GenericPlatform/GenericPlatformApplicationMisc.h"
-
-#include "HAL/PlatformApplicationMisc.h"
+struct FSlateBrush;
 
 
 #define LOCTEXT_NAMESPACE "CustomizableObjectDebugger"
diff --git a/Engine/Plugins/Experimental/Mutable/Source/CustomizableObjectEditor/Private/MuCOE/CustomizableObjectDebugger.h b/Engine/Plugins/Experimental/Mutable/Source/CustomizableObjectEditor/Private/MuCOE/CustomizableObjectDebugger.h
--- a/Engine/Plugins/Experimental/Mutable/Source/CustomizableObjectEditor/Private/MuCOE/CustomizableObjectDebugger.h
+++ b/Engine/Plugins/Experimental/Mutable/Source/CustomizableObjectEditor/Private/MuCOE/CustomizableObjectDebugger.h
@@ -2,10 +2,21 @@
 
 #pragma once
 
-#include "MuCOE/ICustomizableObjectDebugger.h"
+#include "Containers/UnrealString.h"
+#include "Delegates/Delegate.h"
+#include "HAL/Platform.h"
+#include "Internationalization/Text.h"
+#include "Math/Color.h"
 #include "Misc/NotifyHook.h"
-#include "EditorUndoClient.h"
-#include "Toolkits/AssetEditorToolkit.h"
+#include "MuCOE/ICustomizableObjectDebugger.h"
+#include "Templates/SharedPointer.h"
+#include "Toolkits/IToolkit.h"
+#include "UObject/GCObject.h"
+#include "UObject/NameTypes.h"
+
+class FReferenceCollector;
+class UCustomizableObject;
+struct FSlateBrush;
 
 
 DECLARE_DELEGATE(FCreatePreviewInstanceFlagDelegate);
diff --git a/Engine/Plugins/Experimental/Mutable/Source/CustomizableObjectEditor/Private/MuCOE/CustomizableObjectDebuggerActions.cpp b/Engine/Plugins/Experimental/Mutable/Source/CustomizableObjectEditor/Private/MuCOE/CustomizableObjectDebuggerActions.cpp
--- a/Engine/Plugins/Experimental/Mutable/Source/CustomizableObjectEditor/Private/MuCOE/CustomizableObjectDebuggerActions.cpp
+++ b/Engine/Plugins/Experimental/Mutable/Source/CustomizableObjectEditor/Private/MuCOE/CustomizableObjectDebuggerActions.cpp
@@ -1,10 +1,13 @@
 // Copyright Epic Games, Inc. All Rights Reserved.
 
 #include "MuCOE/CustomizableObjectDebuggerActions.h"
-#include "MuCOE/CustomizableObjectEditor.h"
-#include "Styling/AppStyle.h"
-#include "MuCOE/CustomizableObjectEditorStyle.h"
+
 #include "Framework/Commands/InputChord.h"
+#include "Framework/Commands/UICommandInfo.h"
+#include "Internationalization/Internationalization.h"
+#include "MuCOE/CustomizableObjectEditorStyle.h"
+#include "UObject/NameTypes.h"
+#include "UObject/UnrealNames.h"
 
 #define LOCTEXT_NAMESPACE "CustomizableObjectEditor"
 
diff --git a/Engine/Plugins/Experimental/Mutable/Source/CustomizableObjectEditor/Private/MuCOE/CustomizableObjectDetails.cpp b/Engine/Plugins/Experimental/Mutable/Source/CustomizableObjectEditor/Private/MuCOE/CustomizableObjectDetails.cpp
--- a/Engine/Plugins/Experimental/Mutable/Source/CustomizableObjectEditor/Private/MuCOE/CustomizableObjectDetails.cpp
+++ b/Engine/Plugins/Experimental/Mutable/Source/CustomizableObjectEditor/Private/MuCOE/CustomizableObjectDetails.cpp
@@ -1,17 +1,41 @@
 // Copyright Epic Games, Inc. All Rights Reserved.
 
 #include "MuCOE/CustomizableObjectDetails.h"
-#include "MuCOE/CustomizableObjectEditorModule.h"
 
-#include "DetailLayoutBuilder.h"
+#include "Containers/UnrealString.h"
 #include "DetailCategoryBuilder.h"
-#include "PropertyCustomizationHelpers.h"
-#include "LevelEditor.h"
-#include "LevelEditorActions.h"
-
-#include "Runtime/Launch/Resources/Version.h"
-
+#include "DetailLayoutBuilder.h"
+#include "DetailWidgetRow.h"
+#include "Framework/Views/ITypedTableView.h"
+#include "HAL/Platform.h"
+#include "HAL/PlatformCrt.h"
+#include "IDetailsView.h"
+#include "Internationalization/Internationalization.h"
+#include "Internationalization/Text.h"
+#include "Layout/Children.h"
+#include "Layout/Visibility.h"
+#include "Math/Color.h"
+#include "Misc/Attribute.h"
+#include "MuCO/CustomizableObject.h"
 #include "MuR/Model.h"
+#include "MuR/Parameters.h"
+#include "MuR/Ptr.h"
+#include "PropertyHandle.h"
+#include "SlotBase.h"
+#include "Styling/SlateColor.h"
+#include "Templates/Casts.h"
+#include "Types/SlateEnums.h"
+#include "UObject/WeakObjectPtr.h"
+#include "UObject/WeakObjectPtrTemplates.h"
+#include "Widgets/DeclarativeSyntaxSupport.h"
+#include "Widgets/SBoxPanel.h"
+#include "Widgets/SCompoundWidget.h"
+#include "Widgets/Text/STextBlock.h"
+#include "Widgets/Views/SExpanderArrow.h"
+#include "Widgets/Views/SHeaderRow.h"
+#include "Widgets/Views/STableRow.h"
+
+class ITableRow;
 
 
 #define LOCTEXT_NAMESPACE "CustomizableObjectDetails"
diff --git a/Engine/Plugins/Experimental/Mutable/Source/CustomizableObjectEditor/Private/MuCOE/CustomizableObjectDetails.h b/Engine/Plugins/Experimental/Mutable/Source/CustomizableObjectEditor/Private/MuCOE/CustomizableObjectDetails.h
--- a/Engine/Plugins/Experimental/Mutable/Source/CustomizableObjectEditor/Private/MuCOE/CustomizableObjectDetails.h
+++ b/Engine/Plugins/Experimental/Mutable/Source/CustomizableObjectEditor/Private/MuCOE/CustomizableObjectDetails.h
@@ -2,13 +2,21 @@
 
 #pragma once
 
+#include "Containers/Array.h"
+#include "Containers/BitArray.h"
+#include "Containers/Set.h"
+#include "Containers/SparseArray.h"
+#include "Delegates/Delegate.h"
 #include "IDetailCustomization.h"
-#include "Widgets/Views/STreeView.h"
-
+#include "Misc/Optional.h"
 #include "Templates/SharedPointer.h"
-#include "UObject/WeakObjectPtr.h"
+#include "Types/SlateConstants.h"
+#include "Widgets/Views/STableViewBase.h"
+#include "Widgets/Views/STreeView.h"
 
 class FStateDetailsNode;
+class IDetailLayoutBuilder;
+class ITableRow;
 
 
 class FCustomizableObjectDetails : public IDetailCustomization
diff --git a/Engine/Plugins/Experimental/Mutable/Source/CustomizableObjectEditor/Private/MuCOE/CustomizableObjectEditor.cpp b/Engine/Plugins/Experimental/Mutable/Source/CustomizableObjectEditor/Private/MuCOE/CustomizableObjectEditor.cpp
--- a/Engine/Plugins/Experimental/Mutable/Source/CustomizableObjectEditor/Private/MuCOE/CustomizableObjectEditor.cpp
+++ b/Engine/Plugins/Experimental/Mutable/Source/CustomizableObjectEditor/Private/MuCOE/CustomizableObjectEditor.cpp
@@ -1,86 +1,137 @@
 // Copyright Epic Games, Inc. All Rights Reserved.
 
 #include "MuCOE/CustomizableObjectEditor.h"
-#include "MuCOE/CustomizableObjectEditorModule.h"
-#include "MuCOE/CustomizableObjectEditorStyle.h"
 
+#include "AdvancedPreviewSceneModule.h"
+#include "Animation/AnimSingleNodeInstance.h"
+#include "Animation/DebugSkelMeshComponent.h"
+#include "Animation/PoseAsset.h"
+#include "Animation/SmartName.h"
+#include "AssetRegistry/ARFilter.h"
+#include "AssetRegistry/AssetData.h"
+#include "AssetRegistry/AssetRegistryModule.h"
+#include "AssetRegistry/IAssetRegistry.h"
+#include "Components/SkeletalMeshComponent.h"
+#include "Containers/Set.h"
+#include "DetailsViewArgs.h"
+#include "EdGraph/EdGraph.h"
+#include "EdGraph/EdGraphPin.h"
+#include "EdGraph/EdGraphSchema.h"
+#include "EdGraphNode_Comment.h"
+#include "EdGraphUtilities.h"
+#include "Editor.h"
+#include "Editor/EditorEngine.h"
+#include "Engine/Engine.h"
+#include "Engine/EngineTypes.h"
+#include "Engine/StaticMesh.h"
+#include "Engine/StreamableManager.h"
+#include "FileHelpers.h"
+#include "Framework/Commands/GenericCommands.h"
+#include "Framework/Commands/UIAction.h"
+#include "Framework/Commands/UICommandList.h"
+#include "Framework/Docking/TabManager.h"
+#include "Framework/MultiBox/MultiBoxBuilder.h"
+#include "Framework/MultiBox/MultiBoxExtender.h"
+#include "Framework/Notifications/NotificationManager.h"
+#include "Framework/SlateDelegates.h"
+#include "GraphEditorActions.h"
+#include "HAL/PlatformApplicationMisc.h"
+#include "HAL/PlatformTime.h"
+#include "IDetailsView.h"
+#include "Internationalization/Internationalization.h"
+#include "Layout/BasicLayoutWidgetSlot.h"
+#include "Layout/ChildrenBase.h"
+#include "Layout/SlateRect.h"
+#include "Logging/LogCategory.h"
+#include "Logging/LogMacros.h"
+#include "Logging/TokenizedMessage.h"
+#include "MaterialGraph/MaterialGraphNode.h"
+#include "MaterialGraph/MaterialGraphNode_Comment.h"
+#include "Misc/AssertionMacros.h"
+#include "Misc/Attribute.h"
+#include "Misc/CString.h"
+#include "Misc/Guid.h"
+#include "Modules/ModuleManager.h"
+#include "MuCO/CustomizableObject.h"
 #include "MuCO/CustomizableObjectInstance.h"
+#include "MuCO/CustomizableObjectSystem.h"
+#include "MuCO/CustomizableSkeletalComponent.h"
+#include "MuCOE/CustomizableObjectBakeHelpers.h"
+#include "MuCOE/CustomizableObjectCustomSettings.h"
 #include "MuCOE/CustomizableObjectEditorActions.h"
+#include "MuCOE/CustomizableObjectEditorLogger.h"
+#include "MuCOE/CustomizableObjectEditorModule.h"
+#include "MuCOE/CustomizableObjectEditorNodeContextCommands.h"
+#include "MuCOE/CustomizableObjectEditorStyle.h"
+#include "MuCOE/CustomizableObjectEditorViewportClient.h"
+#include "MuCOE/CustomizableObjectEditor_Deprecated.h"
+#include "MuCOE/CustomizableObjectGraph.h"
 #include "MuCOE/CustomizableObjectPreviewScene.h"
-#include "MuCOE/CustomizableObjectDebugger.h"
+#include "MuCOE/EdGraphSchema_CustomizableObject.h"
+#include "MuCOE/GraphTraversal.h"
+#include "MuCOE/Nodes/CustomizableObjectNode.h"
+#include "MuCOE/Nodes/CustomizableObjectNodeColorArithmeticOp.h"
+#include "MuCOE/Nodes/CustomizableObjectNodeEditMaterial.h"
+#include "MuCOE/Nodes/CustomizableObjectNodeEnumParameter.h"
+#include "MuCOE/Nodes/CustomizableObjectNodeFloatParameter.h"
+#include "MuCOE/Nodes/CustomizableObjectNodeGroupProjectorParameter.h"
+#include "MuCOE/Nodes/CustomizableObjectNodeLayoutBlocks.h"
+#include "MuCOE/Nodes/CustomizableObjectNodeMaterial.h"
+#include "MuCOE/Nodes/CustomizableObjectNodeMaterialVariation.h"
+#include "MuCOE/Nodes/CustomizableObjectNodeMeshClipMorph.h"
+#include "MuCOE/Nodes/CustomizableObjectNodeMeshClipWithMesh.h"
+#include "MuCOE/Nodes/CustomizableObjectNodeObject.h"
+#include "MuCOE/Nodes/CustomizableObjectNodeObjectGroup.h"
+#include "MuCOE/Nodes/CustomizableObjectNodeProjectorConstant.h"
+#include "MuCOE/Nodes/CustomizableObjectNodeProjectorParameter.h"
+#include "MuCOE/Nodes/CustomizableObjectNodeRemoveMeshBlocks.h"
+#include "MuCOE/Nodes/CustomizableObjectNodeSkeletalMesh.h"
+#include "MuCOE/Nodes/CustomizableObjectNodeStaticMesh.h"
+#include "MuCOE/Nodes/CustomizableObjectNodeTable.h"
+#include "MuCOE/SCustomizableObjectEditorAdvancedPreviewSettings.h"
+#include "MuCOE/SCustomizableObjectEditorPerformanceReport.h"
+#include "MuCOE/SCustomizableObjectEditorTagExplorer.h"
+#include "MuCOE/SCustomizableObjectEditorTextureAnalyzer.h"
 #include "MuCOE/SCustomizableObjectEditorViewport.h"
-#include "MuCOE/SCustomizableObjectProperties.h"
-#include "MuCOE/SCustomizableObjectSystemVersion.h"
 #include "MuCOE/SCustomizableObjectNodeLayoutBlocksEditor.h"
 #include "MuCOE/SCustomizableObjectNodeLayoutBlocksSelector.h"
-#include "MuCOE/CustomizableObjectPreviewScene.h"
-#include "MuCOE/PinViewer/SPinViewer.h"
-#include "MuCOE/SCustomizableObjectEditorTextureAnalyzer.h"
-#include "MuCOE/SCustomizableObjectEditorPerformanceReport.h"
-#include "MuCOE/SCustomizableObjectEditorTagExplorer.h"
-#include "MuCOE/GraphTraversal.h"
-
-#include "Editor.h"
-#include "AssetRegistry/AssetRegistryModule.h"
-#include "DesktopPlatformModule.h"
-#include "ContentBrowserModule.h"
 #include "PropertyEditorModule.h"
-
-#include "Interfaces/ITargetPlatform.h"
-#include "Interfaces/ITargetPlatformManagerModule.h"
-#include "Widgets/Notifications/SNotificationList.h"
-#include "Widgets/Docking/SDockTab.h"
-#include "Widgets/Input/STextComboBox.h"
-#include "Widgets/Layout/SSplitter.h"
-#include "Widgets/Layout/SSpacer.h"
-#include "Widgets/Input/SSearchBox.h"
-#include "Types/SlateEnums.h"
-#include "Slate/SceneViewport.h"
 #include "SNodePanel.h"
-#include "SGraphNode.h"
-#include "Logging/MessageLog.h"
-#include "Logging/TokenizedMessage.h"
-
-#include "PreviewScene.h"
 #include "ScopedTransaction.h"
-#include "FileHelpers.h"
-#include "BusyCursor.h"
-#include "EditorDirectories.h"
-
-#include "WorkspaceMenuStructureModule.h"
-#include "WorkspaceMenuStructure.h"
-#include "Engine/Selection.h"
-#include "MaterialGraph/MaterialGraphNode.h"
-#include "MaterialGraph/MaterialGraphNode_Comment.h"
-#include "Components/SkeletalMeshComponent.h"
-
-#include "GraphEditorActions.h"
-#include "EdGraphUtilities.h"
-#include "MuCOE/EdGraphSchema_CustomizableObject.h"
-#include "IDetailsView.h"
-#include "WorkspaceMenuStructure.h"
-#include "Framework/Commands/GenericCommands.h"
-#include "Framework/Notifications/NotificationManager.h"
-#include "Framework/MultiBox/MultiBoxExtender.h"
-
-#include "Runtime/Launch/Resources/Version.h"
-#include "AdvancedPreviewSceneModule.h"
-#include "MuCOE/CustomizableObjectCustomSettings.h"
-#include "MuCOE/SCustomizableObjectEditorAdvancedPreviewSettings.h"
-#include "MuCOE/CustomizableObjectEditorNodeContextCommands.h"
-#include "Animation/AnimSingleNodeInstance.h"
-#include "AdvancedPreviewScene.h"
-#include "MuCOE/CustomizableObjectBakeHelpers.h"
-#include "MuCOE/UnrealEditorPortabilityHelpers.h"
-#include "GenericPlatform/GenericPlatformApplicationMisc.h"
-#include "ISinglePropertyView.h"
-#include "Misc/UObjectToken.h"
-
-#include "Widgets/Layout/SScrollBox.h"
+#include "SlotBase.h"
+#include "Styling/AppStyle.h"
+#include "Styling/ISlateStyle.h"
+#include "Templates/Casts.h"
+#include "Templates/Function.h"
+#include "Templates/SubclassOf.h"
+#include "Textures/SlateIcon.h"
+#include "Toolkits/AssetEditorToolkit.h"
+#include "Trace/Detail/Channel.h"
+#include "Types/SlateEnums.h"
+#include "Types/SlateStructs.h"
+#include "UObject/Class.h"
+#include "UObject/EnumProperty.h"
+#include "UObject/Field.h"
+#include "UObject/Package.h"
+#include "UObject/TopLevelAssetPath.h"
+#include "UObject/UnrealNames.h"
+#include "UObject/UnrealType.h"
+#include "Widgets/DeclarativeSyntaxSupport.h"
+#include "Widgets/Docking/SDockTab.h"
+#include "Widgets/Input/SSearchBox.h"
+#include "Widgets/Input/STextComboBox.h"
+#include "Widgets/Layout/SBorder.h"
+#include "Widgets/Layout/SBox.h"
 #include "Widgets/Layout/SScrollBar.h"
+#include "Widgets/Layout/SScrollBox.h"
+#include "Widgets/Layout/SSplitter.h"
+#include "Widgets/Notifications/SNotificationList.h"
+#include "Widgets/SBoxPanel.h"
+
+class FAdvancedPreviewScene;
+class FWorkspaceItem;
+class IToolkitHost;
+class SWidget;
 
 
-#include "HAL/PlatformApplicationMisc.h"
-
-
 #define LOCTEXT_NAMESPACE "CustomizableObjectEditor"
diff --git a/Engine/Plugins/Experimental/Mutable/Source/CustomizableObjectEditor/Private/MuCOE/CustomizableObjectEditor.h b/Engine/Plugins/Experimental/Mutable/Source/CustomizableObjectEditor/Private/MuCOE/CustomizableObjectEditor.h
--- a/Engine/Plugins/Experimental/Mutable/Source/CustomizableObjectEditor/Private/MuCOE/CustomizableObjectEditor.h
+++ b/Engine/Plugins/Experimental/Mutable/Source/CustomizableObjectEditor/Private/MuCOE/CustomizableObjectEditor.h
@@ -2,25 +2,68 @@
 
 #pragma once
 
-#include "MuCOE/ICustomizableObjectEditor.h"
-#include "Misc/NotifyHook.h"
-#include "TickableEditorObject.h"
+#include "Containers/Array.h"
+#include "Containers/EnumAsByte.h"
+#include "Containers/Map.h"
+#include "Containers/UnrealString.h"
+#include "Delegates/Delegate.h"
+#include "EdGraph/EdGraphNode.h"
 #include "EditorUndoClient.h"
-#include "Toolkits/AssetEditorToolkit.h"
+#include "Engine/TextureDefines.h"
 #include "GraphEditor.h"
+#include "Internationalization/Text.h"
+#include "Math/Color.h"
+#include "Math/UnrealMathSSE.h"
+#include "Math/Vector.h"
+#include "Math/Vector2D.h"
+#include "Misc/NotifyHook.h"
+#include "MuCO/CustomizableObjectParameterTypeDefinitions.h"
 #include "MuCOE/CustomizableObjectCompiler.h"
-#include "Widgets/Input/SNumericEntryBox.h"
-#include "Framework/MultiBox/MultiBoxBuilder.h"
-#include "MuCOE/StatePerformanceTesting.h"
-#include "MuCO/CustomizableObjectInstance.h"
-#include "MuCOE/SCustomizableObjectEditorViewport.h"
-#include "MuCOE/CustomizableObjectEditorViewportClient.h"
-#include "Animation/DebugSkelMeshComponent.h"
+#include "MuCOE/ICustomizableObjectEditor.h"
+#include "PixelFormat.h"
+#include "Stats/Stats2.h"
+#include "Templates/SharedPointer.h"
+#include "TickableEditorObject.h"
+#include "Toolkits/IToolkit.h"
+#include "Types/SlateEnums.h"
+#include "UObject/GCObject.h"
+#include "UObject/NameTypes.h"
+#include "UObject/Object.h"
 #include "UObject/ObjectMacros.h"
-#include "Templates/Function.h"
+#include "UObject/ObjectPtr.h"
+#include "UObject/UObjectGlobals.h"
 
 #include "CustomizableObjectEditor.generated.h"
 
+class FCustomizableObjectEditorViewportClient;
+class FProperty;
+class FSpawnTabArgs;
+class FUICommandList;
+class IToolkitHost;
+class SCustomizableObjectEditorViewportTabBody;
+class SDockTab;
+class STextComboBox;
+class SWidget;
+class UCustomizableObject;
+class UCustomizableObjectInstance;
+class UCustomizableObjectNode;
+class UCustomizableObjectNodeMeshClipMorph;
+class UCustomizableObjectNodeMeshClipWithMesh;
+class UCustomizableObjectNodeObject;
+class UCustomizableObjectNodeProjectorConstant;
+class UCustomizableObjectNodeProjectorParameter;
+class UCustomizableSkeletalComponent;
+class UDebugSkelMeshComponent;
+class UEdGraph;
+class UMaterialInterface;
+class UPoseAsset;
+class UStaticMeshComponent;
+class UTexture;
+struct FAssetData;
+struct FFrame;
+struct FPropertyChangedEvent;
+template <typename FuncType> class TFunction;
+
 DECLARE_DELEGATE(FCreatePreviewInstanceFlagDelegate);
 
 
diff --git a/Engine/Plugins/Experimental/Mutable/Source/CustomizableObjectEditor/Private/MuCOE/CustomizableObjectEditorActions.cpp b/Engine/Plugins/Experimental/Mutable/Source/CustomizableObjectEditor/Private/MuCOE/CustomizableObjectEditorActions.cpp
--- a/Engine/Plugins/Experimental/Mutable/Source/CustomizableObjectEditor/Private/MuCOE/CustomizableObjectEditorActions.cpp
+++ b/Engine/Plugins/Experimental/Mutable/Source/CustomizableObjectEditor/Private/MuCOE/CustomizableObjectEditorActions.cpp
@@ -1,11 +1,13 @@
 // Copyright Epic Games, Inc. All Rights Reserved.
 
 #include "MuCOE/CustomizableObjectEditorActions.h"
-#include "MuCOE/CustomizableObjectEditor.h"
-#include "MuCOE/SCustomizableObjectEditorViewport.h"
-#include "Styling/AppStyle.h"
-#include "MuCOE/CustomizableObjectEditorStyle.h"
+
 #include "Framework/Commands/InputChord.h"
+#include "Framework/Commands/UICommandInfo.h"
+#include "Internationalization/Internationalization.h"
+#include "MuCOE/CustomizableObjectEditorStyle.h"
+#include "UObject/NameTypes.h"
+#include "UObject/UnrealNames.h"
 
 #define LOCTEXT_NAMESPACE "CustomizableObjectEditor"
 
diff --git a/Engine/Plugins/Experimental/Mutable/Source/CustomizableObjectEditor/Private/MuCOE/CustomizableObjectEditorLevelProfile.cpp b/Engine/Plugins/Experimental/Mutable/Source/CustomizableObjectEditor/Private/MuCOE/CustomizableObjectEditorLevelProfile.cpp
--- a/Engine/Plugins/Experimental/Mutable/Source/CustomizableObjectEditor/Private/MuCOE/CustomizableObjectEditorLevelProfile.cpp
+++ b/Engine/Plugins/Experimental/Mutable/Source/CustomizableObjectEditor/Private/MuCOE/CustomizableObjectEditorLevelProfile.cpp
@@ -1,15 +1,32 @@
 // Copyright Epic Games, Inc. All Rights Reserved.
 
 #include "MuCOE/CustomizableObjectEditorLevelProfile.h"
-#include "MuCOE/CustomizableObjectEditorSettings.h"
-#include "Engine/SkyLight.h"
+
+#include "AssetViewerSettings.h"
+#include "Components/LightComponent.h"
+#include "Components/SkyLightComponent.h"
 #include "Engine/DirectionalLight.h"
 #include "Engine/PostProcessVolume.h"
+#include "Engine/Scene.h"
+#include "Engine/SkyLight.h"
+#include "Engine/TextureCube.h"
+#include "Engine/World.h"
 #include "EngineUtils.h"
-#include "Misc/MessageDialog.h"
 #include "GameFramework/Actor.h"
-#include "Components/SkyLightComponent.h"
-#include "Components/LightComponent.h"
+#include "HAL/PlatformCrt.h"
+#include "HAL/PlatformMisc.h"
+#include "Internationalization/Internationalization.h"
+#include "Math/Color.h"
+#include "Math/Rotator.h"
+#include "Math/UnrealMathSSE.h"
+#include "Math/Vector.h"
+#include "Misc/MessageDialog.h"
+#include "Templates/Casts.h"
+#include "UObject/ObjectPtr.h"
+#include "UObject/SoftObjectPtr.h"
+#include "UObject/UObjectGlobals.h"
+
+class UClass;
 
 
 #define LOCTEXT_NAMESPACE "LevelSettingsPreview"
diff --git a/Engine/Plugins/Experimental/Mutable/Source/CustomizableObjectEditor/Private/MuCOE/CustomizableObjectEditorLogger.cpp b/Engine/Plugins/Experimental/Mutable/Source/CustomizableObjectEditor/Private/MuCOE/CustomizableObjectEditorLogger.cpp
--- a/Engine/Plugins/Experimental/Mutable/Source/CustomizableObjectEditor/Private/MuCOE/CustomizableObjectEditorLogger.cpp
+++ b/Engine/Plugins/Experimental/Mutable/Source/CustomizableObjectEditor/Private/MuCOE/CustomizableObjectEditorLogger.cpp
@@ -2,12 +2,27 @@
 
 #include "MuCOE/CustomizableObjectEditorLogger.h"
 
+#include "Containers/UnrealString.h"
+#include "Delegates/Delegate.h"
+#include "Editor.h"
+#include "Editor/EditorEngine.h"
+#include "Framework/Notifications/NotificationManager.h"
+#include "HAL/PlatformCrt.h"
+#include "Internationalization/Internationalization.h"
+#include "Logging/MessageLog.h"
+#include "MessageLogModule.h"
+#include "Misc/AssertionMacros.h"
+#include "Misc/UObjectToken.h"
+#include "Modules/ModuleManager.h"
 #include "MuCOE/CustomizableObjectEditorModule.h"
 #include "MuCOE/CustomizableObjectGraph.h"
 #include "MuCOE/ICustomizableObjectEditor.h"
-#include "MessageLogModule.h"
-#include "Framework/Notifications/NotificationManager.h"
-#include "Misc/UObjectToken.h"
+#include "Subsystems/AssetEditorSubsystem.h"
+#include "Templates/Casts.h"
+#include "Templates/UnrealTemplate.h"
+#include "UObject/NameTypes.h"
+#include "UObject/Object.h"
+#include "UObject/WeakObjectPtr.h"
 #include "Widgets/Notifications/SNotificationList.h"
 
 #define LOCTEXT_NAMESPACE "CustomizableObject"
diff --git a/Engine/Plugins/Experimental/Mutable/Source/CustomizableObjectEditor/Private/MuCOE/CustomizableObjectEditorModule.cpp b/Engine/Plugins/Experimental/Mutable/Source/CustomizableObjectEditor/Private/MuCOE/CustomizableObjectEditorModule.cpp
--- a/Engine/Plugins/Experimental/Mutable/Source/CustomizableObjectEditor/Private/MuCOE/CustomizableObjectEditorModule.cpp
+++ b/Engine/Plugins/Experimental/Mutable/Source/CustomizableObjectEditor/Private/MuCOE/CustomizableObjectEditorModule.cpp
@@ -1,66 +1,104 @@
 // Copyright Epic Games, Inc. All Rights Reserved.
 
 #include "MuCOE/CustomizableObjectEditorModule.h"
-#include "MuCOE/CustomizableObjectEditor.h"
-#include "MuCOE/CustomizableObjectInstanceEditor.h"
-#include "MuCOE/CustomizableObjectDebugger.h"
+
+#include "AssetToolsModule.h"
+#include "Containers/Array.h"
+#include "Containers/EnumAsByte.h"
+#include "Containers/IndirectArray.h"
+#include "Delegates/Delegate.h"
+#include "Editor.h"
+#include "Editor/EditorEngine.h"
+#include "Engine/Engine.h"
+#include "Engine/EngineTypes.h"
+#include "Engine/World.h"
+#include "GameFramework/Pawn.h"
+#include "HAL/IConsoleManager.h"
+#include "HAL/PlatformCrt.h"
+#include "HAL/PlatformMath.h"
+#include "HAL/UnrealMemory.h"
+#include "IAssetTools.h"
+#include "IAssetTypeActions.h"
+#include "ISettingsModule.h"
+#include "ISettingsSection.h"
+#include "Internationalization/Internationalization.h"
+#include "Kismet/GameplayStatics.h"
+#include "MessageLogModule.h"
+#include "Misc/CString.h"
+#include "Modules/ModuleManager.h"
+#include "MuCO/CustomizableObject.h"
+#include "MuCO/CustomizableObjectDGGUI.h"
 #include "MuCO/CustomizableObjectInstance.h"
+#include "MuCO/CustomizableObjectSystem.h"		// For defines related to memory function replacements.
+#include "MuCO/CustomizableSkeletalComponent.h"
 #include "MuCO/CustomizableSkeletalMeshActor.h"
-#include "MuCOE/CustomizableObjectInstanceFactory.h"
-#include "MuCOE/CustomizableObjectDetails.h"
+#include "MuCOE/AssetTypeActions_CustomizableObject.h"
+#include "MuCOE/AssetTypeActions_CustomizableObjectInstance.h"
 #include "MuCOE/CustomizableInstanceDetails.h"
-#include "MuCO/CustomizableObjectSystem.h"		// For defines related to memory function replacements.
+#include "MuCOE/CustomizableObjectCompiler.h"
+#include "MuCOE/CustomizableObjectCustomSettings.h"
+#include "MuCOE/CustomizableObjectCustomSettingsDetails.h"
+#include "MuCOE/CustomizableObjectDebugger.h"
+#include "MuCOE/CustomizableObjectDetails.h"
+#include "MuCOE/CustomizableObjectEditor.h"
+#include "MuCOE/CustomizableObjectEditorLogger.h"
+#include "MuCOE/CustomizableObjectEditorSettings.h"
 #include "MuCOE/CustomizableObjectEditorStyle.h"
+#include "MuCOE/CustomizableObjectIdentifierCustomization.h"
+#include "MuCOE/CustomizableObjectInstanceEditor.h"
+#include "MuCOE/CustomizableObjectInstanceFactory.h"
+#include "MuCOE/CustomizableObjectNodeObjectGroupDetails.h"
 #include "MuCOE/Nodes/CustomizableObjectNodeCopyMaterial.h"
+#include "MuCOE/Nodes/CustomizableObjectNodeCopyMaterialDetails.h"
 #include "MuCOE/Nodes/CustomizableObjectNodeEditMaterial.h"
+#include "MuCOE/Nodes/CustomizableObjectNodeEditMaterialBase.h"
+#include "MuCOE/Nodes/CustomizableObjectNodeEditMaterialBaseDetails.h"
 #include "MuCOE/Nodes/CustomizableObjectNodeExtendMaterial.h"
+#include "MuCOE/Nodes/CustomizableObjectNodeExternalPin.h"
+#include "MuCOE/Nodes/CustomizableObjectNodeExternalPinDetails.h"
 #include "MuCOE/Nodes/CustomizableObjectNodeLayoutBlocks.h"
+#include "MuCOE/Nodes/CustomizableObjectNodeLayoutBlocksDetails.h"
+#include "MuCOE/Nodes/CustomizableObjectNodeMaterial.h"
 #include "MuCOE/Nodes/CustomizableObjectNodeMeshClipMorph.h"
+#include "MuCOE/Nodes/CustomizableObjectNodeMeshClipMorphDetails.h"
+#include "MuCOE/Nodes/CustomizableObjectNodeMeshClipWithMesh.h"
+#include "MuCOE/Nodes/CustomizableObjectNodeMeshClipWithMeshDetails.h"
 #include "MuCOE/Nodes/CustomizableObjectNodeMeshMorph.h"
+#include "MuCOE/Nodes/CustomizableObjectNodeMeshMorphDetails.h"
+#include "MuCOE/Nodes/CustomizableObjectNodeMeshReshapeCommon.h"
+#include "MuCOE/Nodes/CustomizableObjectNodeMeshReshapeSelectionDetails.h"
 #include "MuCOE/Nodes/CustomizableObjectNodeMorphMaterial.h"
-#include "MuCOE/Nodes/CustomizableObjectNodeProjectorConstant.h"
-#include "MuCOE/Nodes/CustomizableObjectNodeProjectorParameter.h"
-#include "MuCOE/Nodes/CustomizableObjectNodeMeshReshape.h"
-#include "MuCOE/Nodes/CustomizableObjectNodeLayoutBlocksDetails.h"
-#include "MuCOE/Nodes/CustomizableObjectNodeParentedMaterialDetails.h"
-#include "MuCOE/Nodes/CustomizableObjectNodeEditMaterialBaseDetails.h"
 #include "MuCOE/Nodes/CustomizableObjectNodeMorphMaterialDetails.h"
+#include "MuCOE/Nodes/CustomizableObjectNodeObject.h"
 #include "MuCOE/Nodes/CustomizableObjectNodeObjectDetails.h"
-#include "MuCOE/CustomizableObjectNodeObjectGroupDetails.h"
-#include "MuCOE/Nodes/CustomizableObjectNodeMeshMorphDetails.h"
+#include "MuCOE/Nodes/CustomizableObjectNodeObjectGroup.h"
+#include "MuCOE/Nodes/CustomizableObjectNodeParentedMaterialDetails.h"
+#include "MuCOE/Nodes/CustomizableObjectNodePinViewerDetails.h"
+#include "MuCOE/Nodes/CustomizableObjectNodeProjectorConstant.h"
+#include "MuCOE/Nodes/CustomizableObjectNodeProjectorParameter.h"
 #include "MuCOE/Nodes/CustomizableObjectNodeProjectorParameterDetails.h"
-#include "MuCOE/Nodes/CustomizableObjectNodeMeshClipMorphDetails.h"
-#include "MuCOE/Nodes/CustomizableObjectNodeMeshClipWithMeshDetails.h"
-#include "MuCOE/Nodes/CustomizableObjectNodeExternalPinDetails.h"
-#include "MuCOE/Nodes/CustomizableObjectNodeCopyMaterialDetails.h"
+#include "MuCOE/Nodes/CustomizableObjectNodeSkeletalMesh.h"
 #include "MuCOE/Nodes/CustomizableObjectNodeSkeletalMeshDetails.h"
-#include "MuCOE/Nodes/CustomizableObjectNodePinViewerDetails.h"
-#include "MuCOE/Nodes/CustomizableObjectNodeTableDetails.h"
-#include "MuCOE/Nodes/CustomizableObjectNodeMeshReshapeSelectionDetails.h"
-#include "MuCOE/CustomizableObjectIdentifierCustomization.h"
-#include "Modules/ModuleManager.h"
-#include "Interfaces/IPluginManager.h"
-#include "AssetToolsModule.h"
-#include "MuCOE/AssetTypeActions_CustomizableObject.h"
-#include "MuCOE/AssetTypeActions_CustomizableObjectInstance.h"
-#include "Styling/SlateStyle.h"
-#include "Brushes/SlateImageBrush.h"
-#include "MessageLogModule.h"
-#include "Logging/MessageLog.h"
-#include "MuCOE/CustomizableObjectEditorSettings.h"
-#include "ISettingsModule.h"
-#include "ISettingsSection.h"
-#include "MuCOE/CustomizableObjectCustomSettingsDetails.h"
-#include "AssetRegistry/AssetRegistryModule.h"
-#include "HAL/FileManager.h"
-#include "CoreGlobals.h"
-#include "MuCOE/CustomizableObjectCustomSettings.h"
-#include "Settings/ProjectPackagingSettings.h"
-#include "MuCOE/UnrealEditorPortabilityHelpers.h"
-#include "CookOnTheSide/CookOnTheFlyServer.h" // This is necessary to register the cook delegate.
 #include "MuCOE/Nodes/CustomizableObjectNodeStaticMesh.h"
-#include "MuCO/CustomizableObjectDGGUI.h"
-#include "Kismet/GameplayStatics.h"
+#include "MuCOE/Nodes/CustomizableObjectNodeTable.h"
+#include "MuCOE/Nodes/CustomizableObjectNodeTableDetails.h"
+#include "PropertyEditorDelegates.h"
+#include "PropertyEditorModule.h"
+#include "Templates/SubclassOf.h"
+#include "Toolkits/AssetEditorToolkit.h"
+#include "UObject/Class.h"
+#include "UObject/ObjectPtr.h"
+#include "UObject/UObjectGlobals.h"
+#include "UObject/UObjectIterator.h"
+#include "UObject/WeakObjectPtr.h"
+
+class AActor;
+class FString;
+class ICustomizableObjectDebugger;
+class ICustomizableObjectEditor;
+class ICustomizableObjectInstanceEditor;
+class IToolkitHost;
+class UObject;
 
 
 const FName CustomizableObjectEditorAppIdentifier = FName(TEXT("CustomizableObjectEditorApp"));
