Ran IWYU on Mutable plugin

diff --git a/Engine/Plugins/Experimental/Mutable/Source/CustomizableObjectEditor/Private/MuCOE/CustomizableObjectEditorNodeContextCommands.cpp b/Engine/Plugins/Experimental/Mutable/Source/CustomizableObjectEditor/Private/MuCOE/CustomizableObjectEditorNodeContextCommands.cpp
--- a/Engine/Plugins/Experimental/Mutable/Source/CustomizableObjectEditor/Private/MuCOE/CustomizableObjectEditorNodeContextCommands.cpp
+++ b/Engine/Plugins/Experimental/Mutable/Source/CustomizableObjectEditor/Private/MuCOE/CustomizableObjectEditorNodeContextCommands.cpp
@@ -2,6 +2,10 @@
 
 #include "MuCOE/CustomizableObjectEditorNodeContextCommands.h"
 
+#include "Framework/Commands/InputChord.h"
+#include "Framework/Commands/UICommandInfo.h"
+#include "InputCoreTypes.h"
+
 #define LOCTEXT_NAMESPACE "CustomizableObjectEditorNodeContextCommands"
 
 void FCustomizableObjectEditorNodeContextCommands::RegisterCommands()
diff --git a/Engine/Plugins/Experimental/Mutable/Source/CustomizableObjectEditor/Private/MuCOE/CustomizableObjectEditorNodeContextCommands.h b/Engine/Plugins/Experimental/Mutable/Source/CustomizableObjectEditor/Private/MuCOE/CustomizableObjectEditorNodeContextCommands.h
--- a/Engine/Plugins/Experimental/Mutable/Source/CustomizableObjectEditor/Private/MuCOE/CustomizableObjectEditorNodeContextCommands.h
+++ b/Engine/Plugins/Experimental/Mutable/Source/CustomizableObjectEditor/Private/MuCOE/CustomizableObjectEditorNodeContextCommands.h
@@ -3,7 +3,14 @@
 #pragma once
 
 #include "Framework/Commands/Commands.h"
+#include "HAL/Platform.h"
+#include "Internationalization/Internationalization.h"
 #include "Styling/AppStyle.h"
+#include "Templates/SharedPointer.h"
+#include "UObject/NameTypes.h"
+#include "UObject/UnrealNames.h"
+
+class FUICommandInfo;
 
 /**
 * Class containing commands for viewport menu actions
diff --git a/Engine/Plugins/Experimental/Mutable/Source/CustomizableObjectEditor/Private/MuCOE/CustomizableObjectEditorStyle.cpp b/Engine/Plugins/Experimental/Mutable/Source/CustomizableObjectEditor/Private/MuCOE/CustomizableObjectEditorStyle.cpp
--- a/Engine/Plugins/Experimental/Mutable/Source/CustomizableObjectEditor/Private/MuCOE/CustomizableObjectEditorStyle.cpp
+++ b/Engine/Plugins/Experimental/Mutable/Source/CustomizableObjectEditor/Private/MuCOE/CustomizableObjectEditorStyle.cpp
@@ -2,16 +2,16 @@
 
 #include "MuCOE/CustomizableObjectEditorStyle.h"
 
-#include "Framework/Application/SlateApplication.h"
+#include "Brushes/SlateImageBrush.h"
+#include "Interfaces/IPluginManager.h"
+#include "Math/Vector2D.h"
+#include "Misc/AssertionMacros.h"
+#include "Misc/Paths.h"
 #include "Styling/AppStyle.h"
-#include "Slate/SlateGameResources.h"
+#include "Styling/ISlateStyle.h"
+#include "Styling/SlateStyle.h"
 #include "Styling/SlateStyleRegistry.h"
 #include "Styling/SlateTypes.h"
-#include "Styling/CoreStyle.h"
-#include "Interfaces/IPluginManager.h"
-#include "Settings/EditorStyleSettings.h"
-#include "Styling/StarshipCoreStyle.h"
-#include "Styling/StyleColors.h"
 
 
 TSharedPtr<FSlateStyleSet> FCustomizableObjectEditorStyle::CustomizableObjectEditorStyleInstance;
diff --git a/Engine/Plugins/Experimental/Mutable/Source/CustomizableObjectEditor/Private/MuCOE/CustomizableObjectEditorViewportClient.cpp b/Engine/Plugins/Experimental/Mutable/Source/CustomizableObjectEditor/Private/MuCOE/CustomizableObjectEditorViewportClient.cpp
--- a/Engine/Plugins/Experimental/Mutable/Source/CustomizableObjectEditor/Private/MuCOE/CustomizableObjectEditorViewportClient.cpp
+++ b/Engine/Plugins/Experimental/Mutable/Source/CustomizableObjectEditor/Private/MuCOE/CustomizableObjectEditorViewportClient.cpp
@@ -1,70 +1,134 @@
 // Copyright Epic Games, Inc. All Rights Reserved.
 
 #include "MuCOE/CustomizableObjectEditorViewportClient.h"
-#include "MuCOE/CustomizableObjectEditorActions.h"
-#include "MuCOE/CustomizableObjectBakeHelpers.h"
 
-#include "MouseDeltaTracker.h"
-#include "PreviewScene.h"
-#include "ISocketManager.h"
-#include "Slate/SceneViewport.h"
-#include "Components/ReflectionCaptureComponent.h"
+#include "AdvancedPreviewScene.h"
+#include "Animation/AnimationAsset.h"
+#include "Animation/DebugSkelMeshComponent.h"
+#include "Animation/PoseAsset.h"
+#include "AssetViewerSettings.h"
+#include "BatchedElements.h"
+#include "CanvasTypes.h"
+#include "Components/LightComponent.h"
+#include "Components/PointLightComponent.h"
 #include "Components/SphereReflectionCaptureComponent.h"
-
-#include "Templates/Decay.h"
-#include "Templates/UnrealTemplate.h"
-
-#include "MuCOE/CustomizableObjectEditor.h"
-#include "BusyCursor.h"
-#include "MeshBuild.h"
-#include "PreviewScene.h"
-#include "ObjectTools.h"
 #include "Components/SpotLightComponent.h"
-#include "Engine/Canvas.h"
-#include "Materials/MaterialInstance.h"
-#include "Materials/MaterialInstanceDynamic.h"
+#include "Components/StaticMeshComponent.h"
+#include "Containers/ArrayView.h"
+#include "Containers/ContainersFwd.h"
+#include "Containers/IndirectArray.h"
+#include "ContentBrowserDelegates.h"
+#include "ContentBrowserModule.h"
+#include "Delegates/Delegate.h"
 #include "DynamicMeshBuilder.h"
-#include "SceneManagement.h"
-#include "Animation/PoseAsset.h"
-#include "AssetRegistry/AssetRegistryModule.h"
-#include "AssetToolsModule.h"
-#include "Preferences/PersonaOptions.h"
-#include "UnrealEdGlobals.h"
+#include "EdGraph/EdGraph.h"
+#include "Editor.h"
+#include "Editor/EditorEngine.h"
+#include "Editor/EditorPerProjectUserSettings.h"
+#include "EditorComponents.h"
 #include "EditorModeManager.h"
-#include "Editor/UnrealEdEngine.h"
-#include "Widgets/Layout/SUniformGridPanel.h"
-#include "Widgets/Input/SEditableTextBox.h"
-#include "ScopedTransaction.h"
-#include "IContentBrowserSingleton.h"
-#include "ContentBrowserModule.h"
+#include "Engine/AssetUserData.h"
+#include "Engine/Engine.h"
+#include "Engine/EngineBaseTypes.h"
+#include "Engine/EngineTypes.h"
+#include "Engine/SkeletalMesh.h"
+#include "Engine/StaticMesh.h"
+#include "Engine/Texture.h"
+#include "Engine/Texture2D.h"
 #include "FileHelpers.h"
-
-#include "ScopedTransaction.h"
-#include "IContentBrowserSingleton.h"
-#include "ContentBrowserModule.h"
-#include "Widgets/Input/SButton.h"
-#include "HAL/PlatformFileManager.h"
-#include "HAL/FileManager.h"
-#include "Misc/MessageDialog.h"
-#include "PackageTools.h"
-#include "UnrealEdGlobals.h"
+#include "Fonts/SlateFontInfo.h"
 #include "Framework/Notifications/NotificationManager.h"
-#include "Widgets/Notifications/SNotificationList.h"
-#include "AssetViewerSettings.h"
-#include "UObject/UObjectGlobals.h"
-#include "Editor/EditorPerProjectUserSettings.h"
-#include "AdvancedPreviewScene.h"
-#include "UnrealWidget.h"
-
+#include "HAL/PlatformCrt.h"
+#include "HAL/PlatformMisc.h"
+#include "HitProxies.h"
+#include "IContentBrowserSingleton.h"
+#include "Input/Reply.h"
+#include "InputCoreTypes.h"
+#include "Internationalization/Internationalization.h"
+#include "Internationalization/Text.h"
+#include "Layout/Margin.h"
+#include "Logging/LogCategory.h"
+#include "Logging/LogMacros.h"
+#include "MaterialTypes.h"
+#include "Materials/Material.h"
+#include "Materials/MaterialExpressionTextureBase.h"
+#include "Materials/MaterialInstance.h"
+#include "Materials/MaterialInstanceDynamic.h"
+#include "Materials/MaterialInterface.h"
+#include "Math/Box.h"
+#include "Math/IntPoint.h"
+#include "Math/Transform.h"
+#include "Math/Vector2D.h"
+#include "Math/Vector4.h"
+#include "Misc/AssertionMacros.h"
+#include "Misc/Attribute.h"
+#include "Misc/CString.h"
+#include "Misc/Guid.h"
+#include "Misc/MessageDialog.h"
+#include "Misc/PackageName.h"
+#include "Misc/Paths.h"
+#include "Modules/ModuleManager.h"
+#include "MuCO/CustomizableObject.h"
+#include "MuCO/CustomizableObjectInstance.h"
+#include "MuCO/CustomizableObjectMipDataProvider.h"
+#include "MuCO/UnrealBakeHelpers.h"
 #include "MuCO/UnrealPortabilityHelpers.h"
-#include "MuCOE/UnrealEditorPortabilityHelpers.h"
-#include "FileHelpers.h"
-
+#include "MuCOE/CustomizableObjectBakeHelpers.h"
 #include "MuCOE/CustomizableObjectPreviewScene.h"
+#include "MuCOE/CustomizableObjectWidget.h"
+#include "MuCOE/ICustomizableObjectInstanceEditor.h"
+#include "MuCOE/Nodes/CustomizableObjectNodeMeshClipMorph.h"
+#include "MuCOE/Nodes/CustomizableObjectNodeMeshClipWithMesh.h"
+#include "MuCOE/Nodes/CustomizableObjectNodeProjectorConstant.h"
+#include "MuCOE/Nodes/CustomizableObjectNodeProjectorParameter.h"
 #include "MuCOE/UnrealEditorPortabilityHelpers.h"
-#include "MuCO/UnrealBakeHelpers.h"
-#include "MuCO/CustomizableObjectMipDataProvider.h"
-#include "EditorStyleSet.h"
+#include "ObjectTools.h"
+#include "Preferences/PersonaOptions.h"
+#include "PreviewScene.h"
+#include "RawIndexBuffer.h"
+#include "Rendering/MultiSizeIndexContainer.h"
+#include "Rendering/SkeletalMeshLODRenderData.h"
+#include "Rendering/SkeletalMeshModel.h"
+#include "Rendering/SkeletalMeshRenderData.h"
+#include "Rendering/StaticMeshVertexBuffer.h"
+#include "SceneManagement.h"
+#include "SceneView.h"
+#include "ScopedTransaction.h"
+#include "ShowFlags.h"
+#include "SlotBase.h"
+#include "StaticMeshResources.h"
+#include "Styling/AppStyle.h"
+#include "Styling/SlateTypes.h"
+#include "Subsystems/AssetEditorSubsystem.h"
+#include "Templates/Casts.h"
+#include "Templates/Decay.h"
+#include "Templates/UnrealTemplate.h"
+#include "Trace/Detail/Channel.h"
+#include "Types/SlateEnums.h"
+#include "UObject/Class.h"
+#include "UObject/Object.h"
+#include "UObject/ObjectPtr.h"
+#include "UObject/Package.h"
+#include "UObject/UObjectGlobals.h"
+#include "UObject/UnrealNames.h"
+#include "UObject/WeakObjectPtr.h"
+#include "UnrealClient.h"
+#include "UnrealWidget.h"
+#include "Widgets/DeclarativeSyntaxSupport.h"
+#include "Widgets/Input/SButton.h"
+#include "Widgets/Input/SCheckBox.h"
+#include "Widgets/Input/SEditableTextBox.h"
+#include "Widgets/Layout/SBorder.h"
+#include "Widgets/Layout/SUniformGridPanel.h"
+#include "Widgets/Notifications/SNotificationList.h"
+#include "Widgets/SBoxPanel.h"
+#include "Widgets/SWindow.h"
+#include "Widgets/Text/STextBlock.h"
 
+class FMaterialRenderProxy;
+class UFont;
+class UMaterialExpression;
+class UTextureMipDataProviderFactory;
+
 #define LOCTEXT_NAMESPACE "CustomizableObjectEditor" 
 
diff --git a/Engine/Plugins/Experimental/Mutable/Source/CustomizableObjectEditor/Private/MuCOE/CustomizableObjectEditorViewportLODCommands.cpp b/Engine/Plugins/Experimental/Mutable/Source/CustomizableObjectEditor/Private/MuCOE/CustomizableObjectEditorViewportLODCommands.cpp
--- a/Engine/Plugins/Experimental/Mutable/Source/CustomizableObjectEditor/Private/MuCOE/CustomizableObjectEditorViewportLODCommands.cpp
+++ b/Engine/Plugins/Experimental/Mutable/Source/CustomizableObjectEditor/Private/MuCOE/CustomizableObjectEditorViewportLODCommands.cpp
@@ -1,7 +1,9 @@
 // Copyright Epic Games, Inc. All Rights Reserved.
 
 #include "MuCOE/CustomizableObjectEditorViewportLODCommands.h"
-#include "MuCOE/CustomizableObjectEditorModule.h"
+
+#include "Framework/Commands/InputChord.h"
+#include "Framework/Commands/UICommandInfo.h"
 
 #define LOCTEXT_NAMESPACE "CustomizableObjectEditorViewportLODCommands"
 
diff --git a/Engine/Plugins/Experimental/Mutable/Source/CustomizableObjectEditor/Private/MuCOE/CustomizableObjectEditorViewportMenuCommands.cpp b/Engine/Plugins/Experimental/Mutable/Source/CustomizableObjectEditor/Private/MuCOE/CustomizableObjectEditorViewportMenuCommands.cpp
--- a/Engine/Plugins/Experimental/Mutable/Source/CustomizableObjectEditor/Private/MuCOE/CustomizableObjectEditorViewportMenuCommands.cpp
+++ b/Engine/Plugins/Experimental/Mutable/Source/CustomizableObjectEditor/Private/MuCOE/CustomizableObjectEditorViewportMenuCommands.cpp
@@ -2,6 +2,9 @@
 
 #include "MuCOE/CustomizableObjectEditorViewportMenuCommands.h"
 
+#include "Framework/Commands/InputChord.h"
+#include "Framework/Commands/UICommandInfo.h"
+
 #define LOCTEXT_NAMESPACE "CustomizableObjectEditorViewportMenuCommands"
 
 void FCustomizableObjectEditorViewportMenuCommands::RegisterCommands()
diff --git a/Engine/Plugins/Experimental/Mutable/Source/CustomizableObjectEditor/Private/MuCOE/CustomizableObjectFactory.cpp b/Engine/Plugins/Experimental/Mutable/Source/CustomizableObjectEditor/Private/MuCOE/CustomizableObjectFactory.cpp
--- a/Engine/Plugins/Experimental/Mutable/Source/CustomizableObjectEditor/Private/MuCOE/CustomizableObjectFactory.cpp
+++ b/Engine/Plugins/Experimental/Mutable/Source/CustomizableObjectEditor/Private/MuCOE/CustomizableObjectFactory.cpp
@@ -1,9 +1,16 @@
 // Copyright Epic Games, Inc. All Rights Reserved.
 
 #include "MuCOE/CustomizableObjectFactory.h"
-#include "MuCOE/CustomizableObjectEditorModule.h"
 
-#include "MuR/Serialisation.h"
+#include "MuCO/CustomizableObject.h"
+#include "MuCOE/CustomizableObjectGraph.h"
+#include "Templates/SubclassOf.h"
+#include "UObject/ObjectPtr.h"
+#include "UObject/UnrealNames.h"
+
+class FFeedbackContext;
+class UClass;
+class UObject;
 
 #define LOCTEXT_NAMESPACE "CustomizableObjectFactory"
 
diff --git a/Engine/Plugins/Experimental/Mutable/Source/CustomizableObjectEditor/Private/MuCOE/CustomizableObjectGraph.cpp b/Engine/Plugins/Experimental/Mutable/Source/CustomizableObjectEditor/Private/MuCOE/CustomizableObjectGraph.cpp
--- a/Engine/Plugins/Experimental/Mutable/Source/CustomizableObjectEditor/Private/MuCOE/CustomizableObjectGraph.cpp
+++ b/Engine/Plugins/Experimental/Mutable/Source/CustomizableObjectEditor/Private/MuCOE/CustomizableObjectGraph.cpp
@@ -1,7 +1,19 @@
 // Copyright Epic Games, Inc. All Rights Reserved.
 
 #include "MuCOE/CustomizableObjectGraph.h"
-#include "MuCOE/CustomizableObjectEditorModule.h"
+
+#include "Containers/Array.h"
+#include "EdGraph/EdGraphNode.h"
+#include "HAL/Platform.h"
+#include "HAL/PlatformCrt.h"
+#include "MuCOE/EdGraphSchema_CustomizableObject.h"
+#include "MuCOE/Nodes/CustomizableObjectNode.h"
+#include "MuCOE/Nodes/CustomizableObjectNodeObject.h"
+#include "Templates/Casts.h"
+#include "Templates/SubclassOf.h"
+#include "UObject/ObjectPtr.h"
+
+class UObject;
 
 
 UCustomizableObjectGraph::UCustomizableObjectGraph()
diff --git a/Engine/Plugins/Experimental/Mutable/Source/CustomizableObjectEditor/Private/MuCOE/CustomizableObjectIdentifierCustomization.cpp b/Engine/Plugins/Experimental/Mutable/Source/CustomizableObjectEditor/Private/MuCOE/CustomizableObjectIdentifierCustomization.cpp
--- a/Engine/Plugins/Experimental/Mutable/Source/CustomizableObjectEditor/Private/MuCOE/CustomizableObjectIdentifierCustomization.cpp
+++ b/Engine/Plugins/Experimental/Mutable/Source/CustomizableObjectEditor/Private/MuCOE/CustomizableObjectIdentifierCustomization.cpp
@@ -1,12 +1,26 @@
 // Copyright Epic Games, Inc. All Rights Reserved.
 
 #include "MuCOE/CustomizableObjectIdentifierCustomization.h"
-#include "MuCOE/CustomizableObjectEditorModule.h"
-#include "MuCO/CustomizableObjectIdentifier.h"
-#include "MuCOE/CustomizableObjectCompiler.h"
-#include "PropertyCustomizationHelpers.h"
+
+#include "AssetRegistry/AssetData.h"
 #include "AssetRegistry/AssetRegistryModule.h"
+#include "AssetRegistry/IAssetRegistry.h"
+#include "Containers/Array.h"
+#include "DetailWidgetRow.h"
+#include "HAL/Platform.h"
+#include "HAL/PlatformCrt.h"
+#include "Misc/AssertionMacros.h"
+#include "Misc/Guid.h"
+#include "Modules/ModuleManager.h"
+#include "MuCO/CustomizableObject.h"
 #include "MuCOE/GraphTraversal.h"
+#include "MuCOE/Nodes/CustomizableObjectNodeObject.h"
+#include "PropertyCustomizationHelpers.h"
+#include "PropertyHandle.h"
+#include "Templates/Casts.h"
+#include "UObject/TopLevelAssetPath.h"
+#include "UObject/UnrealType.h"
+#include "Widgets/DeclarativeSyntaxSupport.h"
 
 
 #define LOCTEXT_NAMESPACE "FCustomizableObjectIdentifierCustomization"
diff --git a/Engine/Plugins/Experimental/Mutable/Source/CustomizableObjectEditor/Private/MuCOE/CustomizableObjectInstanceEditor.cpp b/Engine/Plugins/Experimental/Mutable/Source/CustomizableObjectEditor/Private/MuCOE/CustomizableObjectInstanceEditor.cpp
--- a/Engine/Plugins/Experimental/Mutable/Source/CustomizableObjectEditor/Private/MuCOE/CustomizableObjectInstanceEditor.cpp
+++ b/Engine/Plugins/Experimental/Mutable/Source/CustomizableObjectEditor/Private/MuCOE/CustomizableObjectInstanceEditor.cpp
@@ -1,53 +1,74 @@
 // Copyright Epic Games, Inc. All Rights Reserved.
 
 #include "MuCOE/CustomizableObjectInstanceEditor.h"
-#include "MuCOE/CustomizableObjectEditorStyle.h"
 
-#include "MuCOE/CustomizableObjectEditor.h"
-#include "MuCOE/SCustomizableObjectEditorViewport.h"
-#include "MuCOE/CustomizableObjectEditorActions.h"
-#include "MuCOE/CustomizableObjectInstanceEditorActions.h"
-#include "MuCOE/SCustomizableObjectProperties.h"
+#include "AdvancedPreviewSceneModule.h"
+#include "Animation/AnimSingleNodeInstance.h"
+#include "Animation/DebugSkelMeshComponent.h"
+#include "Animation/PoseAsset.h"
+#include "Animation/SmartName.h"
+#include "AssetRegistry/AssetData.h"
+#include "AssetRegistry/AssetRegistryModule.h"
+#include "AssetRegistry/IAssetRegistry.h"
+#include "Components/SkeletalMeshComponent.h"
+#include "ContentBrowserModule.h"
+#include "DetailsViewArgs.h"
+#include "EdGraph/EdGraph.h"
+#include "Editor.h"
+#include "Editor/EditorEngine.h"
+#include "Engine/EngineTypes.h"
+#include "FileHelpers.h"
+#include "Framework/Commands/UIAction.h"
+#include "Framework/Commands/UICommandList.h"
+#include "Framework/Docking/TabManager.h"
+#include "Framework/MultiBox/MultiBoxBuilder.h"
+#include "Framework/MultiBox/MultiBoxExtender.h"
+#include "Framework/Notifications/NotificationManager.h"
+#include "HAL/PlatformCrt.h"
+#include "HAL/PlatformTime.h"
+#include "IContentBrowserSingleton.h"
+#include "IDetailsView.h"
+#include "Internationalization/Internationalization.h"
+#include "Logging/LogCategory.h"
+#include "Logging/LogMacros.h"
+#include "Misc/AssertionMacros.h"
+#include "Misc/Attribute.h"
+#include "Modules/ModuleManager.h"
 #include "MuCO/CustomizableObject.h"
 #include "MuCO/CustomizableObjectInstance.h"
+#include "MuCO/CustomizableObjectParameterTypeDefinitions.h"
 #include "MuCO/CustomizableObjectSystem.h"
+#include "MuCO/CustomizableSkeletalComponent.h"
+#include "MuCOE/CustomizableObjectBakeHelpers.h"
 #include "MuCOE/CustomizableObjectCompiler.h"
+#include "MuCOE/CustomizableObjectEditorActions.h"
+#include "MuCOE/CustomizableObjectEditorModule.h"
+#include "MuCOE/CustomizableObjectEditorViewportClient.h"
+#include "MuCOE/CustomizableObjectInstanceEditorActions.h"
 #include "MuCOE/CustomizableObjectPreviewScene.h"
-#include "MuCOE/SCustomizableObjectSystemVersion.h"
-
-#include "AssetRegistry/AssetRegistryModule.h"
-#include "PreviewScene.h"
-#include "ScopedTransaction.h"
-#include "FileHelpers.h"
-#include "BusyCursor.h"
-#include "IDetailsView.h"
-
-#include "Slate/SceneViewport.h"
-#include "Widgets/Docking/SDockTab.h"
-#include "Components/SkeletalMeshComponent.h"
+#include "MuCOE/Nodes/CustomizableObjectNodeObject.h"
+#include "MuCOE/SCustomizableObjectEditorTextureAnalyzer.h"
+#include "MuCOE/SCustomizableObjectEditorViewport.h"
 #include "MuCOE/UnrealEditorPortabilityHelpers.h"
-
-#include "Editor.h"
-#include "Engine/Selection.h"
-#include "ContentBrowserModule.h"
-#include "IContentBrowserSingleton.h"
-#include "WorkspaceMenuStructureModule.h"
-#include "IDetailsView.h"
 #include "PropertyEditorModule.h"
-#include "WorkspaceMenuStructure.h"
-#include "Framework/Notifications/NotificationManager.h"
-#include "AssetRegistry/AssetRegistryModule.h"
-#include "AdvancedPreviewSceneModule.h"
+#include "Selection.h"
+#include "Subsystems/AssetEditorSubsystem.h"
+#include "Templates/Casts.h"
+#include "Toolkits/AssetEditorToolkit.h"
+#include "Trace/Detail/Channel.h"
+#include "Types/SlateEnums.h"
+#include "UObject/Class.h"
+#include "UObject/ObjectPtr.h"
+#include "UObject/Package.h"
+#include "UObject/UnrealNames.h"
+#include "Widgets/DeclarativeSyntaxSupport.h"
+#include "Widgets/Docking/SDockTab.h"
 #include "Widgets/Notifications/SNotificationList.h"
-#include "MuCOE/CustomizableObjectBakeHelpers.h"
-#include "Animation/AnimSingleNodeInstance.h"
+#include "WorkspaceMenuStructure.h"
+#include "WorkspaceMenuStructureModule.h"
 
-#include "MuCOE/SCustomizableObjectEditorTextureAnalyzer.h"
-#include "MuCOE/CustomizableObjectEditorViewportClient.h"
+class FAdvancedPreviewScene;
+class FProperty;
 
-#include "Engine/Texture.h"
-#include "Widgets/Input/SButton.h"
-#include "Materials/MaterialInstance.h"
-
 #define LOCTEXT_NAMESPACE "CustomizableObjectInstanceEditor"
 
diff --git a/Engine/Plugins/Experimental/Mutable/Source/CustomizableObjectEditor/Private/MuCOE/CustomizableObjectInstanceEditor.h b/Engine/Plugins/Experimental/Mutable/Source/CustomizableObjectEditor/Private/MuCOE/CustomizableObjectInstanceEditor.h
--- a/Engine/Plugins/Experimental/Mutable/Source/CustomizableObjectEditor/Private/MuCOE/CustomizableObjectInstanceEditor.h
+++ b/Engine/Plugins/Experimental/Mutable/Source/CustomizableObjectEditor/Private/MuCOE/CustomizableObjectInstanceEditor.h
@@ -2,13 +2,35 @@
 
 #pragma once
 
+#include "Containers/Array.h"
+#include "Containers/Map.h"
+#include "Containers/UnrealString.h"
+#include "Delegates/Delegate.h"
+#include "HAL/Platform.h"
+#include "Internationalization/Text.h"
+#include "Math/Color.h"
 #include "MuCOE/ICustomizableObjectInstanceEditor.h"
-#include "Toolkits/AssetEditorToolkit.h"
-#include "GraphEditor.h"
+#include "Stats/Stats2.h"
+#include "Templates/SharedPointer.h"
+#include "Templates/UniquePtr.h"
 #include "TickableEditorObject.h"
-#include "Animation/DebugSkelMeshComponent.h"
+#include "Toolkits/IToolkit.h"
+#include "UObject/GCObject.h"
+#include "UObject/NameTypes.h"
+#include "UObject/Object.h"
+#include "UObject/UObjectGlobals.h"
+
 #include "CustomizableObjectInstanceEditor.generated.h"
 
+class FProperty;
+class FSpawnTabArgs;
+class SDockTab;
+class SWidget;
+class UCustomizableObject;
+class UCustomizableObjectInstance;
+class UPoseAsset;
+struct FFrame;
+
 DECLARE_DELEGATE(FCreatePreviewInstanceFlagDelegate);
 
 /**
diff --git a/Engine/Plugins/Experimental/Mutable/Source/CustomizableObjectEditor/Private/MuCOE/CustomizableObjectInstanceEditorActions.cpp b/Engine/Plugins/Experimental/Mutable/Source/CustomizableObjectEditor/Private/MuCOE/CustomizableObjectInstanceEditorActions.cpp
--- a/Engine/Plugins/Experimental/Mutable/Source/CustomizableObjectEditor/Private/MuCOE/CustomizableObjectInstanceEditorActions.cpp
+++ b/Engine/Plugins/Experimental/Mutable/Source/CustomizableObjectEditor/Private/MuCOE/CustomizableObjectInstanceEditorActions.cpp
@@ -1,9 +1,13 @@
 // Copyright Epic Games, Inc. All Rights Reserved.
 
 #include "MuCOE/CustomizableObjectInstanceEditorActions.h"
-#include "MuCOE/SCustomizableObjectEditorViewport.h"
+
+#include "Framework/Commands/InputChord.h"
+#include "Framework/Commands/UICommandInfo.h"
+#include "Internationalization/Internationalization.h"
 #include "MuCOE/CustomizableObjectEditorStyle.h"
-#include "Styling/AppStyle.h"
+#include "UObject/NameTypes.h"
+#include "UObject/UnrealNames.h"
 
 
 #define LOCTEXT_NAMESPACE "CustomizableObjectEditor"
diff --git a/Engine/Plugins/Experimental/Mutable/Source/CustomizableObjectEditor/Private/MuCOE/CustomizableObjectInstanceFactory.cpp b/Engine/Plugins/Experimental/Mutable/Source/CustomizableObjectEditor/Private/MuCOE/CustomizableObjectInstanceFactory.cpp
--- a/Engine/Plugins/Experimental/Mutable/Source/CustomizableObjectEditor/Private/MuCOE/CustomizableObjectInstanceFactory.cpp
+++ b/Engine/Plugins/Experimental/Mutable/Source/CustomizableObjectEditor/Private/MuCOE/CustomizableObjectInstanceFactory.cpp
@@ -1,11 +1,31 @@
 // Copyright Epic Games, Inc. All Rights Reserved.
 
 #include "MuCOE/CustomizableObjectInstanceFactory.h"
-#include "MuCOE/CustomizableObjectEditorModule.h"
-#include "MuCO/CustomizableSkeletalMeshActor.h"
+
+#include "Animation/SkeletalMeshActor.h"
+#include "AssetRegistry/AssetData.h"
 #include "AssetRegistry/AssetRegistryModule.h"
 #include "Components/SkeletalMeshComponent.h"
+#include "Containers/Array.h"
+#include "EdGraph/EdGraph.h"
+#include "Engine/SkeletalMesh.h"
+#include "Engine/World.h"
+#include "GameFramework/Actor.h"
+#include "HAL/PlatformCrt.h"
+#include "Internationalization/Internationalization.h"
+#include "Internationalization/Text.h"
+#include "Math/Vector.h"
+#include "Modules/ModuleManager.h"
+#include "MuCO/CustomizableObject.h"
 #include "MuCO/CustomizableObjectInstance.h"
+#include "MuCO/CustomizableSkeletalComponent.h"
+#include "MuCO/CustomizableSkeletalMeshActor.h"
+#include "MuCOE/Nodes/CustomizableObjectNodeObject.h"
+#include "Templates/Casts.h"
+#include "Templates/SubclassOf.h"
+#include "UObject/Class.h"
+#include "UObject/Object.h"
+#include "UObject/ObjectPtr.h"
 
 
 #define LOCTEXT_NAMESPACE "CustomizableObjectInstanceFactory"
diff --git a/Engine/Plugins/Experimental/Mutable/Source/CustomizableObjectEditor/Private/MuCOE/CustomizableObjectLayout.cpp b/Engine/Plugins/Experimental/Mutable/Source/CustomizableObjectEditor/Private/MuCOE/CustomizableObjectLayout.cpp
--- a/Engine/Plugins/Experimental/Mutable/Source/CustomizableObjectEditor/Private/MuCOE/CustomizableObjectLayout.cpp
+++ b/Engine/Plugins/Experimental/Mutable/Source/CustomizableObjectEditor/Private/MuCOE/CustomizableObjectLayout.cpp
@@ -1,19 +1,28 @@
 // Copyright Epic Games, Inc. All Rights Reserved.
 
 #include "MuCOE/CustomizableObjectLayout.h"
+
+#include "Containers/SparseArray.h"
+#include "EdGraph/EdGraph.h"
+#include "EdGraph/EdGraphPin.h"
+#include "Engine/SkeletalMesh.h"
+#include "Engine/StaticMesh.h"
+#include "HAL/PlatformCrt.h"
+#include "MuCO/CustomizableObject.h"
 #include "MuCOE/CustomizableObjectCompiler.h"
-#include "MuCOE/ICustomizableObjectEditor.h"
 #include "MuCOE/GenerateMutableSource/GenerateMutableSource.h"
+#include "MuCOE/GenerateMutableSource/GenerateMutableSourceMesh.h"
 #include "MuCOE/GraphTraversal.h"
-
+#include "MuCOE/ICustomizableObjectEditor.h"
 #include "MuCOE/Nodes/CustomizableObjectNode.h"
-#include "MuCOE/Nodes/CustomizableObjectNodeSkeletalMesh.h"
-#include "MuCOE/Nodes/CustomizableObjectNodeTable.h"
 #include "MuCOE/Nodes/CustomizableObjectNodeLayoutBlocks.h"
-
-#include "MuCOE/GenerateMutableSource/GenerateMutableSourceMesh.h"
-
+#include "MuCOE/Nodes/CustomizableObjectNodeMesh.h"
+#include "MuCOE/Nodes/CustomizableObjectNodeTable.h"
+#include "MuR/Ptr.h"
 #include "MuT/NodeLayout.h"
+#include "MuT/Table.h"
+#include "Templates/Casts.h"
+#include "Templates/SharedPointer.h"
 
 #define LOCTEXT_NAMESPACE "CustomizableObjectEditor"
 
diff --git a/Engine/Plugins/Experimental/Mutable/Source/CustomizableObjectEditor/Private/MuCOE/CustomizableObjectNodeObjectGroupDetails.cpp b/Engine/Plugins/Experimental/Mutable/Source/CustomizableObjectEditor/Private/MuCOE/CustomizableObjectNodeObjectGroupDetails.cpp
--- a/Engine/Plugins/Experimental/Mutable/Source/CustomizableObjectEditor/Private/MuCOE/CustomizableObjectNodeObjectGroupDetails.cpp
+++ b/Engine/Plugins/Experimental/Mutable/Source/CustomizableObjectEditor/Private/MuCOE/CustomizableObjectNodeObjectGroupDetails.cpp
@@ -1,12 +1,32 @@
 // Copyright Epic Games, Inc. All Rights Reserved.
 
 #include "MuCOE/CustomizableObjectNodeObjectGroupDetails.h"
-#include "MuCOE/Nodes/CustomizableObjectNodeObjectGroup.h"
+
+#include "Containers/Array.h"
+#include "Containers/Map.h"
+#include "Containers/UnrealString.h"
+#include "DetailCategoryBuilder.h"
+#include "DetailLayoutBuilder.h"
+#include "DetailWidgetRow.h"
+#include "EdGraph/EdGraph.h"
+#include "HAL/PlatformCrt.h"
+#include "IDetailsView.h"
+#include "Internationalization/Internationalization.h"
+#include "Misc/Attribute.h"
+#include "Misc/Guid.h"
+#include "MuCO/CustomizableObject.h"
+#include "MuCO/CustomizableObjectIdentifier.h"
 #include "MuCOE/CustomizableObjectEditor.h"
-#include "MuCOE/CustomizableObjectEditorModule.h"
+#include "MuCOE/CustomizableObjectGraph.h"
+#include "MuCOE/Nodes/CustomizableObjectNodeObject.h"
+#include "MuCOE/Nodes/CustomizableObjectNodeObjectGroup.h"
 #include "PropertyCustomizationHelpers.h"
-#include "DetailLayoutBuilder.h"
-#include "DetailCategoryBuilder.h"
+#include "Templates/Casts.h"
+#include "UObject/Object.h"
+#include "UObject/ObjectPtr.h"
+#include "UObject/WeakObjectPtr.h"
+#include "UObject/WeakObjectPtrTemplates.h"
+#include "Widgets/DeclarativeSyntaxSupport.h"
 
 
 #define LOCTEXT_NAMESPACE "CustomizableObjectGroupDetails"
diff --git a/Engine/Plugins/Experimental/Mutable/Source/CustomizableObjectEditor/Private/MuCOE/CustomizableObjectNodeObjectGroupDetails.h b/Engine/Plugins/Experimental/Mutable/Source/CustomizableObjectEditor/Private/MuCOE/CustomizableObjectNodeObjectGroupDetails.h
--- a/Engine/Plugins/Experimental/Mutable/Source/CustomizableObjectEditor/Private/MuCOE/CustomizableObjectNodeObjectGroupDetails.h
+++ b/Engine/Plugins/Experimental/Mutable/Source/CustomizableObjectEditor/Private/MuCOE/CustomizableObjectNodeObjectGroupDetails.h
@@ -3,10 +3,9 @@
 #pragma once
 
 #include "IDetailCustomization.h"
-//#include "Templates/SharedPointer.h"
-//#include "UObject/WeakObjectPtr.h"
-//#include "Input/Reply.h"
-//#include "PropertyHandle.h"
+#include "Templates/SharedPointer.h"
+
+class IDetailLayoutBuilder;
 
 
 class FCustomizableObjectNodeObjectGroupDetails : public IDetailCustomization
diff --git a/Engine/Plugins/Experimental/Mutable/Source/CustomizableObjectEditor/Private/MuCOE/CustomizableObjectPin.cpp b/Engine/Plugins/Experimental/Mutable/Source/CustomizableObjectEditor/Private/MuCOE/CustomizableObjectPin.cpp
--- a/Engine/Plugins/Experimental/Mutable/Source/CustomizableObjectEditor/Private/MuCOE/CustomizableObjectPin.cpp
+++ b/Engine/Plugins/Experimental/Mutable/Source/CustomizableObjectEditor/Private/MuCOE/CustomizableObjectPin.cpp
@@ -2,6 +2,8 @@
 
 #include "MuCOE/CustomizableObjectPin.h"
 
+#include "EdGraph/EdGraphPin.h"
+
 bool IsPinOrphan(const UEdGraphPin &Pin)
 {
 	return Pin.bOrphanedPin;
diff --git a/Engine/Plugins/Experimental/Mutable/Source/CustomizableObjectEditor/Private/MuCOE/CustomizableObjectPreviewScene.cpp b/Engine/Plugins/Experimental/Mutable/Source/CustomizableObjectEditor/Private/MuCOE/CustomizableObjectPreviewScene.cpp
--- a/Engine/Plugins/Experimental/Mutable/Source/CustomizableObjectEditor/Private/MuCOE/CustomizableObjectPreviewScene.cpp
+++ b/Engine/Plugins/Experimental/Mutable/Source/CustomizableObjectEditor/Private/MuCOE/CustomizableObjectPreviewScene.cpp
@@ -2,6 +2,10 @@
 
 #include "MuCOE/CustomizableObjectPreviewScene.h"
 
+#include "PreviewScene.h"
+
+class UStaticMeshComponent;
+
 FCustomizableObjectPreviewScene::FCustomizableObjectPreviewScene(ConstructionValues CVS, float InFloorOffset)
 	: FAdvancedPreviewScene(CVS, InFloorOffset)
 {
diff --git a/Engine/Plugins/Experimental/Mutable/Source/CustomizableObjectEditor/Private/MuCOE/CustomizableObjectPreviewScene.h b/Engine/Plugins/Experimental/Mutable/Source/CustomizableObjectEditor/Private/MuCOE/CustomizableObjectPreviewScene.h
--- a/Engine/Plugins/Experimental/Mutable/Source/CustomizableObjectEditor/Private/MuCOE/CustomizableObjectPreviewScene.h
+++ b/Engine/Plugins/Experimental/Mutable/Source/CustomizableObjectEditor/Private/MuCOE/CustomizableObjectPreviewScene.h
@@ -4,6 +4,8 @@
 
 #include "AdvancedPreviewScene.h"
 
+class UStaticMeshComponent;
+
 class FCustomizableObjectPreviewScene : public FAdvancedPreviewScene
 {
 public:
diff --git a/Engine/Plugins/Experimental/Mutable/Source/CustomizableObjectEditor/Private/MuCOE/CustomizableObjectWidget.cpp b/Engine/Plugins/Experimental/Mutable/Source/CustomizableObjectEditor/Private/MuCOE/CustomizableObjectWidget.cpp
--- a/Engine/Plugins/Experimental/Mutable/Source/CustomizableObjectEditor/Private/MuCOE/CustomizableObjectWidget.cpp
+++ b/Engine/Plugins/Experimental/Mutable/Source/CustomizableObjectEditor/Private/MuCOE/CustomizableObjectWidget.cpp
@@ -1,11 +1,28 @@
 // Copyright Epic Games, Inc. All Rights Reserved.
 
 #include "MuCOE/CustomizableObjectWidget.h"
-#include "MuCOE/CustomizableObjectEditorModule.h"
+
 #include "DynamicMeshBuilder.h"
+#include "Engine/EngineTypes.h"
+#include "EngineDefines.h"
+#include "Materials/Material.h"
+#include "Materials/MaterialInstance.h"
 #include "Materials/MaterialInstanceDynamic.h"
-#include "UnrealWidget.h"
+#include "Math/IntRect.h"
+#include "Math/NumericLimits.h"
+#include "Math/Plane.h"
+#include "Math/RotationMatrix.h"
+#include "Math/Rotator.h"
+#include "Math/ScaleMatrix.h"
+#include "Math/TranslationMatrix.h"
 #include "MuCOE/UnrealEditorPortabilityHelpers.h"
+#include "SceneManagement.h"
+#include "Templates/UnrealTemplate.h"
+#include "UObject/UObjectGlobals.h"
+#include "UnrealWidget.h"
+
+class FMaterialRenderProxy;
+class UMaterialInterface;
 
 
 IMPLEMENT_HIT_PROXY(HCustomizableObjectWidgetAxis,HHitProxy);
diff --git a/Engine/Plugins/Experimental/Mutable/Source/CustomizableObjectEditor/Private/MuCOE/EdGraphSchema_CustomizableObject.cpp b/Engine/Plugins/Experimental/Mutable/Source/CustomizableObjectEditor/Private/MuCOE/EdGraphSchema_CustomizableObject.cpp
--- a/Engine/Plugins/Experimental/Mutable/Source/CustomizableObjectEditor/Private/MuCOE/EdGraphSchema_CustomizableObject.cpp
+++ b/Engine/Plugins/Experimental/Mutable/Source/CustomizableObjectEditor/Private/MuCOE/EdGraphSchema_CustomizableObject.cpp
@@ -2,30 +2,103 @@
 
 #include "MuCOE/EdGraphSchema_CustomizableObject.h"
 
-#include "MuCOE/CustomizableObjectEditorModule.h"
-#include "MuCOE/CustomizableObjectEditor.h"
-
-#include "MuCOE/Nodes/CustomizableObjectNodeMeshGeometryOperation.h"
-#include "MuCOE/Nodes/CustomizableObjectNodeMeshReshape.h"
-
-#include "GraphEditorActions.h" 
-
-#include "Runtime/Launch/Resources/Version.h"
-
-#include "ToolMenu.h"
-#include "ToolMenuSection.h"
-
 #include "AssetRegistry/AssetData.h"
-#include "ScopedTransaction.h"
-#include "Framework/Commands/GenericCommands.h"
+#include "Containers/EnumAsByte.h"
+#include "Containers/Map.h"
+#include "CoreGlobals.h"
+#include "Delegates/Delegate.h"
+#include "EdGraph/EdGraph.h"
+#include "EdGraph/EdGraphNode.h"
+#include "EdGraph/EdGraphPin.h"
 #include "EdGraphNode_Comment.h"
-#include "Toolkits/ToolkitManager.h"
-
 #include "Engine/SkeletalMesh.h"
 #include "Engine/StaticMesh.h"
-
+#include "Engine/Texture2D.h"
+#include "Framework/Commands/GenericCommands.h"
+#include "Framework/Commands/UIAction.h"
+#include "Framework/MultiBox/MultiBoxBuilder.h"
+#include "GraphEditorActions.h"
+#include "Internationalization/Internationalization.h"
+#include "Logging/LogCategory.h"
+#include "Logging/LogMacros.h"
+#include "Materials/MaterialInterface.h"
+#include "Math/UnrealMathSSE.h"
+#include "Misc/AssertionMacros.h"
+#include "MuCO/CustomizableObject.h"
 #include "MuCOE/CustomizableObjectEditorNodeContextCommands.h"
+#include "MuCOE/ICustomizableObjectEditor.h"
+#include "MuCOE/Nodes/CustomizableObjectNode.h"
+#include "MuCOE/Nodes/CustomizableObjectNodeAnimationPose.h"
+#include "MuCOE/Nodes/CustomizableObjectNodeColorArithmeticOp.h"
+#include "MuCOE/Nodes/CustomizableObjectNodeColorConstant.h"
+#include "MuCOE/Nodes/CustomizableObjectNodeColorFromFloats.h"
+#include "MuCOE/Nodes/CustomizableObjectNodeColorParameter.h"
+#include "MuCOE/Nodes/CustomizableObjectNodeColorSwitch.h"
+#include "MuCOE/Nodes/CustomizableObjectNodeColorVariation.h"
+#include "MuCOE/Nodes/CustomizableObjectNodeCopyMaterial.h"
+#include "MuCOE/Nodes/CustomizableObjectNodeCurve.h"
+#include "MuCOE/Nodes/CustomizableObjectNodeEditMaterial.h"
+#include "MuCOE/Nodes/CustomizableObjectNodeEnumParameter.h"
+#include "MuCOE/Nodes/CustomizableObjectNodeExposePin.h"
+#include "MuCOE/Nodes/CustomizableObjectNodeExtendMaterial.h"
+#include "MuCOE/Nodes/CustomizableObjectNodeExternalPin.h"
+#include "MuCOE/Nodes/CustomizableObjectNodeFloatConstant.h"
+#include "MuCOE/Nodes/CustomizableObjectNodeFloatParameter.h"
+#include "MuCOE/Nodes/CustomizableObjectNodeFloatSwitch.h"
+#include "MuCOE/Nodes/CustomizableObjectNodeFloatVariation.h"
+#include "MuCOE/Nodes/CustomizableObjectNodeGroupProjectorParameter.h"
+#include "MuCOE/Nodes/CustomizableObjectNodeLayoutBlocks.h"
+#include "MuCOE/Nodes/CustomizableObjectNodeMaterial.h"
+#include "MuCOE/Nodes/CustomizableObjectNodeMaterialVariation.h"
+#include "MuCOE/Nodes/CustomizableObjectNodeMeshClipDeform.h"
+#include "MuCOE/Nodes/CustomizableObjectNodeMeshClipMorph.h"
+#include "MuCOE/Nodes/CustomizableObjectNodeMeshClipWithMesh.h"
+#include "MuCOE/Nodes/CustomizableObjectNodeMeshGeometryOperation.h"
+#include "MuCOE/Nodes/CustomizableObjectNodeMeshMorph.h"
+#include "MuCOE/Nodes/CustomizableObjectNodeMeshMorphStackApplication.h"
+#include "MuCOE/Nodes/CustomizableObjectNodeMeshMorphStackDefinition.h"
+#include "MuCOE/Nodes/CustomizableObjectNodeMeshReshape.h"
+#include "MuCOE/Nodes/CustomizableObjectNodeMeshSwitch.h"
+#include "MuCOE/Nodes/CustomizableObjectNodeMeshVariation.h"
+#include "MuCOE/Nodes/CustomizableObjectNodeMorphMaterial.h"
+#include "MuCOE/Nodes/CustomizableObjectNodeObject.h"
+#include "MuCOE/Nodes/CustomizableObjectNodeObjectChild.h"
+#include "MuCOE/Nodes/CustomizableObjectNodeObjectGroup.h"
+#include "MuCOE/Nodes/CustomizableObjectNodeProjectorConstant.h"
+#include "MuCOE/Nodes/CustomizableObjectNodeProjectorParameter.h"
+#include "MuCOE/Nodes/CustomizableObjectNodeRemoveMesh.h"
+#include "MuCOE/Nodes/CustomizableObjectNodeRemoveMeshBlocks.h"
+#include "MuCOE/Nodes/CustomizableObjectNodeSkeletalMesh.h"
+#include "MuCOE/Nodes/CustomizableObjectNodeStaticMesh.h"
+#include "MuCOE/Nodes/CustomizableObjectNodeTable.h"
+#include "MuCOE/Nodes/CustomizableObjectNodeTexture.h"
+#include "MuCOE/Nodes/CustomizableObjectNodeTextureBinarise.h"
+#include "MuCOE/Nodes/CustomizableObjectNodeTextureColourMap.h"
+#include "MuCOE/Nodes/CustomizableObjectNodeTextureFromChannels.h"
+#include "MuCOE/Nodes/CustomizableObjectNodeTextureFromColor.h"
+#include "MuCOE/Nodes/CustomizableObjectNodeTextureInterpolate.h"
+#include "MuCOE/Nodes/CustomizableObjectNodeTextureInvert.h"
+#include "MuCOE/Nodes/CustomizableObjectNodeTextureLayer.h"
+#include "MuCOE/Nodes/CustomizableObjectNodeTextureParameter.h"
+#include "MuCOE/Nodes/CustomizableObjectNodeTextureProject.h"
+#include "MuCOE/Nodes/CustomizableObjectNodeTextureSample.h"
+#include "MuCOE/Nodes/CustomizableObjectNodeTextureSwitch.h"
+#include "MuCOE/Nodes/CustomizableObjectNodeTextureToChannels.h"
+#include "MuCOE/Nodes/CustomizableObjectNodeTextureTransform.h"
+#include "MuCOE/Nodes/CustomizableObjectNodeTextureVariation.h"
+#include "ScopedTransaction.h"
 #include "Settings/EditorStyleSettings.h"
+#include "Templates/Casts.h"
+#include "Templates/UnrealTemplate.h"
+#include "Textures/SlateIcon.h"
+#include "ToolMenu.h"
+#include "ToolMenuSection.h"
+#include "Toolkits/ToolkitManager.h"
+#include "Trace/Detail/Channel.h"
+#include "UObject/Object.h"
+#include "UObject/ObjectHandle.h"
+
+class IToolkit;
 
 
 #define LOCTEXT_NAMESPACE "CustomizableObjectSchema"
diff --git a/Engine/Plugins/Experimental/Mutable/Source/CustomizableObjectEditor/Private/MuCOE/FMutableViewportClient.cpp b/Engine/Plugins/Experimental/Mutable/Source/CustomizableObjectEditor/Private/MuCOE/FMutableViewportClient.cpp
--- a/Engine/Plugins/Experimental/Mutable/Source/CustomizableObjectEditor/Private/MuCOE/FMutableViewportClient.cpp
+++ b/Engine/Plugins/Experimental/Mutable/Source/CustomizableObjectEditor/Private/MuCOE/FMutableViewportClient.cpp
@@ -3,6 +3,8 @@
 #include "MuCOE/FMutableViewportClient.h"
 
 #include "AdvancedPreviewScene.h"
+#include "Editor.h"
+#include "UnrealWidget.h"
 
 FMutableMeshViewportClient::FMutableMeshViewportClient(const TSharedRef<FAdvancedPreviewScene>& InPreviewScene)
 : FEditorViewportClient(&GLevelEditorModeTools(), &InPreviewScene.Get())
diff --git a/Engine/Plugins/Experimental/Mutable/Source/CustomizableObjectEditor/Private/MuCOE/FMutableViewportClient.h b/Engine/Plugins/Experimental/Mutable/Source/CustomizableObjectEditor/Private/MuCOE/FMutableViewportClient.h
--- a/Engine/Plugins/Experimental/Mutable/Source/CustomizableObjectEditor/Private/MuCOE/FMutableViewportClient.h
+++ b/Engine/Plugins/Experimental/Mutable/Source/CustomizableObjectEditor/Private/MuCOE/FMutableViewportClient.h
@@ -2,9 +2,10 @@
 
 #pragma once
 
-#include "AdvancedPreviewScene.h"
-#include "SEditorViewport.h"
-#include "UnrealWidget.h"
+#include "EditorViewportClient.h"
+#include "Templates/SharedPointer.h"
+
+class FAdvancedPreviewScene;
 
 /**
  * Super simple viewport client designed to work SMutableMeshViewport
diff --git a/Engine/Plugins/Experimental/Mutable/Source/CustomizableObjectEditor/Private/MuCOE/GenerateMutableSource/GenerateMutableSource.cpp b/Engine/Plugins/Experimental/Mutable/Source/CustomizableObjectEditor/Private/MuCOE/GenerateMutableSource/GenerateMutableSource.cpp
--- a/Engine/Plugins/Experimental/Mutable/Source/CustomizableObjectEditor/Private/MuCOE/GenerateMutableSource/GenerateMutableSource.cpp
+++ b/Engine/Plugins/Experimental/Mutable/Source/CustomizableObjectEditor/Private/MuCOE/GenerateMutableSource/GenerateMutableSource.cpp
@@ -2,60 +2,89 @@
 
 #include "MuCOE/GenerateMutableSource/GenerateMutableSource.h"
 
-#include "MuCOE/GenerateMutableSource/GenerateMutableSourceSurface.h"
-#include "MuCOE/GenerateMutableSource/GenerateMutableSourceGroupProjector.h"
-#include "MuCOE/GenerateMutableSource/GenerateMutableSourceModifier.h"
-
-#include "Templates/UnrealTemplate.h"
-#include "Templates/Decay.h"
-#include "GenericPlatform/GenericPlatformMath.h"
-#include "StaticMeshResources.h"
-#include "Interfaces/ITargetPlatform.h"
+#include "Animation/AnimInstance.h"
+#include "Animation/Skeleton.h"
+#include "Containers/IndirectArray.h"
+#include "Containers/StringConv.h"
+#include "EdGraph/EdGraph.h"
+#include "EdGraph/EdGraphNode.h"
 #include "Engine/SkeletalMesh.h"
-#include "Engine/StaticMesh.h"
-#include "Engine/AssetManager.h"
+#include "Engine/SkeletalMeshSocket.h"
+#include "Engine/Texture2D.h"
 #include "Engine/TextureLODSettings.h"
-#include "Animation/Skeleton.h"
-#include "Animation/MorphTarget.h"
-#include "PhysicsEngine/PhysicsAsset.h"
-#include "ClothingAsset.h"
-#include "MeshUtilities.h"
-#include "EdGraphToken.h"
-#include "AssetRegistry/AssetRegistryModule.h"
-#include "Animation/AnimInstance.h"
-#include "MuCOE/UnrealEditorPortabilityHelpers.h"
+#include "HAL/PlatformCrt.h"
+#include "HAL/PlatformMisc.h"
+#include "Interfaces/ITargetPlatform.h"
+#include "Internationalization/Internationalization.h"
+#include "Internationalization/Text.h"
+#include "Logging/LogCategory.h"
+#include "Logging/LogMacros.h"
+#include "Logging/TokenizedMessage.h"
+#include "Math/BoxSphereBounds.h"
+#include "Math/Rotator.h"
+#include "Math/Vector.h"
+#include "Misc/AssertionMacros.h"
+#include "MuCO/CustomizableObject.h"
+#include "MuCO/CustomizableObjectIdentifier.h"
+#include "MuCO/CustomizableObjectParameterTypeDefinitions.h"
+#include "MuCO/CustomizableObjectUIData.h"
 #include "MuCO/UnrealPortabilityHelpers.h"
-#include "PlatformInfo.h"
-#include "Engine/SkeletalMeshSocket.h"
-
 #include "MuCOE/CustomizableObjectCompiler.h"
-#include "MuCOE/UnrealToMutableTextureConversionUtils.h"
-#include "MuCO/UnrealConversionUtils.h"
-#include "MuCOE/CustomizableObjectEditorModule.h"
-#include "MuCO/CustomizableObjectSystem.h"
-#include "MuCO/CustomizableObject.h"
+#include "MuCOE/EdGraphSchema_CustomizableObject.h"
+#include "MuCOE/GenerateMutableSource/GenerateMutableSourceGroupProjector.h"
+#include "MuCOE/GenerateMutableSource/GenerateMutableSourceModifier.h"
+#include "MuCOE/GenerateMutableSource/GenerateMutableSourceSurface.h"
 #include "MuCOE/GraphTraversal.h"
-#include "MuCOE/MutableUtils.h"
-
+#include "MuCOE/Nodes/CustomizableObjectNodeEditMaterialBase.h"
+#include "MuCOE/Nodes/CustomizableObjectNodeExtendMaterial.h"
+#include "MuCOE/Nodes/CustomizableObjectNodeMaterial.h"
+#include "MuCOE/Nodes/CustomizableObjectNodeMaterialVariation.h"
+#include "MuCOE/Nodes/CustomizableObjectNodeMesh.h"
+#include "MuCOE/Nodes/CustomizableObjectNodeMeshClipMorph.h"
+#include "MuCOE/Nodes/CustomizableObjectNodeMeshClipWithMesh.h"
+#include "MuCOE/Nodes/CustomizableObjectNodeModifierBase.h"
+#include "MuCOE/Nodes/CustomizableObjectNodeObjectGroup.h"
+#include "MuCOE/Nodes/CustomizableObjectNodeTable.h"
+#include "MuCOE/Nodes/CustomizableObjectNodeTexture.h"
+#include "MuCOE/Nodes/CustomizableObjectNodeTextureBinarise.h"
+#include "MuCOE/Nodes/CustomizableObjectNodeTextureColourMap.h"
+#include "MuCOE/Nodes/CustomizableObjectNodeTextureFromChannels.h"
+#include "MuCOE/Nodes/CustomizableObjectNodeTextureInterpolate.h"
+#include "MuCOE/Nodes/CustomizableObjectNodeTextureInvert.h"
+#include "MuCOE/Nodes/CustomizableObjectNodeTextureLayer.h"
+#include "MuCOE/Nodes/CustomizableObjectNodeTextureParameter.h"
+#include "MuCOE/Nodes/CustomizableObjectNodeTextureProject.h"
+#include "MuCOE/Nodes/CustomizableObjectNodeTextureSwitch.h"
+#include "MuCOE/Nodes/CustomizableObjectNodeTextureToChannels.h"
+#include "MuCOE/Nodes/CustomizableObjectNodeTextureTransform.h"
+#include "MuCOE/Nodes/CustomizableObjectNodeTextureVariation.h"
+#include "MuCOE/UnrealEditorPortabilityHelpers.h"
+#include "MuR/Mesh.h"
+#include "MuR/Ptr.h"
+#include "MuR/Skeleton.h"
 #include "MuT/Node.h"
-#include "MuT/NodeImageBinarise.h"
-#include "MuT/NodeMeshConstant.h"
+#include "MuT/NodeComponent.h"
+#include "MuT/NodeComponentEdit.h"
+#include "MuT/NodeComponentNew.h"
+#include "MuT/NodeLOD.h"
 #include "MuT/NodeMeshApplyPose.h"
-#include "MuT/NodeObjectNew.h"
+#include "MuT/NodeMeshConstant.h"
+#include "MuT/NodeModifierMeshClipWithMesh.h"
 #include "MuT/NodeObjectGroup.h"
-#include "MuT/NodeLOD.h"
-#include "MuT/NodeComponentNew.h"
-#include "MuT/NodeComponentEdit.h"
-#include "MuT/NodeSurfaceNew.h"
+#include "MuT/NodeObjectNew.h"
 #include "MuT/NodeSurfaceEdit.h"
-#include "MuT/Compiler.h"
-#include "MuT/Streams.h"
-#include "MuT/NodeModifierMeshClipMorphPlane.h"
-#include "MuT/NodeModifierMeshClipWithMesh.h"
+#include "MuT/NodeSurfaceNew.h"
+#include "PerPlatformProperties.h"
+#include "PerQualityLevelProperties.h"
+#include "PhysicsEngine/PhysicsAsset.h"
+#include "PlatformInfo.h"
+#include "Rendering/SkeletalMeshLODRenderData.h"
+#include "Templates/Casts.h"
+#include "Templates/ChooseClass.h"
+#include "Templates/SubclassOf.h"
+#include "Trace/Detail/Channel.h"
+#include "UObject/ObjectPtr.h"
+#include "UObject/UObjectGlobals.h"
 
-#include "MuR/Mesh.h"
-#include "MuR/Skeleton.h"
-#include "MuR/Image.h"
-
 #define LOCTEXT_NAMESPACE "CustomizableObjectEditor"
 
diff --git a/Engine/Plugins/Experimental/Mutable/Source/CustomizableObjectEditor/Private/MuCOE/GenerateMutableSource/GenerateMutableSource.h b/Engine/Plugins/Experimental/Mutable/Source/CustomizableObjectEditor/Private/MuCOE/GenerateMutableSource/GenerateMutableSource.h
--- a/Engine/Plugins/Experimental/Mutable/Source/CustomizableObjectEditor/Private/MuCOE/GenerateMutableSource/GenerateMutableSource.h
+++ b/Engine/Plugins/Experimental/Mutable/Source/CustomizableObjectEditor/Private/MuCOE/GenerateMutableSource/GenerateMutableSource.h
@@ -2,28 +2,60 @@
 
 #pragma once
 
-#include "MuCO/CustomizableObject.h"
-#include "MuT/NodeMeshApplyPose.h"
-#include "MuCOE/Nodes/CustomizableObjectNodeObject.h"
-#include "MuCOE/Nodes/CustomizableObjectNodeMaterial.h"
-#include "MuCOE/Nodes/CustomizableObjectNodeMesh.h"
-#include "MuCOE/Nodes/CustomizableObjectNodeMeshMorph.h"
-
 #include "Animation/Skeleton.h"
+#include "Containers/Array.h"
+#include "Containers/EnumAsByte.h"
+#include "Containers/Map.h"
+#include "Containers/Set.h"
+#include "Containers/SparseArray.h"
+#include "Containers/UnrealString.h"
+#include "EdGraph/EdGraphPin.h"
 #include "Engine/SkeletalMesh.h"
-
-#include "MuT/NodeObject.h"
-#include "MuT/NodeImageConstant.h"
-#include "MuT/NodeSurfaceNew.h"
+#include "Engine/Texture.h"
+#include "Engine/TextureDefines.h"
+#include "Math/NumericLimits.h"
+#include "Math/Transform.h"
+#include "Math/UnrealMathSSE.h"
+#include "Misc/Guid.h"
+#include "MuCO/CustomizableObject.h"
+#include "MuCO/CustomizableObjectClothingTypes.h"
+#include "MuCOE/Nodes/CustomizableObjectNode.h"
+#include "MuCOE/Nodes/CustomizableObjectNodeMaterialBase.h"
+#include "MuCOE/Nodes/CustomizableObjectNodeObject.h"
+#include "MuR/Mesh.h"
+#include "MuR/RefCounted.h"
+#include "MuT/Node.h"
 #include "MuT/NodeComponentNew.h"
+#include "MuT/NodeImageConstant.h"
+#include "MuT/NodeMeshApplyPose.h"
 #include "MuT/NodeModifierMeshClipWithMesh.h"
+#include "MuT/NodeObject.h"
 #include "MuT/NodeProjector.h"
-#include "MuT/NodeScalarParameter.h"
 #include "MuT/NodeScalarEnumParameter.h"
-#include "MuR/Mesh.h"
+#include "MuT/NodeScalarParameter.h"
+#include "MuT/NodeSurfaceNew.h"
 #include "MuT/Table.h"
-
+#include "RHIDefinitions.h"
+#include "ReferenceSkeleton.h"
+#include "Templates/TypeHash.h"
+#include "UObject/NameTypes.h"
+#include "UObject/Package.h"
+#include "UObject/SoftObjectPtr.h"
+#include "UObject/UnrealNames.h"
+
+class FCustomizableObjectCompiler;
+class UAnimInstance;
+class UCustomizableObjectNodeMaterial;
+class UCustomizableObjectNodeMeshMorph;
+class UCustomizableObjectNodeObjectGroup;
+class UEdGraphNode;
+class UMaterialInterface;
+class UObject;
+class UPhysicsAsset;
+class UTexture2D;
+struct FCustomizableObjectIdPair;
 struct FMutableGraphGenerationContext;
+struct FParameterUIData;
 
 
 struct FGeneratedImageProperties
diff --git a/Engine/Plugins/Experimental/Mutable/Source/CustomizableObjectEditor/Private/MuCOE/GenerateMutableSource/GenerateMutableSourceColor.cpp b/Engine/Plugins/Experimental/Mutable/Source/CustomizableObjectEditor/Private/MuCOE/GenerateMutableSource/GenerateMutableSourceColor.cpp
--- a/Engine/Plugins/Experimental/Mutable/Source/CustomizableObjectEditor/Private/MuCOE/GenerateMutableSource/GenerateMutableSourceColor.cpp
+++ b/Engine/Plugins/Experimental/Mutable/Source/CustomizableObjectEditor/Private/MuCOE/GenerateMutableSource/GenerateMutableSourceColor.cpp
@@ -2,43 +2,56 @@
 
 #include "MuCOE/GenerateMutableSource/GenerateMutableSourceColor.h"
 
+#include "Containers/Array.h"
+#include "Containers/Map.h"
+#include "Containers/Set.h"
+#include "Containers/StringConv.h"
+#include "Containers/UnrealString.h"
+#include "EdGraph/EdGraphPin.h"
+#include "Engine/DataTable.h"
+#include "HAL/Platform.h"
+#include "Internationalization/Internationalization.h"
+#include "Internationalization/Text.h"
+#include "Math/Color.h"
+#include "Misc/AssertionMacros.h"
+#include "Misc/Guid.h"
+#include "MuCO/CustomizableObjectParameterTypeDefinitions.h"
+#include "MuCO/CustomizableObjectUIData.h"
+#include "MuCOE/CustomizableObjectCompiler.h"
+#include "MuCOE/EdGraphSchema_CustomizableObject.h"
+#include "MuCOE/GenerateMutableSource/GenerateMutableSource.h"
 #include "MuCOE/GenerateMutableSource/GenerateMutableSourceFloat.h"
 #include "MuCOE/GenerateMutableSource/GenerateMutableSourceImage.h"
 #include "MuCOE/GenerateMutableSource/GenerateMutableSourceTable.h"
-
-#include "Templates/UnrealTemplate.h"
-#include "GenericPlatform/GenericPlatformMath.h"
-#include "Engine/AssetManager.h"
-#include "Engine/TextureLODSettings.h"
-#include "ClothingAsset.h"
-#include "MeshUtilities.h"
-#include "EdGraphToken.h"
-#include "AssetRegistry/AssetRegistryModule.h"
-#include "MuCOE/UnrealEditorPortabilityHelpers.h"
-#include "PlatformInfo.h"
-
-#include "MuCOE/CustomizableObjectCompiler.h"
-#include "MuCOE/UnrealToMutableTextureConversionUtils.h"
-#include "MuCO/UnrealConversionUtils.h"
-#include "MuCOE/CustomizableObjectEditorModule.h"
-#include "MuCO/CustomizableObject.h"
 #include "MuCOE/GraphTraversal.h"
-
+#include "MuCOE/Nodes/CustomizableObjectNode.h"
+#include "MuCOE/Nodes/CustomizableObjectNodeColorArithmeticOp.h"
+#include "MuCOE/Nodes/CustomizableObjectNodeColorConstant.h"
+#include "MuCOE/Nodes/CustomizableObjectNodeColorFromFloats.h"
+#include "MuCOE/Nodes/CustomizableObjectNodeColorParameter.h"
+#include "MuCOE/Nodes/CustomizableObjectNodeColorSwitch.h"
+#include "MuCOE/Nodes/CustomizableObjectNodeColorVariation.h"
+#include "MuCOE/Nodes/CustomizableObjectNodeTable.h"
+#include "MuCOE/Nodes/CustomizableObjectNodeTextureSample.h"
+#include "MuR/Ptr.h"
+#include "MuR/RefCounted.h"
 #include "MuT/Node.h"
-#include "MuT/NodeScalarConstant.h"
-#include "MuT/NodeScalarEnumParameter.h"
+#include "MuT/NodeColourArithmeticOperation.h"
 #include "MuT/NodeColourConstant.h"
+#include "MuT/NodeColourFromScalars.h"
 #include "MuT/NodeColourParameter.h"
-#include "MuT/NodeColourSwitch.h"
 #include "MuT/NodeColourSampleImage.h"
-#include "MuT/NodeColourArithmeticOperation.h"
-#include "MuT/NodeColourFromScalars.h"
-#include "MuT/NodeColourVariation.h"
+#include "MuT/NodeColourSwitch.h"
 #include "MuT/NodeColourTable.h"
+#include "MuT/NodeColourVariation.h"
 #include "MuT/NodeImageBinarise.h"
-#include "MuT/NodeSurfaceNew.h"
-#include "MuT/Compiler.h"
-#include "MuT/Streams.h"
+#include "MuT/NodeScalar.h"
+#include "MuT/NodeScalarEnumParameter.h"
+#include "MuT/Table.h"
+#include "Templates/Casts.h"
+#include "UObject/NameTypes.h"
+#include "UObject/ObjectPtr.h"
+#include "UObject/UObjectGlobals.h"
 
 
 #define LOCTEXT_NAMESPACE "CustomizableObjectEditor"
diff --git a/Engine/Plugins/Experimental/Mutable/Source/CustomizableObjectEditor/Private/MuCOE/GenerateMutableSource/GenerateMutableSourceColor.h b/Engine/Plugins/Experimental/Mutable/Source/CustomizableObjectEditor/Private/MuCOE/GenerateMutableSource/GenerateMutableSourceColor.h
--- a/Engine/Plugins/Experimental/Mutable/Source/CustomizableObjectEditor/Private/MuCOE/GenerateMutableSource/GenerateMutableSourceColor.h
+++ b/Engine/Plugins/Experimental/Mutable/Source/CustomizableObjectEditor/Private/MuCOE/GenerateMutableSource/GenerateMutableSourceColor.h
@@ -2,9 +2,10 @@
 
 #pragma once
 
-#include "MuCOE/GenerateMutableSource/GenerateMutableSource.h"
-
 #include "MuT/NodeColour.h"
 
+class UEdGraphPin;
+struct FMutableGraphGenerationContext;
+
 /** Convert a CustomizableObject Source Graph into a mutable source graph. */
 mu::NodeColourPtr GenerateMutableSourceColor(const UEdGraphPin* Pin, FMutableGraphGenerationContext& GenerationContext);
diff --git a/Engine/Plugins/Experimental/Mutable/Source/CustomizableObjectEditor/Private/MuCOE/GenerateMutableSource/GenerateMutableSourceFloat.cpp b/Engine/Plugins/Experimental/Mutable/Source/CustomizableObjectEditor/Private/MuCOE/GenerateMutableSource/GenerateMutableSourceFloat.cpp
--- a/Engine/Plugins/Experimental/Mutable/Source/CustomizableObjectEditor/Private/MuCOE/GenerateMutableSource/GenerateMutableSourceFloat.cpp
+++ b/Engine/Plugins/Experimental/Mutable/Source/CustomizableObjectEditor/Private/MuCOE/GenerateMutableSource/GenerateMutableSourceFloat.cpp
@@ -1,49 +1,61 @@
 // Copyright Epic Games, Inc. All Rights Reserved.
 
 #include "MuCOE/GenerateMutableSource/GenerateMutableSourceFloat.h"
-#include "MuCOE/GenerateMutableSource/GenerateMutableSourceColor.h"
-#include "MuCOE/GenerateMutableSource/GenerateMutableSourceProjector.h"
-#include "MuCOE/GenerateMutableSource/GenerateMutableSourceSurface.h"
-#include "MuCOE/GenerateMutableSource/GenerateMutableSourceGroupProjector.h"
-#include "MuCOE/GenerateMutableSource/GenerateMutableSourceImage.h"
-#include "MuCOE/GenerateMutableSource/GenerateMutableSourceMesh.h"
-#include "MuCOE/GenerateMutableSource/GenerateMutableSourceLayout.h"
-#include "MuCOE/GenerateMutableSource/GenerateMutableSourceTable.h"
 
-#include "Templates/UnrealTemplate.h"
-#include "GenericPlatform/GenericPlatformMath.h"
-#include "Engine/AssetManager.h"
-#include "Engine/TextureLODSettings.h"
+#include "Containers/Array.h"
+#include "Containers/EnumAsByte.h"
+#include "Containers/Map.h"
+#include "Containers/Set.h"
+#include "Containers/SparseArray.h"
+#include "Containers/StringConv.h"
+#include "Containers/UnrealString.h"
+#include "Curves/CurveBase.h"
+#include "Curves/CurveFloat.h"
 #include "Curves/CurveLinearColor.h"
 #include "Curves/CurveVector.h"
-#include "Curves/CurveFloat.h"
-#include "ClothingAsset.h"
-#include "MeshUtilities.h"
-#include "EdGraphToken.h"
-#include "AssetRegistry/AssetRegistryModule.h"
-#include "MuCOE/UnrealEditorPortabilityHelpers.h"
-#include "PlatformInfo.h"
-
+#include "Curves/RealCurve.h"
+#include "Curves/RichCurve.h"
+#include "EdGraph/EdGraphPin.h"
+#include "Engine/DataTable.h"
+#include "HAL/PlatformCrt.h"
+#include "Internationalization/Internationalization.h"
+#include "Internationalization/Text.h"
+#include "Math/UnrealMathSSE.h"
+#include "Misc/AssertionMacros.h"
+#include "Misc/Guid.h"
+#include "MuCO/CustomizableObjectParameterTypeDefinitions.h"
+#include "MuCO/CustomizableObjectUIData.h"
 #include "MuCOE/CustomizableObjectCompiler.h"
-#include "MuCOE/CustomizableObjectEditorModule.h"
-#include "MuCO/CustomizableObject.h"
+#include "MuCOE/EdGraphSchema_CustomizableObject.h"
+#include "MuCOE/GenerateMutableSource/GenerateMutableSource.h"
+#include "MuCOE/GenerateMutableSource/GenerateMutableSourceImage.h"
+#include "MuCOE/GenerateMutableSource/GenerateMutableSourceTable.h"
 #include "MuCOE/GraphTraversal.h"
-
+#include "MuCOE/Nodes/CustomizableObjectNode.h"
+#include "MuCOE/Nodes/CustomizableObjectNodeCurve.h"
+#include "MuCOE/Nodes/CustomizableObjectNodeEnumParameter.h"
+#include "MuCOE/Nodes/CustomizableObjectNodeFloatConstant.h"
+#include "MuCOE/Nodes/CustomizableObjectNodeFloatParameter.h"
+#include "MuCOE/Nodes/CustomizableObjectNodeFloatSwitch.h"
+#include "MuCOE/Nodes/CustomizableObjectNodeFloatVariation.h"
+#include "MuCOE/Nodes/CustomizableObjectNodeTable.h"
+#include "MuR/Ptr.h"
+#include "MuR/RefCounted.h"
 #include "MuT/Node.h"
+#include "MuT/NodeImageBinarise.h"
 #include "MuT/NodeScalarConstant.h"
+#include "MuT/NodeScalarCurve.h"
+#include "MuT/NodeScalarEnumParameter.h"
 #include "MuT/NodeScalarParameter.h"
 #include "MuT/NodeScalarSwitch.h"
-#include "MuT/NodeScalarEnumParameter.h"
-#include "MuT/NodeScalarCurve.h"
-#include "MuT/NodeScalarVariation.h"
 #include "MuT/NodeScalarTable.h"
-#include "MuT/NodeImageBinarise.h"
+#include "MuT/NodeScalarVariation.h"
 #include "MuT/NodeSurfaceNew.h"
-#include "MuT/Compiler.h"
-#include "MuT/Streams.h"
-
-#include "MuR/Skeleton.h"
-#include "MuR/Image.h"
+#include "MuT/Table.h"
+#include "Templates/Casts.h"
+#include "UObject/NameTypes.h"
+#include "UObject/ObjectPtr.h"
+#include "UObject/UObjectGlobals.h"
 
 #define LOCTEXT_NAMESPACE "CustomizableObjectEditor"
 
diff --git a/Engine/Plugins/Experimental/Mutable/Source/CustomizableObjectEditor/Private/MuCOE/GenerateMutableSource/GenerateMutableSourceFloat.h b/Engine/Plugins/Experimental/Mutable/Source/CustomizableObjectEditor/Private/MuCOE/GenerateMutableSource/GenerateMutableSourceFloat.h
--- a/Engine/Plugins/Experimental/Mutable/Source/CustomizableObjectEditor/Private/MuCOE/GenerateMutableSource/GenerateMutableSourceFloat.h
+++ b/Engine/Plugins/Experimental/Mutable/Source/CustomizableObjectEditor/Private/MuCOE/GenerateMutableSource/GenerateMutableSourceFloat.h
@@ -2,8 +2,10 @@
 
 #pragma once
 
-#include "MuCOE/GenerateMutableSource/GenerateMutableSource.h"
+#include "MuT/NodeScalar.h"
 
+class UEdGraphPin;
+struct FMutableGraphGenerationContext;
 
 /** Convert a CustomizableObject Source Graph into a mutable source graph. */
 mu::NodeScalarPtr GenerateMutableSourceFloat(const UEdGraphPin* Pin, FMutableGraphGenerationContext& GenerationContext);
diff --git a/Engine/Plugins/Experimental/Mutable/Source/CustomizableObjectEditor/Private/MuCOE/GenerateMutableSource/GenerateMutableSourceGroupProjector.cpp b/Engine/Plugins/Experimental/Mutable/Source/CustomizableObjectEditor/Private/MuCOE/GenerateMutableSource/GenerateMutableSourceGroupProjector.cpp
--- a/Engine/Plugins/Experimental/Mutable/Source/CustomizableObjectEditor/Private/MuCOE/GenerateMutableSource/GenerateMutableSourceGroupProjector.cpp
+++ b/Engine/Plugins/Experimental/Mutable/Source/CustomizableObjectEditor/Private/MuCOE/GenerateMutableSource/GenerateMutableSourceGroupProjector.cpp
@@ -2,51 +2,69 @@
 
 #include "MuCOE/GenerateMutableSource/GenerateMutableSourceGroupProjector.h"
 
-#include "MuCOE/GenerateMutableSource/GenerateMutableSourceFloat.h"
-#include "MuCOE/GenerateMutableSource/GenerateMutableSourceProjector.h"
-
-#include "Templates/UnrealTemplate.h"
-#include "GenericPlatform/GenericPlatformMath.h"
-#include "Interfaces/ITargetPlatform.h"
-#include "Engine/AssetManager.h"
-#include "Engine/TextureLODSettings.h"
-#include "ClothingAsset.h"
-#include "MeshUtilities.h"
-#include "EdGraphToken.h"
-#include "AssetRegistry/AssetRegistryModule.h"
-#include "Materials/MaterialInstance.h"
-#include "MuCOE/UnrealEditorPortabilityHelpers.h"
-#include "PlatformInfo.h"
-
-#include "MuCOE/CustomizableObjectCompiler.h"
-#include "MuCOE/UnrealToMutableTextureConversionUtils.h"
-#include "MuCO/UnrealConversionUtils.h"
-#include "MuCOE/CustomizableObjectEditorModule.h"
+#include "Containers/Array.h"
+#include "Containers/BitArray.h"
+#include "Containers/ContainersFwd.h"
+#include "Containers/Set.h"
+#include "Containers/SparseArray.h"
+#include "Containers/StringConv.h"
+#include "Containers/UnrealString.h"
+#include "EdGraph/EdGraphPin.h"
+#include "Engine/DataTable.h"
+#include "Engine/Texture.h"
+#include "Engine/Texture2D.h"
+#include "Engine/TextureDefines.h"
+#include "HAL/PlatformCrt.h"
+#include "Internationalization/Text.h"
+#include "Logging/TokenizedMessage.h"
+#include "MaterialTypes.h"
+#include "Materials/MaterialInterface.h"
+#include "Math/IntVector.h"
+#include "Math/Transform.h"
+#include "Math/UnrealMathSSE.h"
+#include "Misc/AssertionMacros.h"
+#include "Misc/Guid.h"
 #include "MuCO/CustomizableObject.h"
-#include "MuCOE/GraphTraversal.h"
+#include "MuCO/CustomizableObjectParameterTypeDefinitions.h"
+#include "MuCO/CustomizableObjectUIData.h"
 #include "MuCO/MultilayerProjector.h"
-
-#include "MuT/NodeScalarConstant.h"
-#include "MuT/NodeScalarParameter.h"
-#include "MuT/NodeScalarEnumParameter.h"
+#include "MuCOE/CustomizableObjectCompiler.h"
+#include "MuCOE/GenerateMutableSource/GenerateMutableSource.h"
+#include "MuCOE/GenerateMutableSource/GenerateMutableSourceProjector.h"
+#include "MuCOE/Nodes/CustomizableObjectNode.h"
+#include "MuCOE/Nodes/CustomizableObjectNodeAnimationPose.h"
+#include "MuCOE/Nodes/CustomizableObjectNodeExtendMaterial.h"
+#include "MuCOE/Nodes/CustomizableObjectNodeGroupProjectorParameter.h"
+#include "MuCOE/Nodes/CustomizableObjectNodeMaterial.h"
+#include "MuCOE/Nodes/CustomizableObjectNodeMaterialBase.h"
+#include "MuCOE/Nodes/CustomizableObjectNodeObjectGroup.h"
+#include "MuR/Image.h"
+#include "MuR/Ptr.h"
+#include "MuR/RefCounted.h"
 #include "MuT/NodeColourConstant.h"
 #include "MuT/NodeColourFromScalars.h"
-#include "MuT/NodeProjector.h"
-#include "MuT/NodeImageBinarise.h"
 #include "MuT/NodeImageConstant.h"
-#include "MuT/NodeImageSwitch.h"
 #include "MuT/NodeImageLayerColour.h"
-#include "MuT/NodeImageSwizzle.h"
-#include "MuT/NodeImageProject.h"
-#include "MuT/NodeImagePlainColour.h"
 #include "MuT/NodeImageMultiLayer.h"
+#include "MuT/NodeImagePlainColour.h"
+#include "MuT/NodeImageProject.h"
 #include "MuT/NodeImageResize.h"
-#include "MuT/NodeMeshConstant.h"
-#include "MuT/NodeMeshSwitch.h"
+#include "MuT/NodeImageSwitch.h"
+#include "MuT/NodeImageSwizzle.h"
 #include "MuT/NodeMeshApplyPose.h"
-#include "MuT/Compiler.h"
-#include "MuT/Streams.h"
+#include "MuT/NodeMeshSwitch.h"
+#include "MuT/NodeProjector.h"
 #include "MuT/NodeRangeFromScalar.h"
+#include "MuT/NodeScalar.h"
+#include "MuT/NodeScalarConstant.h"
+#include "MuT/NodeScalarEnumParameter.h"
+#include "MuT/NodeScalarParameter.h"
+#include "Templates/Casts.h"
+#include "UObject/NameTypes.h"
+#include "UObject/ObjectPtr.h"
+#include "UObject/SoftObjectPtr.h"
+
+class UPoseAsset;
 
 #define LOCTEXT_NAMESPACE "CustomizableObjectEditor"
 
diff --git a/Engine/Plugins/Experimental/Mutable/Source/CustomizableObjectEditor/Private/MuCOE/GenerateMutableSource/GenerateMutableSourceGroupProjector.h b/Engine/Plugins/Experimental/Mutable/Source/CustomizableObjectEditor/Private/MuCOE/GenerateMutableSource/GenerateMutableSourceGroupProjector.h
--- a/Engine/Plugins/Experimental/Mutable/Source/CustomizableObjectEditor/Private/MuCOE/GenerateMutableSource/GenerateMutableSourceGroupProjector.h
+++ b/Engine/Plugins/Experimental/Mutable/Source/CustomizableObjectEditor/Private/MuCOE/GenerateMutableSource/GenerateMutableSourceGroupProjector.h
@@ -2,8 +2,18 @@
 
 #pragma once
 
-#include "MuCOE/GenerateMutableSource/GenerateMutableSource.h"
+#include "Containers/Map.h"
+#include "HAL/Platform.h"
+#include "MuT/NodeModifierMeshClipWithMesh.h"
+#include "MuT/NodeSurfaceNew.h"
 
+class FString;
+class UCustomizableObjectNodeMaterial;
+class UCustomizableObjectNodeMaterialBase;
+class UCustomizableObjectNodeObjectGroup;
+class UEdGraphPin;
+class UTexture2D;
+struct FMutableGraphGenerationContext;
 
 mu::NodeImagePtr GenerateMutableGroupProjection(const int32 NodeLOD, const int32 ImageIndex, mu::NodeMeshPtr MeshNode, FMutableGraphGenerationContext& GenerationContext,
 	const UCustomizableObjectNodeMaterialBase* NodeMaterialBase, bool& bShareProjectionTexturesBetweenLODs, bool& bIsGroupProjectorImage,
diff --git a/Engine/Plugins/Experimental/Mutable/Source/CustomizableObjectEditor/Private/MuCOE/GenerateMutableSource/GenerateMutableSourceImage.cpp b/Engine/Plugins/Experimental/Mutable/Source/CustomizableObjectEditor/Private/MuCOE/GenerateMutableSource/GenerateMutableSourceImage.cpp
--- a/Engine/Plugins/Experimental/Mutable/Source/CustomizableObjectEditor/Private/MuCOE/GenerateMutableSource/GenerateMutableSourceImage.cpp
+++ b/Engine/Plugins/Experimental/Mutable/Source/CustomizableObjectEditor/Private/MuCOE/GenerateMutableSource/GenerateMutableSourceImage.cpp
@@ -2,53 +2,87 @@
 
 #include "MuCOE/GenerateMutableSource/GenerateMutableSourceImage.h"
 
-#include "MuCOE/GenerateMutableSource/GenerateMutableSourceFloat.h"
+#include "Containers/Array.h"
+#include "Containers/EnumAsByte.h"
+#include "Containers/Map.h"
+#include "Containers/Set.h"
+#include "Containers/StringConv.h"
+#include "Containers/UnrealString.h"
+#include "EdGraph/EdGraphPin.h"
+#include "Engine/DataTable.h"
+#include "Engine/Texture2D.h"
+#include "HAL/PlatformCrt.h"
+#include "Internationalization/Internationalization.h"
+#include "Internationalization/Text.h"
+#include "Logging/TokenizedMessage.h"
+#include "Math/IntPoint.h"
+#include "Math/IntVector.h"
+#include "Math/UnrealMathSSE.h"
+#include "Misc/AssertionMacros.h"
+#include "Misc/CString.h"
+#include "Misc/Guid.h"
+#include "MuCO/CustomizableObjectParameterTypeDefinitions.h"
+#include "MuCO/CustomizableObjectUIData.h"
+#include "MuCOE/CustomizableObjectCompiler.h"
+#include "MuCOE/EdGraphSchema_CustomizableObject.h"
+#include "MuCOE/GenerateMutableSource/GenerateMutableSource.h"
 #include "MuCOE/GenerateMutableSource/GenerateMutableSourceColor.h"
-#include "MuCOE/GenerateMutableSource/GenerateMutableSourceProjector.h"
+#include "MuCOE/GenerateMutableSource/GenerateMutableSourceFloat.h"
 #include "MuCOE/GenerateMutableSource/GenerateMutableSourceMesh.h"
+#include "MuCOE/GenerateMutableSource/GenerateMutableSourceProjector.h"
 #include "MuCOE/GenerateMutableSource/GenerateMutableSourceTable.h"
-
-#include "Templates/UnrealTemplate.h"
-#include "GenericPlatform/GenericPlatformMath.h"
-#include "Engine/AssetManager.h"
-#include "Engine/TextureLODSettings.h"
-#include "ClothingAsset.h"
-#include "MeshUtilities.h"
-#include "EdGraphToken.h"
-#include "AssetRegistry/AssetRegistryModule.h"
+#include "MuCOE/GraphTraversal.h"
+#include "MuCOE/Nodes/CustomizableObjectNode.h"
+#include "MuCOE/Nodes/CustomizableObjectNodeMesh.h"
+#include "MuCOE/Nodes/CustomizableObjectNodeTable.h"
+#include "MuCOE/Nodes/CustomizableObjectNodeTexture.h"
+#include "MuCOE/Nodes/CustomizableObjectNodeTextureBinarise.h"
+#include "MuCOE/Nodes/CustomizableObjectNodeTextureColourMap.h"
+#include "MuCOE/Nodes/CustomizableObjectNodeTextureFromChannels.h"
+#include "MuCOE/Nodes/CustomizableObjectNodeTextureFromColor.h"
+#include "MuCOE/Nodes/CustomizableObjectNodeTextureInterpolate.h"
+#include "MuCOE/Nodes/CustomizableObjectNodeTextureInvert.h"
+#include "MuCOE/Nodes/CustomizableObjectNodeTextureLayer.h"
+#include "MuCOE/Nodes/CustomizableObjectNodeTextureParameter.h"
+#include "MuCOE/Nodes/CustomizableObjectNodeTextureProject.h"
+#include "MuCOE/Nodes/CustomizableObjectNodeTextureSwitch.h"
+#include "MuCOE/Nodes/CustomizableObjectNodeTextureToChannels.h"
+#include "MuCOE/Nodes/CustomizableObjectNodeTextureTransform.h"
+#include "MuCOE/Nodes/CustomizableObjectNodeTextureVariation.h"
 #include "MuCOE/UnrealEditorPortabilityHelpers.h"
-#include "PlatformInfo.h"
-
-#include "MuCOE/CustomizableObjectCompiler.h"
 #include "MuCOE/UnrealToMutableTextureConversionUtils.h"
-#include "MuCO/UnrealConversionUtils.h"
-#include "MuCOE/CustomizableObjectEditorModule.h"
-#include "MuCO/CustomizableObject.h"
-#include "MuCOE/GraphTraversal.h"
-
+#include "MuR/Image.h"
+#include "MuR/Ptr.h"
+#include "MuR/RefCounted.h"
 #include "MuT/Node.h"
-#include "MuT/NodeScalarConstant.h"
-#include "MuT/NodeScalarEnumParameter.h"
-#include "MuT/NodeProjector.h"
+#include "MuT/NodeImage.h"
 #include "MuT/NodeImageBinarise.h"
+#include "MuT/NodeImageColourMap.h"
 #include "MuT/NodeImageConstant.h"
-#include "MuT/NodeImageParameter.h"
-#include "MuT/NodeImageSwitch.h"
+#include "MuT/NodeImageFormat.h"
 #include "MuT/NodeImageInterpolate.h"
+#include "MuT/NodeImageInvert.h"
 #include "MuT/NodeImageLayer.h"
 #include "MuT/NodeImageLayerColour.h"
-#include "MuT/NodeImageSwizzle.h"
-#include "MuT/NodeImageProject.h"
+#include "MuT/NodeImageParameter.h"
 #include "MuT/NodeImagePlainColour.h"
-#include "MuT/NodeImageFormat.h"
+#include "MuT/NodeImageProject.h"
 #include "MuT/NodeImageResize.h"
-#include "MuT/NodeImageInvert.h"
-#include "MuT/NodeImageColourMap.h"
-#include "MuT/NodeImageVariation.h"
+#include "MuT/NodeImageSwitch.h"
+#include "MuT/NodeImageSwizzle.h"
 #include "MuT/NodeImageTable.h"
 #include "MuT/NodeImageTransform.h"
-#include "MuT/NodeLayout.h"
-#include "MuT/Compiler.h"
+#include "MuT/NodeImageVariation.h"
+#include "MuT/NodeProjector.h"
+#include "MuT/NodeScalar.h"
+#include "MuT/NodeScalarConstant.h"
+#include "MuT/NodeScalarEnumParameter.h"
+#include "Templates/Casts.h"
+#include "Templates/Tuple.h"
+#include "TextureResource.h"
+#include "UObject/NameTypes.h"
+#include "UObject/ObjectPtr.h"
+#include "UObject/UObjectGlobals.h"
 
 #define LOCTEXT_NAMESPACE "CustomizableObjectEditor"
 
diff --git a/Engine/Plugins/Experimental/Mutable/Source/CustomizableObjectEditor/Private/MuCOE/GenerateMutableSource/GenerateMutableSourceImage.h b/Engine/Plugins/Experimental/Mutable/Source/CustomizableObjectEditor/Private/MuCOE/GenerateMutableSource/GenerateMutableSourceImage.h
--- a/Engine/Plugins/Experimental/Mutable/Source/CustomizableObjectEditor/Private/MuCOE/GenerateMutableSource/GenerateMutableSourceImage.h
+++ b/Engine/Plugins/Experimental/Mutable/Source/CustomizableObjectEditor/Private/MuCOE/GenerateMutableSource/GenerateMutableSourceImage.h
@@ -2,8 +2,16 @@
 
 #pragma once
 
-#include "MuCOE/GenerateMutableSource/GenerateMutableSource.h"
 
+#include "MuT/NodeImageConstant.h"
+#include "MuT/NodeSurfaceNew.h"
+#include "MuT/Table.h"
+
+class FCustomizableObjectCompiler;
+class UCustomizableObjectNode;
+class UEdGraphPin;
+class UTexture2D;
+struct FMutableGraphGenerationContext;
 
 mu::ImagePtr ConvertTextureUnrealToMutable(UTexture2D* Texture, const UCustomizableObjectNode* Node, FCustomizableObjectCompiler* Compiler, bool bIsNormalComposite);
 
diff --git a/Engine/Plugins/Experimental/Mutable/Source/CustomizableObjectEditor/Private/MuCOE/GenerateMutableSource/GenerateMutableSourceLayout.cpp b/Engine/Plugins/Experimental/Mutable/Source/CustomizableObjectEditor/Private/MuCOE/GenerateMutableSource/GenerateMutableSourceLayout.cpp
--- a/Engine/Plugins/Experimental/Mutable/Source/CustomizableObjectEditor/Private/MuCOE/GenerateMutableSource/GenerateMutableSourceLayout.cpp
+++ b/Engine/Plugins/Experimental/Mutable/Source/CustomizableObjectEditor/Private/MuCOE/GenerateMutableSource/GenerateMutableSourceLayout.cpp
@@ -2,28 +2,31 @@
 
 #include "MuCOE/GenerateMutableSource/GenerateMutableSourceLayout.h"
 
-#include "MuCOE/GenerateMutableSource/GenerateMutableSourceFloat.h"
-
-#include "Templates/UnrealTemplate.h"
-#include "Engine/AssetManager.h"
-#include "Engine/TextureLODSettings.h"
-#include "ClothingAsset.h"
-#include "EdGraphToken.h"
-#include "AssetRegistry/AssetRegistryModule.h"
-#include "MuCOE/UnrealEditorPortabilityHelpers.h"
-#include "PlatformInfo.h"
-
+#include "Containers/Array.h"
+#include "Containers/Map.h"
+#include "Containers/Set.h"
+#include "Containers/UnrealString.h"
+#include "EdGraph/EdGraphPin.h"
+#include "Internationalization/Internationalization.h"
+#include "Internationalization/Text.h"
+#include "Logging/TokenizedMessage.h"
+#include "Math/IntPoint.h"
+#include "Math/UnrealMathSSE.h"
+#include "Misc/AssertionMacros.h"
 #include "MuCOE/CustomizableObjectCompiler.h"
-#include "MuCOE/UnrealToMutableTextureConversionUtils.h"
-#include "MuCO/UnrealConversionUtils.h"
-#include "MuCOE/CustomizableObjectEditorModule.h"
-#include "MuCO/CustomizableObject.h"
+#include "MuCOE/CustomizableObjectLayout.h"
+#include "MuCOE/GenerateMutableSource/GenerateMutableSource.h"
 #include "MuCOE/GraphTraversal.h"
-
+#include "MuCOE/Nodes/CustomizableObjectNode.h"
+#include "MuCOE/Nodes/CustomizableObjectNodeLayoutBlocks.h"
+#include "MuCOE/Nodes/CustomizableObjectNodeSkeletalMesh.h"
+#include "MuR/Layout.h"
+#include "MuR/Ptr.h"
+#include "MuR/RefCounted.h"
 #include "MuT/Node.h"
 #include "MuT/NodeLayout.h"
-#include "MuT/Compiler.h"
-#include "MuT/Streams.h"
+#include "Templates/Casts.h"
+#include "UObject/ObjectPtr.h"
 
 #define LOCTEXT_NAMESPACE "CustomizableObjectEditor"
 
diff --git a/Engine/Plugins/Experimental/Mutable/Source/CustomizableObjectEditor/Private/MuCOE/GenerateMutableSource/GenerateMutableSourceLayout.h b/Engine/Plugins/Experimental/Mutable/Source/CustomizableObjectEditor/Private/MuCOE/GenerateMutableSource/GenerateMutableSourceLayout.h
--- a/Engine/Plugins/Experimental/Mutable/Source/CustomizableObjectEditor/Private/MuCOE/GenerateMutableSource/GenerateMutableSourceLayout.h
+++ b/Engine/Plugins/Experimental/Mutable/Source/CustomizableObjectEditor/Private/MuCOE/GenerateMutableSource/GenerateMutableSourceLayout.h
@@ -2,8 +2,10 @@
 
 #pragma once
 
-#include "MuCOE/GenerateMutableSource/GenerateMutableSource.h"
 #include "MuT/NodeLayout.h"
 
+class UEdGraphPin;
+struct FMutableGraphGenerationContext;
+
 
 mu::NodeLayoutPtr GenerateMutableSourceLayout(const UEdGraphPin* Pin, FMutableGraphGenerationContext& GenerationContext);
diff --git a/Engine/Plugins/Experimental/Mutable/Source/CustomizableObjectEditor/Private/MuCOE/GenerateMutableSource/GenerateMutableSourceMesh.cpp b/Engine/Plugins/Experimental/Mutable/Source/CustomizableObjectEditor/Private/MuCOE/GenerateMutableSource/GenerateMutableSourceMesh.cpp
--- a/Engine/Plugins/Experimental/Mutable/Source/CustomizableObjectEditor/Private/MuCOE/GenerateMutableSource/GenerateMutableSourceMesh.cpp
+++ b/Engine/Plugins/Experimental/Mutable/Source/CustomizableObjectEditor/Private/MuCOE/GenerateMutableSource/GenerateMutableSourceMesh.cpp
@@ -2,55 +2,133 @@
 
 #include "MuCOE/GenerateMutableSource/GenerateMutableSourceMesh.h"
 
-#include "MuCOE/GenerateMutableSource/GenerateMutableSourceFloat.h"
-#include "MuCOE/GenerateMutableSource/GenerateMutableSourceLayout.h"
-#include "MuCOE/GenerateMutableSource/GenerateMutableSourceTable.h"
-
-#include "Templates/UnrealTemplate.h"
+#include "Animation/MorphTarget.h"
+#include "Animation/PoseAsset.h"
+#include "Animation/Skeleton.h"
+#include "BoneContainer.h"
+#include "BoneIndices.h"
+#include "ClothConfigBase.h"
+#include "ClothLODData.h"
+#include "ClothingAsset.h"
+#include "ClothingAssetBase.h"
+#include "Containers/Array.h"
+#include "Containers/ArrayView.h"
+#include "Containers/IndirectArray.h"
+#include "Containers/Map.h"
+#include "Containers/Queue.h"
+#include "Containers/Set.h"
+#include "Containers/StringConv.h"
+#include "EdGraph/EdGraphNode.h"
+#include "EdGraph/EdGraphPin.h"
+#include "Engine/DataTable.h"
 #include "Engine/SkeletalMesh.h"
 #include "Engine/StaticMesh.h"
-#include "Engine/AssetManager.h"
-#include "Engine/TextureLODSettings.h"
-#include "PhysicsEngine/PhysicsAsset.h"
-#include "ClothingAsset.h"
+#include "GPUSkinPublicDefs.h"
+#include "GameplayTagContainer.h"
+#include "HAL/PlatformCrt.h"
+#include "HAL/UnrealMemory.h"
+#include "Internationalization/Internationalization.h"
+#include "Internationalization/Text.h"
+#include "Logging/LogCategory.h"
+#include "Logging/LogMacros.h"
+#include "Logging/TokenizedMessage.h"
+#include "Math/IntPoint.h"
+#include "Math/Matrix.h"
+#include "Math/NumericLimits.h"
+#include "Math/Quat.h"
+#include "Math/Rotator.h"
+#include "Math/Transform.h"
+#include "Math/UnrealMathSSE.h"
+#include "Math/Vector.h"
+#include "Math/Vector2D.h"
 #include "MeshUtilities.h"
-#include "EdGraphToken.h"
-#include "AssetRegistry/AssetRegistryModule.h"
-#include "MuCOE/UnrealEditorPortabilityHelpers.h"
+#include "Misc/AssertionMacros.h"
+#include "Misc/Guid.h"
+#include "Modules/ModuleManager.h"
+#include "MuCO/CustomizableObject.h"
+#include "MuCO/CustomizableObjectClothingTypes.h"
+#include "MuCO/CustomizableObjectInstance.h"
+#include "MuCO/CustomizableObjectUIData.h"
 #include "MuCO/UnrealPortabilityHelpers.h"
-#include "PlatformInfo.h"
-#include "Templates/Decay.h"
-#include "GenericPlatform/GenericPlatformMath.h"
-#include "StaticMeshResources.h"
-#include "Animation/Skeleton.h"
-#include "Animation/MorphTarget.h"
-#include "Engine/SkeletalMeshSocket.h"
-
 #include "MuCOE/CustomizableObjectCompiler.h"
-#include "MuCOE/UnrealToMutableTextureConversionUtils.h"
-#include "MuCO/UnrealConversionUtils.h"
-#include "MuCOE/CustomizableObjectEditorModule.h"
-#include "MuCO/CustomizableObjectSystem.h"
-#include "MuCO/CustomizableObject.h"
+#include "MuCOE/CustomizableObjectLayout.h"
+#include "MuCOE/GenerateMutableSource/GenerateMutableSource.h"
+#include "MuCOE/GenerateMutableSource/GenerateMutableSourceFloat.h"
+#include "MuCOE/GenerateMutableSource/GenerateMutableSourceLayout.h"
+#include "MuCOE/GenerateMutableSource/GenerateMutableSourceTable.h"
 #include "MuCOE/GraphTraversal.h"
-#include "MuCOE/MutableUtils.h"
-
+#include "MuCOE/Nodes/CustomizableObjectNode.h"
+#include "MuCOE/Nodes/CustomizableObjectNodeAnimationPose.h"
+#include "MuCOE/Nodes/CustomizableObjectNodeFloatConstant.h"
+#include "MuCOE/Nodes/CustomizableObjectNodeFloatParameter.h"
+#include "MuCOE/Nodes/CustomizableObjectNodeMeshGeometryOperation.h"
+#include "MuCOE/Nodes/CustomizableObjectNodeMeshMorph.h"
+#include "MuCOE/Nodes/CustomizableObjectNodeMeshMorphStackApplication.h"
+#include "MuCOE/Nodes/CustomizableObjectNodeMeshMorphStackDefinition.h"
+#include "MuCOE/Nodes/CustomizableObjectNodeMeshReshape.h"
+#include "MuCOE/Nodes/CustomizableObjectNodeMeshReshapeCommon.h"
+#include "MuCOE/Nodes/CustomizableObjectNodeMeshSwitch.h"
+#include "MuCOE/Nodes/CustomizableObjectNodeMeshVariation.h"
+#include "MuCOE/Nodes/CustomizableObjectNodeObject.h"
+#include "MuCOE/Nodes/CustomizableObjectNodeSkeletalMesh.h"
+#include "MuCOE/Nodes/CustomizableObjectNodeStaticMesh.h"
+#include "MuCOE/Nodes/CustomizableObjectNodeTable.h"
+#include "MuCOE/UnrealEditorPortabilityHelpers.h"
+#include "MuR/Layout.h"
+#include "MuR/Mesh.h"
+#include "MuR/MeshBufferSet.h"
+#include "MuR/PhysicsBody.h"
+#include "MuR/Ptr.h"
+#include "MuR/RefCounted.h"
+#include "MuR/Skeleton.h"
 #include "MuT/Node.h"
-#include "MuT/NodeScalarConstant.h"
-#include "MuT/NodeScalarEnumParameter.h"
 #include "MuT/NodeLayout.h"
-#include "MuT/NodeMeshConstant.h"
-#include "MuT/NodeMeshMorph.h"
-#include "MuT/NodeMeshMakeMorph.h"
-#include "MuT/NodeMeshSwitch.h"
+#include "MuT/NodeMesh.h"
 #include "MuT/NodeMeshApplyPose.h"
-#include "MuT/NodeMeshVariation.h"
+#include "MuT/NodeMeshConstant.h"
 #include "MuT/NodeMeshGeometryOperation.h"
+#include "MuT/NodeMeshMakeMorph.h"
+#include "MuT/NodeMeshMorph.h"
 #include "MuT/NodeMeshReshape.h"
+#include "MuT/NodeMeshSwitch.h"
 #include "MuT/NodeMeshTable.h"
-#include "MuT/NodeSurfaceNew.h"
-#include "MuT/Compiler.h"
-#include "MuT/Streams.h"
+#include "MuT/NodeMeshVariation.h"
+#include "MuT/NodeScalar.h"
+#include "MuT/NodeScalarEnumParameter.h"
+#include "PhysicsEngine/AggregateGeom.h"
+#include "PhysicsEngine/BoxElem.h"
+#include "PhysicsEngine/ConvexElem.h"
+#include "PhysicsEngine/PhysicsAsset.h"
+#include "PhysicsEngine/ShapeElem.h"
+#include "PhysicsEngine/SphereElem.h"
+#include "PhysicsEngine/SphylElem.h"
+#include "PhysicsEngine/TaperedCapsuleElem.h"
+#include "RawIndexBuffer.h"
+#include "ReferenceSkeleton.h"
+#include "Rendering/MultiSizeIndexContainer.h"
+#include "Rendering/PositionVertexBuffer.h"
+#include "Rendering/SkeletalMeshLODModel.h"
+#include "Rendering/SkeletalMeshLODRenderData.h"
+#include "Rendering/SkeletalMeshModel.h"
+#include "Rendering/SkeletalMeshRenderData.h"
+#include "Rendering/StaticMeshVertexBuffer.h"
+#include "Serialization/MemoryWriter.h"
+#include "Serialization/StructuredArchiveAdapters.h"
+#include "SkeletalMeshTypes.h"
+#include "StaticMeshResources.h"
+#include "Templates/Casts.h"
+#include "Templates/Decay.h"
+#include "Templates/TypeHash.h"
+#include "Templates/UnrealTemplate.h"
+#include "Trace/Detail/Channel.h"
+#include "UObject/Class.h"
+#include "UObject/NameTypes.h"
+#include "UObject/Object.h"
+#include "UObject/ObjectPtr.h"
+#include "UObject/Package.h"
+#include "UObject/SoftObjectPtr.h"
+#include "UObject/UObjectGlobals.h"
+#include "UObject/UnrealNames.h"
 
 #define LOCTEXT_NAMESPACE "CustomizableObjectEditor"
 
diff --git a/Engine/Plugins/Experimental/Mutable/Source/CustomizableObjectEditor/Private/MuCOE/GenerateMutableSource/GenerateMutableSourceMesh.h b/Engine/Plugins/Experimental/Mutable/Source/CustomizableObjectEditor/Private/MuCOE/GenerateMutableSource/GenerateMutableSourceMesh.h
--- a/Engine/Plugins/Experimental/Mutable/Source/CustomizableObjectEditor/Private/MuCOE/GenerateMutableSource/GenerateMutableSourceMesh.h
+++ b/Engine/Plugins/Experimental/Mutable/Source/CustomizableObjectEditor/Private/MuCOE/GenerateMutableSource/GenerateMutableSourceMesh.h
@@ -2,8 +2,18 @@
 
 #pragma once
 
-#include "MuCOE/GenerateMutableSource/GenerateMutableSource.h"
-
+#include "Containers/UnrealString.h"
+#include "HAL/Platform.h"
+#include "MuT/NodeModifierMeshClipWithMesh.h"
+#include "MuT/Table.h"
+
+class UCustomizableObjectNode;
+class UEdGraphPin;
+class UObject;
+class USkeletalMesh;
+class UStaticMesh;
+struct FMutableGraphGenerationContext;
+struct FMutableGraphMeshGenerationData;
 
 mu::MeshPtr ConvertSkeletalMeshToMutable(USkeletalMesh* InSkeletalMesh, int LOD, int MaterialIndex, FMutableGraphGenerationContext& GenerationContext, const UCustomizableObjectNode* CurrentNode);
 
diff --git a/Engine/Plugins/Experimental/Mutable/Source/CustomizableObjectEditor/Private/MuCOE/GenerateMutableSource/GenerateMutableSourceModifier.cpp b/Engine/Plugins/Experimental/Mutable/Source/CustomizableObjectEditor/Private/MuCOE/GenerateMutableSource/GenerateMutableSourceModifier.cpp
--- a/Engine/Plugins/Experimental/Mutable/Source/CustomizableObjectEditor/Private/MuCOE/GenerateMutableSource/GenerateMutableSourceModifier.cpp
+++ b/Engine/Plugins/Experimental/Mutable/Source/CustomizableObjectEditor/Private/MuCOE/GenerateMutableSource/GenerateMutableSourceModifier.cpp
@@ -2,35 +2,43 @@
 
 #include "MuCOE/GenerateMutableSource/GenerateMutableSourceModifier.h"
 
-#include "MuCOE/GenerateMutableSource/GenerateMutableSourceFloat.h"
-#include "MuCOE/GenerateMutableSource/GenerateMutableSourceMesh.h"
-
-#include "Templates/UnrealTemplate.h"
+#include "Containers/Array.h"
+#include "Containers/Map.h"
+#include "Containers/Set.h"
+#include "Containers/StringConv.h"
+#include "Containers/UnrealString.h"
+#include "EdGraph/EdGraphPin.h"
 #include "Engine/SkeletalMesh.h"
 #include "Engine/StaticMesh.h"
-#include "Engine/AssetManager.h"
-#include "Engine/TextureLODSettings.h"
-#include "ClothingAsset.h"
-#include "EdGraphToken.h"
-#include "AssetRegistry/AssetRegistryModule.h"
-#include "MuCOE/UnrealEditorPortabilityHelpers.h"
-#include "PlatformInfo.h"
-
-#include "MuCOE/CustomizableObjectCompiler.h"
-#include "MuCOE/UnrealToMutableTextureConversionUtils.h"
-#include "MuCO/UnrealConversionUtils.h"
-#include "MuCOE/CustomizableObjectEditorModule.h"
+#include "HAL/PlatformCrt.h"
+#include "Internationalization/Internationalization.h"
+#include "Math/MathFwd.h"
+#include "Math/Matrix.h"
+#include "Math/Transform.h"
+#include "Math/Vector.h"
+#include "Misc/AssertionMacros.h"
 #include "MuCO/CustomizableObject.h"
+#include "MuCOE/CustomizableObjectCompiler.h"
+#include "MuCOE/GenerateMutableSource/GenerateMutableSource.h"
+#include "MuCOE/GenerateMutableSource/GenerateMutableSourceMesh.h"
 #include "MuCOE/GraphTraversal.h"
 #include "MuCOE/MutableUtils.h"
-
-#include "MuT/NodeMeshConstant.h"
+#include "MuCOE/Nodes/CustomizableObjectNode.h"
+#include "MuCOE/Nodes/CustomizableObjectNodeMeshClipDeform.h"
+#include "MuCOE/Nodes/CustomizableObjectNodeMeshClipMorph.h"
+#include "MuCOE/Nodes/CustomizableObjectNodeMeshClipWithMesh.h"
+#include "MuCOE/Nodes/CustomizableObjectNodeStaticMesh.h"
+#include "MuR/Mesh.h"
+#include "MuR/Ptr.h"
+#include "MuR/RefCounted.h"
+#include "MuT/Node.h"
 #include "MuT/NodeMeshTransform.h"
-#include "MuT/Compiler.h"
-#include "MuT/Streams.h"
+#include "MuT/NodeModifierMeshClipDeform.h"
 #include "MuT/NodeModifierMeshClipMorphPlane.h"
 #include "MuT/NodeModifierMeshClipWithMesh.h"
-#include "MuT/NodeModifierMeshClipDeform.h"
+#include "Templates/Casts.h"
+#include "UObject/NameTypes.h"
+#include "UObject/ObjectPtr.h"
 
 #define LOCTEXT_NAMESPACE "CustomizableObjectEditor"
 
diff --git a/Engine/Plugins/Experimental/Mutable/Source/CustomizableObjectEditor/Private/MuCOE/GenerateMutableSource/GenerateMutableSourceModifier.h b/Engine/Plugins/Experimental/Mutable/Source/CustomizableObjectEditor/Private/MuCOE/GenerateMutableSource/GenerateMutableSourceModifier.h
--- a/Engine/Plugins/Experimental/Mutable/Source/CustomizableObjectEditor/Private/MuCOE/GenerateMutableSource/GenerateMutableSourceModifier.h
+++ b/Engine/Plugins/Experimental/Mutable/Source/CustomizableObjectEditor/Private/MuCOE/GenerateMutableSource/GenerateMutableSourceModifier.h
@@ -2,7 +2,9 @@
 
 #pragma once
 
-#include "MuCOE/GenerateMutableSource/GenerateMutableSource.h"
+#include "MuT/NodeModifier.h"
 
+class UEdGraphPin;
+struct FMutableGraphGenerationContext;
 
 mu::NodeModifierPtr GenerateMutableSourceModifier(const UEdGraphPin* Pin, FMutableGraphGenerationContext& GenerationContext);
diff --git a/Engine/Plugins/Experimental/Mutable/Source/CustomizableObjectEditor/Private/MuCOE/GenerateMutableSource/GenerateMutableSourceProjector.cpp b/Engine/Plugins/Experimental/Mutable/Source/CustomizableObjectEditor/Private/MuCOE/GenerateMutableSource/GenerateMutableSourceProjector.cpp
--- a/Engine/Plugins/Experimental/Mutable/Source/CustomizableObjectEditor/Private/MuCOE/GenerateMutableSource/GenerateMutableSourceProjector.cpp
+++ b/Engine/Plugins/Experimental/Mutable/Source/CustomizableObjectEditor/Private/MuCOE/GenerateMutableSource/GenerateMutableSourceProjector.cpp
@@ -2,30 +2,29 @@
 
 #include "MuCOE/GenerateMutableSource/GenerateMutableSourceProjector.h"
 
-#include "MuCOE/GenerateMutableSource/GenerateMutableSourceFloat.h"
-
-#include "Templates/UnrealTemplate.h"
-#include "GenericPlatform/GenericPlatformMath.h"
-#include "Engine/AssetManager.h"
-#include "Engine/TextureLODSettings.h"
-#include "ClothingAsset.h"
-#include "MeshUtilities.h"
-#include "EdGraphToken.h"
-#include "AssetRegistry/AssetRegistryModule.h"
-#include "MuCOE/UnrealEditorPortabilityHelpers.h"
-#include "PlatformInfo.h"
-
+#include "Containers/Map.h"
+#include "Containers/Set.h"
+#include "Containers/StringConv.h"
+#include "Containers/UnrealString.h"
+#include "EdGraph/EdGraphPin.h"
+#include "Internationalization/Internationalization.h"
+#include "Math/UnrealMathSSE.h"
+#include "Math/Vector.h"
+#include "Misc/AssertionMacros.h"
+#include "Misc/Guid.h"
+#include "MuCO/CustomizableObjectParameterTypeDefinitions.h"
+#include "MuCO/CustomizableObjectUIData.h"
 #include "MuCOE/CustomizableObjectCompiler.h"
-#include "MuCOE/UnrealToMutableTextureConversionUtils.h"
-#include "MuCO/UnrealConversionUtils.h"
-#include "MuCOE/CustomizableObjectEditorModule.h"
-#include "MuCO/CustomizableObject.h"
-
+#include "MuCOE/GenerateMutableSource/GenerateMutableSource.h"
+#include "MuCOE/Nodes/CustomizableObjectNode.h"
+#include "MuCOE/Nodes/CustomizableObjectNodeProjectorConstant.h"
+#include "MuCOE/Nodes/CustomizableObjectNodeProjectorParameter.h"
+#include "MuR/Parameters.h"
+#include "MuR/Ptr.h"
+#include "MuR/RefCounted.h"
 #include "MuT/Node.h"
-#include "MuT/NodeScalarConstant.h"
 #include "MuT/NodeProjector.h"
-#include "MuT/Compiler.h"
-#include "MuT/Streams.h"
+#include "Templates/Casts.h"
 
 #define LOCTEXT_NAMESPACE "CustomizableObjectEditor"
 
diff --git a/Engine/Plugins/Experimental/Mutable/Source/CustomizableObjectEditor/Private/MuCOE/GenerateMutableSource/GenerateMutableSourceProjector.h b/Engine/Plugins/Experimental/Mutable/Source/CustomizableObjectEditor/Private/MuCOE/GenerateMutableSource/GenerateMutableSourceProjector.h
--- a/Engine/Plugins/Experimental/Mutable/Source/CustomizableObjectEditor/Private/MuCOE/GenerateMutableSource/GenerateMutableSourceProjector.h
+++ b/Engine/Plugins/Experimental/Mutable/Source/CustomizableObjectEditor/Private/MuCOE/GenerateMutableSource/GenerateMutableSourceProjector.h
@@ -2,8 +2,10 @@
 
 #pragma once
 
-#include "MuCOE/GenerateMutableSource/GenerateMutableSource.h"
+#include "MuT/NodeProjector.h"
 
+class UEdGraphPin;
+struct FMutableGraphGenerationContext;
 
 /** Convert a CustomizableObject Source Graph into a mutable source graph. */
 mu::NodeProjectorPtr GenerateMutableSourceProjector(const UEdGraphPin* Pin, FMutableGraphGenerationContext& GenerationContext);
diff --git a/Engine/Plugins/Experimental/Mutable/Source/CustomizableObjectEditor/Private/MuCOE/GenerateMutableSource/GenerateMutableSourceSurface.cpp b/Engine/Plugins/Experimental/Mutable/Source/CustomizableObjectEditor/Private/MuCOE/GenerateMutableSource/GenerateMutableSourceSurface.cpp
--- a/Engine/Plugins/Experimental/Mutable/Source/CustomizableObjectEditor/Private/MuCOE/GenerateMutableSource/GenerateMutableSourceSurface.cpp
+++ b/Engine/Plugins/Experimental/Mutable/Source/CustomizableObjectEditor/Private/MuCOE/GenerateMutableSource/GenerateMutableSourceSurface.cpp
@@ -2,59 +2,90 @@
 
 #include "MuCOE/GenerateMutableSource/GenerateMutableSourceSurface.h"
 
-#include "MuCOE/GenerateMutableSource/GenerateMutableSourceFloat.h"
+#include "Containers/Array.h"
+#include "Containers/EnumAsByte.h"
+#include "Containers/Map.h"
+#include "Containers/Set.h"
+#include "Containers/StringConv.h"
+#include "Containers/UnrealString.h"
+#include "EdGraph/EdGraphNode.h"
+#include "EdGraph/EdGraphPin.h"
+#include "Engine/DataTable.h"
+#include "Engine/SkeletalMesh.h"
+#include "Engine/StaticMesh.h"
+#include "Engine/Texture2D.h"
+#include "Engine/TextureDefines.h"
+#include "GPUSkinPublicDefs.h"
+#include "HAL/PlatformCrt.h"
+#include "Interfaces/ITargetPlatform.h"
+#include "Internationalization/Internationalization.h"
+#include "Internationalization/Text.h"
+#include "Logging/LogCategory.h"
+#include "Logging/LogMacros.h"
+#include "MaterialTypes.h"
+#include "Materials/MaterialInstance.h"
+#include "Materials/MaterialInterface.h"
+#include "Math/UnrealMathSSE.h"
+#include "Misc/AssertionMacros.h"
+#include "Misc/CString.h"
+#include "Misc/Guid.h"
+#include "MuCO/CustomizableObject.h"
+#include "MuCO/CustomizableObjectInstance.h"
+#include "MuCO/CustomizableObjectUIData.h"
+#include "MuCO/MutableMeshBufferUtils.h"
+#include "MuCOE/CustomizableObjectCompiler.h"
+#include "MuCOE/CustomizableObjectLayout.h"
+#include "MuCOE/GenerateMutableSource/GenerateMutableSource.h"
 #include "MuCOE/GenerateMutableSource/GenerateMutableSourceColor.h"
+#include "MuCOE/GenerateMutableSource/GenerateMutableSourceFloat.h"
+#include "MuCOE/GenerateMutableSource/GenerateMutableSourceGroupProjector.h"
 #include "MuCOE/GenerateMutableSource/GenerateMutableSourceImage.h"
 #include "MuCOE/GenerateMutableSource/GenerateMutableSourceMesh.h"
 #include "MuCOE/GenerateMutableSource/GenerateMutableSourceTable.h"
-#include "MuCOE/GenerateMutableSource/GenerateMutableSourceGroupProjector.h"
-
-#include "Templates/UnrealTemplate.h"
-#include "GenericPlatform/GenericPlatformMath.h"
-#include "Interfaces/ITargetPlatform.h"
-#include "Engine/SkeletalMesh.h"
-#include "Engine/AssetManager.h"
-#include "Engine/TextureLODSettings.h"
-#include "ClothingAsset.h"
-#include "MeshUtilities.h"
-#include "EdGraphToken.h"
-#include "AssetRegistry/AssetRegistryModule.h"
-#include "Materials/MaterialInstance.h"
-#include "MuCOE/UnrealEditorPortabilityHelpers.h"
-#include "PlatformInfo.h"
-#include "Engine/SkeletalMeshSocket.h"
-
-#include "MuCOE/CustomizableObjectCompiler.h"
-#include "MuCO/UnrealConversionUtils.h"
-#include "MuCOE/CustomizableObjectEditorModule.h"
-#include "MuCO/CustomizableObjectSystem.h"
-#include "MuCO/CustomizableObject.h"
 #include "MuCOE/GraphTraversal.h"
 #include "MuCOE/MutableUtils.h"
-#include "MuCO/MutableMeshBufferUtils.h"
-
+#include "MuCOE/Nodes/CustomizableObjectNode.h"
+#include "MuCOE/Nodes/CustomizableObjectNodeCopyMaterial.h"
+#include "MuCOE/Nodes/CustomizableObjectNodeEditMaterial.h"
+#include "MuCOE/Nodes/CustomizableObjectNodeExtendMaterial.h"
+#include "MuCOE/Nodes/CustomizableObjectNodeFloatConstant.h"
+#include "MuCOE/Nodes/CustomizableObjectNodeFloatParameter.h"
+#include "MuCOE/Nodes/CustomizableObjectNodeMaterial.h"
+#include "MuCOE/Nodes/CustomizableObjectNodeMaterialBase.h"
+#include "MuCOE/Nodes/CustomizableObjectNodeMaterialVariation.h"
+#include "MuCOE/Nodes/CustomizableObjectNodeMesh.h"
+#include "MuCOE/Nodes/CustomizableObjectNodeMorphMaterial.h"
+#include "MuCOE/Nodes/CustomizableObjectNodeRemoveMesh.h"
+#include "MuCOE/Nodes/CustomizableObjectNodeRemoveMeshBlocks.h"
+#include "MuCOE/Nodes/CustomizableObjectNodeSkeletalMesh.h"
+#include "MuCOE/Nodes/CustomizableObjectNodeTable.h"
+#include "MuR/Image.h"
+#include "MuR/MeshBufferSet.h"
+#include "MuR/Ptr.h"
+#include "MuR/RefCounted.h"
 #include "MuT/Node.h"
-#include "MuT/NodeColourConstant.h"
-#include "MuT/NodeImageBinarise.h"
+#include "MuT/NodeColour.h"
 #include "MuT/NodeImageConstant.h"
-#include "MuT/NodeImageMipmap.h"
 #include "MuT/NodeImageFormat.h"
-#include "MuT/NodeImageResize.h"
+#include "MuT/NodeImageMipmap.h"
 #include "MuT/NodeImageNormalComposite.h"
-#include "MuT/NodePatchImage.h"
+#include "MuT/NodeImageResize.h"
 #include "MuT/NodeMeshConstant.h"
 #include "MuT/NodeMeshFormat.h"
 #include "MuT/NodeMeshFragment.h"
-#include "MuT/NodeSurfaceNew.h"
+#include "MuT/NodePatchImage.h"
+#include "MuT/NodePatchMesh.h"
+#include "MuT/NodeSurface.h"
 #include "MuT/NodeSurfaceEdit.h"
+#include "MuT/NodeSurfaceNew.h"
 #include "MuT/NodeSurfaceVariation.h"
-#include "MuT/NodePatchMesh.h"
-#include "MuT/Compiler.h"
-#include "MuT/Streams.h"
-
-#include "MuR/Mesh.h"
-#include "MuR/Skeleton.h"
-#include "MuR/Image.h"
+#include "MuT/Table.h"
+#include "Templates/Casts.h"
+#include "Trace/Detail/Channel.h"
+#include "UObject/NameTypes.h"
+#include "UObject/Object.h"
+#include "UObject/ObjectPtr.h"
+#include "UObject/UnrealNames.h"
 
 #define LOCTEXT_NAMESPACE "CustomizableObjectEditor"
 
diff --git a/Engine/Plugins/Experimental/Mutable/Source/CustomizableObjectEditor/Private/MuCOE/GenerateMutableSource/GenerateMutableSourceSurface.h b/Engine/Plugins/Experimental/Mutable/Source/CustomizableObjectEditor/Private/MuCOE/GenerateMutableSource/GenerateMutableSourceSurface.h
--- a/Engine/Plugins/Experimental/Mutable/Source/CustomizableObjectEditor/Private/MuCOE/GenerateMutableSource/GenerateMutableSourceSurface.h
+++ b/Engine/Plugins/Experimental/Mutable/Source/CustomizableObjectEditor/Private/MuCOE/GenerateMutableSource/GenerateMutableSourceSurface.h
@@ -2,7 +2,10 @@
 
 #pragma once
 
-#include "MuCOE/GenerateMutableSource/GenerateMutableSource.h"
+#include "MuT/NodeComponentNew.h"
 
+class UEdGraphPin;
+struct FMutableGraphGenerationContext;
+struct FMutableGraphSurfaceGenerationData;
 
 mu::NodeSurfacePtr GenerateMutableSourceSurface(const UEdGraphPin * Pin, FMutableGraphGenerationContext & GenerationContext, FMutableGraphSurfaceGenerationData& SurfaceData);
diff --git a/Engine/Plugins/Experimental/Mutable/Source/CustomizableObjectEditor/Private/MuCOE/GenerateMutableSource/GenerateMutableSourceTable.cpp b/Engine/Plugins/Experimental/Mutable/Source/CustomizableObjectEditor/Private/MuCOE/GenerateMutableSource/GenerateMutableSourceTable.cpp
--- a/Engine/Plugins/Experimental/Mutable/Source/CustomizableObjectEditor/Private/MuCOE/GenerateMutableSource/GenerateMutableSourceTable.cpp
+++ b/Engine/Plugins/Experimental/Mutable/Source/CustomizableObjectEditor/Private/MuCOE/GenerateMutableSource/GenerateMutableSourceTable.cpp
@@ -2,35 +2,52 @@
 
 #include "MuCOE/GenerateMutableSource/GenerateMutableSourceTable.h"
 
-#include "MuCOE/GenerateMutableSource/GenerateMutableSourceFloat.h"
-#include "MuCOE/GenerateMutableSource/GenerateMutableSourceMesh.h"
-
-#include "Templates/UnrealTemplate.h"
+#include "Animation/AnimInstance.h"
+#include "Containers/Array.h"
+#include "Containers/IndirectArray.h"
+#include "Containers/Map.h"
+#include "Containers/SparseArray.h"
+#include "Containers/StringConv.h"
+#include "EdGraph/EdGraphPin.h"
+#include "Engine/DataTable.h"
 #include "Engine/SkeletalMesh.h"
-#include "Engine/AssetManager.h"
-#include "Engine/TextureLODSettings.h"
-#include "ClothingAsset.h"
-#include "MeshUtilities.h"
-#include "EdGraphToken.h"
-#include "AssetRegistry/AssetRegistryModule.h"
+#include "Engine/StaticMesh.h"
+#include "Engine/Texture.h"
+#include "Engine/Texture2D.h"
+#include "GameplayTagContainer.h"
+#include "HAL/Platform.h"
+#include "HAL/PlatformCrt.h"
+#include "Internationalization/Internationalization.h"
+#include "Internationalization/Text.h"
+#include "MaterialTypes.h"
+#include "Materials/Material.h"
 #include "Materials/MaterialInstance.h"
-#include "MuCOE/UnrealEditorPortabilityHelpers.h"
-#include "PlatformInfo.h"
-
+#include "Math/Color.h"
+#include "Misc/Guid.h"
+#include "MuCO/CustomizableObjectParameterTypeDefinitions.h"
+#include "MuCO/CustomizableObjectUIData.h"
+#include "MuCO/UnrealPortabilityHelpers.h"
 #include "MuCOE/CustomizableObjectCompiler.h"
-#include "MuCOE/CustomizableObjectEditorModule.h"
-#include "MuCO/CustomizableObject.h"
-
-#include "MuT/Compiler.h"
-#include "MuT/Streams.h"
-
+#include "MuCOE/GenerateMutableSource/GenerateMutableSource.h"
+#include "MuCOE/GenerateMutableSource/GenerateMutableSourceMesh.h"
+#include "MuCOE/Nodes/CustomizableObjectNode.h"
+#include "MuCOE/Nodes/CustomizableObjectNodeTable.h"
 #include "MuR/Mesh.h"
-#include "MuR/Skeleton.h"
-#include "MuR/Image.h"
+#include "MuR/Ptr.h"
+#include "Rendering/SkeletalMeshLODModel.h"
+#include "Rendering/SkeletalMeshModel.h"
+#include "StaticMeshResources.h"
+#include "Templates/Casts.h"
+#include "UObject/Class.h"
+#include "UObject/Field.h"
+#include "UObject/NameTypes.h"
+#include "UObject/Object.h"
+#include "UObject/ObjectPtr.h"
+#include "UObject/SoftObjectPtr.h"
+#include "UObject/UnrealType.h"
 
 #define LOCTEXT_NAMESPACE "CustomizableObjectEditor"
 
-
 void FillTableColumn(const UEdGraphPin* Pin, mu::TablePtr MutableTable, FString ColumnName, FString RowName, int32 RowIdx, uint8* CellData, FProperty* Property, FMutableGraphGenerationContext& GenerationContext)
 {
 	UCustomizableObjectNodeTable* CustomNodeTable = Cast<UCustomizableObjectNodeTable>(Pin->GetOwningNode());
