Ran IWYU on Mutable plugin

diff --git a/Engine/Plugins/Experimental/Mutable/Source/CustomizableObjectEditor/Private/MuCOE/GenerateMutableSource/GenerateMutableSourceTable.h b/Engine/Plugins/Experimental/Mutable/Source/CustomizableObjectEditor/Private/MuCOE/GenerateMutableSource/GenerateMutableSourceTable.h
--- a/Engine/Plugins/Experimental/Mutable/Source/CustomizableObjectEditor/Private/MuCOE/GenerateMutableSource/GenerateMutableSourceTable.h
+++ b/Engine/Plugins/Experimental/Mutable/Source/CustomizableObjectEditor/Private/MuCOE/GenerateMutableSource/GenerateMutableSourceTable.h
@@ -2,7 +2,10 @@
 
 #pragma once
 
-#include "MuCOE/GenerateMutableSource/GenerateMutableSource.h"
+#include "Containers/UnrealString.h"
+#include "MuT/Table.h"
 
+class UEdGraphPin;
+struct FMutableGraphGenerationContext;
 
 mu::TablePtr GenerateMutableSourceTable(const FString& TableName, const UEdGraphPin* Pin, FMutableGraphGenerationContext& GenerationContext);
diff --git a/Engine/Plugins/Experimental/Mutable/Source/CustomizableObjectEditor/Private/MuCOE/GraphTraversal.cpp b/Engine/Plugins/Experimental/Mutable/Source/CustomizableObjectEditor/Private/MuCOE/GraphTraversal.cpp
--- a/Engine/Plugins/Experimental/Mutable/Source/CustomizableObjectEditor/Private/MuCOE/GraphTraversal.cpp
+++ b/Engine/Plugins/Experimental/Mutable/Source/CustomizableObjectEditor/Private/MuCOE/GraphTraversal.cpp
@@ -1,11 +1,22 @@
 // Copyright Epic Games, Inc. All Rights Reserved.
 
 #include "MuCOE/GraphTraversal.h"
+
+#include "Containers/EnumAsByte.h"
+#include "Containers/Set.h"
+#include "EdGraph/EdGraph.h"
+#include "EdGraph/EdGraphNode.h"
+#include "EdGraph/EdGraphPin.h"
+#include "HAL/Platform.h"
+#include "Misc/AssertionMacros.h"
 #include "MuCO/CustomizableObject.h"
 #include "MuCOE/CustomizableObjectPin.h"
+#include "MuCOE/EdGraphSchema_CustomizableObject.h"
 #include "MuCOE/Nodes/CustomizableObjectNodeEnumParameter.h"
+#include "MuCOE/Nodes/CustomizableObjectNodeExposePin.h"
 #include "MuCOE/Nodes/CustomizableObjectNodeExtendMaterial.h"
 #include "MuCOE/Nodes/CustomizableObjectNodeExternalPin.h"
+#include "MuCOE/Nodes/CustomizableObjectNodeMaterial.h"
 #include "MuCOE/Nodes/CustomizableObjectNodeMaterialVariation.h"
 #include "MuCOE/Nodes/CustomizableObjectNodeMeshGeometryOperation.h"
 #include "MuCOE/Nodes/CustomizableObjectNodeMeshMorph.h"
@@ -18,6 +29,11 @@
 #include "MuCOE/Nodes/CustomizableObjectNodeSkeletalMesh.h"
 #include "MuCOE/Nodes/CustomizableObjectNodeStaticMesh.h"
 #include "MuCOE/Nodes/CustomizableObjectNodeTable.h"
+#include "Templates/Casts.h"
+#include "Templates/Function.h"
+#include "UObject/NameTypes.h"
+#include "UObject/ObjectPtr.h"
+#include "UObject/UObjectIterator.h"
 
 
 /** Function to follow an External Pin node.
diff --git a/Engine/Plugins/Experimental/Mutable/Source/CustomizableObjectEditor/Private/MuCOE/MutableMeshPreviewUtils.cpp b/Engine/Plugins/Experimental/Mutable/Source/CustomizableObjectEditor/Private/MuCOE/MutableMeshPreviewUtils.cpp
--- a/Engine/Plugins/Experimental/Mutable/Source/CustomizableObjectEditor/Private/MuCOE/MutableMeshPreviewUtils.cpp
+++ b/Engine/Plugins/Experimental/Mutable/Source/CustomizableObjectEditor/Private/MuCOE/MutableMeshPreviewUtils.cpp
@@ -1,18 +1,48 @@
 // Copyright Epic Games, Inc. All Rights Reserved.
 
 #include "MuCOE/MutableMeshPreviewUtils.h"
-#include "MuCO/MutableMeshBufferUtils.h"
-#include "MuCO/UnrealConversionUtils.h"
 
 #include "Animation/Skeleton.h"
-#include "MuCO/CustomizableObjectSystemPrivate.h"
+#include "Containers/Array.h"
+#include "Containers/IndirectArray.h"
+#include "Containers/Map.h"
+#include "Containers/SparseArray.h"
+#include "CoreGlobals.h"
+#include "Engine/EngineTypes.h"
+#include "Engine/SkeletalMesh.h"
+#include "HAL/PlatformCrt.h"
+#include "Logging/LogCategory.h"
+#include "Logging/LogMacros.h"
+#include "MaterialShared.h"
 #include "Materials/Material.h"
-#include "MuR/MeshPrivate.h"
-#include "MuR/OpMeshFormat.h"
+#include "Math/BoxSphereBounds.h"
+#include "Math/Matrix.h"
+#include "Math/UnrealMathSSE.h"
+#include "Misc/AssertionMacros.h"
+#include "MuCO/CustomizableObject.h"
+#include "MuCO/CustomizableObjectInstance.h"
+#include "MuCO/MutableMeshBufferUtils.h"
+#include "MuCO/UnrealConversionUtils.h"
+#include "MuCO/UnrealPortabilityHelpers.h"
 #include "MuR/Mesh.h"
+#include "MuR/MeshBufferSet.h"
 #include "MuR/MutableTrace.h"
+#include "MuR/OpMeshFormat.h"
 #include "MuR/Parameters.h"
 #include "MuR/Ptr.h"
+#include "MuR/RefCounted.h"
+#include "MuR/Skeleton.h"
+#include "PerPlatformProperties.h"
+#include "ReferenceSkeleton.h"
+#include "Rendering/SkeletalMeshLODRenderData.h"
+#include "Trace/Detail/Channel.h"
+#include "UObject/NameTypes.h"
+#include "UObject/ObjectPtr.h"
+#include "UObject/Package.h"
+#include "UObject/UObjectGlobals.h"
+#include "UObject/UnrealNames.h"
+
+class UMaterialInterface;
 
 
 namespace MutableMeshPreviewUtils
diff --git a/Engine/Plugins/Experimental/Mutable/Source/CustomizableObjectEditor/Private/MuCOE/MutableUtils.cpp b/Engine/Plugins/Experimental/Mutable/Source/CustomizableObjectEditor/Private/MuCOE/MutableUtils.cpp
--- a/Engine/Plugins/Experimental/Mutable/Source/CustomizableObjectEditor/Private/MuCOE/MutableUtils.cpp
+++ b/Engine/Plugins/Experimental/Mutable/Source/CustomizableObjectEditor/Private/MuCOE/MutableUtils.cpp
@@ -2,7 +2,22 @@
 
 #include "MuCOE/MutableUtils.h"
 
+#include "Containers/IndirectArray.h"
+#include "Containers/Map.h"
+#include "GenericPlatform/GenericPlatformMath.h"
+#include "HAL/PlatformCrt.h"
+#include "Math/UnrealMathSSE.h"
+#include "Math/Vector.h"
+#include "Misc/AssertionMacros.h"
+#include "RawIndexBuffer.h"
+#include "Rendering/MultiSizeIndexContainer.h"
+#include "Rendering/PositionVertexBuffer.h"
+#include "Rendering/SkeletalMeshLODModel.h"
+#include "Rendering/SkeletalMeshLODRenderData.h"
 #include "Rendering/SkeletalMeshModel.h"
+#include "Rendering/SkeletalMeshRenderData.h"
+#include "Rendering/StaticMeshVertexBuffer.h"
+#include "StaticMeshResources.h"
 
 
 TArray<FVector2f> GetUV(const USkeletalMesh& SkeletalMesh, const int32 LODIndex, const int32 SectionIndex, const int32 UVIndex)
diff --git a/Engine/Plugins/Experimental/Mutable/Source/CustomizableObjectEditor/Private/MuCOE/Nodes/CustomizableObjectNode.cpp b/Engine/Plugins/Experimental/Mutable/Source/CustomizableObjectEditor/Private/MuCOE/Nodes/CustomizableObjectNode.cpp
--- a/Engine/Plugins/Experimental/Mutable/Source/CustomizableObjectEditor/Private/MuCOE/Nodes/CustomizableObjectNode.cpp
+++ b/Engine/Plugins/Experimental/Mutable/Source/CustomizableObjectEditor/Private/MuCOE/Nodes/CustomizableObjectNode.cpp
@@ -2,22 +2,35 @@
 
 #include "MuCOE/Nodes/CustomizableObjectNode.h"
 
-#include "MuCOE/Nodes/CustomizableObjectNodeExposePin.h"
-#include "MuCOE/CustomizableObjectEditorModule.h"
-#include "Toolkits/ToolkitManager.h"
-
-#include "MuCOE/UnrealEditorPortabilityHelpers.h"
-
-#include "MuCOE/CustomizableObjectEditor.h"
+#include "Containers/EnumAsByte.h"
+#include "Containers/Queue.h"
+#include "EdGraph/EdGraph.h"
+#include "HAL/PlatformCrt.h"
+#include "HAL/PlatformMath.h"
+#include "Internationalization/Internationalization.h"
+#include "Logging/TokenizedMessage.h"
+#include "Misc/AssertionMacros.h"
 #include "MuCO/CustomizableObjectCustomVersion.h"
-#include "MuCOE/CustomizableObjectGraph.h"
-#include "MuCOE/EdGraphSchema_CustomizableObject.h"
-
+#include "MuCOE/CustomizableObjectEditor.h"
+#include "MuCOE/CustomizableObjectEditorLogger.h"
 #include "MuCOE/CustomizableObjectGraph.h"
 #include "MuCOE/CustomizableObjectPin.h"
 #include "MuCOE/EdGraphSchema_CustomizableObject.h"
 #include "MuCOE/GraphTraversal.h"
-#include "Runtime/Launch/Resources/Version.h" // Include engine version macros
+#include "MuCOE/ICustomizableObjectEditor.h"
+#include "MuCOE/Nodes/CustomizableObjectNodeObject.h"
+#include "MuCOE/Nodes/CustomizableObjectNodeObjectGroup.h"
+#include "MuCOE/RemapPins/CustomizableObjectNodeRemapPins.h"
+#include "MuCOE/RemapPins/CustomizableObjectNodeRemapPinsByName.h"
+#include "MuCOE/RemapPins/CustomizableObjectNodeRemapPinsByPosition.h"
+#include "Serialization/Archive.h"
+#include "Templates/Casts.h"
+#include "Templates/Tuple.h"
+#include "Toolkits/ToolkitManager.h"
+
+class IToolkit;
+class SWidget;
+struct FPropertyChangedEvent;
 
 #define LOCTEXT_NAMESPACE "CustomizableObjectEditor"
 
diff --git a/Engine/Plugins/Experimental/Mutable/Source/CustomizableObjectEditor/Private/MuCOE/Nodes/CustomizableObjectNodeAnimationPose.cpp b/Engine/Plugins/Experimental/Mutable/Source/CustomizableObjectEditor/Private/MuCOE/Nodes/CustomizableObjectNodeAnimationPose.cpp
--- a/Engine/Plugins/Experimental/Mutable/Source/CustomizableObjectEditor/Private/MuCOE/Nodes/CustomizableObjectNodeAnimationPose.cpp
+++ b/Engine/Plugins/Experimental/Mutable/Source/CustomizableObjectEditor/Private/MuCOE/Nodes/CustomizableObjectNodeAnimationPose.cpp
@@ -1,15 +1,35 @@
 // Copyright Epic Games, Inc. All Rights Reserved.
 
 #include "MuCOE/Nodes/CustomizableObjectNodeAnimationPose.h"
-#include "MuCOE/CustomizableObjectEditorModule.h"
-#include "Animation/AnimInstanceProxy.h"
-#include "Animation/AnimInstance.h"
-#include "Animation/DebugSkelMeshComponent.h"
+
+#include "Animation/AnimCurveTypes.h"
 #include "Animation/AnimInstance.h"
-#include "Animation/AnimSingleNodeInstance.h"
+#include "Animation/AnimTypes.h"
+#include "Animation/AnimationAsset.h"
 #include "Animation/AnimationPoseData.h"
-
-#include "IPersonaPreviewScene.h"
+#include "Animation/AttributesRuntime.h"
+#include "Animation/DebugSkelMeshComponent.h"
+#include "Animation/PoseAsset.h"
+#include "Animation/SmartName.h"
+#include "BoneContainer.h"
+#include "BoneIndices.h"
+#include "BonePose.h"
+#include "Components/SkeletalMeshComponent.h"
+#include "Containers/UnrealString.h"
+#include "EdGraph/EdGraphPin.h"
+#include "Engine/SkeletalMesh.h"
+#include "HAL/Platform.h"
+#include "HAL/PlatformCrt.h"
+#include "Internationalization/Internationalization.h"
+#include "Math/TransformVectorized.h"
+#include "Misc/AssertionMacros.h"
+#include "Misc/MemStack.h"
+#include "MuCOE/EdGraphSchema_CustomizableObject.h"
+#include "ReferenceSkeleton.h"
+#include "UObject/NameTypes.h"
+
+class UCustomizableObjectNodeRemapPins;
+struct FPropertyChangedEvent;
 
 #define LOCTEXT_NAMESPACE "CustomizableObjectEditor"
 
diff --git a/Engine/Plugins/Experimental/Mutable/Source/CustomizableObjectEditor/Private/MuCOE/Nodes/CustomizableObjectNodeColorArithmeticOp.cpp b/Engine/Plugins/Experimental/Mutable/Source/CustomizableObjectEditor/Private/MuCOE/Nodes/CustomizableObjectNodeColorArithmeticOp.cpp
--- a/Engine/Plugins/Experimental/Mutable/Source/CustomizableObjectEditor/Private/MuCOE/Nodes/CustomizableObjectNodeColorArithmeticOp.cpp
+++ b/Engine/Plugins/Experimental/Mutable/Source/CustomizableObjectEditor/Private/MuCOE/Nodes/CustomizableObjectNodeColorArithmeticOp.cpp
@@ -1,7 +1,16 @@
 // Copyright Epic Games, Inc. All Rights Reserved.
 
 #include "MuCOE/Nodes/CustomizableObjectNodeColorArithmeticOp.h"
-#include "MuCOE/CustomizableObjectEditorModule.h"
+
+#include "Containers/UnrealString.h"
+#include "EdGraph/EdGraphPin.h"
+#include "Internationalization/Internationalization.h"
+#include "MuCOE/EdGraphSchema_CustomizableObject.h"
+#include "UObject/Class.h"
+#include "UObject/NameTypes.h"
+
+class UCustomizableObjectNodeRemapPins;
+struct FPropertyChangedEvent;
 
 #define LOCTEXT_NAMESPACE "CustomizableObjectEditor"
 
diff --git a/Engine/Plugins/Experimental/Mutable/Source/CustomizableObjectEditor/Private/MuCOE/Nodes/CustomizableObjectNodeColorConstant.cpp b/Engine/Plugins/Experimental/Mutable/Source/CustomizableObjectEditor/Private/MuCOE/Nodes/CustomizableObjectNodeColorConstant.cpp
--- a/Engine/Plugins/Experimental/Mutable/Source/CustomizableObjectEditor/Private/MuCOE/Nodes/CustomizableObjectNodeColorConstant.cpp
+++ b/Engine/Plugins/Experimental/Mutable/Source/CustomizableObjectEditor/Private/MuCOE/Nodes/CustomizableObjectNodeColorConstant.cpp
@@ -1,13 +1,40 @@
 // Copyright Epic Games, Inc. All Rights Reserved.
 
 #include "MuCOE/Nodes/CustomizableObjectNodeColorConstant.h"
-#include "MuCOE/CustomizableObjectEditorModule.h"
-#include "MuCOE/CustomizableObjectEditorStyle.h"
 
-#include "Widgets/Input/SSpinBox.h"
-#include "Brushes/SlateColorBrush.h"
+#include "EdGraph/EdGraphPin.h"
+#include "GenericPlatform/ICursor.h"
+#include "HAL/Platform.h"
+#include "Input/Events.h"
+#include "InputCoreTypes.h"
+#include "Internationalization/Internationalization.h"
+#include "Layout/Margin.h"
+#include "Layout/Visibility.h"
+#include "Math/Vector2D.h"
+#include "Misc/Attribute.h"
+#include "Misc/Optional.h"
+#include "MuCOE/CustomizableObjectEditorStyle.h"
+#include "MuCOE/CustomizableObjectGraph.h"
+#include "MuCOE/EdGraphSchema_CustomizableObject.h"
+#include "SlotBase.h"
+#include "Styling/AppStyle.h"
+#include "Styling/CoreStyle.h"
+#include "Styling/ISlateStyle.h"
+#include "Templates/Casts.h"
+#include "UObject/NameTypes.h"
+#include "UObject/Object.h"
 #include "Widgets/Colors/SColorBlock.h"
+#include "Widgets/Images/SImage.h"
 #include "Widgets/Input/SCheckBox.h"
+#include "Widgets/Input/SSpinBox.h"
+#include "Widgets/SBoxPanel.h"
+#include "Widgets/SOverlay.h"
+#include "Widgets/Text/STextBlock.h"
+
+class UCustomizableObjectNodeRemapPins;
+struct FGeometry;
+struct FPropertyChangedEvent;
+struct FSlateBrush;
 
 #define LOCTEXT_NAMESPACE "CustomizableObjectEditor"
 
diff --git a/Engine/Plugins/Experimental/Mutable/Source/CustomizableObjectEditor/Private/MuCOE/Nodes/CustomizableObjectNodeColorFromFloats.cpp b/Engine/Plugins/Experimental/Mutable/Source/CustomizableObjectEditor/Private/MuCOE/Nodes/CustomizableObjectNodeColorFromFloats.cpp
--- a/Engine/Plugins/Experimental/Mutable/Source/CustomizableObjectEditor/Private/MuCOE/Nodes/CustomizableObjectNodeColorFromFloats.cpp
+++ b/Engine/Plugins/Experimental/Mutable/Source/CustomizableObjectEditor/Private/MuCOE/Nodes/CustomizableObjectNodeColorFromFloats.cpp
@@ -1,7 +1,14 @@
 // Copyright Epic Games, Inc. All Rights Reserved.
 
 #include "MuCOE/Nodes/CustomizableObjectNodeColorFromFloats.h"
-#include "MuCOE/CustomizableObjectEditorModule.h"
+
+#include "EdGraph/EdGraphPin.h"
+#include "Internationalization/Internationalization.h"
+#include "MuCOE/EdGraphSchema_CustomizableObject.h"
+#include "UObject/NameTypes.h"
+
+class UCustomizableObjectNodeRemapPins;
+struct FPropertyChangedEvent;
 
 #define LOCTEXT_NAMESPACE "CustomizableObjectEditor"
 
diff --git a/Engine/Plugins/Experimental/Mutable/Source/CustomizableObjectEditor/Private/MuCOE/Nodes/CustomizableObjectNodeColorParameter.cpp b/Engine/Plugins/Experimental/Mutable/Source/CustomizableObjectEditor/Private/MuCOE/Nodes/CustomizableObjectNodeColorParameter.cpp
--- a/Engine/Plugins/Experimental/Mutable/Source/CustomizableObjectEditor/Private/MuCOE/Nodes/CustomizableObjectNodeColorParameter.cpp
+++ b/Engine/Plugins/Experimental/Mutable/Source/CustomizableObjectEditor/Private/MuCOE/Nodes/CustomizableObjectNodeColorParameter.cpp
@@ -1,7 +1,13 @@
 // Copyright Epic Games, Inc. All Rights Reserved.
 
 #include "MuCOE/Nodes/CustomizableObjectNodeColorParameter.h"
-#include "MuCOE/CustomizableObjectEditorModule.h"
+
+#include "EdGraph/EdGraphPin.h"
+#include "Internationalization/Internationalization.h"
+#include "MuCOE/EdGraphSchema_CustomizableObject.h"
+#include "UObject/NameTypes.h"
+
+class UCustomizableObjectNodeRemapPins;
 
 #define LOCTEXT_NAMESPACE "CustomizableObjectEditor"
 
diff --git a/Engine/Plugins/Experimental/Mutable/Source/CustomizableObjectEditor/Private/MuCOE/Nodes/CustomizableObjectNodeColorSwitch.cpp b/Engine/Plugins/Experimental/Mutable/Source/CustomizableObjectEditor/Private/MuCOE/Nodes/CustomizableObjectNodeColorSwitch.cpp
--- a/Engine/Plugins/Experimental/Mutable/Source/CustomizableObjectEditor/Private/MuCOE/Nodes/CustomizableObjectNodeColorSwitch.cpp
+++ b/Engine/Plugins/Experimental/Mutable/Source/CustomizableObjectEditor/Private/MuCOE/Nodes/CustomizableObjectNodeColorSwitch.cpp
@@ -2,7 +2,11 @@
 
 #include "MuCOE/Nodes/CustomizableObjectNodeColorSwitch.h"
 
+#include "EdGraph/EdGraphPin.h"
+#include "HAL/Platform.h"
+#include "MuCO/CustomizableObjectCustomVersion.h"
 #include "MuCOE/EdGraphSchema_CustomizableObject.h"
+#include "Serialization/Archive.h"
 
 #define LOCTEXT_NAMESPACE "CustomizableObjectEditor"
 
diff --git a/Engine/Plugins/Experimental/Mutable/Source/CustomizableObjectEditor/Private/MuCOE/Nodes/CustomizableObjectNodeColorVariation.cpp b/Engine/Plugins/Experimental/Mutable/Source/CustomizableObjectEditor/Private/MuCOE/Nodes/CustomizableObjectNodeColorVariation.cpp
--- a/Engine/Plugins/Experimental/Mutable/Source/CustomizableObjectEditor/Private/MuCOE/Nodes/CustomizableObjectNodeColorVariation.cpp
+++ b/Engine/Plugins/Experimental/Mutable/Source/CustomizableObjectEditor/Private/MuCOE/Nodes/CustomizableObjectNodeColorVariation.cpp
@@ -1,9 +1,15 @@
 // Copyright Epic Games, Inc. All Rights Reserved.
 
 #include "MuCOE/Nodes/CustomizableObjectNodeColorVariation.h"
-#include "MuCOE/CustomizableObjectEditorModule.h"
 
+#include "EdGraph/EdGraphPin.h"
+#include "Internationalization/Internationalization.h"
+#include "MuCOE/EdGraphSchema_CustomizableObject.h"
 #include "MuCOE/UnrealEditorPortabilityHelpers.h"
+#include "UObject/NameTypes.h"
+#include "UObject/UnrealType.h"
+
+class UCustomizableObjectNodeRemapPins;
 
 #define LOCTEXT_NAMESPACE "CustomizableObjectEditor"
 
diff --git a/Engine/Plugins/Experimental/Mutable/Source/CustomizableObjectEditor/Private/MuCOE/Nodes/CustomizableObjectNodeCopyMaterial.cpp b/Engine/Plugins/Experimental/Mutable/Source/CustomizableObjectEditor/Private/MuCOE/Nodes/CustomizableObjectNodeCopyMaterial.cpp
--- a/Engine/Plugins/Experimental/Mutable/Source/CustomizableObjectEditor/Private/MuCOE/Nodes/CustomizableObjectNodeCopyMaterial.cpp
+++ b/Engine/Plugins/Experimental/Mutable/Source/CustomizableObjectEditor/Private/MuCOE/Nodes/CustomizableObjectNodeCopyMaterial.cpp
@@ -1,11 +1,21 @@
 // Copyright Epic Games, Inc. All Rights Reserved.
 
 #include "MuCOE/Nodes/CustomizableObjectNodeCopyMaterial.h"
-#include "MuCOE/Nodes/CustomizableObjectNodeExternalPin.h"
-#include "MuCOE/Nodes/CustomizableObjectNodeObject.h"
+
+#include "Containers/EnumAsByte.h"
+#include "Containers/UnrealString.h"
+#include "EdGraph/EdGraphPin.h"
+#include "HAL/Platform.h"
+#include "Internationalization/Internationalization.h"
 #include "MuCOE/EdGraphSchema_CustomizableObject.h"
-#include "MuCOE/GenerateMutableSource/GenerateMutableSource.h"
 #include "MuCOE/GraphTraversal.h"
+#include "MuCOE/Nodes/CustomizableObjectNodeExposePin.h"
+#include "MuCOE/Nodes/CustomizableObjectNodeExternalPin.h"
+#include "MuCOE/Nodes/CustomizableObjectNodeObject.h"
+#include "MuCOE/Nodes/CustomizableObjectNodeSkeletalMesh.h"
+#include "MuCOE/RemapPins/CustomizableObjectNodeRemapPins.h"
+#include "Templates/Casts.h"
+#include "UObject/NameTypes.h"
 
 #define LOCTEXT_NAMESPACE "CustomizableObjectEditor"
 
diff --git a/Engine/Plugins/Experimental/Mutable/Source/CustomizableObjectEditor/Private/MuCOE/Nodes/CustomizableObjectNodeCopyMaterialDetails.cpp b/Engine/Plugins/Experimental/Mutable/Source/CustomizableObjectEditor/Private/MuCOE/Nodes/CustomizableObjectNodeCopyMaterialDetails.cpp
--- a/Engine/Plugins/Experimental/Mutable/Source/CustomizableObjectEditor/Private/MuCOE/Nodes/CustomizableObjectNodeCopyMaterialDetails.cpp
+++ b/Engine/Plugins/Experimental/Mutable/Source/CustomizableObjectEditor/Private/MuCOE/Nodes/CustomizableObjectNodeCopyMaterialDetails.cpp
@@ -2,7 +2,9 @@
 
 #include "MuCOE/Nodes/CustomizableObjectNodeCopyMaterialDetails.h"
 
+#include "Containers/Array.h"
 #include "DetailLayoutBuilder.h"
+#include "UObject/NameTypes.h"
 
 TSharedRef<IDetailCustomization> FCustomizableObjectNodeCopyMaterialDetails::MakeInstance()
 {
diff --git a/Engine/Plugins/Experimental/Mutable/Source/CustomizableObjectEditor/Private/MuCOE/Nodes/CustomizableObjectNodeCopyMaterialDetails.h b/Engine/Plugins/Experimental/Mutable/Source/CustomizableObjectEditor/Private/MuCOE/Nodes/CustomizableObjectNodeCopyMaterialDetails.h
--- a/Engine/Plugins/Experimental/Mutable/Source/CustomizableObjectEditor/Private/MuCOE/Nodes/CustomizableObjectNodeCopyMaterialDetails.h
+++ b/Engine/Plugins/Experimental/Mutable/Source/CustomizableObjectEditor/Private/MuCOE/Nodes/CustomizableObjectNodeCopyMaterialDetails.h
@@ -3,6 +3,9 @@
 #pragma once
 
 #include "IDetailCustomization.h"
+#include "Templates/SharedPointer.h"
+
+class IDetailLayoutBuilder;
 
 /** Copy Material node details panel. Hides all properties from the inheret Material node. */
 class FCustomizableObjectNodeCopyMaterialDetails : public IDetailCustomization
diff --git a/Engine/Plugins/Experimental/Mutable/Source/CustomizableObjectEditor/Private/MuCOE/Nodes/CustomizableObjectNodeCurve.cpp b/Engine/Plugins/Experimental/Mutable/Source/CustomizableObjectEditor/Private/MuCOE/Nodes/CustomizableObjectNodeCurve.cpp
--- a/Engine/Plugins/Experimental/Mutable/Source/CustomizableObjectEditor/Private/MuCOE/Nodes/CustomizableObjectNodeCurve.cpp
+++ b/Engine/Plugins/Experimental/Mutable/Source/CustomizableObjectEditor/Private/MuCOE/Nodes/CustomizableObjectNodeCurve.cpp
@@ -1,9 +1,22 @@
 // Copyright Epic Games, Inc. All Rights Reserved.
 
 #include "MuCOE/Nodes/CustomizableObjectNodeCurve.h"
-#include "MuCOE/CustomizableObjectEditorModule.h"
 
+#include "Containers/Array.h"
+#include "Containers/UnrealString.h"
+#include "Curves/CurveBase.h"
+#include "Curves/RichCurve.h"
+#include "EdGraph/EdGraphPin.h"
+#include "Internationalization/Internationalization.h"
+#include "MuCO/CustomizableObjectCustomVersion.h"
+#include "MuCOE/EdGraphSchema_CustomizableObject.h"
+#include "MuCOE/RemapPins/CustomizableObjectNodeRemapPinsByPosition.h"
 #include "MuCOE/UnrealEditorPortabilityHelpers.h"
+#include "Serialization/Archive.h"
+#include "UObject/NameTypes.h"
+#include "UObject/UnrealType.h"
+
+class UCustomizableObjectNodeRemapPins;
 
 #define LOCTEXT_NAMESPACE "CustomizableObjectEditor"
 
diff --git a/Engine/Plugins/Experimental/Mutable/Source/CustomizableObjectEditor/Private/MuCOE/Nodes/CustomizableObjectNodeEditMaterial.cpp b/Engine/Plugins/Experimental/Mutable/Source/CustomizableObjectEditor/Private/MuCOE/Nodes/CustomizableObjectNodeEditMaterial.cpp
--- a/Engine/Plugins/Experimental/Mutable/Source/CustomizableObjectEditor/Private/MuCOE/Nodes/CustomizableObjectNodeEditMaterial.cpp
+++ b/Engine/Plugins/Experimental/Mutable/Source/CustomizableObjectEditor/Private/MuCOE/Nodes/CustomizableObjectNodeEditMaterial.cpp
@@ -1,11 +1,27 @@
 // Copyright Epic Games, Inc. All Rights Reserved.
 
 #include "MuCOE/Nodes/CustomizableObjectNodeEditMaterial.h"
-#include "MuCOE/CustomizableObjectEditorModule.h"
+
+#include "HAL/PlatformCrt.h"
+#include "Internationalization/Internationalization.h"
+#include "Logging/LogCategory.h"
+#include "Logging/LogMacros.h"
+#include "MaterialTypes.h"
+#include "MuCO/CustomizableObject.h"
+#include "MuCO/CustomizableObjectCustomVersion.h"
+#include "MuCOE/CustomizableObjectEditor_Deprecated.h"
+#include "MuCOE/CustomizableObjectLayout.h"
 #include "MuCOE/EdGraphSchema_CustomizableObject.h"
-#include "MuCOE/Nodes/CustomizableObjectNodeExtendMaterial.h"
 #include "MuCOE/Nodes/CustomizableObjectNodeMaterial.h"
-#include "MuCOE/CustomizableObjectLayout.h"
+#include "MuCOE/Nodes/CustomizableObjectNodeMaterialBase.h"
+#include "Serialization/Archive.h"
+#include "Templates/Casts.h"
+#include "Trace/Detail/Channel.h"
+#include "UObject/NameTypes.h"
+#include "UObject/Package.h"
+
+class FCustomizableObjectNodeParentedMaterial;
+class UCustomizableObjectNodeRemapPins;
 
 #define LOCTEXT_NAMESPACE "CustomizableObjectEditor"
 
diff --git a/Engine/Plugins/Experimental/Mutable/Source/CustomizableObjectEditor/Private/MuCOE/Nodes/CustomizableObjectNodeEditMaterialBase.cpp b/Engine/Plugins/Experimental/Mutable/Source/CustomizableObjectEditor/Private/MuCOE/Nodes/CustomizableObjectNodeEditMaterialBase.cpp
--- a/Engine/Plugins/Experimental/Mutable/Source/CustomizableObjectEditor/Private/MuCOE/Nodes/CustomizableObjectNodeEditMaterialBase.cpp
+++ b/Engine/Plugins/Experimental/Mutable/Source/CustomizableObjectEditor/Private/MuCOE/Nodes/CustomizableObjectNodeEditMaterialBase.cpp
@@ -1,7 +1,12 @@
 // Copyright Epic Games, Inc. All Rights Reserved.
 
 #include "MuCOE/Nodes/CustomizableObjectNodeEditMaterialBase.h"
-#include "MuCOE/CustomizableObjectEditorModule.h"
+
+#include "MuCO/CustomizableObject.h"
+#include "MuCOE/CustomizableObjectGraph.h"
+#include "MuCOE/ICustomizableObjectEditor.h"
+#include "Templates/Casts.h"
+#include "Templates/SharedPointer.h"
 
 
 void UCustomizableObjectNodeEditMaterialBase::BeginPostDuplicate(bool bDuplicateForPIE)
diff --git a/Engine/Plugins/Experimental/Mutable/Source/CustomizableObjectEditor/Private/MuCOE/Nodes/CustomizableObjectNodeEditMaterialBaseDetails.cpp b/Engine/Plugins/Experimental/Mutable/Source/CustomizableObjectEditor/Private/MuCOE/Nodes/CustomizableObjectNodeEditMaterialBaseDetails.cpp
--- a/Engine/Plugins/Experimental/Mutable/Source/CustomizableObjectEditor/Private/MuCOE/Nodes/CustomizableObjectNodeEditMaterialBaseDetails.cpp
+++ b/Engine/Plugins/Experimental/Mutable/Source/CustomizableObjectEditor/Private/MuCOE/Nodes/CustomizableObjectNodeEditMaterialBaseDetails.cpp
@@ -1,24 +1,37 @@
 // Copyright Epic Games, Inc. All Rights Reserved.
 
 #include "MuCOE/Nodes/CustomizableObjectNodeEditMaterialBaseDetails.h"
-#include "MuCOE/CustomizableObjectEditorModule.h"
-#include "PropertyCustomizationHelpers.h"
-#include "MuCOE/CustomizableObjectEditorUtilities.h"
 
-#include "DetailLayoutBuilder.h"
+#include "Containers/UnrealString.h"
 #include "DetailCategoryBuilder.h"
-#include "LevelEditor.h"
-#include "LevelEditorActions.h"
-
-#include "Widgets/Text/STextBlock.h"
-#include "Widgets/Input/STextComboBox.h"
+#include "DetailLayoutBuilder.h"
+#include "DetailWidgetRow.h"
+#include "HAL/Platform.h"
+#include "HAL/PlatformCrt.h"
+#include "IDetailsView.h"
+#include "Internationalization/Internationalization.h"
+#include "Layout/Margin.h"
+#include "Misc/AssertionMacros.h"
+#include "Misc/Attribute.h"
+#include "MuCOE/Nodes/CustomizableObjectNodeEditMaterialBase.h"
+#include "MuCOE/Nodes/CustomizableObjectNodeMaterial.h"
+#include "MuCOE/Nodes/CustomizableObjectNodeMaterialBase.h"
+#include "PropertyCustomizationHelpers.h"
+#include "PropertyHandle.h"
+#include "SlotBase.h"
+#include "Styling/AppStyle.h"
+#include "Templates/Casts.h"
 #include "Templates/SharedPointer.h"
 #include "UObject/WeakObjectPtr.h"
+#include "UObject/WeakObjectPtrTemplates.h"
+#include "Widgets/DeclarativeSyntaxSupport.h"
+#include "Widgets/Input/STextComboBox.h"
+#include "Widgets/Layout/SBorder.h"
+#include "Widgets/SBoxPanel.h"
+#include "Widgets/Text/STextBlock.h"
 
-#include "Engine/SkeletalMesh.h"
-#include "Engine/StaticMesh.h"
-
-#include "Runtime/Launch/Resources/Version.h"
+class IDetailCustomization;
+class UObject;
 
 
 #define LOCTEXT_NAMESPACE "CustomizableObjectDetails"
diff --git a/Engine/Plugins/Experimental/Mutable/Source/CustomizableObjectEditor/Private/MuCOE/Nodes/CustomizableObjectNodeEditMaterialBaseDetails.h b/Engine/Plugins/Experimental/Mutable/Source/CustomizableObjectEditor/Private/MuCOE/Nodes/CustomizableObjectNodeEditMaterialBaseDetails.h
--- a/Engine/Plugins/Experimental/Mutable/Source/CustomizableObjectEditor/Private/MuCOE/Nodes/CustomizableObjectNodeEditMaterialBaseDetails.h
+++ b/Engine/Plugins/Experimental/Mutable/Source/CustomizableObjectEditor/Private/MuCOE/Nodes/CustomizableObjectNodeEditMaterialBaseDetails.h
@@ -2,15 +2,17 @@
 
 #pragma once
 
+#include "Containers/Array.h"
 #include "MuCOE/Nodes/CustomizableObjectNodeParentedMaterialDetails.h"
-
 #include "Templates/SharedPointer.h"
-#include "UObject/WeakObjectPtr.h"
-#include "Misc/Guid.h"
 #include "Types/SlateEnums.h"
-#include "PropertyHandle.h"
-#include "MuCOE/CustomizableObjectEditorModule.h"
-#include "MuCOE/Nodes/CustomizableObjectNodeEditMaterialBase.h"
+
+class FString;
+class IDetailCustomization;
+class IDetailLayoutBuilder;
+class IPropertyHandle;
+class STextComboBox;
+class UCustomizableObjectNodeEditMaterialBase;
 
 class FCustomizableObjectNodeEditMaterialBaseDetails : public FCustomizableObjectNodeParentedMaterialDetails
 {
diff --git a/Engine/Plugins/Experimental/Mutable/Source/CustomizableObjectEditor/Private/MuCOE/Nodes/CustomizableObjectNodeEnumParameter.cpp b/Engine/Plugins/Experimental/Mutable/Source/CustomizableObjectEditor/Private/MuCOE/Nodes/CustomizableObjectNodeEnumParameter.cpp
--- a/Engine/Plugins/Experimental/Mutable/Source/CustomizableObjectEditor/Private/MuCOE/Nodes/CustomizableObjectNodeEnumParameter.cpp
+++ b/Engine/Plugins/Experimental/Mutable/Source/CustomizableObjectEditor/Private/MuCOE/Nodes/CustomizableObjectNodeEnumParameter.cpp
@@ -1,7 +1,14 @@
 // Copyright Epic Games, Inc. All Rights Reserved.
 
 #include "MuCOE/Nodes/CustomizableObjectNodeEnumParameter.h"
-#include "MuCOE/CustomizableObjectEditorModule.h"
+
+#include "EdGraph/EdGraphPin.h"
+#include "Internationalization/Internationalization.h"
+#include "MuCOE/EdGraphSchema_CustomizableObject.h"
+#include "UObject/NameTypes.h"
+#include "UObject/UnrealType.h"
+
+class UCustomizableObjectNodeRemapPins;
 
 #define LOCTEXT_NAMESPACE "CustomizableObjectEditor"
 
diff --git a/Engine/Plugins/Experimental/Mutable/Source/CustomizableObjectEditor/Private/MuCOE/Nodes/CustomizableObjectNodeExposePin.cpp b/Engine/Plugins/Experimental/Mutable/Source/CustomizableObjectEditor/Private/MuCOE/Nodes/CustomizableObjectNodeExposePin.cpp
--- a/Engine/Plugins/Experimental/Mutable/Source/CustomizableObjectEditor/Private/MuCOE/Nodes/CustomizableObjectNodeExposePin.cpp
+++ b/Engine/Plugins/Experimental/Mutable/Source/CustomizableObjectEditor/Private/MuCOE/Nodes/CustomizableObjectNodeExposePin.cpp
@@ -1,10 +1,15 @@
 // Copyright Epic Games, Inc. All Rights Reserved.
 
 #include "MuCOE/Nodes/CustomizableObjectNodeExposePin.h"
-#include "MuCOE/CustomizableObjectEditorModule.h"
-#include "MuCO/CustomizableObject.h"
 
-#include "MuCOE/UnrealEditorPortabilityHelpers.h"
+#include "EdGraph/EdGraphPin.h"
+#include "Internationalization/Internationalization.h"
+#include "MuCOE/EdGraphSchema_CustomizableObject.h"
+#include "MuCOE/Nodes/CustomizableObjectNodeExternalPin.h"
+#include "Templates/Casts.h"
+#include "UObject/UnrealType.h"
+
+class UCustomizableObjectNodeRemapPins;
 
 #define LOCTEXT_NAMESPACE "CustomizableObjectEditor"
 
diff --git a/Engine/Plugins/Experimental/Mutable/Source/CustomizableObjectEditor/Private/MuCOE/Nodes/CustomizableObjectNodeExtendMaterial.cpp b/Engine/Plugins/Experimental/Mutable/Source/CustomizableObjectEditor/Private/MuCOE/Nodes/CustomizableObjectNodeExtendMaterial.cpp
--- a/Engine/Plugins/Experimental/Mutable/Source/CustomizableObjectEditor/Private/MuCOE/Nodes/CustomizableObjectNodeExtendMaterial.cpp
+++ b/Engine/Plugins/Experimental/Mutable/Source/CustomizableObjectEditor/Private/MuCOE/Nodes/CustomizableObjectNodeExtendMaterial.cpp
@@ -1,7 +1,24 @@
 // Copyright Epic Games, Inc. All Rights Reserved.
 
 #include "MuCOE/Nodes/CustomizableObjectNodeExtendMaterial.h"
-#include "MuCOE/CustomizableObjectEditorModule.h"
+
+#include "EdGraph/EdGraphPin.h"
+#include "HAL/Platform.h"
+#include "Internationalization/Internationalization.h"
+#include "MaterialTypes.h"
+#include "MuCO/CustomizableObject.h"
+#include "MuCO/CustomizableObjectCustomVersion.h"
+#include "MuCOE/CustomizableObjectEditor_Deprecated.h"
+#include "MuCOE/EdGraphSchema_CustomizableObject.h"
+#include "MuCOE/Nodes/CustomizableObjectNodeCopyMaterial.h"
+#include "MuCOE/Nodes/CustomizableObjectNodeMaterial.h"
+#include "MuCOE/Nodes/CustomizableObjectNodeParentedNode.h"
+#include "Serialization/Archive.h"
+#include "Templates/Casts.h"
+#include "UObject/NameTypes.h"
+
+class UCustomizableObjectNode;
+class UCustomizableObjectNodeRemapPins;
 
 #define LOCTEXT_NAMESPACE "CustomizableObjectEditor"
 
diff --git a/Engine/Plugins/Experimental/Mutable/Source/CustomizableObjectEditor/Private/MuCOE/Nodes/CustomizableObjectNodeExternalPin.cpp b/Engine/Plugins/Experimental/Mutable/Source/CustomizableObjectEditor/Private/MuCOE/Nodes/CustomizableObjectNodeExternalPin.cpp
--- a/Engine/Plugins/Experimental/Mutable/Source/CustomizableObjectEditor/Private/MuCOE/Nodes/CustomizableObjectNodeExternalPin.cpp
+++ b/Engine/Plugins/Experimental/Mutable/Source/CustomizableObjectEditor/Private/MuCOE/Nodes/CustomizableObjectNodeExternalPin.cpp
@@ -1,9 +1,20 @@
 // Copyright Epic Games, Inc. All Rights Reserved.
 
 #include "MuCOE/Nodes/CustomizableObjectNodeExternalPin.h"
-#include "MuCOE/CustomizableObjectEditorModule.h"
-#include "MuCOE/CustomizableObjectPin.h"
+
+#include "Delegates/Delegate.h"
+#include "EdGraph/EdGraphPin.h"
+#include "HAL/PlatformCrt.h"
+#include "Internationalization/Internationalization.h"
 #include "MuCO/CustomizableObject.h"
+#include "MuCO/CustomizableObjectCustomVersion.h"
+#include "MuCOE/CustomizableObjectGraph.h"
+#include "MuCOE/CustomizableObjectPin.h"
+#include "MuCOE/EdGraphSchema_CustomizableObject.h"
+#include "MuCOE/Nodes/CustomizableObjectNodeExposePin.h"
+#include "MuCOE/RemapPins/CustomizableObjectNodeRemapPins.h"
+#include "Serialization/Archive.h"
+#include "Templates/Casts.h"
 
 #define LOCTEXT_NAMESPACE "CustomizableObjectEditor"
 
diff --git a/Engine/Plugins/Experimental/Mutable/Source/CustomizableObjectEditor/Private/MuCOE/Nodes/CustomizableObjectNodeExternalPinDetails.cpp b/Engine/Plugins/Experimental/Mutable/Source/CustomizableObjectEditor/Private/MuCOE/Nodes/CustomizableObjectNodeExternalPinDetails.cpp
--- a/Engine/Plugins/Experimental/Mutable/Source/CustomizableObjectEditor/Private/MuCOE/Nodes/CustomizableObjectNodeExternalPinDetails.cpp
+++ b/Engine/Plugins/Experimental/Mutable/Source/CustomizableObjectEditor/Private/MuCOE/Nodes/CustomizableObjectNodeExternalPinDetails.cpp
@@ -1,23 +1,36 @@
 // Copyright Epic Games, Inc. All Rights Reserved.
 
 #include "MuCOE/Nodes/CustomizableObjectNodeExternalPinDetails.h"
-#include "MuCOE/CustomizableObjectEditorModule.h"
-#include "PropertyCustomizationHelpers.h"
-#include "MuCOE/CustomizableObjectEditorUtilities.h"
-#include "MuCO/CustomizableObject.h"
-#include "MuCOE/Nodes/CustomizableObjectNodeObjectGroup.h"
 
-#include "ScopedTransaction.h"
-#include "DetailLayoutBuilder.h"
+#include "AssetRegistry/AssetData.h"
+#include "Containers/UnrealString.h"
 #include "DetailCategoryBuilder.h"
-#include "LevelEditor.h"
-#include "LevelEditorActions.h"
-#include "GraphEditor.h"
+#include "DetailLayoutBuilder.h"
+#include "DetailWidgetRow.h"
+#include "EdGraph/EdGraph.h"
+#include "HAL/PlatformCrt.h"
+#include "IDetailsView.h"
+#include "Internationalization/Internationalization.h"
+#include "Layout/Margin.h"
+#include "Misc/Attribute.h"
+#include "Misc/Guid.h"
+#include "MuCO/CustomizableObject.h"
+#include "MuCOE/CustomizableObjectEditorUtilities.h"
+#include "MuCOE/Nodes/CustomizableObjectNodeExposePin.h"
+#include "MuCOE/Nodes/CustomizableObjectNodeExternalPin.h"
+#include "PropertyCustomizationHelpers.h"
+#include "ScopedTransaction.h"
+#include "Styling/AppStyle.h"
+#include "Templates/Casts.h"
+#include "UObject/NameTypes.h"
+#include "UObject/ObjectPtr.h"
+#include "UObject/WeakObjectPtr.h"
+#include "UObject/WeakObjectPtrTemplates.h"
+#include "Widgets/DeclarativeSyntaxSupport.h"
+#include "Widgets/Input/STextComboBox.h"
 #include "Widgets/Layout/SBorder.h"
 #include "Widgets/Text/STextBlock.h"
 
-#include "Runtime/Launch/Resources/Version.h"
-
 
 #define LOCTEXT_NAMESPACE "CustomizableObjectDetails"
 
diff --git a/Engine/Plugins/Experimental/Mutable/Source/CustomizableObjectEditor/Private/MuCOE/Nodes/CustomizableObjectNodeExternalPinDetails.h b/Engine/Plugins/Experimental/Mutable/Source/CustomizableObjectEditor/Private/MuCOE/Nodes/CustomizableObjectNodeExternalPinDetails.h
--- a/Engine/Plugins/Experimental/Mutable/Source/CustomizableObjectEditor/Private/MuCOE/Nodes/CustomizableObjectNodeExternalPinDetails.h
+++ b/Engine/Plugins/Experimental/Mutable/Source/CustomizableObjectEditor/Private/MuCOE/Nodes/CustomizableObjectNodeExternalPinDetails.h
@@ -2,14 +2,15 @@
 
 #pragma once
 
+#include "Containers/Array.h"
 #include "IDetailCustomization.h"
-#include "Types/SlateEnums.h"
 #include "Templates/SharedPointer.h"
-#include "UObject/WeakObjectPtr.h"
-#include "PropertyHandle.h"
-#include "AssetRegistry/AssetData.h"
-#include "Widgets/Input/STextComboBox.h"
-#include "MuCOE/Nodes/CustomizableObjectNodeExternalPin.h"
+#include "Types/SlateEnums.h"
+
+class FString;
+class IDetailLayoutBuilder;
+class UCustomizableObjectNodeExternalPin;
+struct FAssetData;
 
 
 class FCustomizableObjectNodeExternalPinDetails : public IDetailCustomization
diff --git a/Engine/Plugins/Experimental/Mutable/Source/CustomizableObjectEditor/Private/MuCOE/Nodes/CustomizableObjectNodeFloatConstant.cpp b/Engine/Plugins/Experimental/Mutable/Source/CustomizableObjectEditor/Private/MuCOE/Nodes/CustomizableObjectNodeFloatConstant.cpp
--- a/Engine/Plugins/Experimental/Mutable/Source/CustomizableObjectEditor/Private/MuCOE/Nodes/CustomizableObjectNodeFloatConstant.cpp
+++ b/Engine/Plugins/Experimental/Mutable/Source/CustomizableObjectEditor/Private/MuCOE/Nodes/CustomizableObjectNodeFloatConstant.cpp
@@ -1,11 +1,35 @@
 // Copyright Epic Games, Inc. All Rights Reserved.
 
 #include "MuCOE/Nodes/CustomizableObjectNodeFloatConstant.h"
-#include "MuCOE/CustomizableObjectEditorModule.h"
-#include "MuCOE/CustomizableObjectEditorStyle.h"
 
-#include "Widgets/Input/SSpinBox.h"
+#include "EdGraph/EdGraphPin.h"
+#include "GenericPlatform/ICursor.h"
+#include "HAL/Platform.h"
+#include "Internationalization/Internationalization.h"
+#include "Layout/Margin.h"
+#include "Layout/Visibility.h"
+#include "Misc/Attribute.h"
+#include "Misc/Optional.h"
+#include "MuCOE/CustomizableObjectEditorStyle.h"
+#include "MuCOE/CustomizableObjectGraph.h"
+#include "MuCOE/EdGraphSchema_CustomizableObject.h"
+#include "SlotBase.h"
+#include "Styling/AppStyle.h"
+#include "Styling/CoreStyle.h"
+#include "Styling/ISlateStyle.h"
+#include "Templates/Casts.h"
+#include "UObject/NameTypes.h"
+#include "UObject/Object.h"
+#include "Widgets/Images/SImage.h"
 #include "Widgets/Input/SCheckBox.h"
+#include "Widgets/Input/SSpinBox.h"
+#include "Widgets/SBoxPanel.h"
+#include "Widgets/SOverlay.h"
+#include "Widgets/Text/STextBlock.h"
+
+class UCustomizableObjectNodeRemapPins;
+struct FPropertyChangedEvent;
+struct FSlateBrush;
 
 #define LOCTEXT_NAMESPACE "CustomizableObjectEditor"
 
diff --git a/Engine/Plugins/Experimental/Mutable/Source/CustomizableObjectEditor/Private/MuCOE/Nodes/CustomizableObjectNodeFloatParameter.cpp b/Engine/Plugins/Experimental/Mutable/Source/CustomizableObjectEditor/Private/MuCOE/Nodes/CustomizableObjectNodeFloatParameter.cpp
--- a/Engine/Plugins/Experimental/Mutable/Source/CustomizableObjectEditor/Private/MuCOE/Nodes/CustomizableObjectNodeFloatParameter.cpp
+++ b/Engine/Plugins/Experimental/Mutable/Source/CustomizableObjectEditor/Private/MuCOE/Nodes/CustomizableObjectNodeFloatParameter.cpp
@@ -1,8 +1,18 @@
 // Copyright Epic Games, Inc. All Rights Reserved.
 
 #include "MuCOE/Nodes/CustomizableObjectNodeFloatParameter.h"
-#include "MuCOE/CustomizableObjectEditorModule.h"
+
+#include "EdGraph/EdGraphPin.h"
+#include "HAL/PlatformCrt.h"
+#include "Internationalization/Internationalization.h"
+#include "MuCO/CustomizableObjectCustomVersion.h"
+#include "MuCOE/EdGraphSchema_CustomizableObject.h"
 #include "MuCOE/UnrealEditorPortabilityHelpers.h"
+#include "Serialization/Archive.h"
+#include "UObject/NameTypes.h"
+#include "UObject/UnrealType.h"
+
+class UCustomizableObjectNodeRemapPins;
 
 #define LOCTEXT_NAMESPACE "CustomizableObjectEditor"
 
diff --git a/Engine/Plugins/Experimental/Mutable/Source/CustomizableObjectEditor/Private/MuCOE/Nodes/CustomizableObjectNodeFloatSwitch.cpp b/Engine/Plugins/Experimental/Mutable/Source/CustomizableObjectEditor/Private/MuCOE/Nodes/CustomizableObjectNodeFloatSwitch.cpp
--- a/Engine/Plugins/Experimental/Mutable/Source/CustomizableObjectEditor/Private/MuCOE/Nodes/CustomizableObjectNodeFloatSwitch.cpp
+++ b/Engine/Plugins/Experimental/Mutable/Source/CustomizableObjectEditor/Private/MuCOE/Nodes/CustomizableObjectNodeFloatSwitch.cpp
@@ -2,7 +2,11 @@
 
 #include "MuCOE/Nodes/CustomizableObjectNodeFloatSwitch.h"
 
+#include "EdGraph/EdGraphPin.h"
+#include "HAL/Platform.h"
+#include "MuCO/CustomizableObjectCustomVersion.h"
 #include "MuCOE/EdGraphSchema_CustomizableObject.h"
+#include "Serialization/Archive.h"
 
 #define LOCTEXT_NAMESPACE "CustomizableObjectEditor"
 
diff --git a/Engine/Plugins/Experimental/Mutable/Source/CustomizableObjectEditor/Private/MuCOE/Nodes/CustomizableObjectNodeFloatVariation.cpp b/Engine/Plugins/Experimental/Mutable/Source/CustomizableObjectEditor/Private/MuCOE/Nodes/CustomizableObjectNodeFloatVariation.cpp
--- a/Engine/Plugins/Experimental/Mutable/Source/CustomizableObjectEditor/Private/MuCOE/Nodes/CustomizableObjectNodeFloatVariation.cpp
+++ b/Engine/Plugins/Experimental/Mutable/Source/CustomizableObjectEditor/Private/MuCOE/Nodes/CustomizableObjectNodeFloatVariation.cpp
@@ -1,9 +1,15 @@
 // Copyright Epic Games, Inc. All Rights Reserved.
 
 #include "MuCOE/Nodes/CustomizableObjectNodeFloatVariation.h"
-#include "MuCOE/CustomizableObjectEditorModule.h"
 
+#include "EdGraph/EdGraphPin.h"
+#include "Internationalization/Internationalization.h"
+#include "MuCOE/EdGraphSchema_CustomizableObject.h"
 #include "MuCOE/UnrealEditorPortabilityHelpers.h"
+#include "UObject/NameTypes.h"
+#include "UObject/UnrealType.h"
+
+class UCustomizableObjectNodeRemapPins;
 
 #define LOCTEXT_NAMESPACE "CustomizableObjectEditor"
 
diff --git a/Engine/Plugins/Experimental/Mutable/Source/CustomizableObjectEditor/Private/MuCOE/Nodes/CustomizableObjectNodeGroupProjectorParameter.cpp b/Engine/Plugins/Experimental/Mutable/Source/CustomizableObjectEditor/Private/MuCOE/Nodes/CustomizableObjectNodeGroupProjectorParameter.cpp
--- a/Engine/Plugins/Experimental/Mutable/Source/CustomizableObjectEditor/Private/MuCOE/Nodes/CustomizableObjectNodeGroupProjectorParameter.cpp
+++ b/Engine/Plugins/Experimental/Mutable/Source/CustomizableObjectEditor/Private/MuCOE/Nodes/CustomizableObjectNodeGroupProjectorParameter.cpp
@@ -1,8 +1,25 @@
 // Copyright Epic Games, Inc. All Rights Reserved.
 
 #include "MuCOE/Nodes/CustomizableObjectNodeGroupProjectorParameter.h"
-#include "MuCOE/UnrealEditorPortabilityHelpers.h"
+
+#include "Containers/Map.h"
+#include "EdGraph/EdGraphPin.h"
+#include "Engine/DataTable.h"
+#include "Engine/Texture2D.h"
+#include "HAL/PlatformCrt.h"
+#include "Internationalization/Internationalization.h"
+#include "Logging/LogCategory.h"
+#include "Logging/LogMacros.h"
+#include "MuCO/CustomizableObject.h"
+#include "MuCO/CustomizableObjectCustomVersion.h"
 #include "MuCOE/EdGraphSchema_CustomizableObject.h"
+#include "MuCOE/UnrealEditorPortabilityHelpers.h"
+#include "Serialization/Archive.h"
+#include "Trace/Detail/Channel.h"
+#include "UObject/PropertyPortFlags.h"
+#include "UObject/UnrealType.h"
+
+class UCustomizableObjectNodeRemapPins;
 
 #define LOCTEXT_NAMESPACE "CustomizableObjectEditor"
 
diff --git a/Engine/Plugins/Experimental/Mutable/Source/CustomizableObjectEditor/Private/MuCOE/Nodes/CustomizableObjectNodeLayoutBlocks.cpp b/Engine/Plugins/Experimental/Mutable/Source/CustomizableObjectEditor/Private/MuCOE/Nodes/CustomizableObjectNodeLayoutBlocks.cpp
--- a/Engine/Plugins/Experimental/Mutable/Source/CustomizableObjectEditor/Private/MuCOE/Nodes/CustomizableObjectNodeLayoutBlocks.cpp
+++ b/Engine/Plugins/Experimental/Mutable/Source/CustomizableObjectEditor/Private/MuCOE/Nodes/CustomizableObjectNodeLayoutBlocks.cpp
@@ -1,8 +1,28 @@
 // Copyright Epic Games, Inc. All Rights Reserved.
 
 #include "MuCOE/Nodes/CustomizableObjectNodeLayoutBlocks.h"
-#include "MuCOE/CustomizableObjectEditorModule.h"
+
+#include "Containers/ArrayView.h"
+#include "Delegates/Delegate.h"
+#include "EdGraph/EdGraphPin.h"
+#include "HAL/PlatformCrt.h"
+#include "Internationalization/Internationalization.h"
+#include "Math/Vector2D.h"
+#include "Misc/Guid.h"
+#include "MuCO/CustomizableObjectCustomVersion.h"
+#include "MuCOE/EdGraphSchema_CustomizableObject.h"
 #include "MuCOE/GraphTraversal.h"
+#include "MuCOE/Nodes/CustomizableObjectNodeMesh.h"
+#include "MuCOE/Nodes/CustomizableObjectNodeSkeletalMesh.h"
+#include "MuCOE/Nodes/CustomizableObjectNodeStaticMesh.h"
+#include "Serialization/Archive.h"
+#include "Templates/Casts.h"
+#include "UObject/NameTypes.h"
+#include "UObject/WeakObjectPtr.h"
+
+class UCustomizableObjectNodeRemapPins;
+class UObject;
+struct FPropertyChangedEvent;
 
 
 #define LOCTEXT_NAMESPACE "CustomizableObjectEditor"
diff --git a/Engine/Plugins/Experimental/Mutable/Source/CustomizableObjectEditor/Private/MuCOE/Nodes/CustomizableObjectNodeLayoutBlocksDetails.cpp b/Engine/Plugins/Experimental/Mutable/Source/CustomizableObjectEditor/Private/MuCOE/Nodes/CustomizableObjectNodeLayoutBlocksDetails.cpp
--- a/Engine/Plugins/Experimental/Mutable/Source/CustomizableObjectEditor/Private/MuCOE/Nodes/CustomizableObjectNodeLayoutBlocksDetails.cpp
+++ b/Engine/Plugins/Experimental/Mutable/Source/CustomizableObjectEditor/Private/MuCOE/Nodes/CustomizableObjectNodeLayoutBlocksDetails.cpp
@@ -1,17 +1,24 @@
 // Copyright Epic Games, Inc. All Rights Reserved.
 
 #include "MuCOE/Nodes/CustomizableObjectNodeLayoutBlocksDetails.h"
-#include "MuCOE/CustomizableObjectEditorModule.h"
-#include "PropertyCustomizationHelpers.h"
 
-#include "DetailLayoutBuilder.h"
 #include "DetailCategoryBuilder.h"
-#include "LevelEditor.h"
-#include "LevelEditorActions.h"
-
+#include "DetailLayoutBuilder.h"
+#include "DetailWidgetRow.h"
+#include "IDetailsView.h"
+#include "Internationalization/Internationalization.h"
+#include "Math/IntPoint.h"
+#include "Math/IntRect.h"
+#include "MuCOE/CustomizableObjectLayout.h"
+#include "MuCOE/Nodes/CustomizableObjectNodeLayoutBlocks.h"
+#include "Templates/Casts.h"
+#include "UObject/ObjectPtr.h"
+#include "UObject/WeakObjectPtr.h"
+#include "UObject/WeakObjectPtrTemplates.h"
+#include "Widgets/DeclarativeSyntaxSupport.h"
 #include "Widgets/Text/STextBlock.h"
 
-#include "Runtime/Launch/Resources/Version.h"
+class FString;
 
 
 #define LOCTEXT_NAMESPACE "CustomizableObjectDetails"
diff --git a/Engine/Plugins/Experimental/Mutable/Source/CustomizableObjectEditor/Private/MuCOE/Nodes/CustomizableObjectNodeLayoutBlocksDetails.h b/Engine/Plugins/Experimental/Mutable/Source/CustomizableObjectEditor/Private/MuCOE/Nodes/CustomizableObjectNodeLayoutBlocksDetails.h
--- a/Engine/Plugins/Experimental/Mutable/Source/CustomizableObjectEditor/Private/MuCOE/Nodes/CustomizableObjectNodeLayoutBlocksDetails.h
+++ b/Engine/Plugins/Experimental/Mutable/Source/CustomizableObjectEditor/Private/MuCOE/Nodes/CustomizableObjectNodeLayoutBlocksDetails.h
@@ -3,10 +3,14 @@
 #pragma once
 
 
+#include "Containers/Array.h"
 #include "IDetailCustomization.h"
+#include "Math/MathFwd.h"
 #include "Templates/SharedPointer.h"
 #include "Types/SlateEnums.h"
-#include "UObject/WeakObjectPtr.h"
+
+class FString;
+class IDetailLayoutBuilder;
 
 
 class FCustomizableObjectNodeLayoutBlocksDetails : public IDetailCustomization
diff --git a/Engine/Plugins/Experimental/Mutable/Source/CustomizableObjectEditor/Private/MuCOE/Nodes/CustomizableObjectNodeMaterial.cpp b/Engine/Plugins/Experimental/Mutable/Source/CustomizableObjectEditor/Private/MuCOE/Nodes/CustomizableObjectNodeMaterial.cpp
--- a/Engine/Plugins/Experimental/Mutable/Source/CustomizableObjectEditor/Private/MuCOE/Nodes/CustomizableObjectNodeMaterial.cpp
+++ b/Engine/Plugins/Experimental/Mutable/Source/CustomizableObjectEditor/Private/MuCOE/Nodes/CustomizableObjectNodeMaterial.cpp
@@ -1,22 +1,48 @@
 // Copyright Epic Games, Inc. All Rights Reserved.
 
 #include "MuCOE/Nodes/CustomizableObjectNodeMaterial.h"
+
+#include "Containers/EnumAsByte.h"
+#include "Containers/Set.h"
+#include "CoreTypes.h"
+#include "EdGraph/EdGraph.h"
+#include "EdGraph/EdGraphPin.h"
+#include "Engine/Texture.h"
+#include "Engine/Texture2D.h"
+#include "HAL/PlatformCrt.h"
+#include "Internationalization/Internationalization.h"
+#include "MaterialCachedData.h"
+#include "MaterialExpressionIO.h"
 #include "Materials/Material.h"
+#include "Materials/MaterialExpressionTextureCoordinate.h"
+#include "Materials/MaterialExpressionTextureSample.h"
 #include "Materials/MaterialFunctionInterface.h"
-#include "MuCOE/CustomizableObjectCompiler.h"
+#include "Materials/MaterialInterface.h"
+#include "Materials/MaterialLayersFunctions.h"
+#include "Misc/AssertionMacros.h"
+#include "MuCO/CustomizableObjectCustomVersion.h"
+#include "MuCOE/CustomizableObjectEditor_Deprecated.h"
 #include "MuCOE/EdGraphSchema_CustomizableObject.h"
 #include "MuCOE/GraphTraversal.h"
-#include "ObjectEditorUtils.h"
-#include "MuCOE/Nodes/SPinViewerNodeMaterialPinImageDetails.h"
-#include "Materials/MaterialExpressionTextureCoordinate.h"
+#include "MuCOE/Nodes/CustomizableObjectNodeCopyMaterial.h"
 #include "MuCOE/Nodes/CustomizableObjectNodeSkeletalMesh.h"
 #include "MuCOE/Nodes/CustomizableObjectNodeStaticMesh.h"
-#include "MuCOE/Nodes/SCustomizableObjectNodeMaterial.h"
 #include "MuCOE/Nodes/CustomizableObjectNodeTable.h"
-#include "MuCOE/Nodes/CustomizableObjectNodeCopyMaterial.h"
-#include "Slate/Private/Widgets/Views/SListPanel.h"
-#include "Materials/MaterialInstance.h"
-#include "MaterialCachedData.h"
+#include "MuCOE/Nodes/SCustomizableObjectNodeMaterial.h"
+#include "MuCOE/Nodes/SPinViewerNodeMaterialPinImageDetails.h"
+#include "MuCOE/RemapPins/CustomizableObjectNodeRemapPins.h"
+#include "ObjectEditorUtils.h"
+#include "Serialization/Archive.h"
+#include "Templates/Casts.h"
+#include "Templates/ChooseClass.h"
+#include "UObject/UnrealType.h"
+#include "UObject/WeakObjectPtr.h"
+#include "Widgets/DeclarativeSyntaxSupport.h"
+
+class SGraphNode;
+class SWidget;
+class UCustomizableObjectLayout;
+class UObject;
 
 
 #define LOCTEXT_NAMESPACE "CustomizableObjectEditor"
diff --git a/Engine/Plugins/Experimental/Mutable/Source/CustomizableObjectEditor/Private/MuCOE/Nodes/CustomizableObjectNodeMaterialBase.cpp b/Engine/Plugins/Experimental/Mutable/Source/CustomizableObjectEditor/Private/MuCOE/Nodes/CustomizableObjectNodeMaterialBase.cpp
--- a/Engine/Plugins/Experimental/Mutable/Source/CustomizableObjectEditor/Private/MuCOE/Nodes/CustomizableObjectNodeMaterialBase.cpp
+++ b/Engine/Plugins/Experimental/Mutable/Source/CustomizableObjectEditor/Private/MuCOE/Nodes/CustomizableObjectNodeMaterialBase.cpp
@@ -1,7 +1,8 @@
 // Copyright Epic Games, Inc. All Rights Reserved.
 
 #include "MuCOE/Nodes/CustomizableObjectNodeMaterialBase.h"
-#include "MuCOE/CustomizableObjectEditorModule.h"
+
+class UEdGraphPin;
 
 
 bool UCustomizableObjectNodeMaterialBase::ShouldBreakExistingConnections(const UEdGraphPin* InputPin, const UEdGraphPin* OutputPin) const
diff --git a/Engine/Plugins/Experimental/Mutable/Source/CustomizableObjectEditor/Private/MuCOE/Nodes/CustomizableObjectNodeMaterialVariation.cpp b/Engine/Plugins/Experimental/Mutable/Source/CustomizableObjectEditor/Private/MuCOE/Nodes/CustomizableObjectNodeMaterialVariation.cpp
--- a/Engine/Plugins/Experimental/Mutable/Source/CustomizableObjectEditor/Private/MuCOE/Nodes/CustomizableObjectNodeMaterialVariation.cpp
+++ b/Engine/Plugins/Experimental/Mutable/Source/CustomizableObjectEditor/Private/MuCOE/Nodes/CustomizableObjectNodeMaterialVariation.cpp
@@ -1,9 +1,14 @@
 // Copyright Epic Games, Inc. All Rights Reserved.
 
 #include "MuCOE/Nodes/CustomizableObjectNodeMaterialVariation.h"
-#include "MuCOE/CustomizableObjectEditorModule.h"
 
-#include "MuCOE/UnrealEditorPortabilityHelpers.h"
+#include "EdGraph/EdGraphPin.h"
+#include "Internationalization/Internationalization.h"
+#include "MuCOE/EdGraphSchema_CustomizableObject.h"
+#include "UObject/NameTypes.h"
+#include "UObject/UnrealType.h"
+
+class UCustomizableObjectNodeRemapPins;
 
 #define LOCTEXT_NAMESPACE "CustomizableObjectEditor"
 
diff --git a/Engine/Plugins/Experimental/Mutable/Source/CustomizableObjectEditor/Private/MuCOE/Nodes/CustomizableObjectNodeMeshClipDeform.cpp b/Engine/Plugins/Experimental/Mutable/Source/CustomizableObjectEditor/Private/MuCOE/Nodes/CustomizableObjectNodeMeshClipDeform.cpp
--- a/Engine/Plugins/Experimental/Mutable/Source/CustomizableObjectEditor/Private/MuCOE/Nodes/CustomizableObjectNodeMeshClipDeform.cpp
+++ b/Engine/Plugins/Experimental/Mutable/Source/CustomizableObjectEditor/Private/MuCOE/Nodes/CustomizableObjectNodeMeshClipDeform.cpp
@@ -1,9 +1,13 @@
 // Copyright Epic Games, Inc. All Rights Reserved.
 
 #include "MuCOE/Nodes/CustomizableObjectNodeMeshClipDeform.h"
-#include "MuCOE/CustomizableObjectEditorModule.h"
-#include "Toolkits/ToolkitManager.h"
-#include "MuCOE/UnrealEditorPortabilityHelpers.h"
+
+#include "EdGraph/EdGraphPin.h"
+#include "Internationalization/Internationalization.h"
+#include "MuCOE/EdGraphSchema_CustomizableObject.h"
+#include "UObject/NameTypes.h"
+
+class UCustomizableObjectNodeRemapPins;
 
 #define LOCTEXT_NAMESPACE "CustomizableObjectEditor"
 
diff --git a/Engine/Plugins/Experimental/Mutable/Source/CustomizableObjectEditor/Private/MuCOE/Nodes/CustomizableObjectNodeMeshClipMorph.cpp b/Engine/Plugins/Experimental/Mutable/Source/CustomizableObjectEditor/Private/MuCOE/Nodes/CustomizableObjectNodeMeshClipMorph.cpp
--- a/Engine/Plugins/Experimental/Mutable/Source/CustomizableObjectEditor/Private/MuCOE/Nodes/CustomizableObjectNodeMeshClipMorph.cpp
+++ b/Engine/Plugins/Experimental/Mutable/Source/CustomizableObjectEditor/Private/MuCOE/Nodes/CustomizableObjectNodeMeshClipMorph.cpp
@@ -1,9 +1,18 @@
 // Copyright Epic Games, Inc. All Rights Reserved.
 
 #include "MuCOE/Nodes/CustomizableObjectNodeMeshClipMorph.h"
-#include "MuCOE/CustomizableObjectEditorModule.h"
-#include "Toolkits/ToolkitManager.h"
+
+#include "Internationalization/Internationalization.h"
+#include "MuCO/CustomizableObjectCustomVersion.h"
+#include "MuCOE/EdGraphSchema_CustomizableObject.h"
+#include "MuCOE/ICustomizableObjectEditor.h"
 #include "MuCOE/UnrealEditorPortabilityHelpers.h"
+#include "Serialization/Archive.h"
+#include "Templates/SharedPointer.h"
+#include "UObject/UnrealType.h"
+
+class UCustomizableObjectNodeRemapPins;
+class UEdGraphPin;
 
 #define LOCTEXT_NAMESPACE "CustomizableObjectEditor"
 
diff --git a/Engine/Plugins/Experimental/Mutable/Source/CustomizableObjectEditor/Private/MuCOE/Nodes/CustomizableObjectNodeMeshClipMorphDetails.cpp b/Engine/Plugins/Experimental/Mutable/Source/CustomizableObjectEditor/Private/MuCOE/Nodes/CustomizableObjectNodeMeshClipMorphDetails.cpp
--- a/Engine/Plugins/Experimental/Mutable/Source/CustomizableObjectEditor/Private/MuCOE/Nodes/CustomizableObjectNodeMeshClipMorphDetails.cpp
+++ b/Engine/Plugins/Experimental/Mutable/Source/CustomizableObjectEditor/Private/MuCOE/Nodes/CustomizableObjectNodeMeshClipMorphDetails.cpp
@@ -1,27 +1,45 @@
 // Copyright Epic Games, Inc. All Rights Reserved.
 
 #include "MuCOE/Nodes/CustomizableObjectNodeMeshClipMorphDetails.h"
-#include "MuCOE/CustomizableObjectEditorModule.h"
-#include "PropertyCustomizationHelpers.h"
-#include "MuCOE/CustomizableObjectEditorUtilities.h"
-#include "MuCO/CustomizableObject.h"
-#include "MuCOE/Nodes/CustomizableObjectNodeObjectGroup.h"
-#include "MuCOE/CustomizableObjectCompiler.h"
-#include "MuCOE/GraphTraversal.h"
 
-#include "DetailLayoutBuilder.h"
+#include "Containers/UnrealString.h"
+#include "Delegates/Delegate.h"
 #include "DetailCategoryBuilder.h"
-#include "LevelEditor.h"
-#include "LevelEditorActions.h"
-
+#include "DetailLayoutBuilder.h"
+#include "DetailWidgetRow.h"
 #include "Engine/SkeletalMesh.h"
-
+#include "Fonts/SlateFontInfo.h"
+#include "HAL/PlatformCrt.h"
+#include "IDetailsView.h"
+#include "Internationalization/Internationalization.h"
+#include "Layout/Margin.h"
+#include "Math/Transform.h"
+#include "Math/TransformVectorized.h"
+#include "Math/UnrealMathSSE.h"
+#include "Math/Vector.h"
+#include "Misc/Attribute.h"
+#include "MuCO/CustomizableObject.h"
+#include "MuCOE/CustomizableObjectEditorUtilities.h"
+#include "MuCOE/CustomizableObjectGraph.h"
+#include "MuCOE/GraphTraversal.h"
+#include "MuCOE/Nodes/CustomizableObjectNodeMeshClipMorph.h"
+#include "PropertyCustomizationHelpers.h"
+#include "PropertyHandle.h"
+#include "ReferenceSkeleton.h"
+#include "SlotBase.h"
+#include "Styling/AppStyle.h"
+#include "Templates/Casts.h"
+#include "UObject/NameTypes.h"
+#include "UObject/WeakObjectPtr.h"
+#include "UObject/WeakObjectPtrTemplates.h"
+#include "Widgets/DeclarativeSyntaxSupport.h"
+#include "Widgets/Input/SCheckBox.h"
+#include "Widgets/Input/STextComboBox.h"
 #include "Widgets/Layout/SBorder.h"
+#include "Widgets/SBoxPanel.h"
 #include "Widgets/Text/STextBlock.h"
-#include "Widgets/Input/STextComboBox.h"
-#include "Widgets/Input/SCheckBox.h"
 
-#include "Runtime/Launch/Resources/Version.h"
+class UCustomizableObjectNodeObject;
 
 
 #define LOCTEXT_NAMESPACE "MeshClipMorphDetails"
diff --git a/Engine/Plugins/Experimental/Mutable/Source/CustomizableObjectEditor/Private/MuCOE/Nodes/CustomizableObjectNodeMeshClipMorphDetails.h b/Engine/Plugins/Experimental/Mutable/Source/CustomizableObjectEditor/Private/MuCOE/Nodes/CustomizableObjectNodeMeshClipMorphDetails.h
--- a/Engine/Plugins/Experimental/Mutable/Source/CustomizableObjectEditor/Private/MuCOE/Nodes/CustomizableObjectNodeMeshClipMorphDetails.h
+++ b/Engine/Plugins/Experimental/Mutable/Source/CustomizableObjectEditor/Private/MuCOE/Nodes/CustomizableObjectNodeMeshClipMorphDetails.h
@@ -2,12 +2,15 @@
 
 #pragma once
 
+#include "Containers/Array.h"
 #include "IDetailCustomization.h"
-#include "Types/SlateEnums.h"
+#include "Styling/SlateTypes.h"
 #include "Templates/SharedPointer.h"
-#include "UObject/WeakObjectPtr.h"
-#include "Misc/Guid.h"
-#include "PropertyHandle.h"
+#include "Types/SlateEnums.h"
+
+class FString;
+class IDetailLayoutBuilder;
+class IPropertyHandle;
 
 
 class FCustomizableObjectNodeMeshClipMorphDetails : public IDetailCustomization
diff --git a/Engine/Plugins/Experimental/Mutable/Source/CustomizableObjectEditor/Private/MuCOE/Nodes/CustomizableObjectNodeMeshClipWithMesh.cpp b/Engine/Plugins/Experimental/Mutable/Source/CustomizableObjectEditor/Private/MuCOE/Nodes/CustomizableObjectNodeMeshClipWithMesh.cpp
--- a/Engine/Plugins/Experimental/Mutable/Source/CustomizableObjectEditor/Private/MuCOE/Nodes/CustomizableObjectNodeMeshClipWithMesh.cpp
+++ b/Engine/Plugins/Experimental/Mutable/Source/CustomizableObjectEditor/Private/MuCOE/Nodes/CustomizableObjectNodeMeshClipWithMesh.cpp
@@ -1,9 +1,19 @@
 // Copyright Epic Games, Inc. All Rights Reserved.
 
 #include "MuCOE/Nodes/CustomizableObjectNodeMeshClipWithMesh.h"
-#include "MuCOE/CustomizableObjectEditorModule.h"
-#include "MuCOE/Nodes/CustomizableObjectNodeRemoveMesh.h"
-#include "Toolkits/ToolkitManager.h"
+
+#include "EdGraph/EdGraphPin.h"
+#include "HAL/PlatformCrt.h"
+#include "Internationalization/Internationalization.h"
+#include "Misc/Guid.h"
+#include "MuCOE/CustomizableObjectGraph.h"
+#include "MuCOE/EdGraphSchema_CustomizableObject.h"
+#include "MuCOE/ICustomizableObjectEditor.h"
+#include "Templates/Casts.h"
+#include "Templates/SharedPointer.h"
+#include "UObject/NameTypes.h"
+
+class UCustomizableObjectNodeRemapPins;
 
 #define LOCTEXT_NAMESPACE "CustomizableObjectEditor"
 
diff --git a/Engine/Plugins/Experimental/Mutable/Source/CustomizableObjectEditor/Private/MuCOE/Nodes/CustomizableObjectNodeMeshClipWithMeshDetails.cpp b/Engine/Plugins/Experimental/Mutable/Source/CustomizableObjectEditor/Private/MuCOE/Nodes/CustomizableObjectNodeMeshClipWithMeshDetails.cpp
--- a/Engine/Plugins/Experimental/Mutable/Source/CustomizableObjectEditor/Private/MuCOE/Nodes/CustomizableObjectNodeMeshClipWithMeshDetails.cpp
+++ b/Engine/Plugins/Experimental/Mutable/Source/CustomizableObjectEditor/Private/MuCOE/Nodes/CustomizableObjectNodeMeshClipWithMeshDetails.cpp
@@ -1,27 +1,46 @@
 // Copyright Epic Games, Inc. All Rights Reserved.
 
 #include "MuCOE/Nodes/CustomizableObjectNodeMeshClipWithMeshDetails.h"
-#include "MuCOE/CustomizableObjectEditorModule.h"
-#include "PropertyCustomizationHelpers.h"
-#include "MuCOE/CustomizableObjectEditorUtilities.h"
-#include "PropertyHandle.h"
-#include "MuCO/CustomizableObject.h"
 
-#include "DetailLayoutBuilder.h"
+#include "AssetRegistry/AssetData.h"
+#include "CoreTypes.h"
 #include "DetailCategoryBuilder.h"
+#include "DetailLayoutBuilder.h"
+#include "DetailWidgetRow.h"
+#include "EdGraph/EdGraph.h"
+#include "EdGraph/EdGraphPin.h"
+#include "HAL/PlatformCrt.h"
+#include "IDetailsView.h"
+#include "Internationalization/Internationalization.h"
+#include "Internationalization/Text.h"
+#include "Layout/Children.h"
+#include "Layout/Margin.h"
+#include "Layout/Visibility.h"
+#include "Materials/MaterialInterface.h"
+#include "Misc/Attribute.h"
+#include "Misc/Guid.h"
+#include "MuCO/CustomizableObject.h"
+#include "MuCOE/CustomizableObjectEditorUtilities.h"
 #include "MuCOE/GraphTraversal.h"
-#include "LevelEditor.h"
-#include "LevelEditorActions.h"
+#include "MuCOE/Nodes/CustomizableObjectNodeMaterial.h"
+#include "MuCOE/Nodes/CustomizableObjectNodeMeshClipWithMesh.h"
+#include "MuCOE/Nodes/CustomizableObjectNodeStaticMesh.h"
+#include "PropertyCustomizationHelpers.h"
 #include "ScopedTransaction.h"
-
+#include "SlotBase.h"
+#include "Styling/AppStyle.h"
+#include "Templates/Casts.h"
+#include "UObject/ObjectPtr.h"
+#include "UObject/WeakObjectPtr.h"
+#include "UObject/WeakObjectPtrTemplates.h"
 #include "Widgets/Images/SImage.h"
-#include "Widgets/Text/STextBlock.h"
-#include "Widgets/Input/STextComboBox.h"
-#include "Widgets/Input/SEditableTextBox.h"
 #include "Widgets/Input/SButton.h"
-#include "AssetRegistry/AssetData.h"
+#include "Widgets/Input/SEditableTextBox.h"
+#include "Widgets/Input/STextComboBox.h"
+#include "Widgets/SBoxPanel.h"
+#include "Widgets/Text/STextBlock.h"
 
-#include "Runtime/Launch/Resources/Version.h"
+class IPropertyHandle;
 
 
 #define LOCTEXT_NAMESPACE "CustomizableObjectNodeMeshClipWithMeshDetails"
diff --git a/Engine/Plugins/Experimental/Mutable/Source/CustomizableObjectEditor/Private/MuCOE/Nodes/CustomizableObjectNodeMeshClipWithMeshDetails.h b/Engine/Plugins/Experimental/Mutable/Source/CustomizableObjectEditor/Private/MuCOE/Nodes/CustomizableObjectNodeMeshClipWithMeshDetails.h
--- a/Engine/Plugins/Experimental/Mutable/Source/CustomizableObjectEditor/Private/MuCOE/Nodes/CustomizableObjectNodeMeshClipWithMeshDetails.h
+++ b/Engine/Plugins/Experimental/Mutable/Source/CustomizableObjectEditor/Private/MuCOE/Nodes/CustomizableObjectNodeMeshClipWithMeshDetails.h
@@ -2,14 +2,22 @@
 
 #pragma once
 
+#include "Containers/Array.h"
+#include "Containers/UnrealString.h"
+#include "HAL/Platform.h"
 #include "IDetailCustomization.h"
+#include "Input/Reply.h"
 #include "Templates/SharedPointer.h"
-#include "UObject/WeakObjectPtr.h"
 #include "Types/SlateEnums.h"
+#include "Widgets/DeclarativeSyntaxSupport.h"
 #include "Widgets/SCompoundWidget.h"
-#include "Widgets/SBoxPanel.h"
-#include "AssetRegistry/AssetData.h"
-#include "MuCOE/Nodes/CustomizableObjectNodeMeshClipWithMesh.h"
+
+class FText;
+class IDetailLayoutBuilder;
+class SHorizontalBox;
+class SVerticalBox;
+class UCustomizableObjectNodeMeshClipWithMesh;
+struct FAssetData;
 
 // Widget that represents the array of tags of the clip with mesh node
 class STagView : public SCompoundWidget
diff --git a/Engine/Plugins/Experimental/Mutable/Source/CustomizableObjectEditor/Private/MuCOE/Nodes/CustomizableObjectNodeMeshGeometryOperation.cpp b/Engine/Plugins/Experimental/Mutable/Source/CustomizableObjectEditor/Private/MuCOE/Nodes/CustomizableObjectNodeMeshGeometryOperation.cpp
--- a/Engine/Plugins/Experimental/Mutable/Source/CustomizableObjectEditor/Private/MuCOE/Nodes/CustomizableObjectNodeMeshGeometryOperation.cpp
+++ b/Engine/Plugins/Experimental/Mutable/Source/CustomizableObjectEditor/Private/MuCOE/Nodes/CustomizableObjectNodeMeshGeometryOperation.cpp
@@ -1,7 +1,13 @@
 // Copyright Epic Games, Inc. All Rights Reserved.
 
 #include "MuCOE/Nodes/CustomizableObjectNodeMeshGeometryOperation.h"
-#include "MuCOE/CustomizableObjectEditorModule.h"
+
+#include "EdGraph/EdGraphPin.h"
+#include "Internationalization/Internationalization.h"
+#include "MuCOE/EdGraphSchema_CustomizableObject.h"
+#include "UObject/NameTypes.h"
+
+class UCustomizableObjectNodeRemapPins;
 
 #define LOCTEXT_NAMESPACE "CustomizableObjectEditor"
 
diff --git a/Engine/Plugins/Experimental/Mutable/Source/CustomizableObjectEditor/Private/MuCOE/Nodes/CustomizableObjectNodeMeshMorph.cpp b/Engine/Plugins/Experimental/Mutable/Source/CustomizableObjectEditor/Private/MuCOE/Nodes/CustomizableObjectNodeMeshMorph.cpp
--- a/Engine/Plugins/Experimental/Mutable/Source/CustomizableObjectEditor/Private/MuCOE/Nodes/CustomizableObjectNodeMeshMorph.cpp
+++ b/Engine/Plugins/Experimental/Mutable/Source/CustomizableObjectEditor/Private/MuCOE/Nodes/CustomizableObjectNodeMeshMorph.cpp
@@ -1,8 +1,24 @@
 // Copyright Epic Games, Inc. All Rights Reserved.
 
 #include "MuCOE/Nodes/CustomizableObjectNodeMeshMorph.h"
-#include "MuCOE/CustomizableObjectEditorModule.h"
+
+#include "Animation/MorphTarget.h"
+#include "Containers/EnumAsByte.h"
+#include "EdGraph/EdGraph.h"
+#include "EdGraph/EdGraphPin.h"
+#include "Engine/SkeletalMesh.h"
+#include "Internationalization/Internationalization.h"
+#include "MuCOE/EdGraphSchema_CustomizableObject.h"
 #include "MuCOE/GraphTraversal.h"
+#include "MuCOE/ICustomizableObjectEditor.h"
+#include "MuCOE/Nodes/CustomizableObjectNodeMeshReshapeCommon.h"
+#include "MuCOE/Nodes/CustomizableObjectNodeSkeletalMesh.h"
+#include "Templates/Casts.h"
+#include "Templates/SharedPointer.h"
+#include "UObject/NameTypes.h"
+#include "UObject/ObjectPtr.h"
+
+class UCustomizableObjectNodeRemapPins;
 
 #define LOCTEXT_NAMESPACE "CustomizableObjectEditor"
 
diff --git a/Engine/Plugins/Experimental/Mutable/Source/CustomizableObjectEditor/Private/MuCOE/Nodes/CustomizableObjectNodeMeshMorphDetails.cpp b/Engine/Plugins/Experimental/Mutable/Source/CustomizableObjectEditor/Private/MuCOE/Nodes/CustomizableObjectNodeMeshMorphDetails.cpp
--- a/Engine/Plugins/Experimental/Mutable/Source/CustomizableObjectEditor/Private/MuCOE/Nodes/CustomizableObjectNodeMeshMorphDetails.cpp
+++ b/Engine/Plugins/Experimental/Mutable/Source/CustomizableObjectEditor/Private/MuCOE/Nodes/CustomizableObjectNodeMeshMorphDetails.cpp
@@ -1,28 +1,35 @@
 // Copyright Epic Games, Inc. All Rights Reserved.
 
 #include "MuCOE/Nodes/CustomizableObjectNodeMeshMorphDetails.h"
-#include "MuCOE/CustomizableObjectEditorModule.h"
-#include "MuCOE/CustomizableObjectCompiler.h"
-#include "MuCOE/CustomizableObjectEditorUtilities.h"
 
-#include "PropertyCustomizationHelpers.h"
-#include "DetailLayoutBuilder.h"
+#include "Animation/MorphTarget.h"
+#include "Containers/UnrealString.h"
 #include "DetailCategoryBuilder.h"
+#include "DetailLayoutBuilder.h"
+#include "DetailWidgetRow.h"
+#include "EdGraph/EdGraphPin.h"
+#include "Engine/SkeletalMesh.h"
+#include "HAL/PlatformCrt.h"
+#include "IDetailsView.h"
+#include "Internationalization/Internationalization.h"
+#include "Layout/Margin.h"
+#include "Misc/Attribute.h"
+#include "MuCOE/CustomizableObjectEditorUtilities.h"
 #include "MuCOE/GraphTraversal.h"
-#include "LevelEditor.h"
-#include "LevelEditorActions.h"
-
+#include "MuCOE/Nodes/CustomizableObjectNodeMeshMorph.h"
+#include "MuCOE/Nodes/CustomizableObjectNodeSkeletalMesh.h"
+#include "MuCOE/Nodes/CustomizableObjectNodeTable.h"
+#include "PropertyCustomizationHelpers.h"
+#include "PropertyHandle.h"
+#include "Styling/AppStyle.h"
+#include "Templates/Casts.h"
+#include "UObject/ObjectPtr.h"
+#include "UObject/WeakObjectPtr.h"
+#include "UObject/WeakObjectPtrTemplates.h"
+#include "Widgets/DeclarativeSyntaxSupport.h"
+#include "Widgets/Input/STextComboBox.h"
 #include "Widgets/Layout/SBorder.h"
 #include "Widgets/Text/STextBlock.h"
-#include "Widgets/Input/STextComboBox.h"
-
-#include "Engine/SkeletalMesh.h"
-#include "Engine/StaticMesh.h"
-#include "Animation/MorphTarget.h"
-
-#include "BoneSelectionWidget.h"
-
-#include "Runtime/Launch/Resources/Version.h"
 
 
 #define LOCTEXT_NAMESPACE "CustomizableObjectDetails"
diff --git a/Engine/Plugins/Experimental/Mutable/Source/CustomizableObjectEditor/Private/MuCOE/Nodes/CustomizableObjectNodeMeshMorphDetails.h b/Engine/Plugins/Experimental/Mutable/Source/CustomizableObjectEditor/Private/MuCOE/Nodes/CustomizableObjectNodeMeshMorphDetails.h
--- a/Engine/Plugins/Experimental/Mutable/Source/CustomizableObjectEditor/Private/MuCOE/Nodes/CustomizableObjectNodeMeshMorphDetails.h
+++ b/Engine/Plugins/Experimental/Mutable/Source/CustomizableObjectEditor/Private/MuCOE/Nodes/CustomizableObjectNodeMeshMorphDetails.h
@@ -2,11 +2,14 @@
 
 #pragma once
 
+#include "Containers/Array.h"
 #include "IDetailCustomization.h"
-#include "Types/SlateEnums.h"
 #include "Templates/SharedPointer.h"
-#include "UObject/WeakObjectPtr.h"
-#include "PropertyHandle.h"
+#include "Types/SlateEnums.h"
+
+class FString;
+class IDetailLayoutBuilder;
+class IPropertyHandle;
 
 class FCustomizableObjectNodeMeshMorphDetails : public IDetailCustomization
 {
diff --git a/Engine/Plugins/Experimental/Mutable/Source/CustomizableObjectEditor/Private/MuCOE/Nodes/CustomizableObjectNodeMeshMorphStackApplication.cpp b/Engine/Plugins/Experimental/Mutable/Source/CustomizableObjectEditor/Private/MuCOE/Nodes/CustomizableObjectNodeMeshMorphStackApplication.cpp
--- a/Engine/Plugins/Experimental/Mutable/Source/CustomizableObjectEditor/Private/MuCOE/Nodes/CustomizableObjectNodeMeshMorphStackApplication.cpp
+++ b/Engine/Plugins/Experimental/Mutable/Source/CustomizableObjectEditor/Private/MuCOE/Nodes/CustomizableObjectNodeMeshMorphStackApplication.cpp
@@ -1,8 +1,23 @@
 // Copyright Epic Games, Inc. All Rights Reserved.
 
 #include "MuCOE/Nodes/CustomizableObjectNodeMeshMorphStackApplication.h"
-#include "MuCOE/CustomizableObjectEditorModule.h"
+
+#include "Animation/MorphTarget.h"
+#include "Containers/EnumAsByte.h"
+#include "EdGraph/EdGraphPin.h"
+#include "Engine/SkeletalMesh.h"
+#include "HAL/PlatformCrt.h"
+#include "Internationalization/Internationalization.h"
+#include "MuCOE/EdGraphSchema_CustomizableObject.h"
 #include "MuCOE/GraphTraversal.h"
+#include "MuCOE/ICustomizableObjectEditor.h"
+#include "MuCOE/Nodes/CustomizableObjectNodeSkeletalMesh.h"
+#include "Templates/Casts.h"
+#include "Templates/SharedPointer.h"
+#include "UObject/NameTypes.h"
+#include "UObject/ObjectPtr.h"
+
+class UCustomizableObjectNodeRemapPins;
 
 
 #define LOCTEXT_NAMESPACE "CustomizableObjectEditor"
diff --git a/Engine/Plugins/Experimental/Mutable/Source/CustomizableObjectEditor/Private/MuCOE/Nodes/CustomizableObjectNodeMeshMorphStackDefinition.cpp b/Engine/Plugins/Experimental/Mutable/Source/CustomizableObjectEditor/Private/MuCOE/Nodes/CustomizableObjectNodeMeshMorphStackDefinition.cpp
--- a/Engine/Plugins/Experimental/Mutable/Source/CustomizableObjectEditor/Private/MuCOE/Nodes/CustomizableObjectNodeMeshMorphStackDefinition.cpp
+++ b/Engine/Plugins/Experimental/Mutable/Source/CustomizableObjectEditor/Private/MuCOE/Nodes/CustomizableObjectNodeMeshMorphStackDefinition.cpp
@@ -1,9 +1,24 @@
 // Copyright Epic Games, Inc. All Rights Reserved.
 
 #include "MuCOE/Nodes/CustomizableObjectNodeMeshMorphStackDefinition.h"
-#include "MuCOE/Nodes/CustomizableObjectNodeSkeletalMesh.h"
-#include "MuCOE/CustomizableObjectEditorModule.h"
+
+#include "Animation/MorphTarget.h"
+#include "Containers/EnumAsByte.h"
+#include "EdGraph/EdGraphPin.h"
+#include "Engine/SkeletalMesh.h"
+#include "HAL/PlatformCrt.h"
+#include "Internationalization/Internationalization.h"
+#include "MuCOE/EdGraphSchema_CustomizableObject.h"
 #include "MuCOE/GraphTraversal.h"
+#include "MuCOE/ICustomizableObjectEditor.h"
+#include "MuCOE/Nodes/CustomizableObjectNodeSkeletalMesh.h"
+#include "MuCOE/Nodes/CustomizableObjectNodeTable.h"
+#include "Templates/Casts.h"
+#include "Templates/SharedPointer.h"
+#include "UObject/NameTypes.h"
+#include "UObject/ObjectPtr.h"
+
+class UCustomizableObjectNodeRemapPins;
 
 
 #define LOCTEXT_NAMESPACE "CustomizableObjectEditor"
diff --git a/Engine/Plugins/Experimental/Mutable/Source/CustomizableObjectEditor/Private/MuCOE/Nodes/CustomizableObjectNodeMeshReshape.cpp b/Engine/Plugins/Experimental/Mutable/Source/CustomizableObjectEditor/Private/MuCOE/Nodes/CustomizableObjectNodeMeshReshape.cpp
--- a/Engine/Plugins/Experimental/Mutable/Source/CustomizableObjectEditor/Private/MuCOE/Nodes/CustomizableObjectNodeMeshReshape.cpp
+++ b/Engine/Plugins/Experimental/Mutable/Source/CustomizableObjectEditor/Private/MuCOE/Nodes/CustomizableObjectNodeMeshReshape.cpp
@@ -1,9 +1,14 @@
 // Copyright Epic Games, Inc. All Rights Reserved.
 
 #include "MuCOE/Nodes/CustomizableObjectNodeMeshReshape.h"
-#include "MuCOE/CustomizableObjectEditorModule.h"
-#include "DetailLayoutBuilder.h"
-#include "MuCOE/GraphTraversal.h"
+
+#include "EdGraph/EdGraphPin.h"
+#include "Internationalization/Internationalization.h"
+#include "MuCOE/EdGraphSchema_CustomizableObject.h"
+#include "MuCOE/Nodes/CustomizableObjectNodeMeshReshapeCommon.h"
+#include "UObject/NameTypes.h"
+
+class UCustomizableObjectNodeRemapPins;
 
 #define LOCTEXT_NAMESPACE "CustomizableObjectEditor"
 
diff --git a/Engine/Plugins/Experimental/Mutable/Source/CustomizableObjectEditor/Private/MuCOE/Nodes/CustomizableObjectNodeMeshReshapeDetails.cpp b/Engine/Plugins/Experimental/Mutable/Source/CustomizableObjectEditor/Private/MuCOE/Nodes/CustomizableObjectNodeMeshReshapeDetails.cpp
--- a/Engine/Plugins/Experimental/Mutable/Source/CustomizableObjectEditor/Private/MuCOE/Nodes/CustomizableObjectNodeMeshReshapeDetails.cpp
+++ b/Engine/Plugins/Experimental/Mutable/Source/CustomizableObjectEditor/Private/MuCOE/Nodes/CustomizableObjectNodeMeshReshapeDetails.cpp
@@ -1,22 +1,10 @@
 // Copyright Epic Games, Inc. All Rights Reserved.
 
 #include "MuCOE/Nodes/CustomizableObjectNodeMeshReshapeDetails.h"
-#include "MuCOE/Nodes/CustomizableObjectNodeMeshReshape.h"
-#include "MuCOE/Nodes/CustomizableObjectNodeSkeletalMesh.h"
-#include "MuCOE/Nodes/CustomizableObjectNodeTable.h"
-#include "MuCOE/GenerateMutableSource/GenerateMutableSource.h"
-
-#include "DetailLayoutBuilder.h"
-#include "DetailCategoryBuilder.h"
-#include "PropertyCustomizationHelpers.h"
 
 #include "Templates/SharedPointer.h"
-#include "UObject/WeakObjectPtr.h"
-#include "Engine/SkeletalMesh.h"
 
-#include "Widgets/Text/STextBlock.h"
-#include "BoneSelectionWidget.h"
-#include "MuCOE/GraphTraversal.h"
+class IDetailLayoutBuilder;
 
 #define LOCTEXT_NAMESPACE "CustomizableObjectDetails"
 
diff --git a/Engine/Plugins/Experimental/Mutable/Source/CustomizableObjectEditor/Private/MuCOE/Nodes/CustomizableObjectNodeMeshReshapeDetails.h b/Engine/Plugins/Experimental/Mutable/Source/CustomizableObjectEditor/Private/MuCOE/Nodes/CustomizableObjectNodeMeshReshapeDetails.h
--- a/Engine/Plugins/Experimental/Mutable/Source/CustomizableObjectEditor/Private/MuCOE/Nodes/CustomizableObjectNodeMeshReshapeDetails.h
+++ b/Engine/Plugins/Experimental/Mutable/Source/CustomizableObjectEditor/Private/MuCOE/Nodes/CustomizableObjectNodeMeshReshapeDetails.h
@@ -3,7 +3,9 @@
 #pragma once
 
 #include "IDetailCustomization.h"
-#include "IPropertyTypeCustomization.h"
+#include "Templates/SharedPointer.h"
+
+class IDetailLayoutBuilder;
 
 
 // TODO: This is no longer used, remove.
diff --git a/Engine/Plugins/Experimental/Mutable/Source/CustomizableObjectEditor/Private/MuCOE/Nodes/CustomizableObjectNodeMeshReshapeSelectionDetails.cpp b/Engine/Plugins/Experimental/Mutable/Source/CustomizableObjectEditor/Private/MuCOE/Nodes/CustomizableObjectNodeMeshReshapeSelectionDetails.cpp
--- a/Engine/Plugins/Experimental/Mutable/Source/CustomizableObjectEditor/Private/MuCOE/Nodes/CustomizableObjectNodeMeshReshapeSelectionDetails.cpp
+++ b/Engine/Plugins/Experimental/Mutable/Source/CustomizableObjectEditor/Private/MuCOE/Nodes/CustomizableObjectNodeMeshReshapeSelectionDetails.cpp
@@ -2,28 +2,30 @@
 
 #include "MuCOE/Nodes/CustomizableObjectNodeMeshReshapeSelectionDetails.h"
 
-#include "MuCOE/CustomizableObjectEditorModule.h"
-#include "MuCOE/CustomizableObjectCompiler.h"
-#include "MuCOE/CustomizableObjectEditorUtilities.h"
-
-#include "PropertyCustomizationHelpers.h"
-#include "DetailLayoutBuilder.h"
-#include "DetailCategoryBuilder.h"
-#include "MuCOE/GraphTraversal.h"
-#include "LevelEditor.h"
-#include "LevelEditorActions.h"
-
-#include "Widgets/Layout/SBorder.h"
-#include "Widgets/Text/STextBlock.h"
-#include "Widgets/Input/STextComboBox.h"
-
-#include "Engine/SkeletalMesh.h"
-#include "Engine/StaticMesh.h"
-#include "Animation/MorphTarget.h"
-
 #include "BoneSelectionWidget.h"
-
-#include "Runtime/Launch/Resources/Version.h"
+#include "Containers/Array.h"
+#include "Containers/UnrealString.h"
+#include "DetailWidgetRow.h"
+#include "EdGraph/EdGraphPin.h"
+#include "Engine/SkeletalMesh.h"
+#include "HAL/Platform.h"
+#include "HAL/PlatformCrt.h"
+#include "Internationalization/Text.h"
+#include "Misc/AssertionMacros.h"
+#include "Misc/Attribute.h"
+#include "MuCOE/GraphTraversal.h"
+#include "MuCOE/Nodes/CustomizableObjectNodeMeshMorph.h"
+#include "MuCOE/Nodes/CustomizableObjectNodeMeshReshape.h"
+#include "MuCOE/Nodes/CustomizableObjectNodeMeshReshapeCommon.h"
+#include "MuCOE/Nodes/CustomizableObjectNodeSkeletalMesh.h"
+#include "MuCOE/Nodes/CustomizableObjectNodeTable.h"
+#include "PropertyEditorModule.h"
+#include "PropertyHandle.h"
+#include "ReferenceSkeleton.h"
+#include "Templates/Casts.h"
+#include "UObject/Object.h"
+#include "UObject/ObjectPtr.h"
+#include "Widgets/DeclarativeSyntaxSupport.h"
 
 
 #define LOCTEXT_NAMESPACE "CustomizableObjectDetails"
diff --git a/Engine/Plugins/Experimental/Mutable/Source/CustomizableObjectEditor/Private/MuCOE/Nodes/CustomizableObjectNodeMeshReshapeSelectionDetails.h b/Engine/Plugins/Experimental/Mutable/Source/CustomizableObjectEditor/Private/MuCOE/Nodes/CustomizableObjectNodeMeshReshapeSelectionDetails.h
--- a/Engine/Plugins/Experimental/Mutable/Source/CustomizableObjectEditor/Private/MuCOE/Nodes/CustomizableObjectNodeMeshReshapeSelectionDetails.h
+++ b/Engine/Plugins/Experimental/Mutable/Source/CustomizableObjectEditor/Private/MuCOE/Nodes/CustomizableObjectNodeMeshReshapeSelectionDetails.h
@@ -2,12 +2,14 @@
 
 #pragma once
 
-#include "IDetailCustomization.h"
-#include "Types/SlateEnums.h"
+#include "IPropertyTypeCustomization.h"
+#include "Math/Color.h"
 #include "Templates/SharedPointer.h"
-#include "UObject/WeakObjectPtr.h"
-#include "PropertyHandle.h"
-#include "BoneSelectionWidget.h"
+#include "UObject/NameTypes.h"
+
+class IPropertyHandle;
+class IPropertyHandleArray;
+class SBoneSelectionWidget;
 
 class FMeshReshapeBonesReferenceCustomization : public IPropertyTypeCustomization
 {
diff --git a/Engine/Plugins/Experimental/Mutable/Source/CustomizableObjectEditor/Private/MuCOE/Nodes/CustomizableObjectNodeMeshVariation.cpp b/Engine/Plugins/Experimental/Mutable/Source/CustomizableObjectEditor/Private/MuCOE/Nodes/CustomizableObjectNodeMeshVariation.cpp
--- a/Engine/Plugins/Experimental/Mutable/Source/CustomizableObjectEditor/Private/MuCOE/Nodes/CustomizableObjectNodeMeshVariation.cpp
+++ b/Engine/Plugins/Experimental/Mutable/Source/CustomizableObjectEditor/Private/MuCOE/Nodes/CustomizableObjectNodeMeshVariation.cpp
@@ -1,9 +1,14 @@
 // Copyright Epic Games, Inc. All Rights Reserved.
 
 #include "MuCOE/Nodes/CustomizableObjectNodeMeshVariation.h"
-#include "MuCOE/CustomizableObjectEditorModule.h"
 
-#include "MuCOE/UnrealEditorPortabilityHelpers.h"
+#include "EdGraph/EdGraphPin.h"
+#include "Internationalization/Internationalization.h"
+#include "MuCOE/EdGraphSchema_CustomizableObject.h"
+#include "UObject/NameTypes.h"
+#include "UObject/UnrealType.h"
+
+class UCustomizableObjectNodeRemapPins;
 
 #define LOCTEXT_NAMESPACE "CustomizableObjectEditor"
 
diff --git a/Engine/Plugins/Experimental/Mutable/Source/CustomizableObjectEditor/Private/MuCOE/Nodes/CustomizableObjectNodeModifierBase.cpp b/Engine/Plugins/Experimental/Mutable/Source/CustomizableObjectEditor/Private/MuCOE/Nodes/CustomizableObjectNodeModifierBase.cpp
--- a/Engine/Plugins/Experimental/Mutable/Source/CustomizableObjectEditor/Private/MuCOE/Nodes/CustomizableObjectNodeModifierBase.cpp
+++ b/Engine/Plugins/Experimental/Mutable/Source/CustomizableObjectEditor/Private/MuCOE/Nodes/CustomizableObjectNodeModifierBase.cpp
@@ -1,7 +1,6 @@
 // Copyright Epic Games, Inc. All Rights Reserved.
 
 #include "MuCOE/Nodes/CustomizableObjectNodeModifierBase.h"
-#include "MuCOE/CustomizableObjectEditorModule.h"
 
 
 UCustomizableObjectNodeModifierBase::UCustomizableObjectNodeModifierBase()
diff --git a/Engine/Plugins/Experimental/Mutable/Source/CustomizableObjectEditor/Private/MuCOE/Nodes/CustomizableObjectNodeMorphMaterial.cpp b/Engine/Plugins/Experimental/Mutable/Source/CustomizableObjectEditor/Private/MuCOE/Nodes/CustomizableObjectNodeMorphMaterial.cpp
--- a/Engine/Plugins/Experimental/Mutable/Source/CustomizableObjectEditor/Private/MuCOE/Nodes/CustomizableObjectNodeMorphMaterial.cpp
+++ b/Engine/Plugins/Experimental/Mutable/Source/CustomizableObjectEditor/Private/MuCOE/Nodes/CustomizableObjectNodeMorphMaterial.cpp
@@ -1,8 +1,26 @@
 // Copyright Epic Games, Inc. All Rights Reserved.
 
 #include "MuCOE/Nodes/CustomizableObjectNodeMorphMaterial.h"
-#include "MuCOE/CustomizableObjectEditorModule.h"
+
+#include "Animation/MorphTarget.h"
+#include "Containers/Array.h"
+#include "EdGraph/EdGraph.h"
+#include "EdGraph/EdGraphPin.h"
+#include "Engine/SkeletalMesh.h"
+#include "Internationalization/Internationalization.h"
+#include "MuCO/CustomizableObjectCustomVersion.h"
+#include "MuCOE/EdGraphSchema_CustomizableObject.h"
 #include "MuCOE/GraphTraversal.h"
+#include "MuCOE/ICustomizableObjectEditor.h"
+#include "MuCOE/Nodes/CustomizableObjectNodeMaterial.h"
+#include "MuCOE/Nodes/CustomizableObjectNodeSkeletalMesh.h"
+#include "Templates/Casts.h"
+#include "Templates/SharedPointer.h"
+#include "UObject/NameTypes.h"
+#include "UObject/ObjectPtr.h"
+
+class UCustomizableObjectNodeRemapPins;
+struct FPropertyChangedEvent;
 
 #define LOCTEXT_NAMESPACE "CustomizableObjectEditor"
 
diff --git a/Engine/Plugins/Experimental/Mutable/Source/CustomizableObjectEditor/Private/MuCOE/Nodes/CustomizableObjectNodeMorphMaterialDetails.cpp b/Engine/Plugins/Experimental/Mutable/Source/CustomizableObjectEditor/Private/MuCOE/Nodes/CustomizableObjectNodeMorphMaterialDetails.cpp
--- a/Engine/Plugins/Experimental/Mutable/Source/CustomizableObjectEditor/Private/MuCOE/Nodes/CustomizableObjectNodeMorphMaterialDetails.cpp
+++ b/Engine/Plugins/Experimental/Mutable/Source/CustomizableObjectEditor/Private/MuCOE/Nodes/CustomizableObjectNodeMorphMaterialDetails.cpp
@@ -1,30 +1,39 @@
 // Copyright Epic Games, Inc. All Rights Reserved.
 
 #include "MuCOE/Nodes/CustomizableObjectNodeMorphMaterialDetails.h"
-#include "MuCOE/CustomizableObjectEditorModule.h"
-#include "MuCOE/CustomizableObjectCompiler.h"
-#include "PropertyCustomizationHelpers.h"
-#include "MuCOE/CustomizableObjectEditorUtilities.h"
-
-#include "LevelEditor.h"
-#include "LevelEditorActions.h"
 
-#include "DetailLayoutBuilder.h"
+#include "Animation/MorphTarget.h"
+#include "Containers/UnrealString.h"
 #include "DetailCategoryBuilder.h"
-#include "MuCOE/GraphTraversal.h"
-#include "LevelEditor.h"
-#include "LevelEditorActions.h"
-
+#include "DetailLayoutBuilder.h"
+#include "DetailWidgetRow.h"
+#include "EdGraph/EdGraphPin.h"
 #include "Engine/SkeletalMesh.h"
-#include "Engine/StaticMesh.h"
-#include "Animation/MorphTarget.h"
-
-#include "Widgets/Text/STextBlock.h"
-#include "Widgets/Input/STextComboBox.h"
+#include "HAL/PlatformCrt.h"
+#include "IDetailsView.h"
+#include "Internationalization/Internationalization.h"
+#include "Layout/Margin.h"
+#include "Misc/Attribute.h"
+#include "MuCOE/CustomizableObjectEditorUtilities.h"
+#include "MuCOE/GraphTraversal.h"
+#include "MuCOE/Nodes/CustomizableObjectNodeMaterial.h"
+#include "MuCOE/Nodes/CustomizableObjectNodeMaterialBase.h"
+#include "MuCOE/Nodes/CustomizableObjectNodeMorphMaterial.h"
+#include "MuCOE/Nodes/CustomizableObjectNodeSkeletalMesh.h"
+#include "PropertyCustomizationHelpers.h"
+#include "PropertyHandle.h"
+#include "SlotBase.h"
+#include "Styling/AppStyle.h"
+#include "Templates/Casts.h"
 #include "Templates/SharedPointer.h"
+#include "UObject/ObjectPtr.h"
 #include "UObject/WeakObjectPtr.h"
-
-#include "Runtime/Launch/Resources/Version.h"
+#include "UObject/WeakObjectPtrTemplates.h"
+#include "Widgets/DeclarativeSyntaxSupport.h"
+#include "Widgets/Input/STextComboBox.h"
+#include "Widgets/Layout/SBorder.h"
+#include "Widgets/SBoxPanel.h"
+#include "Widgets/Text/STextBlock.h"
 
 
 #define LOCTEXT_NAMESPACE "CustomizableObjectDetails"
diff --git a/Engine/Plugins/Experimental/Mutable/Source/CustomizableObjectEditor/Private/MuCOE/Nodes/CustomizableObjectNodeMorphMaterialDetails.h b/Engine/Plugins/Experimental/Mutable/Source/CustomizableObjectEditor/Private/MuCOE/Nodes/CustomizableObjectNodeMorphMaterialDetails.h
--- a/Engine/Plugins/Experimental/Mutable/Source/CustomizableObjectEditor/Private/MuCOE/Nodes/CustomizableObjectNodeMorphMaterialDetails.h
+++ b/Engine/Plugins/Experimental/Mutable/Source/CustomizableObjectEditor/Private/MuCOE/Nodes/CustomizableObjectNodeMorphMaterialDetails.h
@@ -2,14 +2,15 @@
 
 #pragma once
 
+#include "Containers/Array.h"
 #include "IDetailCustomization.h"
-
 #include "Templates/SharedPointer.h"
-#include "UObject/WeakObjectPtr.h"
-#include "Misc/Guid.h"
 #include "Types/SlateEnums.h"
-#include "PropertyHandle.h"
-#include "MuCOE/CustomizableObjectEditorModule.h"
+
+class FString;
+class IDetailLayoutBuilder;
+class IPropertyHandle;
+class UCustomizableObjectNodeObject;
 
 class FCustomizableObjectNodeMorphMaterialDetails : public IDetailCustomization
 {
diff --git a/Engine/Plugins/Experimental/Mutable/Source/CustomizableObjectEditor/Private/MuCOE/Nodes/CustomizableObjectNodeObject.cpp b/Engine/Plugins/Experimental/Mutable/Source/CustomizableObjectEditor/Private/MuCOE/Nodes/CustomizableObjectNodeObject.cpp
--- a/Engine/Plugins/Experimental/Mutable/Source/CustomizableObjectEditor/Private/MuCOE/Nodes/CustomizableObjectNodeObject.cpp
+++ b/Engine/Plugins/Experimental/Mutable/Source/CustomizableObjectEditor/Private/MuCOE/Nodes/CustomizableObjectNodeObject.cpp
@@ -1,11 +1,26 @@
 // Copyright Epic Games, Inc. All Rights Reserved.
 
 #include "MuCOE/Nodes/CustomizableObjectNodeObject.h"
-#include "MuCOE/CustomizableObjectEditorModule.h"
-#include "MuCOE/CustomizableObjectCompiler.h"
+
+#include "Containers/Queue.h"
+#include "HAL/PlatformCrt.h"
+#include "Internationalization/Internationalization.h"
+#include "Logging/MessageLog.h"
+#include "Logging/TokenizedMessage.h"
+#include "Math/UnrealMathSSE.h"
 #include "MuCO/CustomizableObject.h"
+#include "MuCOE/CustomizableObjectEditorLogger.h"
+#include "MuCOE/CustomizableObjectGraph.h"
+#include "MuCOE/EdGraphSchema_CustomizableObject.h"
 #include "MuCOE/GraphTraversal.h"
-#include "Logging/MessageLog.h"
+#include "MuCOE/ICustomizableObjectEditor.h"
+#include "MuCOE/Nodes/CustomizableObjectNodeMaterial.h"
+#include "MuCOE/Nodes/CustomizableObjectNodeMaterialVariation.h"
+#include "Templates/Casts.h"
+#include "Templates/SharedPointer.h"
+#include "UObject/UnrealType.h"
+
+class UCustomizableObjectNodeRemapPins;
 
 #define LOCTEXT_NAMESPACE "CustomizableObjectEditor"
 
diff --git a/Engine/Plugins/Experimental/Mutable/Source/CustomizableObjectEditor/Private/MuCOE/Nodes/CustomizableObjectNodeObjectChild.cpp b/Engine/Plugins/Experimental/Mutable/Source/CustomizableObjectEditor/Private/MuCOE/Nodes/CustomizableObjectNodeObjectChild.cpp
--- a/Engine/Plugins/Experimental/Mutable/Source/CustomizableObjectEditor/Private/MuCOE/Nodes/CustomizableObjectNodeObjectChild.cpp
+++ b/Engine/Plugins/Experimental/Mutable/Source/CustomizableObjectEditor/Private/MuCOE/Nodes/CustomizableObjectNodeObjectChild.cpp
@@ -1,7 +1,9 @@
 // Copyright Epic Games, Inc. All Rights Reserved.
 
 #include "MuCOE/Nodes/CustomizableObjectNodeObjectChild.h"
-#include "MuCOE/CustomizableObjectEditorModule.h"
+
+#include "HAL/Platform.h"
+#include "Internationalization/Internationalization.h"
 
 #define LOCTEXT_NAMESPACE "CustomizableObjectEditor"
 
diff --git a/Engine/Plugins/Experimental/Mutable/Source/CustomizableObjectEditor/Private/MuCOE/Nodes/CustomizableObjectNodeObjectDetails.cpp b/Engine/Plugins/Experimental/Mutable/Source/CustomizableObjectEditor/Private/MuCOE/Nodes/CustomizableObjectNodeObjectDetails.cpp
--- a/Engine/Plugins/Experimental/Mutable/Source/CustomizableObjectEditor/Private/MuCOE/Nodes/CustomizableObjectNodeObjectDetails.cpp
+++ b/Engine/Plugins/Experimental/Mutable/Source/CustomizableObjectEditor/Private/MuCOE/Nodes/CustomizableObjectNodeObjectDetails.cpp
@@ -2,26 +2,38 @@
 
 #include "MuCOE/Nodes/CustomizableObjectNodeObjectDetails.h"
 
-#include "MuCOE/CustomizableObjectEditorModule.h"
-#include "PropertyCustomizationHelpers.h"
-#include "MuCOE/CustomizableObjectEditorUtilities.h"
+#include "AssetRegistry/AssetData.h"
+#include "Containers/UnrealString.h"
+#include "Delegates/Delegate.h"
+#include "DetailCategoryBuilder.h"
+#include "DetailLayoutBuilder.h"
+#include "DetailWidgetRow.h"
+#include "EdGraph/EdGraph.h"
+#include "HAL/Platform.h"
+#include "HAL/PlatformCrt.h"
+#include "IDetailsView.h"
+#include "Internationalization/Internationalization.h"
+#include "Layout/Margin.h"
+#include "Misc/Attribute.h"
+#include "Misc/Guid.h"
 #include "MuCO/CustomizableObject.h"
+#include "MuCOE/CustomizableObjectEditorUtilities.h"
+#include "MuCOE/Nodes/CustomizableObjectNodeObject.h"
 #include "MuCOE/Nodes/CustomizableObjectNodeObjectGroup.h"
-
+#include "MuCOE/SCustomizableObjectNodeObjectRTMorphTargetOverride.h"
+#include "MuCOE/SCustomizableObjectNodeObjectStatesView.h"
+#include "PropertyCustomizationHelpers.h"
+#include "PropertyHandle.h"
 #include "ScopedTransaction.h"
-#include "DetailLayoutBuilder.h"
-#include "DetailCategoryBuilder.h"
-#include "LevelEditor.h"
-#include "LevelEditorActions.h"
-
-#include "Widgets/Input/SButton.h"
-#include "Widgets/Text/STextBlock.h"
+#include "Styling/AppStyle.h"
+#include "Templates/Casts.h"
+#include "UObject/ObjectPtr.h"
+#include "UObject/WeakObjectPtr.h"
+#include "UObject/WeakObjectPtrTemplates.h"
+#include "Widgets/DeclarativeSyntaxSupport.h"
 #include "Widgets/Input/STextComboBox.h"
-
-#include "Runtime/Launch/Resources/Version.h"
-
-#include "MuCOE/SCustomizableObjectNodeObjectStatesView.h"
-#include "MuCOE/SCustomizableObjectNodeObjectRTMorphTargetOverride.h"
+#include "Widgets/Layout/SBorder.h"
+#include "Widgets/Text/STextBlock.h"
 
 
 #define LOCTEXT_NAMESPACE "CustomizableObjectDetails"
diff --git a/Engine/Plugins/Experimental/Mutable/Source/CustomizableObjectEditor/Private/MuCOE/Nodes/CustomizableObjectNodeObjectDetails.h b/Engine/Plugins/Experimental/Mutable/Source/CustomizableObjectEditor/Private/MuCOE/Nodes/CustomizableObjectNodeObjectDetails.h
--- a/Engine/Plugins/Experimental/Mutable/Source/CustomizableObjectEditor/Private/MuCOE/Nodes/CustomizableObjectNodeObjectDetails.h
+++ b/Engine/Plugins/Experimental/Mutable/Source/CustomizableObjectEditor/Private/MuCOE/Nodes/CustomizableObjectNodeObjectDetails.h
@@ -2,13 +2,16 @@
 
 #pragma once
 
+#include "Containers/Array.h"
 #include "IDetailCustomization.h"
-
-#include "Types/SlateEnums.h"
+#include "Math/MathFwd.h"
 #include "Templates/SharedPointer.h"
-#include "UObject/WeakObjectPtr.h"
-#include "PropertyHandle.h"
-#include "AssetRegistry/AssetData.h"
+#include "Types/SlateEnums.h"
+
+class FString;
+class IDetailLayoutBuilder;
+class IPropertyHandle;
+struct FAssetData;
 
 
 class FCustomizableObjectNodeObjectDetails : public IDetailCustomization
diff --git a/Engine/Plugins/Experimental/Mutable/Source/CustomizableObjectEditor/Private/MuCOE/Nodes/CustomizableObjectNodeObjectGroup.cpp b/Engine/Plugins/Experimental/Mutable/Source/CustomizableObjectEditor/Private/MuCOE/Nodes/CustomizableObjectNodeObjectGroup.cpp
--- a/Engine/Plugins/Experimental/Mutable/Source/CustomizableObjectEditor/Private/MuCOE/Nodes/CustomizableObjectNodeObjectGroup.cpp
+++ b/Engine/Plugins/Experimental/Mutable/Source/CustomizableObjectEditor/Private/MuCOE/Nodes/CustomizableObjectNodeObjectGroup.cpp
@@ -1,7 +1,16 @@
 // Copyright Epic Games, Inc. All Rights Reserved.
 
 #include "MuCOE/Nodes/CustomizableObjectNodeObjectGroup.h"
-#include "MuCOE/CustomizableObjectEditorModule.h"
+
+#include "EdGraph/EdGraphPin.h"
+#include "Internationalization/Internationalization.h"
+#include "MuCO/CustomizableObjectCustomVersion.h"
+#include "MuCOE/EdGraphSchema_CustomizableObject.h"
+#include "Serialization/Archive.h"
+#include "UObject/NameTypes.h"
+
+class UCustomizableObjectNodeRemapPins;
+struct FPropertyChangedEvent;
 
 #define LOCTEXT_NAMESPACE "CustomizableObjectEditor"
 
diff --git a/Engine/Plugins/Experimental/Mutable/Source/CustomizableObjectEditor/Private/MuCOE/Nodes/CustomizableObjectNodeParentedMaterial.cpp b/Engine/Plugins/Experimental/Mutable/Source/CustomizableObjectEditor/Private/MuCOE/Nodes/CustomizableObjectNodeParentedMaterial.cpp
--- a/Engine/Plugins/Experimental/Mutable/Source/CustomizableObjectEditor/Private/MuCOE/Nodes/CustomizableObjectNodeParentedMaterial.cpp
+++ b/Engine/Plugins/Experimental/Mutable/Source/CustomizableObjectEditor/Private/MuCOE/Nodes/CustomizableObjectNodeParentedMaterial.cpp
@@ -1,8 +1,14 @@
 // Copyright Epic Games, Inc. All Rights Reserved.
 
 #include "MuCOE/Nodes/CustomizableObjectNodeParentedMaterial.h"
+
+#include "HAL/Platform.h"
+#include "HAL/PlatformCrt.h"
+#include "MuCOE/Nodes/CustomizableObjectNode.h"
 #include "MuCOE/Nodes/CustomizableObjectNodeCopyMaterial.h"
+#include "MuCOE/Nodes/CustomizableObjectNodeMaterial.h"
 #include "MuCOE/Nodes/CustomizableObjectNodeObject.h"
+#include "Templates/Casts.h"
 
 
 UCustomizableObjectNodeMaterial* FCustomizableObjectNodeParentedMaterial::GetParentMaterialNode() const
diff --git a/Engine/Plugins/Experimental/Mutable/Source/CustomizableObjectEditor/Private/MuCOE/Nodes/CustomizableObjectNodeParentedMaterialDetails.cpp b/Engine/Plugins/Experimental/Mutable/Source/CustomizableObjectEditor/Private/MuCOE/Nodes/CustomizableObjectNodeParentedMaterialDetails.cpp
--- a/Engine/Plugins/Experimental/Mutable/Source/CustomizableObjectEditor/Private/MuCOE/Nodes/CustomizableObjectNodeParentedMaterialDetails.cpp
+++ b/Engine/Plugins/Experimental/Mutable/Source/CustomizableObjectEditor/Private/MuCOE/Nodes/CustomizableObjectNodeParentedMaterialDetails.cpp
@@ -2,25 +2,44 @@
 
 #include "MuCOE/Nodes/CustomizableObjectNodeParentedMaterialDetails.h"
 
-#include "MuCOE/CustomizableObjectEditorModule.h"
-#include "PropertyCustomizationHelpers.h"
-#include "MuCOE/CustomizableObjectEditorUtilities.h"
-
-#include "DetailLayoutBuilder.h"
 #include "DetailCategoryBuilder.h"
-#include "LevelEditor.h"
-#include "LevelEditorActions.h"
+#include "DetailLayoutBuilder.h"
+#include "DetailWidgetRow.h"
+#include "EdGraph/EdGraph.h"
+#include "Engine/SkeletalMesh.h"
+#include "HAL/Platform.h"
+#include "HAL/PlatformCrt.h"
+#include "IDetailsView.h"
+#include "Internationalization/Internationalization.h"
+#include "Layout/Margin.h"
+#include "Materials/MaterialInterface.h"
+#include "Misc/AssertionMacros.h"
+#include "Misc/Attribute.h"
+#include "MuCO/CustomizableObject.h"
+#include "MuCOE/Nodes/CustomizableObjectNode.h"
+#include "MuCOE/Nodes/CustomizableObjectNodeCopyMaterial.h"
+#include "MuCOE/Nodes/CustomizableObjectNodeEditMaterialBase.h"
+#include "MuCOE/Nodes/CustomizableObjectNodeExtendMaterial.h"
+#include "MuCOE/Nodes/CustomizableObjectNodeMaterial.h"
+#include "MuCOE/Nodes/CustomizableObjectNodeObject.h"
+#include "MuCOE/Nodes/CustomizableObjectNodeParentedMaterial.h"
+#include "MuCOE/Nodes/CustomizableObjectNodeSkeletalMesh.h"
+#include "PropertyCustomizationHelpers.h"
 #include "ScopedTransaction.h"
-
-#include "Widgets/Text/STextBlock.h"
-#include "Widgets/Input/STextComboBox.h"
+#include "SlotBase.h"
+#include "Styling/AppStyle.h"
+#include "Templates/Casts.h"
 #include "Templates/SharedPointer.h"
+#include "UObject/ObjectPtr.h"
 #include "UObject/WeakObjectPtr.h"
+#include "Widgets/DeclarativeSyntaxSupport.h"
+#include "Widgets/Input/STextComboBox.h"
+#include "Widgets/Layout/SBorder.h"
+#include "Widgets/SBoxPanel.h"
+#include "Widgets/Text/STextBlock.h"
 
-#include "Engine/SkeletalMesh.h"
-#include "Engine/StaticMesh.h"
-
-#include "Runtime/Launch/Resources/Version.h"
+class IPropertyHandle;
+class UObject;
 
 
 #define LOCTEXT_NAMESPACE "CustomizableObjectDetails"
diff --git a/Engine/Plugins/Experimental/Mutable/Source/CustomizableObjectEditor/Private/MuCOE/Nodes/CustomizableObjectNodeParentedMaterialDetails.h b/Engine/Plugins/Experimental/Mutable/Source/CustomizableObjectEditor/Private/MuCOE/Nodes/CustomizableObjectNodeParentedMaterialDetails.h
--- a/Engine/Plugins/Experimental/Mutable/Source/CustomizableObjectEditor/Private/MuCOE/Nodes/CustomizableObjectNodeParentedMaterialDetails.h
+++ b/Engine/Plugins/Experimental/Mutable/Source/CustomizableObjectEditor/Private/MuCOE/Nodes/CustomizableObjectNodeParentedMaterialDetails.h
@@ -2,17 +2,21 @@
 
 #pragma once
 
+#include "Containers/Array.h"
+#include "Containers/UnrealString.h"
 #include "IDetailCustomization.h"
-
-#include "MuCO/CustomizableObject.h"	
-#include "Templates/SharedPointer.h"
-#include "UObject/WeakObjectPtr.h"
 #include "Misc/Guid.h"
+#include "Templates/SharedPointer.h"
 #include "Types/SlateEnums.h"
-#include "PropertyHandle.h"
-#include "MuCOE/CustomizableObjectEditorModule.h"
-#include "MuCOE/Nodes/CustomizableObjectNode.h"
-#include "MuCOE/Nodes/CustomizableObjectNodeParentedMaterial.h"
+#include "UObject/WeakObjectPtr.h"
+#include "UObject/WeakObjectPtrTemplates.h"
+
+class FCustomizableObjectNodeParentedMaterial;
+class IDetailLayoutBuilder;
+class IPropertyHandle;
+class UCustomizableObject;
+class UCustomizableObjectNode;
+class UCustomizableObjectNodeMaterial;
 
 class FCustomizableObjectNodeParentedMaterialDetails : public IDetailCustomization
 {
diff --git a/Engine/Plugins/Experimental/Mutable/Source/CustomizableObjectEditor/Private/MuCOE/Nodes/CustomizableObjectNodeParentedNode.cpp b/Engine/Plugins/Experimental/Mutable/Source/CustomizableObjectEditor/Private/MuCOE/Nodes/CustomizableObjectNodeParentedNode.cpp
--- a/Engine/Plugins/Experimental/Mutable/Source/CustomizableObjectEditor/Private/MuCOE/Nodes/CustomizableObjectNodeParentedNode.cpp
+++ b/Engine/Plugins/Experimental/Mutable/Source/CustomizableObjectEditor/Private/MuCOE/Nodes/CustomizableObjectNodeParentedNode.cpp
@@ -2,6 +2,12 @@
 
 #include "MuCOE/Nodes/CustomizableObjectNodeParentedNode.h"
 
+#include "Containers/Array.h"
+#include "HAL/PlatformCrt.h"
+#include "MuCOE/Nodes/CustomizableObjectNode.h"
+
+class UCustomizableObject;
+
 
 void ICustomizableObjectNodeParentedNode::SetParentNode(UCustomizableObject* Object, FGuid NodeId)
 {
diff --git a/Engine/Plugins/Experimental/Mutable/Source/CustomizableObjectEditor/Private/MuCOE/Nodes/CustomizableObjectNodePinViewerDetails.cpp b/Engine/Plugins/Experimental/Mutable/Source/CustomizableObjectEditor/Private/MuCOE/Nodes/CustomizableObjectNodePinViewerDetails.cpp
--- a/Engine/Plugins/Experimental/Mutable/Source/CustomizableObjectEditor/Private/MuCOE/Nodes/CustomizableObjectNodePinViewerDetails.cpp
+++ b/Engine/Plugins/Experimental/Mutable/Source/CustomizableObjectEditor/Private/MuCOE/Nodes/CustomizableObjectNodePinViewerDetails.cpp
@@ -2,10 +2,10 @@
 
 #include "MuCOE/Nodes/CustomizableObjectNodePinViewerDetails.h"
 
-#include "DetailLayoutBuilder.h"
-
 #include "MuCOE/PinViewer/SPinViewer.h"
 
+class IDetailLayoutBuilder;
+
 #define LOCTEXT_NAMESPACE "CustomizableObjectEditor"
 
 
