Ran IWYU on Mutable plugin

diff --git a/Engine/Plugins/Experimental/Mutable/Source/CustomizableObjectEditor/Private/MuCOE/Nodes/CustomizableObjectNodePinViewerDetails.h b/Engine/Plugins/Experimental/Mutable/Source/CustomizableObjectEditor/Private/MuCOE/Nodes/CustomizableObjectNodePinViewerDetails.h
--- a/Engine/Plugins/Experimental/Mutable/Source/CustomizableObjectEditor/Private/MuCOE/Nodes/CustomizableObjectNodePinViewerDetails.h
+++ b/Engine/Plugins/Experimental/Mutable/Source/CustomizableObjectEditor/Private/MuCOE/Nodes/CustomizableObjectNodePinViewerDetails.h
@@ -5,6 +5,8 @@
 #include "IDetailCustomization.h"
 #include "Templates/SharedPointer.h"
 
+class IDetailLayoutBuilder;
+
 
 /** Simple Detail Customization which includes the PinViewer. */
 class FCustomizableObjectNodePinViewerDetails : public IDetailCustomization
diff --git a/Engine/Plugins/Experimental/Mutable/Source/CustomizableObjectEditor/Private/MuCOE/Nodes/CustomizableObjectNodeProjectorConstant.cpp b/Engine/Plugins/Experimental/Mutable/Source/CustomizableObjectEditor/Private/MuCOE/Nodes/CustomizableObjectNodeProjectorConstant.cpp
--- a/Engine/Plugins/Experimental/Mutable/Source/CustomizableObjectEditor/Private/MuCOE/Nodes/CustomizableObjectNodeProjectorConstant.cpp
+++ b/Engine/Plugins/Experimental/Mutable/Source/CustomizableObjectEditor/Private/MuCOE/Nodes/CustomizableObjectNodeProjectorConstant.cpp
@@ -1,9 +1,15 @@
 // Copyright Epic Games, Inc. All Rights Reserved.
 
 #include "MuCOE/Nodes/CustomizableObjectNodeProjectorConstant.h"
-#include "Components/SkeletalMeshComponent.h"
-#include "MuCOE/CustomizableObjectEditorModule.h"
+
+#include "Containers/UnrealString.h"
+#include "EdGraph/EdGraphPin.h"
+#include "Internationalization/Internationalization.h"
+#include "MuCOE/EdGraphSchema_CustomizableObject.h"
 #include "MuCOE/UnrealEditorPortabilityHelpers.h"
+#include "UObject/UnrealType.h"
+
+class UCustomizableObjectNodeRemapPins;
 
 #define LOCTEXT_NAMESPACE "CustomizableObjectEditor"
 
diff --git a/Engine/Plugins/Experimental/Mutable/Source/CustomizableObjectEditor/Private/MuCOE/Nodes/CustomizableObjectNodeProjectorParameter.cpp b/Engine/Plugins/Experimental/Mutable/Source/CustomizableObjectEditor/Private/MuCOE/Nodes/CustomizableObjectNodeProjectorParameter.cpp
--- a/Engine/Plugins/Experimental/Mutable/Source/CustomizableObjectEditor/Private/MuCOE/Nodes/CustomizableObjectNodeProjectorParameter.cpp
+++ b/Engine/Plugins/Experimental/Mutable/Source/CustomizableObjectEditor/Private/MuCOE/Nodes/CustomizableObjectNodeProjectorParameter.cpp
@@ -1,9 +1,14 @@
 // Copyright Epic Games, Inc. All Rights Reserved.
 
 #include "MuCOE/Nodes/CustomizableObjectNodeProjectorParameter.h"
-#include "Components/SkeletalMeshComponent.h"
-#include "MuCOE/CustomizableObjectEditorModule.h"
+
+#include "EdGraph/EdGraphPin.h"
+#include "Internationalization/Internationalization.h"
+#include "MuCOE/EdGraphSchema_CustomizableObject.h"
 #include "MuCOE/UnrealEditorPortabilityHelpers.h"
+#include "UObject/UnrealType.h"
+
+class UCustomizableObjectNodeRemapPins;
 
 #define LOCTEXT_NAMESPACE "CustomizableObjectEditor"
 
diff --git a/Engine/Plugins/Experimental/Mutable/Source/CustomizableObjectEditor/Private/MuCOE/Nodes/CustomizableObjectNodeProjectorParameterDetails.cpp b/Engine/Plugins/Experimental/Mutable/Source/CustomizableObjectEditor/Private/MuCOE/Nodes/CustomizableObjectNodeProjectorParameterDetails.cpp
--- a/Engine/Plugins/Experimental/Mutable/Source/CustomizableObjectEditor/Private/MuCOE/Nodes/CustomizableObjectNodeProjectorParameterDetails.cpp
+++ b/Engine/Plugins/Experimental/Mutable/Source/CustomizableObjectEditor/Private/MuCOE/Nodes/CustomizableObjectNodeProjectorParameterDetails.cpp
@@ -1,26 +1,52 @@
 // Copyright Epic Games, Inc. All Rights Reserved.
 
 #include "MuCOE/Nodes/CustomizableObjectNodeProjectorParameterDetails.h"
-#include "MuCOE/CustomizableObjectEditorModule.h"
-#include "PropertyCustomizationHelpers.h"
-#include "MuCOE/CustomizableObjectEditorUtilities.h"
 
-#include "DetailLayoutBuilder.h"
+#include "Containers/UnrealString.h"
+#include "CoreGlobals.h"
+#include "Delegates/Delegate.h"
 #include "DetailCategoryBuilder.h"
-#include "LevelEditor.h"
-#include "LevelEditorActions.h"
-
-#include "Widgets/Input/SButton.h"
-#include "Widgets/Text/STextBlock.h"
-
-#include "Runtime/Launch/Resources/Version.h"
+#include "DetailLayoutBuilder.h"
+#include "DetailWidgetRow.h"
 #include "Engine/SkeletalMesh.h"
-#include "Widgets/Input/STextComboBox.h"
-#include "Math/Vector4.h"
-#include "Math/Vector.h"
-
-#include "Runtime/Launch/Resources/Version.h"
+#include "Fonts/SlateFontInfo.h"
 #include "HAL/PlatformApplicationMisc.h"
+#include "HAL/PlatformCrt.h"
+#include "IDetailsView.h"
+#include "Internationalization/Internationalization.h"
+#include "Layout/Margin.h"
+#include "Math/Transform.h"
+#include "Math/TransformVectorized.h"
+#include "Math/UnrealMathSSE.h"
+#include "Math/Vector.h"
+#include "Misc/Attribute.h"
+#include "Misc/OutputDeviceRedirector.h"
+#include "MuCO/CustomizableObject.h"
+#include "MuCO/CustomizableObjectParameterTypeDefinitions.h"
+#include "MuCOE/CustomizableObjectEditorUtilities.h"
+#include "MuCOE/CustomizableObjectGraph.h"
+#include "MuCOE/ICustomizableObjectEditor.h"
+#include "MuCOE/Nodes/CustomizableObjectNodeProjectorConstant.h"
+#include "MuCOE/Nodes/CustomizableObjectNodeProjectorParameter.h"
+#include "PropertyCustomizationHelpers.h"
+#include "PropertyHandle.h"
+#include "ReferenceSkeleton.h"
+#include "SlotBase.h"
+#include "Styling/AppStyle.h"
+#include "Templates/Casts.h"
+#include "UObject/Class.h"
+#include "UObject/NameTypes.h"
+#include "UObject/Object.h"
+#include "UObject/PropertyPortFlags.h"
+#include "UObject/UObjectBaseUtility.h"
+#include "UObject/WeakObjectPtr.h"
+#include "UObject/WeakObjectPtrTemplates.h"
+#include "Widgets/DeclarativeSyntaxSupport.h"
+#include "Widgets/Input/SButton.h"
+#include "Widgets/Input/STextComboBox.h"
+#include "Widgets/Layout/SBorder.h"
+#include "Widgets/SBoxPanel.h"
+#include "Widgets/Text/STextBlock.h"
 
 
 #define LOCTEXT_NAMESPACE "CustomizableObjectDetails"
diff --git a/Engine/Plugins/Experimental/Mutable/Source/CustomizableObjectEditor/Private/MuCOE/Nodes/CustomizableObjectNodeProjectorParameterDetails.h b/Engine/Plugins/Experimental/Mutable/Source/CustomizableObjectEditor/Private/MuCOE/Nodes/CustomizableObjectNodeProjectorParameterDetails.h
--- a/Engine/Plugins/Experimental/Mutable/Source/CustomizableObjectEditor/Private/MuCOE/Nodes/CustomizableObjectNodeProjectorParameterDetails.h
+++ b/Engine/Plugins/Experimental/Mutable/Source/CustomizableObjectEditor/Private/MuCOE/Nodes/CustomizableObjectNodeProjectorParameterDetails.h
@@ -2,11 +2,15 @@
 
 #pragma once
 
+#include "Containers/Array.h"
 #include "IDetailCustomization.h"
-#include "Templates/SharedPointer.h"
-#include "UObject/WeakObjectPtr.h"
 #include "Input/Reply.h"
-#include "PropertyHandle.h"
+#include "Templates/SharedPointer.h"
+#include "Types/SlateEnums.h"
+
+class FString;
+class IDetailLayoutBuilder;
+class IPropertyHandle;
 
 
 class FCustomizableObjectNodeProjectorParameterDetails : public IDetailCustomization
diff --git a/Engine/Plugins/Experimental/Mutable/Source/CustomizableObjectEditor/Private/MuCOE/Nodes/CustomizableObjectNodeRemoveMesh.cpp b/Engine/Plugins/Experimental/Mutable/Source/CustomizableObjectEditor/Private/MuCOE/Nodes/CustomizableObjectNodeRemoveMesh.cpp
--- a/Engine/Plugins/Experimental/Mutable/Source/CustomizableObjectEditor/Private/MuCOE/Nodes/CustomizableObjectNodeRemoveMesh.cpp
+++ b/Engine/Plugins/Experimental/Mutable/Source/CustomizableObjectEditor/Private/MuCOE/Nodes/CustomizableObjectNodeRemoveMesh.cpp
@@ -1,7 +1,15 @@
 // Copyright Epic Games, Inc. All Rights Reserved.
 
 #include "MuCOE/Nodes/CustomizableObjectNodeRemoveMesh.h"
-#include "MuCOE/CustomizableObjectEditorModule.h"
+
+#include "EdGraph/EdGraphPin.h"
+#include "Internationalization/Internationalization.h"
+#include "MuCOE/EdGraphSchema_CustomizableObject.h"
+#include "MuCOE/ICustomizableObjectEditor.h"
+#include "Templates/SharedPointer.h"
+#include "UObject/NameTypes.h"
+
+class UCustomizableObjectNodeRemapPins;
 
 #define LOCTEXT_NAMESPACE "CustomizableObjectEditor"
 
diff --git a/Engine/Plugins/Experimental/Mutable/Source/CustomizableObjectEditor/Private/MuCOE/Nodes/CustomizableObjectNodeRemoveMeshBlocks.cpp b/Engine/Plugins/Experimental/Mutable/Source/CustomizableObjectEditor/Private/MuCOE/Nodes/CustomizableObjectNodeRemoveMeshBlocks.cpp
--- a/Engine/Plugins/Experimental/Mutable/Source/CustomizableObjectEditor/Private/MuCOE/Nodes/CustomizableObjectNodeRemoveMeshBlocks.cpp
+++ b/Engine/Plugins/Experimental/Mutable/Source/CustomizableObjectEditor/Private/MuCOE/Nodes/CustomizableObjectNodeRemoveMeshBlocks.cpp
@@ -1,7 +1,30 @@
 // Copyright Epic Games, Inc. All Rights Reserved.
 
 #include "MuCOE/Nodes/CustomizableObjectNodeRemoveMeshBlocks.h"
-#include "MuCOE/CustomizableObjectEditorModule.h"
+
+#include "EdGraph/EdGraph.h"
+#include "HAL/PlatformCrt.h"
+#include "Internationalization/Internationalization.h"
+#include "Logging/LogCategory.h"
+#include "Logging/LogMacros.h"
+#include "Misc/Guid.h"
+#include "MuCO/CustomizableObject.h"
+#include "MuCO/CustomizableObjectCustomVersion.h"
+#include "MuCOE/CustomizableObjectLayout.h"
+#include "MuCOE/EdGraphSchema_CustomizableObject.h"
+#include "MuCOE/ICustomizableObjectEditor.h"
+#include "MuCOE/Nodes/CustomizableObjectNodeMaterial.h"
+#include "MuCOE/Nodes/CustomizableObjectNodeMaterialBase.h"
+#include "Serialization/Archive.h"
+#include "Templates/Casts.h"
+#include "Templates/SharedPointer.h"
+#include "Trace/Detail/Channel.h"
+#include "UObject/NameTypes.h"
+#include "UObject/Package.h"
+#include "UObject/UnrealType.h"
+
+class UCustomizableObjectNodeRemapPins;
+class UEdGraphPin;
 
 #define LOCTEXT_NAMESPACE "CustomizableObjectEditor"
 
diff --git a/Engine/Plugins/Experimental/Mutable/Source/CustomizableObjectEditor/Private/MuCOE/Nodes/CustomizableObjectNodeSkeletalMesh.cpp b/Engine/Plugins/Experimental/Mutable/Source/CustomizableObjectEditor/Private/MuCOE/Nodes/CustomizableObjectNodeSkeletalMesh.cpp
--- a/Engine/Plugins/Experimental/Mutable/Source/CustomizableObjectEditor/Private/MuCOE/Nodes/CustomizableObjectNodeSkeletalMesh.cpp
+++ b/Engine/Plugins/Experimental/Mutable/Source/CustomizableObjectEditor/Private/MuCOE/Nodes/CustomizableObjectNodeSkeletalMesh.cpp
@@ -1,21 +1,58 @@
 // Copyright Epic Games, Inc. All Rights Reserved.
 
 #include "MuCOE/Nodes/CustomizableObjectNodeSkeletalMesh.h"
-#include "MuCOE/Nodes/CustomizableObjectNode.h"
-#include "MuCOE/CustomizableObjectEditorModule.h"
+
+#include "AssetThumbnail.h"
+#include "Containers/EnumAsByte.h"
+#include "Containers/IndirectArray.h"
+#include "Containers/Map.h"
+#include "EdGraph/EdGraph.h"
 #include "Engine/SkeletalMesh.h"
+#include "Engine/Texture.h"
+#include "Engine/Texture2D.h"
+#include "GenericPlatform/ICursor.h"
+#include "ISinglePropertyView.h"
+#include "Internationalization/Internationalization.h"
+#include "Layout/Margin.h"
+#include "MaterialTypes.h"
 #include "Materials/Material.h"
+#include "Materials/MaterialInterface.h"
+#include "Misc/AssertionMacros.h"
+#include "Misc/Attribute.h"
+#include "Misc/Guid.h"
+#include "Misc/Optional.h"
+#include "Modules/ModuleManager.h"
+#include "MuCO/CustomizableObjectCustomVersion.h"
 #include "MuCO/UnrealPortabilityHelpers.h"
-#include "MuCOE/UnrealEditorPortabilityHelpers.h"
-#include "NodeFactory.h"
-#include "ISinglePropertyView.h"
-#include "EditorFramework/ThumbnailInfo.h"
-#include "Widgets/Input/SCheckBox.h"
-#include "MuCOE/CustomizableObjectEditorStyle.h"
 #include "MuCOE/CustomizableObjectEditor.h"
+#include "MuCOE/CustomizableObjectEditorStyle.h"
+#include "MuCOE/CustomizableObjectLayout.h"
+#include "MuCOE/EdGraphSchema_CustomizableObject.h"
 #include "MuCOE/GraphTraversal.h"
 #include "MuCOE/MutableUtils.h"
+#include "MuCOE/Nodes/CustomizableObjectNodeLayoutBlocks.h"
+#include "MuCOE/RemapPins/CustomizableObjectNodeRemapPins.h"
 #include "MuCOE/RemapPins/CustomizableObjectNodeRemapPinsByNameDefaultPin.h"
+#include "MuCOE/UnrealEditorPortabilityHelpers.h"
+#include "PropertyEditorModule.h"
+#include "Rendering/SkeletalMeshLODModel.h"
+#include "Rendering/SkeletalMeshModel.h"
+#include "Serialization/Archive.h"
+#include "SlotBase.h"
+#include "Styling/AppStyle.h"
+#include "Styling/ISlateStyle.h"
+#include "Templates/Casts.h"
+#include "Types/SlateEnums.h"
+#include "UObject/NameTypes.h"
+#include "UObject/UnrealType.h"
+#include "Widgets/Images/SImage.h"
+#include "Widgets/Input/SCheckBox.h"
+#include "Widgets/SBoxPanel.h"
+#include "Widgets/SOverlay.h"
+
+class UCustomizableObjectNodeRemapPinsByName;
+class UObject;
+struct FSlateBrush;
 
 
 #define LOCTEXT_NAMESPACE "CustomizableObjectEditor"
diff --git a/Engine/Plugins/Experimental/Mutable/Source/CustomizableObjectEditor/Private/MuCOE/Nodes/CustomizableObjectNodeSkeletalMeshDetails.cpp b/Engine/Plugins/Experimental/Mutable/Source/CustomizableObjectEditor/Private/MuCOE/Nodes/CustomizableObjectNodeSkeletalMeshDetails.cpp
--- a/Engine/Plugins/Experimental/Mutable/Source/CustomizableObjectEditor/Private/MuCOE/Nodes/CustomizableObjectNodeSkeletalMeshDetails.cpp
+++ b/Engine/Plugins/Experimental/Mutable/Source/CustomizableObjectEditor/Private/MuCOE/Nodes/CustomizableObjectNodeSkeletalMeshDetails.cpp
@@ -1,16 +1,23 @@
 // Copyright Epic Games, Inc. All Rights Reserved.
 
 #include "MuCOE/Nodes/CustomizableObjectNodeSkeletalMeshDetails.h"
-#include "MuCOE/CustomizableObjectEditorModule.h"
-#include "PropertyCustomizationHelpers.h"
 
-#include "DetailLayoutBuilder.h"
+#include "Containers/Array.h"
+#include "Delegates/Delegate.h"
 #include "DetailCategoryBuilder.h"
-#include "LevelEditor.h"
-#include "LevelEditorActions.h"
+#include "DetailLayoutBuilder.h"
+#include "DetailWidgetRow.h"
+#include "IDetailsView.h"
+#include "Internationalization/Internationalization.h"
+#include "Misc/AssertionMacros.h"
+#include "MuCOE/Nodes/CustomizableObjectNodeSkeletalMesh.h"
 #include "MuCOE/PinViewer/SPinViewer.h"
-
 #include "MuCOE/SCustomizableObjectNodeSkeletalMeshRTMorphSelector.h"
+#include "PropertyHandle.h"
+#include "Templates/Casts.h"
+#include "UObject/WeakObjectPtr.h"
+#include "UObject/WeakObjectPtrTemplates.h"
+#include "Widgets/DeclarativeSyntaxSupport.h"
 
 #define LOCTEXT_NAMESPACE "CustomizableObjectNodeMaterialDetails"
 
diff --git a/Engine/Plugins/Experimental/Mutable/Source/CustomizableObjectEditor/Private/MuCOE/Nodes/CustomizableObjectNodeSkeletalMeshDetails.h b/Engine/Plugins/Experimental/Mutable/Source/CustomizableObjectEditor/Private/MuCOE/Nodes/CustomizableObjectNodeSkeletalMeshDetails.h
--- a/Engine/Plugins/Experimental/Mutable/Source/CustomizableObjectEditor/Private/MuCOE/Nodes/CustomizableObjectNodeSkeletalMeshDetails.h
+++ b/Engine/Plugins/Experimental/Mutable/Source/CustomizableObjectEditor/Private/MuCOE/Nodes/CustomizableObjectNodeSkeletalMeshDetails.h
@@ -4,8 +4,8 @@
 
 #include "IDetailCustomization.h"
 #include "Templates/SharedPointer.h"
-#include "Types/SlateEnums.h"
-#include "UObject/WeakObjectPtr.h"
+
+class IDetailLayoutBuilder;
 
 
 class FCustomizableObjectNodeSkeletalMeshDetails : public IDetailCustomization
diff --git a/Engine/Plugins/Experimental/Mutable/Source/CustomizableObjectEditor/Private/MuCOE/Nodes/CustomizableObjectNodeStaticMesh.cpp b/Engine/Plugins/Experimental/Mutable/Source/CustomizableObjectEditor/Private/MuCOE/Nodes/CustomizableObjectNodeStaticMesh.cpp
--- a/Engine/Plugins/Experimental/Mutable/Source/CustomizableObjectEditor/Private/MuCOE/Nodes/CustomizableObjectNodeStaticMesh.cpp
+++ b/Engine/Plugins/Experimental/Mutable/Source/CustomizableObjectEditor/Private/MuCOE/Nodes/CustomizableObjectNodeStaticMesh.cpp
@@ -1,17 +1,57 @@
 // Copyright Epic Games, Inc. All Rights Reserved.
 
 #include "MuCOE/Nodes/CustomizableObjectNodeStaticMesh.h"
-#include "MuCOE/CustomizableObjectEditorModule.h"
-#include "MuCOE/UnrealEditorPortabilityHelpers.h"
 
+#include "AssetThumbnail.h"
+#include "Containers/EnumAsByte.h"
+#include "EdGraph/EdGraph.h"
 #include "Engine/StaticMesh.h"
-#include "Materials/Material.h"
+#include "Engine/Texture.h"
+#include "Engine/Texture2D.h"
+#include "GenericPlatform/ICursor.h"
+#include "HAL/PlatformCrt.h"
 #include "ISinglePropertyView.h"
-#include "MuCOE/CustomizableObjectEditorStyle.h"
+#include "Internationalization/Internationalization.h"
+#include "Layout/Margin.h"
+#include "MaterialTypes.h"
+#include "Materials/Material.h"
+#include "Materials/MaterialInterface.h"
+#include "Misc/AssertionMacros.h"
+#include "Misc/Attribute.h"
+#include "Misc/Guid.h"
+#include "Misc/Optional.h"
+#include "Modules/ModuleManager.h"
+#include "MuCO/CustomizableObjectCustomVersion.h"
 #include "MuCOE/CustomizableObjectEditor.h"
+#include "MuCOE/CustomizableObjectEditorStyle.h"
+#include "MuCOE/CustomizableObjectLayout.h"
+#include "MuCOE/EdGraphSchema_CustomizableObject.h"
 #include "MuCOE/GraphTraversal.h"
 #include "MuCOE/MutableUtils.h"
+#include "MuCOE/Nodes/CustomizableObjectNodeLayoutBlocks.h"
+#include "MuCOE/RemapPins/CustomizableObjectNodeRemapPins.h"
 #include "MuCOE/RemapPins/CustomizableObjectNodeRemapPinsByNameDefaultPin.h"
+#include "MuCOE/UnrealEditorPortabilityHelpers.h"
+#include "PropertyEditorModule.h"
+#include "Serialization/Archive.h"
+#include "SlotBase.h"
+#include "StaticMeshResources.h"
+#include "Styling/AppStyle.h"
+#include "Styling/ISlateStyle.h"
+#include "Templates/Casts.h"
+#include "Types/SlateEnums.h"
+#include "UObject/NameTypes.h"
+#include "UObject/UnrealType.h"
+#include "Widgets/Images/SImage.h"
+#include "Widgets/Input/SCheckBox.h"
+#include "Widgets/SBoxPanel.h"
+#include "Widgets/SOverlay.h"
+
+#include <functional>
+
+class UCustomizableObjectNodeRemapPinsByName;
+class UObject;
+struct FSlateBrush;
 
 #define LOCTEXT_NAMESPACE "CustomizableObjectEditor"
 
diff --git a/Engine/Plugins/Experimental/Mutable/Source/CustomizableObjectEditor/Private/MuCOE/Nodes/CustomizableObjectNodeTable.cpp b/Engine/Plugins/Experimental/Mutable/Source/CustomizableObjectEditor/Private/MuCOE/Nodes/CustomizableObjectNodeTable.cpp
--- a/Engine/Plugins/Experimental/Mutable/Source/CustomizableObjectEditor/Private/MuCOE/Nodes/CustomizableObjectNodeTable.cpp
+++ b/Engine/Plugins/Experimental/Mutable/Source/CustomizableObjectEditor/Private/MuCOE/Nodes/CustomizableObjectNodeTable.cpp
@@ -1,18 +1,45 @@
 // Copyright Epic Games, Inc. All Rights Reserved.
 
 #include "MuCOE/Nodes/CustomizableObjectNodeTable.h"
-#include "MuCOE/CustomizableObjectEditorModule.h"
-#include "MuCOE/CustomizableObjectEditor.h"
-#include "MuCOE/SCustomizableObjectNodeLayoutBlocksEditor.h"
-#include "MuCOE/CustomizableObjectEditorLogger.h"
-#include "GameplayTagContainer.h"
 
+#include "Containers/EnumAsByte.h"
+#include "Containers/IndirectArray.h"
+#include "CoreTypes.h"
 #include "DataTableUtils.h"
+#include "Delegates/Delegate.h"
+#include "EdGraph/EdGraphPin.h"
+#include "Engine/SkeletalMesh.h"
+#include "Engine/StaticMesh.h"
+#include "Engine/Texture.h"
+#include "Engine/Texture2D.h"
+#include "HAL/PlatformCrt.h"
+#include "Internationalization/Internationalization.h"
+#include "Logging/TokenizedMessage.h"
+#include "MaterialTypes.h"
+#include "Materials/Material.h"
+#include "Materials/MaterialInstance.h"
+#include "Materials/MaterialInterface.h"
+#include "Misc/AssertionMacros.h"
+#include "Misc/Guid.h"
+#include "MuCO/UnrealPortabilityHelpers.h"
+#include "MuCOE/CustomizableObjectEditor.h"
+#include "MuCOE/CustomizableObjectEditorLogger.h"
+#include "MuCOE/CustomizableObjectLayout.h"
+#include "MuCOE/EdGraphSchema_CustomizableObject.h"
 #include "MuCOE/GraphTraversal.h"
 #include "MuCOE/MutableUtils.h"
-#include "Materials/MaterialInstance.h"
-
-#include "MuCOE/UnrealEditorPortabilityHelpers.h"
+#include "MuCOE/Nodes/CustomizableObjectNodeMaterial.h"
+#include "MuCOE/SCustomizableObjectNodeLayoutBlocksEditor.h"
+#include "Rendering/SkeletalMeshLODModel.h"
+#include "Rendering/SkeletalMeshModel.h"
+#include "StaticMeshResources.h"
+#include "Templates/Casts.h"
+#include "Templates/SharedPointer.h"
+#include "Templates/Tuple.h"
+#include "UObject/Object.h"
+
+class ICustomizableObjectEditor;
+class UCustomizableObjectNodeRemapPins;
 
 #define LOCTEXT_NAMESPACE "CustomizableObjectEditor"
 
diff --git a/Engine/Plugins/Experimental/Mutable/Source/CustomizableObjectEditor/Private/MuCOE/Nodes/CustomizableObjectNodeTableDetails.cpp b/Engine/Plugins/Experimental/Mutable/Source/CustomizableObjectEditor/Private/MuCOE/Nodes/CustomizableObjectNodeTableDetails.cpp
--- a/Engine/Plugins/Experimental/Mutable/Source/CustomizableObjectEditor/Private/MuCOE/Nodes/CustomizableObjectNodeTableDetails.cpp
+++ b/Engine/Plugins/Experimental/Mutable/Source/CustomizableObjectEditor/Private/MuCOE/Nodes/CustomizableObjectNodeTableDetails.cpp
@@ -1,19 +1,42 @@
 // Copyright Epic Games, Inc. All Rights Reserved.
 
 #include "MuCOE/Nodes/CustomizableObjectNodeTableDetails.h"
-#include "MuCOE/Nodes/CustomizableObjectNode.h"
-#include "MuCOE/Nodes/CustomizableObjectNodeTable.h"
 
-#include "DetailLayoutBuilder.h"
+#include "Animation/AnimInstance.h"
+#include "Containers/UnrealString.h"
+#include "DataTableUtils.h"
 #include "DetailCategoryBuilder.h"
+#include "DetailLayoutBuilder.h"
 #include "DetailWidgetRow.h"
-
+#include "EdGraph/EdGraphPin.h"
 #include "Engine/DataTable.h"
+#include "Engine/SkeletalMesh.h"
 #include "Engine/StaticMesh.h"
-#include "Animation/AnimInstance.h"
 #include "GameplayTagContainer.h"
-
+#include "HAL/Platform.h"
+#include "HAL/PlatformCrt.h"
+#include "IDetailsView.h"
+#include "Internationalization/Internationalization.h"
+#include "Internationalization/Text.h"
+#include "Layout/Margin.h"
+#include "Layout/Visibility.h"
+#include "Misc/Attribute.h"
+#include "MuCOE/CustomizableObjectLayout.h"
+#include "MuCOE/Nodes/CustomizableObjectNodeTable.h"
+#include "SlotBase.h"
 #include "Styling/AppStyle.h"
+#include "Templates/Casts.h"
+#include "UObject/Class.h"
+#include "UObject/Field.h"
+#include "UObject/ObjectPtr.h"
+#include "UObject/UnrealType.h"
+#include "UObject/WeakObjectPtr.h"
+#include "UObject/WeakObjectPtrTemplates.h"
+#include "Widgets/DeclarativeSyntaxSupport.h"
+#include "Widgets/Input/STextComboBox.h"
+#include "Widgets/Layout/SBorder.h"
+#include "Widgets/SBoxPanel.h"
+#include "Widgets/Text/STextBlock.h"
 
 #define LOCTEXT_NAMESPACE "CustomizableObjectDetails"
 
diff --git a/Engine/Plugins/Experimental/Mutable/Source/CustomizableObjectEditor/Private/MuCOE/Nodes/CustomizableObjectNodeTableDetails.h b/Engine/Plugins/Experimental/Mutable/Source/CustomizableObjectEditor/Private/MuCOE/Nodes/CustomizableObjectNodeTableDetails.h
--- a/Engine/Plugins/Experimental/Mutable/Source/CustomizableObjectEditor/Private/MuCOE/Nodes/CustomizableObjectNodeTableDetails.h
+++ b/Engine/Plugins/Experimental/Mutable/Source/CustomizableObjectEditor/Private/MuCOE/Nodes/CustomizableObjectNodeTableDetails.h
@@ -2,12 +2,14 @@
 
 #pragma once
 
+#include "Containers/Array.h"
 #include "IDetailCustomization.h"
 #include "Templates/SharedPointer.h"
-
 #include "Types/SlateEnums.h"
-#include "Widgets/Text/STextBlock.h"
-#include "Widgets/Input/STextComboBox.h"
+
+class FString;
+class IDetailLayoutBuilder;
+class STextComboBox;
 
 /** Copy Material node details panel. Hides all properties from the inheret Material node. */
 class FCustomizableObjectNodeTableDetails : public IDetailCustomization
diff --git a/Engine/Plugins/Experimental/Mutable/Source/CustomizableObjectEditor/Private/MuCOE/Nodes/CustomizableObjectNodeTexture.cpp b/Engine/Plugins/Experimental/Mutable/Source/CustomizableObjectEditor/Private/MuCOE/Nodes/CustomizableObjectNodeTexture.cpp
--- a/Engine/Plugins/Experimental/Mutable/Source/CustomizableObjectEditor/Private/MuCOE/Nodes/CustomizableObjectNodeTexture.cpp
+++ b/Engine/Plugins/Experimental/Mutable/Source/CustomizableObjectEditor/Private/MuCOE/Nodes/CustomizableObjectNodeTexture.cpp
@@ -1,11 +1,36 @@
 // Copyright Epic Games, Inc. All Rights Reserved.
 
 #include "MuCOE/Nodes/CustomizableObjectNodeTexture.h"
-#include "MuCOE/CustomizableObjectEditorModule.h"
+
+#include "Containers/EnumAsByte.h"
+#include "Containers/UnrealString.h"
+#include "EdGraph/EdGraphPin.h"
+#include "Engine/Texture2D.h"
+#include "GenericPlatform/ICursor.h"
+#include "HAL/Platform.h"
 #include "ISinglePropertyView.h"
-#include "Widgets/Input/SCheckBox.h"
-#include "Widgets/Images/SImage.h"
+#include "Internationalization/Internationalization.h"
+#include "Layout/Margin.h"
+#include "Math/Vector2D.h"
+#include "Misc/Attribute.h"
+#include "Misc/Optional.h"
+#include "Modules/ModuleManager.h"
 #include "MuCOE/CustomizableObjectEditorStyle.h"
+#include "MuCOE/EdGraphSchema_CustomizableObject.h"
+#include "PropertyEditorModule.h"
+#include "SlotBase.h"
+#include "Styling/AppStyle.h"
+#include "Styling/ISlateStyle.h"
+#include "Templates/Casts.h"
+#include "Types/SlateEnums.h"
+#include "UObject/NameTypes.h"
+#include "Widgets/Images/SImage.h"
+#include "Widgets/Input/SCheckBox.h"
+#include "Widgets/SBoxPanel.h"
+#include "Widgets/SOverlay.h"
+
+class UCustomizableObjectNodeRemapPins;
+struct FGeometry;
 
 #define LOCTEXT_NAMESPACE "CustomizableObjectEditor"
 
diff --git a/Engine/Plugins/Experimental/Mutable/Source/CustomizableObjectEditor/Private/MuCOE/Nodes/CustomizableObjectNodeTextureBinarise.cpp b/Engine/Plugins/Experimental/Mutable/Source/CustomizableObjectEditor/Private/MuCOE/Nodes/CustomizableObjectNodeTextureBinarise.cpp
--- a/Engine/Plugins/Experimental/Mutable/Source/CustomizableObjectEditor/Private/MuCOE/Nodes/CustomizableObjectNodeTextureBinarise.cpp
+++ b/Engine/Plugins/Experimental/Mutable/Source/CustomizableObjectEditor/Private/MuCOE/Nodes/CustomizableObjectNodeTextureBinarise.cpp
@@ -1,7 +1,15 @@
 // Copyright Epic Games, Inc. All Rights Reserved.
 
 #include "MuCOE/Nodes/CustomizableObjectNodeTextureBinarise.h"
-#include "MuCOE/CustomizableObjectEditorModule.h"
+
+#include "Containers/UnrealString.h"
+#include "Internationalization/Internationalization.h"
+#include "MuCO/CustomizableObjectCustomVersion.h"
+#include "MuCOE/EdGraphSchema_CustomizableObject.h"
+#include "Serialization/Archive.h"
+#include "UObject/NameTypes.h"
+
+class UCustomizableObjectNodeRemapPins;
 
 #define LOCTEXT_NAMESPACE "CustomizableObjectEditor"
 
diff --git a/Engine/Plugins/Experimental/Mutable/Source/CustomizableObjectEditor/Private/MuCOE/Nodes/CustomizableObjectNodeTextureColourMap.cpp b/Engine/Plugins/Experimental/Mutable/Source/CustomizableObjectEditor/Private/MuCOE/Nodes/CustomizableObjectNodeTextureColourMap.cpp
--- a/Engine/Plugins/Experimental/Mutable/Source/CustomizableObjectEditor/Private/MuCOE/Nodes/CustomizableObjectNodeTextureColourMap.cpp
+++ b/Engine/Plugins/Experimental/Mutable/Source/CustomizableObjectEditor/Private/MuCOE/Nodes/CustomizableObjectNodeTextureColourMap.cpp
@@ -1,7 +1,14 @@
 // Copyright Epic Games, Inc. All Rights Reserved.
 
 #include "MuCOE/Nodes/CustomizableObjectNodeTextureColourMap.h"
-#include "MuCOE/CustomizableObjectEditorModule.h"
+
+#include "Containers/UnrealString.h"
+#include "EdGraph/EdGraphPin.h"
+#include "Internationalization/Internationalization.h"
+#include "MuCOE/EdGraphSchema_CustomizableObject.h"
+#include "UObject/NameTypes.h"
+
+class UCustomizableObjectNodeRemapPins;
 
 #define LOCTEXT_NAMESPACE "CustomizableObjectEditor"
 
diff --git a/Engine/Plugins/Experimental/Mutable/Source/CustomizableObjectEditor/Private/MuCOE/Nodes/CustomizableObjectNodeTextureFromChannels.cpp b/Engine/Plugins/Experimental/Mutable/Source/CustomizableObjectEditor/Private/MuCOE/Nodes/CustomizableObjectNodeTextureFromChannels.cpp
--- a/Engine/Plugins/Experimental/Mutable/Source/CustomizableObjectEditor/Private/MuCOE/Nodes/CustomizableObjectNodeTextureFromChannels.cpp
+++ b/Engine/Plugins/Experimental/Mutable/Source/CustomizableObjectEditor/Private/MuCOE/Nodes/CustomizableObjectNodeTextureFromChannels.cpp
@@ -1,7 +1,14 @@
 // Copyright Epic Games, Inc. All Rights Reserved.
 
 #include "MuCOE/Nodes/CustomizableObjectNodeTextureFromChannels.h"
-#include "MuCOE/CustomizableObjectEditorModule.h"
+
+#include "Containers/UnrealString.h"
+#include "EdGraph/EdGraphPin.h"
+#include "Internationalization/Internationalization.h"
+#include "MuCOE/EdGraphSchema_CustomizableObject.h"
+#include "UObject/NameTypes.h"
+
+class UCustomizableObjectNodeRemapPins;
 
 #define LOCTEXT_NAMESPACE "CustomizableObjectEditor"
 
diff --git a/Engine/Plugins/Experimental/Mutable/Source/CustomizableObjectEditor/Private/MuCOE/Nodes/CustomizableObjectNodeTextureInterpolate.cpp b/Engine/Plugins/Experimental/Mutable/Source/CustomizableObjectEditor/Private/MuCOE/Nodes/CustomizableObjectNodeTextureInterpolate.cpp
--- a/Engine/Plugins/Experimental/Mutable/Source/CustomizableObjectEditor/Private/MuCOE/Nodes/CustomizableObjectNodeTextureInterpolate.cpp
+++ b/Engine/Plugins/Experimental/Mutable/Source/CustomizableObjectEditor/Private/MuCOE/Nodes/CustomizableObjectNodeTextureInterpolate.cpp
@@ -1,7 +1,13 @@
 // Copyright Epic Games, Inc. All Rights Reserved.
 
 #include "MuCOE/Nodes/CustomizableObjectNodeTextureInterpolate.h"
-#include "MuCOE/CustomizableObjectEditorModule.h"
+
+#include "Internationalization/Internationalization.h"
+#include "MuCOE/EdGraphSchema_CustomizableObject.h"
+#include "UObject/NameTypes.h"
+#include "UObject/UnrealType.h"
+
+class UCustomizableObjectNodeRemapPins;
 
 #define LOCTEXT_NAMESPACE "CustomizableObjectEditor"
 
diff --git a/Engine/Plugins/Experimental/Mutable/Source/CustomizableObjectEditor/Private/MuCOE/Nodes/CustomizableObjectNodeTextureInvert.cpp b/Engine/Plugins/Experimental/Mutable/Source/CustomizableObjectEditor/Private/MuCOE/Nodes/CustomizableObjectNodeTextureInvert.cpp
--- a/Engine/Plugins/Experimental/Mutable/Source/CustomizableObjectEditor/Private/MuCOE/Nodes/CustomizableObjectNodeTextureInvert.cpp
+++ b/Engine/Plugins/Experimental/Mutable/Source/CustomizableObjectEditor/Private/MuCOE/Nodes/CustomizableObjectNodeTextureInvert.cpp
@@ -1,7 +1,16 @@
 // Copyright Epic Games, Inc. All Rights Reserved.
 
 #include "MuCOE/Nodes/CustomizableObjectNodeTextureInvert.h"
-#include "MuCOE/CustomizableObjectEditorModule.h"
+
+#include "Containers/UnrealString.h"
+#include "HAL/Platform.h"
+#include "Internationalization/Internationalization.h"
+#include "MuCO/CustomizableObjectCustomVersion.h"
+#include "MuCOE/EdGraphSchema_CustomizableObject.h"
+#include "Serialization/Archive.h"
+#include "UObject/NameTypes.h"
+
+class UCustomizableObjectNodeRemapPins;
 
 #define LOCTEXT_NAMESPACE "CustomizableObjectEditor"
 
diff --git a/Engine/Plugins/Experimental/Mutable/Source/CustomizableObjectEditor/Private/MuCOE/Nodes/CustomizableObjectNodeTextureLayer.cpp b/Engine/Plugins/Experimental/Mutable/Source/CustomizableObjectEditor/Private/MuCOE/Nodes/CustomizableObjectNodeTextureLayer.cpp
--- a/Engine/Plugins/Experimental/Mutable/Source/CustomizableObjectEditor/Private/MuCOE/Nodes/CustomizableObjectNodeTextureLayer.cpp
+++ b/Engine/Plugins/Experimental/Mutable/Source/CustomizableObjectEditor/Private/MuCOE/Nodes/CustomizableObjectNodeTextureLayer.cpp
@@ -1,8 +1,18 @@
 // Copyright Epic Games, Inc. All Rights Reserved.
 
 #include "MuCOE/Nodes/CustomizableObjectNodeTextureLayer.h"
-#include "MuCOE/CustomizableObjectEditorModule.h"
+
+#include "HAL/PlatformCrt.h"
+#include "Internationalization/Internationalization.h"
+#include "MuCO/CustomizableObjectCustomVersion.h"
+#include "MuCOE/EdGraphSchema_CustomizableObject.h"
 #include "MuCOE/UnrealEditorPortabilityHelpers.h"
+#include "Serialization/Archive.h"
+#include "UObject/NameTypes.h"
+#include "UObject/UnrealType.h"
+
+class UCustomizableObjectNodeRemapPins;
+class UEdGraph;
 
 #define LOCTEXT_NAMESPACE "CustomizableObjectEditor"
 
diff --git a/Engine/Plugins/Experimental/Mutable/Source/CustomizableObjectEditor/Private/MuCOE/Nodes/CustomizableObjectNodeTextureParameter.cpp b/Engine/Plugins/Experimental/Mutable/Source/CustomizableObjectEditor/Private/MuCOE/Nodes/CustomizableObjectNodeTextureParameter.cpp
--- a/Engine/Plugins/Experimental/Mutable/Source/CustomizableObjectEditor/Private/MuCOE/Nodes/CustomizableObjectNodeTextureParameter.cpp
+++ b/Engine/Plugins/Experimental/Mutable/Source/CustomizableObjectEditor/Private/MuCOE/Nodes/CustomizableObjectNodeTextureParameter.cpp
@@ -1,7 +1,14 @@
 // Copyright Epic Games, Inc. All Rights Reserved.
 
 #include "MuCOE/Nodes/CustomizableObjectNodeTextureParameter.h"
-#include "MuCOE/CustomizableObjectEditorModule.h"
+
+#include "EdGraph/EdGraphPin.h"
+#include "HAL/Platform.h"
+#include "Internationalization/Internationalization.h"
+#include "MuCOE/EdGraphSchema_CustomizableObject.h"
+#include "UObject/NameTypes.h"
+
+class UCustomizableObjectNodeRemapPins;
 
 #define LOCTEXT_NAMESPACE "CustomizableObjectEditor"
 
diff --git a/Engine/Plugins/Experimental/Mutable/Source/CustomizableObjectEditor/Private/MuCOE/Nodes/CustomizableObjectNodeTextureProject.cpp b/Engine/Plugins/Experimental/Mutable/Source/CustomizableObjectEditor/Private/MuCOE/Nodes/CustomizableObjectNodeTextureProject.cpp
--- a/Engine/Plugins/Experimental/Mutable/Source/CustomizableObjectEditor/Private/MuCOE/Nodes/CustomizableObjectNodeTextureProject.cpp
+++ b/Engine/Plugins/Experimental/Mutable/Source/CustomizableObjectEditor/Private/MuCOE/Nodes/CustomizableObjectNodeTextureProject.cpp
@@ -1,7 +1,19 @@
 // Copyright Epic Games, Inc. All Rights Reserved.
 
 #include "MuCOE/Nodes/CustomizableObjectNodeTextureProject.h"
-#include "MuCOE/CustomizableObjectEditorModule.h"
+
+#include "Containers/EnumAsByte.h"
+#include "Containers/UnrealString.h"
+#include "EdGraph/EdGraphPin.h"
+#include "HAL/PlatformCrt.h"
+#include "Internationalization/Internationalization.h"
+#include "MuCO/CustomizableObjectCustomVersion.h"
+#include "MuCOE/EdGraphSchema_CustomizableObject.h"
+#include "Serialization/Archive.h"
+#include "UObject/NameTypes.h"
+#include "UObject/UnrealType.h"
+
+class UCustomizableObjectNodeRemapPins;
 
 #define LOCTEXT_NAMESPACE "CustomizableObjectEditor"
 
diff --git a/Engine/Plugins/Experimental/Mutable/Source/CustomizableObjectEditor/Private/MuCOE/Nodes/CustomizableObjectNodeTextureSample.cpp b/Engine/Plugins/Experimental/Mutable/Source/CustomizableObjectEditor/Private/MuCOE/Nodes/CustomizableObjectNodeTextureSample.cpp
--- a/Engine/Plugins/Experimental/Mutable/Source/CustomizableObjectEditor/Private/MuCOE/Nodes/CustomizableObjectNodeTextureSample.cpp
+++ b/Engine/Plugins/Experimental/Mutable/Source/CustomizableObjectEditor/Private/MuCOE/Nodes/CustomizableObjectNodeTextureSample.cpp
@@ -1,7 +1,15 @@
 // Copyright Epic Games, Inc. All Rights Reserved.
 
 #include "MuCOE/Nodes/CustomizableObjectNodeTextureSample.h"
-#include "MuCOE/CustomizableObjectEditorModule.h"
+
+#include "Containers/UnrealString.h"
+#include "EdGraph/EdGraphPin.h"
+#include "Internationalization/Internationalization.h"
+#include "MuCOE/EdGraphSchema_CustomizableObject.h"
+#include "UObject/NameTypes.h"
+
+class UCustomizableObjectNodeRemapPins;
+struct FPropertyChangedEvent;
 
 #define LOCTEXT_NAMESPACE "CustomizableObjectEditor"
 
diff --git a/Engine/Plugins/Experimental/Mutable/Source/CustomizableObjectEditor/Private/MuCOE/Nodes/CustomizableObjectNodeTextureToChannels.cpp b/Engine/Plugins/Experimental/Mutable/Source/CustomizableObjectEditor/Private/MuCOE/Nodes/CustomizableObjectNodeTextureToChannels.cpp
--- a/Engine/Plugins/Experimental/Mutable/Source/CustomizableObjectEditor/Private/MuCOE/Nodes/CustomizableObjectNodeTextureToChannels.cpp
+++ b/Engine/Plugins/Experimental/Mutable/Source/CustomizableObjectEditor/Private/MuCOE/Nodes/CustomizableObjectNodeTextureToChannels.cpp
@@ -1,7 +1,16 @@
 // Copyright Epic Games, Inc. All Rights Reserved.
 
 #include "MuCOE/Nodes/CustomizableObjectNodeTextureToChannels.h"
-#include "MuCOE/CustomizableObjectEditorModule.h"
+
+#include "Containers/UnrealString.h"
+#include "Internationalization/Internationalization.h"
+#include "MuCO/CustomizableObjectCustomVersion.h"
+#include "MuCOE/EdGraphSchema_CustomizableObject.h"
+#include "Serialization/Archive.h"
+#include "UObject/NameTypes.h"
+
+class UCustomizableObjectNodeRemapPins;
+struct FPropertyChangedEvent;
 
 #define LOCTEXT_NAMESPACE "CustomizableObjectEditor"
 
diff --git a/Engine/Plugins/Experimental/Mutable/Source/CustomizableObjectEditor/Private/MuCOE/Nodes/CustomizableObjectNodeTextureTransform.cpp b/Engine/Plugins/Experimental/Mutable/Source/CustomizableObjectEditor/Private/MuCOE/Nodes/CustomizableObjectNodeTextureTransform.cpp
--- a/Engine/Plugins/Experimental/Mutable/Source/CustomizableObjectEditor/Private/MuCOE/Nodes/CustomizableObjectNodeTextureTransform.cpp
+++ b/Engine/Plugins/Experimental/Mutable/Source/CustomizableObjectEditor/Private/MuCOE/Nodes/CustomizableObjectNodeTextureTransform.cpp
@@ -1,7 +1,15 @@
 // Copyright Epic Games, Inc. All Rights Reserved.
 
 #include "MuCOE/Nodes/CustomizableObjectNodeTextureTransform.h"
-#include "MuCOE/CustomizableObjectEditorModule.h"
+
+#include "Containers/UnrealString.h"
+#include "Internationalization/Internationalization.h"
+#include "MuCO/CustomizableObjectCustomVersion.h"
+#include "MuCOE/EdGraphSchema_CustomizableObject.h"
+#include "Serialization/Archive.h"
+#include "UObject/NameTypes.h"
+
+class UCustomizableObjectNodeRemapPins;
 
 #define LOCTEXT_NAMESPACE "CustomizableObjectEditor"
 
diff --git a/Engine/Plugins/Experimental/Mutable/Source/CustomizableObjectEditor/Private/MuCOE/Nodes/CustomizableObjectNodeUseMaterial.cpp b/Engine/Plugins/Experimental/Mutable/Source/CustomizableObjectEditor/Private/MuCOE/Nodes/CustomizableObjectNodeUseMaterial.cpp
--- a/Engine/Plugins/Experimental/Mutable/Source/CustomizableObjectEditor/Private/MuCOE/Nodes/CustomizableObjectNodeUseMaterial.cpp
+++ b/Engine/Plugins/Experimental/Mutable/Source/CustomizableObjectEditor/Private/MuCOE/Nodes/CustomizableObjectNodeUseMaterial.cpp
@@ -2,11 +2,21 @@
 
 #include "MuCOE/Nodes/CustomizableObjectNodeUseMaterial.h"
 
+#include "Delegates/Delegate.h"
+#include "EdGraph/EdGraph.h"
+#include "EdGraph/EdGraphPin.h"
+#include "Misc/Guid.h"
 #include "MuCOE/CustomizableObjectPin.h"
 #include "MuCOE/GraphTraversal.h"
 #include "MuCOE/ICustomizableObjectEditor.h"
-
+#include "MuCOE/Nodes/CustomizableObjectNode.h"
 #include "MuCOE/Nodes/CustomizableObjectNodeMaterial.h"
+#include "MuCOE/Nodes/CustomizableObjectNodeParentedMaterial.h"
+#include "Templates/ChooseClass.h"
+#include "Templates/SharedPointer.h"
+#include "Templates/Tuple.h"
+
+class UCustomizableObject;
 
 
 bool FCustomizableObjectNodeUseMaterial::IsNodeOutDatedAndNeedsRefreshWork()
diff --git a/Engine/Plugins/Experimental/Mutable/Source/CustomizableObjectEditor/Private/MuCOE/Nodes/SCustomizableObjectNodeMaterial.cpp b/Engine/Plugins/Experimental/Mutable/Source/CustomizableObjectEditor/Private/MuCOE/Nodes/SCustomizableObjectNodeMaterial.cpp
--- a/Engine/Plugins/Experimental/Mutable/Source/CustomizableObjectEditor/Private/MuCOE/Nodes/SCustomizableObjectNodeMaterial.cpp
+++ b/Engine/Plugins/Experimental/Mutable/Source/CustomizableObjectEditor/Private/MuCOE/Nodes/SCustomizableObjectNodeMaterial.cpp
@@ -1,8 +1,15 @@
 // Copyright Epic Games, Inc. All Rights Reserved.
 
 #include "MuCOE/Nodes/SCustomizableObjectNodeMaterial.h"
+
+#include "Containers/EnumAsByte.h"
+#include "EdGraph/EdGraphNode.h"
+#include "EdGraph/EdGraphPin.h"
 #include "MuCOE/EdGraphSchema_CustomizableObject.h"
 #include "MuCOE/Nodes/SCustomizableObjectNodeMaterialPinImage.h"
+#include "UObject/NameTypes.h"
+
+class SGraphPin;
 
 
 void SCustomizableObjectNodeMaterial::Construct(const FArguments& InArgs, UEdGraphNode* InGraphNode)
diff --git a/Engine/Plugins/Experimental/Mutable/Source/CustomizableObjectEditor/Private/MuCOE/Nodes/SCustomizableObjectNodeMaterial.h b/Engine/Plugins/Experimental/Mutable/Source/CustomizableObjectEditor/Private/MuCOE/Nodes/SCustomizableObjectNodeMaterial.h
--- a/Engine/Plugins/Experimental/Mutable/Source/CustomizableObjectEditor/Private/MuCOE/Nodes/SCustomizableObjectNodeMaterial.h
+++ b/Engine/Plugins/Experimental/Mutable/Source/CustomizableObjectEditor/Private/MuCOE/Nodes/SCustomizableObjectNodeMaterial.h
@@ -3,6 +3,12 @@
 #pragma once
 
 #include "SGraphNode.h"
+#include "Templates/SharedPointer.h"
+#include "Widgets/DeclarativeSyntaxSupport.h"
+
+class SGraphPin;
+class UEdGraphNode;
+class UEdGraphPin;
 
 
 /** Custom widget for the Material node. */
diff --git a/Engine/Plugins/Experimental/Mutable/Source/CustomizableObjectEditor/Private/MuCOE/Nodes/SCustomizableObjectNodeMaterialPinImage.cpp b/Engine/Plugins/Experimental/Mutable/Source/CustomizableObjectEditor/Private/MuCOE/Nodes/SCustomizableObjectNodeMaterialPinImage.cpp
--- a/Engine/Plugins/Experimental/Mutable/Source/CustomizableObjectEditor/Private/MuCOE/Nodes/SCustomizableObjectNodeMaterialPinImage.cpp
+++ b/Engine/Plugins/Experimental/Mutable/Source/CustomizableObjectEditor/Private/MuCOE/Nodes/SCustomizableObjectNodeMaterialPinImage.cpp
@@ -2,13 +2,24 @@
 
 #include "MuCOE/Nodes/SCustomizableObjectNodeMaterialPinImage.h"
 
+#include "EdGraph/EdGraphPin.h"
+#include "HAL/Platform.h"
+#include "Internationalization/Internationalization.h"
+#include "Math/NumericLimits.h"
+#include "Misc/Attribute.h"
 #include "MuCOE/GraphTraversal.h"
-#include "MuCOE/Nodes/CustomizableObjectNode.h"
 #include "MuCOE/Nodes/CustomizableObjectNodeMaterial.h"
+#include "Styling/AppStyle.h"
+#include "Styling/SlateColor.h"
+#include "Templates/Casts.h"
 #include "Widgets/Input/SEditableTextBox.h"
 #include "Widgets/Layout/SWrapBox.h"
 #include "Widgets/SToolTip.h"
 
+class IToolTip;
+class SWidget;
+struct FSlateBrush;
+
 #define LOCTEXT_NAMESPACE "CustomizableObjectEditor"
 
 
diff --git a/Engine/Plugins/Experimental/Mutable/Source/CustomizableObjectEditor/Private/MuCOE/Nodes/SCustomizableObjectNodeMaterialPinImage.h b/Engine/Plugins/Experimental/Mutable/Source/CustomizableObjectEditor/Private/MuCOE/Nodes/SCustomizableObjectNodeMaterialPinImage.h
--- a/Engine/Plugins/Experimental/Mutable/Source/CustomizableObjectEditor/Private/MuCOE/Nodes/SCustomizableObjectNodeMaterialPinImage.h
+++ b/Engine/Plugins/Experimental/Mutable/Source/CustomizableObjectEditor/Private/MuCOE/Nodes/SCustomizableObjectNodeMaterialPinImage.h
@@ -2,7 +2,15 @@
 
 #pragma once
 
+#include "Internationalization/Text.h"
+#include "Layout/Visibility.h"
 #include "SGraphPin.h"
+#include "Templates/SharedPointer.h"
+#include "Widgets/DeclarativeSyntaxSupport.h"
+
+class SWidget;
+class UEdGraphPin;
+struct FSlateBrush;
 
 
 /** Material node custom Image pin. Allows to define a custom style to Image pins.
diff --git a/Engine/Plugins/Experimental/Mutable/Source/CustomizableObjectEditor/Private/MuCOE/Nodes/SPinViewerNodeMaterialPinImageDetails.cpp b/Engine/Plugins/Experimental/Mutable/Source/CustomizableObjectEditor/Private/MuCOE/Nodes/SPinViewerNodeMaterialPinImageDetails.cpp
--- a/Engine/Plugins/Experimental/Mutable/Source/CustomizableObjectEditor/Private/MuCOE/Nodes/SPinViewerNodeMaterialPinImageDetails.cpp
+++ b/Engine/Plugins/Experimental/Mutable/Source/CustomizableObjectEditor/Private/MuCOE/Nodes/SPinViewerNodeMaterialPinImageDetails.cpp
@@ -2,12 +2,25 @@
 
 #include "MuCOE/Nodes/SPinViewerNodeMaterialPinImageDetails.h"
 
+#include "Containers/BitArray.h"
+#include "Containers/Set.h"
+#include "Containers/SparseArray.h"
+#include "HAL/PlatformCrt.h"
 #include "ISinglePropertyView.h"
-#include "PerPlatformPropertyCustomization.h"
-#include "Widgets/Input/SNumericEntryBox.h"
-#include "MuCOE/SMutableTextComboBox.h"
+#include "Internationalization/Internationalization.h"
+#include "Internationalization/Text.h"
+#include "Layout/Visibility.h"
+#include "Math/UnrealMathSSE.h"
+#include "Misc/AssertionMacros.h"
+#include "Misc/Attribute.h"
+#include "Misc/Optional.h"
 #include "Modules/ModuleManager.h"
+#include "MuCOE/SMutableTextComboBox.h"
+#include "PropertyEditorModule.h"
 #include "ScopedTransaction.h"
+#include "Serialization/Archive.h"
+#include "Templates/UnrealTemplate.h"
+#include "Widgets/SWidget.h"
 
 #define LOCTEXT_NAMESPACE "CustomizableObjectEditor"
 
diff --git a/Engine/Plugins/Experimental/Mutable/Source/CustomizableObjectEditor/Private/MuCOE/Nodes/SPinViewerNodeMaterialPinImageDetails.h b/Engine/Plugins/Experimental/Mutable/Source/CustomizableObjectEditor/Private/MuCOE/Nodes/SPinViewerNodeMaterialPinImageDetails.h
--- a/Engine/Plugins/Experimental/Mutable/Source/CustomizableObjectEditor/Private/MuCOE/Nodes/SPinViewerNodeMaterialPinImageDetails.h
+++ b/Engine/Plugins/Experimental/Mutable/Source/CustomizableObjectEditor/Private/MuCOE/Nodes/SPinViewerNodeMaterialPinImageDetails.h
@@ -2,10 +2,20 @@
 
 #pragma once
 
+#include "Containers/Array.h"
+#include "Containers/Map.h"
+#include "Delegates/Delegate.h"
+#include "HAL/Platform.h"
 #include "MuCOE/Nodes/CustomizableObjectNodeMaterial.h"
 #include "MuCOE/PinViewer/SPinViewerPinDetails.h"
+#include "Templates/SharedPointer.h"
+#include "Types/SlateEnums.h"
+#include "Widgets/DeclarativeSyntaxSupport.h"
 #include "Widgets/Input/SSpinBox.h"
 
+class FString;
+class SWidget;
+
 /** UV Layout Mode UI Enum. */
 UENUM()
 enum class EUVLayoutMode
diff --git a/Engine/Plugins/Experimental/Mutable/Source/CustomizableObjectEditor/Private/MuCOE/PinViewer/SPinViewer.cpp b/Engine/Plugins/Experimental/Mutable/Source/CustomizableObjectEditor/Private/MuCOE/PinViewer/SPinViewer.cpp
--- a/Engine/Plugins/Experimental/Mutable/Source/CustomizableObjectEditor/Private/MuCOE/PinViewer/SPinViewer.cpp
+++ b/Engine/Plugins/Experimental/Mutable/Source/CustomizableObjectEditor/Private/MuCOE/PinViewer/SPinViewer.cpp
@@ -3,13 +3,36 @@
 #include "MuCOE/PinViewer/SPinViewer.h"
 
 #include "DetailCategoryBuilder.h"
+#include "DetailLayoutBuilder.h"
 #include "DetailWidgetRow.h"
-#include "Widgets/Input/SSearchBox.h"
+#include "EdGraph/EdGraphPin.h"
+#include "Framework/Views/ITypedTableView.h"
+#include "HAL/PlatformCrt.h"
+#include "IDetailsView.h"
+#include "Internationalization/Internationalization.h"
+#include "Layout/BasicLayoutWidgetSlot.h"
+#include "Layout/Children.h"
+#include "Misc/AssertionMacros.h"
+#include "Misc/Attribute.h"
+#include "MuCOE/Nodes/CustomizableObjectNode.h"
+#include "MuCOE/PinViewer/SPinViewerListRow.h"
+#include "ScopedTransaction.h"
+#include "SlotBase.h"
+#include "Templates/Casts.h"
+#include "Templates/Tuple.h"
+#include "Types/SlateEnums.h"
+#include "UObject/UObjectGlobals.h"
+#include "UObject/WeakObjectPtr.h"
+#include "UObject/WeakObjectPtrTemplates.h"
 #include "Widgets/Input/SButton.h"
+#include "Widgets/Input/SSearchBox.h"
 #include "Widgets/Layout/SScrollBox.h"
-#include "MuCOE/EdGraphSchema_CustomizableObject.h"
-#include "ScopedTransaction.h"
-#include "MuCOE/PinViewer/SPinViewerListRow.h"
+#include "Widgets/SBoxPanel.h"
+
+class ITableRow;
+class STableViewBase;
+class UObject;
+struct EVisibility;
 
 #define LOCTEXT_NAMESPACE "CustomizableObjectEditor"
 
diff --git a/Engine/Plugins/Experimental/Mutable/Source/CustomizableObjectEditor/Private/MuCOE/PinViewer/SPinViewer.h b/Engine/Plugins/Experimental/Mutable/Source/CustomizableObjectEditor/Private/MuCOE/PinViewer/SPinViewer.h
--- a/Engine/Plugins/Experimental/Mutable/Source/CustomizableObjectEditor/Private/MuCOE/PinViewer/SPinViewer.h
+++ b/Engine/Plugins/Experimental/Mutable/Source/CustomizableObjectEditor/Private/MuCOE/PinViewer/SPinViewer.h
@@ -2,10 +2,34 @@
 
 #pragma once
 
-#include "DetailLayoutBuilder.h"
+#include "Containers/Array.h"
+#include "Containers/BitArray.h"
+#include "Containers/Map.h"
+#include "Containers/Set.h"
+#include "Containers/SparseArray.h"
+#include "Containers/UnrealString.h"
+#include "Delegates/Delegate.h"
+#include "HAL/Platform.h"
+#include "Input/Reply.h"
+#include "Internationalization/Text.h"
+#include "Misc/Optional.h"
+#include "Templates/SharedPointer.h"
+#include "UObject/GCObject.h"
+#include "UObject/NameTypes.h"
+#include "Widgets/DeclarativeSyntaxSupport.h"
 #include "Widgets/SCompoundWidget.h"
-#include "MuCOE/Nodes/CustomizableObjectNode.h"
-#include "EdGraph/EdGraphPin.h"
+#include "Widgets/Views/SHeaderRow.h"
+#include "Widgets/Views/SListView.h"
+
+class FReferenceCollector;
+class IDetailLayoutBuilder;
+class ITableRow;
+class STableViewBase;
+class UCustomizableObjectNode;
+class UEdGraphPin;
+struct EVisibility;
+struct FEdGraphPinReference;
+struct FGuid;
 
 /**
  * Pin Viewer widget. Shows a list of pins a node contains.
diff --git a/Engine/Plugins/Experimental/Mutable/Source/CustomizableObjectEditor/Private/MuCOE/PinViewer/SPinViewerDetailRowIndent.cpp b/Engine/Plugins/Experimental/Mutable/Source/CustomizableObjectEditor/Private/MuCOE/PinViewer/SPinViewerDetailRowIndent.cpp
--- a/Engine/Plugins/Experimental/Mutable/Source/CustomizableObjectEditor/Private/MuCOE/PinViewer/SPinViewerDetailRowIndent.cpp
+++ b/Engine/Plugins/Experimental/Mutable/Source/CustomizableObjectEditor/Private/MuCOE/PinViewer/SPinViewerDetailRowIndent.cpp
@@ -2,8 +2,24 @@
 
 #include "MuCOE/PinViewer/SPinViewerDetailRowIndent.h"
 
+#include "Layout/Children.h"
+#include "Layout/Geometry.h"
+#include "Math/Vector2D.h"
+#include "Misc/Attribute.h"
 #include "MuCOE/Widgets/SMutableExpandableTableRow.h"
+#include "Rendering/DrawElements.h"
+#include "Rendering/RenderingCommon.h"
+#include "Styling/AppStyle.h"
+#include "Styling/ISlateStyle.h"
+#include "Styling/SlateColor.h"
+#include "Types/SlateStructs.h"
 #include "Widgets/Layout/SBox.h"
+#include "Widgets/SWidget.h"
+
+class FPaintArgs;
+class FSlateRect;
+class FWidgetStyle;
+struct FSlateBrush;
 
 
 void SPinViewerDetailRowIndent::Construct(const FArguments& InArgs, TSharedRef<SWidget> DetailsRow)
diff --git a/Engine/Plugins/Experimental/Mutable/Source/CustomizableObjectEditor/Private/MuCOE/PinViewer/SPinViewerDetailRowIndent.h b/Engine/Plugins/Experimental/Mutable/Source/CustomizableObjectEditor/Private/MuCOE/PinViewer/SPinViewerDetailRowIndent.h
--- a/Engine/Plugins/Experimental/Mutable/Source/CustomizableObjectEditor/Private/MuCOE/PinViewer/SPinViewerDetailRowIndent.h
+++ b/Engine/Plugins/Experimental/Mutable/Source/CustomizableObjectEditor/Private/MuCOE/PinViewer/SPinViewerDetailRowIndent.h
@@ -2,10 +2,18 @@
 
 #pragma once
 
-#include "Styling/SlateColor.h"
+#include "HAL/Platform.h"
+#include "Templates/SharedPointer.h"
 #include "Widgets/DeclarativeSyntaxSupport.h"
 #include "Widgets/SCompoundWidget.h"
 
+class FPaintArgs;
+class FSlateRect;
+class FSlateWindowElementList;
+class FWidgetStyle;
+class SWidget;
+struct FGeometry;
+
 
 /** Pin Viewer custom details row left indent.
  * Draws a shadow to differentiate the different indent levels.
diff --git a/Engine/Plugins/Experimental/Mutable/Source/CustomizableObjectEditor/Private/MuCOE/PinViewer/SPinViewerListRow.cpp b/Engine/Plugins/Experimental/Mutable/Source/CustomizableObjectEditor/Private/MuCOE/PinViewer/SPinViewerListRow.cpp
--- a/Engine/Plugins/Experimental/Mutable/Source/CustomizableObjectEditor/Private/MuCOE/PinViewer/SPinViewerListRow.cpp
+++ b/Engine/Plugins/Experimental/Mutable/Source/CustomizableObjectEditor/Private/MuCOE/PinViewer/SPinViewerListRow.cpp
@@ -2,10 +2,26 @@
 
 #include "MuCOE/PinViewer/SPinViewerListRow.h"
 
+#include "Containers/Map.h"
+#include "Delegates/Delegate.h"
+#include "Internationalization/Internationalization.h"
+#include "Internationalization/Text.h"
+#include "Math/Color.h"
+#include "Misc/Attribute.h"
 #include "MuCOE/EdGraphSchema_CustomizableObject.h"
+#include "MuCOE/Nodes/CustomizableObjectNode.h"
 #include "MuCOE/PinViewer/SPinViewer.h"
-#include "Widgets/Input/SCheckBox.h"
 #include "ScopedTransaction.h"
+#include "SlotBase.h"
+#include "Styling/SlateColor.h"
+#include "Types/SlateEnums.h"
+#include "Widgets/Input/SCheckBox.h"
+#include "Widgets/SBoxPanel.h"
+#include "Widgets/SNullWidget.h"
+#include "Widgets/Text/STextBlock.h"
+
+class STableViewBase;
+class SWidget;
 
 #define LOCTEXT_NAMESPACE "CustomizableObjectEditor"
 
diff --git a/Engine/Plugins/Experimental/Mutable/Source/CustomizableObjectEditor/Private/MuCOE/PinViewer/SPinViewerListRow.h b/Engine/Plugins/Experimental/Mutable/Source/CustomizableObjectEditor/Private/MuCOE/PinViewer/SPinViewerListRow.h
--- a/Engine/Plugins/Experimental/Mutable/Source/CustomizableObjectEditor/Private/MuCOE/PinViewer/SPinViewerListRow.h
+++ b/Engine/Plugins/Experimental/Mutable/Source/CustomizableObjectEditor/Private/MuCOE/PinViewer/SPinViewerListRow.h
@@ -2,10 +2,21 @@
 
 #pragma once
 
-#include "MuCOE/Widgets/SMutableExpandableTableRow.h"
+#include "Containers/BitArray.h"
+#include "Delegates/IDelegateInstance.h"
 #include "EdGraph/EdGraphPin.h"
+#include "Layout/Visibility.h"
+#include "Misc/Guid.h"
+#include "MuCOE/Widgets/SMutableExpandableTableRow.h"
+#include "Styling/SlateTypes.h"
+#include "Templates/SharedPointer.h"
+#include "Templates/UnrealTemplate.h"
+#include "UObject/NameTypes.h"
+#include "Widgets/DeclarativeSyntaxSupport.h"
 
 class SPinViewer;
+class STableViewBase;
+class SWidget;
 
 /** Custom row for a SListView that allows to add custom widgets as a row element. */
 class SPinViewerListRow : public SMutableExpandableTableRow<TSharedPtr<FEdGraphPinReference>>
diff --git a/Engine/Plugins/Experimental/Mutable/Source/CustomizableObjectEditor/Private/MuCOE/PinViewer/SPinViewerPinDetails.cpp b/Engine/Plugins/Experimental/Mutable/Source/CustomizableObjectEditor/Private/MuCOE/PinViewer/SPinViewerPinDetails.cpp
--- a/Engine/Plugins/Experimental/Mutable/Source/CustomizableObjectEditor/Private/MuCOE/PinViewer/SPinViewerPinDetails.cpp
+++ b/Engine/Plugins/Experimental/Mutable/Source/CustomizableObjectEditor/Private/MuCOE/PinViewer/SPinViewerPinDetails.cpp
@@ -2,10 +2,23 @@
 
 #include "MuCOE/PinViewer/SPinViewerPinDetails.h"
 
-#include "MuCOE/Widgets/SMutableExpandableTableRow.h"
+#include "Delegates/Delegate.h"
+#include "HAL/PlatformCrt.h"
+#include "Internationalization/Text.h"
+#include "Layout/Clipping.h"
+#include "Layout/Margin.h"
+#include "Misc/Attribute.h"
 #include "MuCOE/PinViewer/SPinViewerDetailRowIndent.h"
-#include "Widgets/Input/STextComboBox.h"
+#include "MuCOE/Widgets/SMutableExpandableTableRow.h"
+#include "SlotBase.h"
 #include "Styling/AppStyle.h"
+#include "Styling/ISlateStyle.h"
+#include "Types/SlateEnums.h"
+#include "Widgets/Layout/SSplitter.h"
+#include "Widgets/SBoxPanel.h"
+#include "Widgets/Text/STextBlock.h"
+
+class SWidget;
 
 #define LOCTEXT_NAMESPACE "CustomizableObjectEditor"
 
diff --git a/Engine/Plugins/Experimental/Mutable/Source/CustomizableObjectEditor/Private/MuCOE/PinViewer/SPinViewerPinDetails.h b/Engine/Plugins/Experimental/Mutable/Source/CustomizableObjectEditor/Private/MuCOE/PinViewer/SPinViewerPinDetails.h
--- a/Engine/Plugins/Experimental/Mutable/Source/CustomizableObjectEditor/Private/MuCOE/PinViewer/SPinViewerPinDetails.h
+++ b/Engine/Plugins/Experimental/Mutable/Source/CustomizableObjectEditor/Private/MuCOE/PinViewer/SPinViewerPinDetails.h
@@ -2,8 +2,14 @@
 
 #pragma once
 
+#include "Styling/SlateColor.h"
+#include "Templates/SharedPointer.h"
+#include "Widgets/DeclarativeSyntaxSupport.h"
 #include "Widgets/Layout/SBorder.h"
-#include "Widgets/SBoxPanel.h"
+
+class FText;
+class SVerticalBox;
+class SWidget;
 
 
 /** Pin Viewer custom details multiple row style.
diff --git a/Engine/Plugins/Experimental/Mutable/Source/CustomizableObjectEditor/Private/MuCOE/RemapPins/CustomizableObjectNodeRemapPinsByName.cpp b/Engine/Plugins/Experimental/Mutable/Source/CustomizableObjectEditor/Private/MuCOE/RemapPins/CustomizableObjectNodeRemapPinsByName.cpp
--- a/Engine/Plugins/Experimental/Mutable/Source/CustomizableObjectEditor/Private/MuCOE/RemapPins/CustomizableObjectNodeRemapPinsByName.cpp
+++ b/Engine/Plugins/Experimental/Mutable/Source/CustomizableObjectEditor/Private/MuCOE/RemapPins/CustomizableObjectNodeRemapPinsByName.cpp
@@ -1,6 +1,11 @@
 // Copyright Epic Games, Inc. All Rights Reserved.
 
 #include "MuCOE/RemapPins/CustomizableObjectNodeRemapPinsByName.h"
+
+#include "Containers/EnumAsByte.h"
+#include "Containers/UnrealString.h"
+#include "EdGraph/EdGraphPin.h"
+#include "HAL/PlatformCrt.h"
 #include "MuCOE/UnrealEditorPortabilityHelpers.h"
 
 
diff --git a/Engine/Plugins/Experimental/Mutable/Source/CustomizableObjectEditor/Private/MuCOE/RemapPins/CustomizableObjectNodeRemapPinsByNameDefaultPin.cpp b/Engine/Plugins/Experimental/Mutable/Source/CustomizableObjectEditor/Private/MuCOE/RemapPins/CustomizableObjectNodeRemapPinsByNameDefaultPin.cpp
--- a/Engine/Plugins/Experimental/Mutable/Source/CustomizableObjectEditor/Private/MuCOE/RemapPins/CustomizableObjectNodeRemapPinsByNameDefaultPin.cpp
+++ b/Engine/Plugins/Experimental/Mutable/Source/CustomizableObjectEditor/Private/MuCOE/RemapPins/CustomizableObjectNodeRemapPinsByNameDefaultPin.cpp
@@ -2,6 +2,12 @@
 
 #include "MuCOE/RemapPins/CustomizableObjectNodeRemapPinsByNameDefaultPin.h"
 
+#include "Containers/EnumAsByte.h"
+#include "EdGraph/EdGraphNode.h"
+#include "EdGraph/EdGraphPin.h"
+#include "HAL/PlatformCrt.h"
+#include "Misc/AssertionMacros.h"
+
 
 void UCustomizableObjectNodeRemapPinsByNameDefaultPin::RemapPins(const TArray<UEdGraphPin*>& OldPins, const TArray<UEdGraphPin*>& NewPins, TMap<UEdGraphPin*, UEdGraphPin*>& PinsToRemap, TArray<UEdGraphPin*>& PinsToOrphan)
 {
diff --git a/Engine/Plugins/Experimental/Mutable/Source/CustomizableObjectEditor/Private/MuCOE/RemapPins/CustomizableObjectNodeRemapPinsByPosition.cpp b/Engine/Plugins/Experimental/Mutable/Source/CustomizableObjectEditor/Private/MuCOE/RemapPins/CustomizableObjectNodeRemapPinsByPosition.cpp
--- a/Engine/Plugins/Experimental/Mutable/Source/CustomizableObjectEditor/Private/MuCOE/RemapPins/CustomizableObjectNodeRemapPinsByPosition.cpp
+++ b/Engine/Plugins/Experimental/Mutable/Source/CustomizableObjectEditor/Private/MuCOE/RemapPins/CustomizableObjectNodeRemapPinsByPosition.cpp
@@ -2,6 +2,12 @@
 
 #include "MuCOE/RemapPins/CustomizableObjectNodeRemapPinsByPosition.h"
 
+#include "Containers/EnumAsByte.h"
+#include "EdGraph/EdGraphNode.h"
+#include "EdGraph/EdGraphPin.h"
+#include "HAL/Platform.h"
+#include "HAL/PlatformCrt.h"
+
 
 /** Remap pins of a given direction by position. */
 void RemapPinsInternal(const TArray<UEdGraphPin*>& OldPins, const TArray<UEdGraphPin*>& NewPins, TMap<UEdGraphPin*, UEdGraphPin*>& PinsToRemap, TArray<UEdGraphPin*>& PinsToOrphan, EEdGraphPinDirection Direction)
diff --git a/Engine/Plugins/Experimental/Mutable/Source/CustomizableObjectEditor/Private/MuCOE/SCustomizableInstanceProperties.cpp b/Engine/Plugins/Experimental/Mutable/Source/CustomizableObjectEditor/Private/MuCOE/SCustomizableInstanceProperties.cpp
--- a/Engine/Plugins/Experimental/Mutable/Source/CustomizableObjectEditor/Private/MuCOE/SCustomizableInstanceProperties.cpp
+++ b/Engine/Plugins/Experimental/Mutable/Source/CustomizableObjectEditor/Private/MuCOE/SCustomizableInstanceProperties.cpp
@@ -2,36 +2,78 @@
 
 #include "MuCOE/SCustomizableInstanceProperties.h"
 
-#include "UObject/PropertyPortFlags.h"
-
+#include "ContentBrowserModule.h"
+#include "CoreGlobals.h"
+#include "Delegates/Delegate.h"
+#include "Editor.h"
+#include "Editor/EditorEngine.h"
+#include "Engine/Engine.h"
+#include "Engine/Texture2D.h"
+#include "Fonts/SlateFontInfo.h"
+#include "Framework/SlateDelegates.h"
+#include "Framework/Text/TextLayout.h"
+#include "HAL/PlatformApplicationMisc.h"
+#include "HAL/PlatformCrt.h"
+#include "HAL/PlatformMisc.h"
+#include "Input/Events.h"
+#include "InputCoreTypes.h"
+#include "Internationalization/Internationalization.h"
+#include "Layout/Children.h"
+#include "Layout/Margin.h"
+#include "Layout/Visibility.h"
+#include "Logging/LogCategory.h"
+#include "Logging/LogMacros.h"
+#include "Math/UnrealMathSSE.h"
+#include "Math/Vector.h"
+#include "Math/Vector2D.h"
+#include "Misc/AssertionMacros.h"
+#include "Misc/Attribute.h"
+#include "Misc/CString.h"
+#include "Misc/Optional.h"
+#include "Misc/OutputDeviceRedirector.h"
+#include "Misc/Paths.h"
+#include "Modules/ModuleManager.h"
+#include "MuCO/CustomizableObject.h"
+#include "MuCO/CustomizableObjectIdentifier.h"
 #include "MuCO/CustomizableObjectInstance.h"
-#include "MuCO/MultilayerProjector.h"
+#include "MuCO/CustomizableObjectParameterTypeDefinitions.h"
 #include "MuCO/CustomizableObjectSystem.h"
+#include "MuCO/CustomizableObjectUIData.h"
+#include "MuCO/MultilayerProjector.h"
+#include "MuCOE/CustomizableInstanceDetails.h"
 #include "MuCOE/CustomizableObjectEditorUtilities.h"
-
-#include "Widgets/Input/STextComboBox.h"
 #include "MuCOE/SMutableTextSearchBox.h"
-#include "Widgets/Colors/SColorPicker.h"
+#include "Serialization/BufferArchive.h"
+#include "Serialization/MemoryReader.h"
+#include "Slate/DeferredCleanupSlateBrush.h"
+#include "SlotBase.h"
+#include "Styling/AppStyle.h"
+#include "Styling/SlateBrush.h"
+#include "Styling/SlateColor.h"
+#include "Trace/Detail/Channel.h"
+#include "Types/SlateStructs.h"
+#include "UObject/Class.h"
+#include "UObject/NameTypes.h"
+#include "UObject/PropertyPortFlags.h"
+#include "UObject/UObjectBaseUtility.h"
+#include "UObject/UObjectGlobals.h"
 #include "Widgets/Colors/SColorBlock.h"
-#include "Widgets/Input/SCheckBox.h"
+#include "Widgets/Colors/SColorPicker.h"
 #include "Widgets/Images/SImage.h"
 #include "Widgets/Input/SButton.h"
+#include "Widgets/Input/SCheckBox.h"
+#include "Widgets/Input/SEditableTextBox.h"
+#include "Widgets/Input/SSearchBox.h"
+#include "Widgets/Input/STextComboBox.h"
+#include "Widgets/Layout/SBorder.h"
+#include "Widgets/Layout/SBox.h"
 #include "Widgets/Layout/SExpandableArea.h"
-#include "Slate/DeferredCleanupSlateBrush.h"
-
-#include "Serialization/BufferArchive.h"
 #include "Widgets/Layout/SUniformGridPanel.h"
+#include "Widgets/SBoxPanel.h"
+#include "Widgets/SWindow.h"
+#include "Widgets/Text/STextBlock.h"
 
-#include "ContentBrowserModule.h"
-#include "IContentBrowserSingleton.h"
-#include "AssetRegistry/AssetRegistryModule.h"
-#include "Editor.h"
-#include "ISkeletonTree.h"
-#include "EnvironmentQuery/EnvQueryTypes.h"
-
-#include "UObject/UObjectGlobals.h"
-
-#include "HAL/PlatformApplicationMisc.h"
+struct FGeometry;
 
 #define LOCTEXT_NAMESPACE "CustomizableObjectEditor"
 
diff --git a/Engine/Plugins/Experimental/Mutable/Source/CustomizableObjectEditor/Private/MuCOE/SCustomizableObjectEditorAdvancedPreviewSettings.cpp b/Engine/Plugins/Experimental/Mutable/Source/CustomizableObjectEditor/Private/MuCOE/SCustomizableObjectEditorAdvancedPreviewSettings.cpp
--- a/Engine/Plugins/Experimental/Mutable/Source/CustomizableObjectEditor/Private/MuCOE/SCustomizableObjectEditorAdvancedPreviewSettings.cpp
+++ b/Engine/Plugins/Experimental/Mutable/Source/CustomizableObjectEditor/Private/MuCOE/SCustomizableObjectEditorAdvancedPreviewSettings.cpp
@@ -1,33 +1,50 @@
 // Copyright Epic Games, Inc. All Rights Reserved.
 
 #include "MuCOE/SCustomizableObjectEditorAdvancedPreviewSettings.h"
-#include "Widgets/SBoxPanel.h"
-#include "Widgets/Text/STextBlock.h"
-#include "Widgets/Layout/SBorder.h"
-#include "Widgets/Input/SButton.h"
-#include "Widgets/Layout/SBox.h"
-#include "Styling/AppStyle.h"
-#include "Widgets/Layout/SUniformGridPanel.h"
-#include "Editor.h"
-#include "Widgets/Input/SEditableTextBox.h"
-#include "Misc/MessageDialog.h"
-#include "Widgets/Docking/SDockTab.h"
+
+#include "AdvancedPreviewSceneModule.h"
+#include "AssetRegistry/ARFilter.h"
+#include "AssetViewerSettings.h"
+#include "ContentBrowserDelegates.h"
 #include "ContentBrowserModule.h"
+#include "Delegates/Delegate.h"
+#include "Editor.h"
+#include "Editor/EditorEngine.h"
+#include "Fonts/SlateFontInfo.h"
+#include "HAL/Platform.h"
+#include "HAL/PlatformCrt.h"
 #include "IContentBrowserSingleton.h"
-#include "Widgets/Layout/SBox.h"
-#include "Engine/LevelStreaming.h"
-#include "Kismet/GameplayStatics.h"
-#include "Engine/SkyLight.h"
-#include "Engine/DirectionalLight.h"
-#include "Engine/PostProcessVolume.h"
-#include "EngineUtils.h"
-#include "MuCOE/CustomizableObjectEditorSettings.h"
-#include "AdvancedPreviewScene.h"
-#include "Components/StaticMeshComponent.h"
-#include "MuCOE/CustomizableObjectPreviewScene.h"
-#include "Serialization/JsonTypes.h"
-#include "Serialization/JsonReader.h"
-#include "Serialization/JsonSerializer.h"
+#include "Internationalization/Internationalization.h"
+#include "Internationalization/Text.h"
+#include "Layout/Children.h"
+#include "Layout/Margin.h"
+#include "Math/Vector2D.h"
+#include "Misc/MessageDialog.h"
+#include "Misc/Paths.h"
+#include "Modules/ModuleManager.h"
+#include "MuCOE/CustomizableObjectEditor.h"
+#include "MuCOE/CustomizableObjectEditorLevelProfile.h"
+#include "SlotBase.h"
+#include "Styling/AppStyle.h"
+#include "Templates/Casts.h"
+#include "Types/SlateEnums.h"
+#include "UObject/NameTypes.h"
+#include "UObject/Object.h"
+#include "UObject/ObjectRedirector.h"
+#include "UObject/SoftObjectPtr.h"
+#include "UObject/TopLevelAssetPath.h"
+#include "UObject/UObjectGlobals.h"
+#include "Widgets/Input/SButton.h"
+#include "Widgets/Input/SEditableTextBox.h"
+#include "Widgets/Input/STextComboBox.h"
+#include "Widgets/Layout/SBorder.h"
+#include "Widgets/Layout/SUniformGridPanel.h"
+#include "Widgets/SBoxPanel.h"
+#include "Widgets/SCompoundWidget.h"
+#include "Widgets/Text/STextBlock.h"
+
+class FAdvancedPreviewScene;
+class SWidget;
 
 #define LOCTEXT_NAMESPACE "LevelSettingsPreview"
 
diff --git a/Engine/Plugins/Experimental/Mutable/Source/CustomizableObjectEditor/Private/MuCOE/SCustomizableObjectEditorAdvancedPreviewSettings.h b/Engine/Plugins/Experimental/Mutable/Source/CustomizableObjectEditor/Private/MuCOE/SCustomizableObjectEditorAdvancedPreviewSettings.h
--- a/Engine/Plugins/Experimental/Mutable/Source/CustomizableObjectEditor/Private/MuCOE/SCustomizableObjectEditorAdvancedPreviewSettings.h
+++ b/Engine/Plugins/Experimental/Mutable/Source/CustomizableObjectEditor/Private/MuCOE/SCustomizableObjectEditorAdvancedPreviewSettings.h
@@ -2,15 +2,24 @@
 
 #pragma once
 
-#include "AdvancedPreviewSceneModule.h"
+#include "AssetRegistry/AssetData.h"
+#include "Containers/Array.h"
+#include "Containers/UnrealString.h"
+#include "HAL/PlatformMisc.h"
+#include "Input/Reply.h"
+#include "Misc/Attribute.h"
 #include "SAdvancedPreviewDetailsTab.h"
+#include "Templates/SharedPointer.h"
+#include "Widgets/DeclarativeSyntaxSupport.h"
 #include "Widgets/SWindow.h"
-//#include "Classes/Engine/Level.h"
-#include "Widgets/Input/STextComboBox.h"
-#include "Widgets/Input/SEditableTextBox.h"
-#include "Widgets/Input/SComboBox.h"
-#include "MuCOE/CustomizableObjectEditorLevelProfile.h"
-#include "MuCOE/CustomizableObjectEditor.h"
+
+class FAdvancedPreviewScene;
+class FCustomizableObjectEditor;
+class SEditableTextBox;
+class STextComboBox;
+class UAssetViewerSettings;
+class UObject;
+struct FPreviewSceneProfile;
 
 
 /** Simple window with a content browser filtered for ULevel assets to select a ULevel with a
diff --git a/Engine/Plugins/Experimental/Mutable/Source/CustomizableObjectEditor/Private/MuCOE/SCustomizableObjectEditorPerformanceReport.cpp b/Engine/Plugins/Experimental/Mutable/Source/CustomizableObjectEditor/Private/MuCOE/SCustomizableObjectEditorPerformanceReport.cpp
--- a/Engine/Plugins/Experimental/Mutable/Source/CustomizableObjectEditor/Private/MuCOE/SCustomizableObjectEditorPerformanceReport.cpp
+++ b/Engine/Plugins/Experimental/Mutable/Source/CustomizableObjectEditor/Private/MuCOE/SCustomizableObjectEditorPerformanceReport.cpp
@@ -2,16 +2,61 @@
 
 
 #include "MuCOE/SCustomizableObjectEditorPerformanceReport.h"
+
+#include "AssetThumbnail.h"
+#include "Containers/Map.h"
+#include "Containers/Set.h"
+#include "Containers/SparseArray.h"
+#include "Delegates/Delegate.h"
+#include "DetailsViewArgs.h"
+#include "Editor.h"
+#include "Editor/EditorEngine.h"
+#include "Framework/MultiBox/MultiBoxBuilder.h"
+#include "Framework/Views/ITypedTableView.h"
+#include "HAL/PlatformCrt.h"
+#include "Internationalization/Internationalization.h"
+#include "Internationalization/Text.h"
+#include "Layout/BasicLayoutWidgetSlot.h"
+#include "Layout/Children.h"
+#include "Math/Color.h"
+#include "Misc/AssertionMacros.h"
+#include "Misc/Attribute.h"
+#include "Modules/ModuleManager.h"
+#include "MuCO/CustomizableObject.h"
 #include "MuCO/CustomizableObjectInstance.h"
+#include "MuCOE/CustomizableObjectEditorStyle.h"
+#include "MuCOE/SCustomizableObjectEditorTextureAnalyzer.h"
 #include "MuCOE/StatePerformanceTesting.h"
 #include "PropertyCustomizationHelpers.h"
-#include "IDetailsView.h"
+#include "PropertyEditorModule.h"
+#include "SlotBase.h"
+#include "Styling/AppStyle.h"
+#include "Styling/ISlateStyle.h"
+#include "Styling/SlateColor.h"
+#include "Styling/SlateTypes.h"
+#include "Types/SlateEnums.h"
+#include "Types/SlateStructs.h"
+#include "UObject/Package.h"
+#include "UObject/UnrealNames.h"
+#include "Widgets/Input/SButton.h"
+#include "Widgets/Input/SComboButton.h"
+#include "Widgets/InvalidateWidgetReason.h"
+#include "Widgets/Layout/SBorder.h"
+#include "Widgets/Layout/SBox.h"
 #include "Widgets/Layout/SScrollBox.h"
 #include "Widgets/Layout/SSplitter.h"
-#include "MuCOE/CustomizableObjectEditorStyle.h"
-#include "Modules/ModuleManager.h"
-#include "PropertyCustomizationHelpers.h"
-#include "AssetThumbnail.h"
+#include "Widgets/SBoxPanel.h"
+#include "Widgets/SOverlay.h"
+#include "Widgets/Text/STextBlock.h"
+#include "Widgets/Views/SHeaderRow.h"
+#include "Widgets/Views/SListView.h"
+#include "Widgets/Views/STableRow.h"
+
+class ITableRow;
+class STableViewBase;
+class SWidget;
+struct FGeometry;
+struct FSlateBrush;
 
 
 #define LOCTEXT_NAMESPACE "CustomizableObjectEditor"
diff --git a/Engine/Plugins/Experimental/Mutable/Source/CustomizableObjectEditor/Private/MuCOE/SCustomizableObjectEditorPerformanceReport.h b/Engine/Plugins/Experimental/Mutable/Source/CustomizableObjectEditor/Private/MuCOE/SCustomizableObjectEditorPerformanceReport.h
--- a/Engine/Plugins/Experimental/Mutable/Source/CustomizableObjectEditor/Private/MuCOE/SCustomizableObjectEditorPerformanceReport.h
+++ b/Engine/Plugins/Experimental/Mutable/Source/CustomizableObjectEditor/Private/MuCOE/SCustomizableObjectEditorPerformanceReport.h
@@ -2,15 +2,41 @@
 
 #pragma once
 
-#include "Widgets/SCompoundWidget.h"
-#include "Widgets/Views/SListView.h"
+#include "Containers/Array.h"
+#include "Containers/BitArray.h"
+#include "Containers/UnrealString.h"
+#include "CoreTypes.h"
+#include "Delegates/IDelegateInstance.h"
+#include "Input/Reply.h"
+#include "Layout/Visibility.h"
 #include "Misc/NotifyHook.h"
-#include "MuCO/CustomizableObject.h"
+#include "Misc/Optional.h"
+#include "MuCOE/StatePerformanceTesting.h"
 #include "MuCOE/StressTest.h"
 #include "MuCOE/Widgets/SMutableExpandableTableRow.h"
-#include "MuCOE/SCustomizableObjectEditorTextureAnalyzer.h"
+#include "Templates/SharedPointer.h"
+#include "Templates/UnrealTemplate.h"
+#include "UObject/GCObject.h"
+#include "UObject/NameTypes.h"
+#include "UObject/Object.h"
+#include "UObject/ObjectPtr.h"
+#include "UObject/UObjectGlobals.h"
+#include "Widgets/DeclarativeSyntaxSupport.h"
+#include "Widgets/Input/SNumericEntryBox.h"
+#include "Widgets/SCompoundWidget.h"
+
 #include "SCustomizableObjectEditorPerformanceReport.generated.h"
 
+class ITableRow;
+class STableViewBase;
+class SWidget;
+class UCustomizableObject;
+class UCustomizableObjectInstance;
+struct FFrame;
+struct FGeometry;
+struct FSlateBrush;
+struct FTableRowStyle;
+
 
 // Instances performing worst in some aspect
 UCLASS()
diff --git a/Engine/Plugins/Experimental/Mutable/Source/CustomizableObjectEditor/Private/MuCOE/SCustomizableObjectEditorTagExplorer.cpp b/Engine/Plugins/Experimental/Mutable/Source/CustomizableObjectEditor/Private/MuCOE/SCustomizableObjectEditorTagExplorer.cpp
--- a/Engine/Plugins/Experimental/Mutable/Source/CustomizableObjectEditor/Private/MuCOE/SCustomizableObjectEditorTagExplorer.cpp
+++ b/Engine/Plugins/Experimental/Mutable/Source/CustomizableObjectEditor/Private/MuCOE/SCustomizableObjectEditorTagExplorer.cpp
@@ -2,21 +2,49 @@
 
 #include "MuCOE/SCustomizableObjectEditorTagExplorer.h"
 
+#include "CoreGlobals.h"
+#include "CoreTypes.h"
+#include "EdGraph/EdGraph.h"
+#include "EdGraph/EdGraphNode.h"
+#include "Editor.h"
+#include "Editor/EditorEngine.h"
+#include "Framework/Views/ITypedTableView.h"
+#include "HAL/PlatformApplicationMisc.h"
+#include "HAL/PlatformCrt.h"
+#include "Internationalization/Internationalization.h"
+#include "Layout/Children.h"
+#include "Logging/LogCategory.h"
+#include "Logging/LogMacros.h"
+#include "Misc/Attribute.h"
+#include "MuCO/CustomizableObject.h"
+#include "MuCOE/CustomizableObjectEditor.h"
 #include "MuCOE/CustomizableObjectGraph.h"
 #include "MuCOE/GraphTraversal.h"
-
-#include "MuCOE/UnrealEditorPortabilityHelpers.h"
-#include "HAL/PlatformApplicationMisc.h"
-
+#include "MuCOE/ICustomizableObjectEditor.h"
 #include "MuCOE/Nodes/CustomizableObjectNode.h"
-#include "MuCOE/Nodes/CustomizableObjectNodeObject.h"
+#include "MuCOE/Nodes/CustomizableObjectNodeExtendMaterial.h"
 #include "MuCOE/Nodes/CustomizableObjectNodeMaterial.h"
+#include "MuCOE/Nodes/CustomizableObjectNodeMaterialVariation.h"
 #include "MuCOE/Nodes/CustomizableObjectNodeMeshClipMorph.h"
 #include "MuCOE/Nodes/CustomizableObjectNodeMeshClipWithMesh.h"
-#include "MuCOE/Nodes/CustomizableObjectNodeMaterialVariation.h"
-#include "MuCOE/Nodes/CustomizableObjectNodeExtendMaterial.h"
-
+#include "MuCOE/Nodes/CustomizableObjectNodeObject.h"
+#include "Serialization/Archive.h"
+#include "SlotBase.h"
+#include "Subsystems/AssetEditorSubsystem.h"
+#include "Templates/Casts.h"
+#include "Templates/UnrealTemplate.h"
+#include "Trace/Detail/Channel.h"
+#include "UObject/ObjectPtr.h"
 #include "Widgets/Input/SButton.h"
+#include "Widgets/SBoxPanel.h"
+#include "Widgets/Text/STextBlock.h"
+#include "Widgets/Views/SHeaderRow.h"
+#include "Widgets/Views/STableRow.h"
+
+class ITableRow;
+class STableViewBase;
+class SWidget;
+class UObject;
 
 #define LOCTEXT_NAMESPACE "SCustomizableObjectEditorTagExplorer"
 
diff --git a/Engine/Plugins/Experimental/Mutable/Source/CustomizableObjectEditor/Private/MuCOE/SCustomizableObjectEditorTagExplorer.h b/Engine/Plugins/Experimental/Mutable/Source/CustomizableObjectEditor/Private/MuCOE/SCustomizableObjectEditorTagExplorer.h
--- a/Engine/Plugins/Experimental/Mutable/Source/CustomizableObjectEditor/Private/MuCOE/SCustomizableObjectEditorTagExplorer.h
+++ b/Engine/Plugins/Experimental/Mutable/Source/CustomizableObjectEditor/Private/MuCOE/SCustomizableObjectEditorTagExplorer.h
@@ -2,10 +2,29 @@
 
 #pragma once
 
+#include "Containers/Array.h"
+#include "Containers/BitArray.h"
+#include "Containers/Map.h"
+#include "Containers/Set.h"
+#include "Containers/SparseArray.h"
+#include "Containers/UnrealString.h"
+#include "Delegates/Delegate.h"
+#include "Input/Reply.h"
+#include "Internationalization/Text.h"
+#include "Misc/Optional.h"
+#include "MuCOE/Nodes/CustomizableObjectNode.h"
+#include "Templates/SharedPointer.h"
+#include "Types/SlateEnums.h"
+#include "Widgets/DeclarativeSyntaxSupport.h"
+#include "Widgets/Input/SComboBox.h"
 #include "Widgets/SCompoundWidget.h"
-#include "MuCOE/CustomizableObjectEditor.h"
+#include "Widgets/Views/SListView.h"
 
-class UCustomizableObjectNode;
+class FCustomizableObjectEditor;
+class ITableRow;
+class STableViewBase;
+class SWidget;
+class UCustomizableObject;
 
 class SCustomizableObjectEditorTagExplorer : public SCompoundWidget
 {
diff --git a/Engine/Plugins/Experimental/Mutable/Source/CustomizableObjectEditor/Private/MuCOE/SCustomizableObjectEditorTextureAnalyzer.cpp b/Engine/Plugins/Experimental/Mutable/Source/CustomizableObjectEditor/Private/MuCOE/SCustomizableObjectEditorTextureAnalyzer.cpp
--- a/Engine/Plugins/Experimental/Mutable/Source/CustomizableObjectEditor/Private/MuCOE/SCustomizableObjectEditorTextureAnalyzer.cpp
+++ b/Engine/Plugins/Experimental/Mutable/Source/CustomizableObjectEditor/Private/MuCOE/SCustomizableObjectEditorTextureAnalyzer.cpp
@@ -2,16 +2,45 @@
 
 #include "MuCOE/SCustomizableObjectEditorTextureAnalyzer.h"
 
+#include "Containers/EnumAsByte.h"
+#include "Editor.h"
+#include "Editor/EditorEngine.h"
+#include "Engine/SkeletalMesh.h"
+#include "Engine/Texture.h"
+#include "Engine/Texture2D.h"
+#include "Framework/Views/ITypedTableView.h"
+#include "HAL/PlatformCrt.h"
 #include "IPropertyTable.h"
-#include "IPropertyTableColumn.h"
 #include "IPropertyTableCell.h"
-#include "IPropertyTableRow.h"
-
-#include "Engine/Texture.h"
-#include "Modules/ModuleManager.h"
+#include "IPropertyTableColumn.h"
+#include "Internationalization/Internationalization.h"
+#include "Internationalization/Text.h"
+#include "Layout/Children.h"
+#include "Layout/Visibility.h"
+#include "MaterialTypes.h"
 #include "Materials/MaterialInstance.h"
-#include "MuCOE/UnrealEditorPortabilityHelpers.h"
+#include "Materials/MaterialInterface.h"
+#include "Misc/Attribute.h"
+#include "Modules/ModuleManager.h"
+#include "MuCO/CustomizableObjectInstance.h"
+#include "MuCOE/CustomizableObjectEditor.h"
+#include "MuCOE/CustomizableObjectInstanceEditor.h"
+#include "PropertyEditorModule.h"
+#include "SlotBase.h"
+#include "Subsystems/AssetEditorSubsystem.h"
+#include "Templates/Casts.h"
+#include "Types/SlateEnums.h"
+#include "UObject/NameTypes.h"
+#include "UObject/Object.h"
+#include "UObject/ObjectPtr.h"
+#include "UObject/UObjectGlobals.h"
+#include "UObject/WeakObjectPtrTemplates.h"
 #include "Widgets/Input/SButton.h"
+#include "Widgets/SBoxPanel.h"
+#include "Widgets/Text/STextBlock.h"
+
+class SWidget;
+struct FGeometry;
 
 #define LOCTEXT_NAMESPACE "SCustomizableObjecEditorTextureAnalyzer"
 
diff --git a/Engine/Plugins/Experimental/Mutable/Source/CustomizableObjectEditor/Private/MuCOE/SCustomizableObjectEditorTextureAnalyzer.h b/Engine/Plugins/Experimental/Mutable/Source/CustomizableObjectEditor/Private/MuCOE/SCustomizableObjectEditorTextureAnalyzer.h
--- a/Engine/Plugins/Experimental/Mutable/Source/CustomizableObjectEditor/Private/MuCOE/SCustomizableObjectEditorTextureAnalyzer.h
+++ b/Engine/Plugins/Experimental/Mutable/Source/CustomizableObjectEditor/Private/MuCOE/SCustomizableObjectEditorTextureAnalyzer.h
@@ -2,9 +2,24 @@
 
 #pragma once
 
+#include "Containers/Array.h"
+#include "Containers/UnrealString.h"
+#include "HAL/Platform.h"
+#include "Input/Reply.h"
+#include "Templates/SharedPointer.h"
+#include "UObject/GCObject.h"
+#include "Widgets/DeclarativeSyntaxSupport.h"
 #include "Widgets/SCompoundWidget.h"
-#include "MuCOE/CustomizableObjectEditor.h"
-#include "MuCOE/CustomizableObjectInstanceEditor.h"
+
+class FCustomizableObjectEditor;
+class FCustomizableObjectInstanceEditor;
+class FReferenceCollector;
+class IPropertyTable;
+class STextBlock;
+class SWidget;
+class UCustomizableObjectInstance;
+class UObject;
+struct FGeometry;
 
 
 class SCustomizableObjecEditorTextureAnalyzer : public SCompoundWidget, public FGCObject
diff --git a/Engine/Plugins/Experimental/Mutable/Source/CustomizableObjectEditor/Private/MuCOE/SCustomizableObjectEditorTransformViewportToolbar.cpp b/Engine/Plugins/Experimental/Mutable/Source/CustomizableObjectEditor/Private/MuCOE/SCustomizableObjectEditorTransformViewportToolbar.cpp
--- a/Engine/Plugins/Experimental/Mutable/Source/CustomizableObjectEditor/Private/MuCOE/SCustomizableObjectEditorTransformViewportToolbar.cpp
+++ b/Engine/Plugins/Experimental/Mutable/Source/CustomizableObjectEditor/Private/MuCOE/SCustomizableObjectEditorTransformViewportToolbar.cpp
@@ -1,12 +1,9 @@
 // Copyright Epic Games, Inc. All Rights Reserved.	
 
 #include "MuCOE/SCustomizableObjectEditorTransformViewportToolbar.h"
-#include "Framework/MultiBox/MultiBoxDefs.h"
-#include "Framework/MultiBox/MultiBoxBuilder.h"
-#include "Styling/AppStyle.h"
+
+#include "Framework/Commands/UICommandList.h"
 #include "MuCOE/CustomizableObjectEditorViewportLODCommands.h"
-#include "MuCOE/SCustomizableObjectEditorViewport.h"
-#include "MuCOE/CustomizableObjectEditorViewportClient.h"
 
 #define LOCTEXT_NAMESPACE "SCustomizableObjectEditorTransformViewportToolbar"
 
diff --git a/Engine/Plugins/Experimental/Mutable/Source/CustomizableObjectEditor/Private/MuCOE/SCustomizableObjectEditorViewport.cpp b/Engine/Plugins/Experimental/Mutable/Source/CustomizableObjectEditor/Private/MuCOE/SCustomizableObjectEditorViewport.cpp
--- a/Engine/Plugins/Experimental/Mutable/Source/CustomizableObjectEditor/Private/MuCOE/SCustomizableObjectEditorViewport.cpp
+++ b/Engine/Plugins/Experimental/Mutable/Source/CustomizableObjectEditor/Private/MuCOE/SCustomizableObjectEditorViewport.cpp
@@ -1,44 +1,79 @@
 // Copyright Epic Games, Inc. All Rights Reserved.
 
 #include "MuCOE/SCustomizableObjectEditorViewport.h"
-#include "MuCOE/CustomizableObjectEditorActions.h"
-#include "MuCOE/SCustomizableObjectEditorViewportToolBar.h"
-#include "MuCOE/CustomizableObjectEditorViewportMenuCommands.h"
-#include "MuCOE/CustomizableObjectEditorViewportLODCommands.h"
-#include "MuCOE/CustomizableObjectEditorViewportClient.h"
-#include "MuCOE/SCustomizableObjectEditorTransformViewportToolbar.h"
-#include "MuCOE/CustomizableObjectEditorUtilities.h"
-#include "MuCOE/Nodes/CustomizableObjectNodeMeshClipMorph.h"
-#include "MuCOE/Nodes/CustomizableObjectNodeMeshClipWithMesh.h"
-
-#include "IPersonaPreviewScene.h"
 
-#include "MouseDeltaTracker.h"
-#include "PreviewScene.h"
-#include "Slate/SceneViewport.h"
-//#include "Widgets/Docking/SDockableTab.h"
-#include "Widgets/Input/STextComboBox.h"
-#include "Widgets/Input/SCheckBox.h"
-
-#include "MuCOE/CustomizableObjectEditor.h"
-#include "FbxMeshUtils.h"
-#include "BusyCursor.h"
-#include "MeshBuild.h"
-#include "ObjectTools.h"
+#include "Animation/AnimationAsset.h"
+#include "Animation/DebugSkelMeshComponent.h"
 #include "Animation/PoseAsset.h"
+#include "AssetRegistry/AssetData.h"
+#include "Containers/EnumAsByte.h"
+#include "Containers/IndirectArray.h"
+#include "Containers/Map.h"
+#include "Delegates/Delegate.h"
 #include "DragAndDrop/AssetDragDropOp.h"
-#include "Viewports.h"
-#include "Framework/Notifications/NotificationManager.h"
-
+#include "Editor.h"
+#include "Editor/EditorEngine.h"
+#include "Editor/UnrealEdEngine.h"
+#include "Editor/UnrealEdTypes.h"
+#include "EditorViewportClient.h"
+#include "EditorViewportCommands.h"
+#include "Engine/Engine.h"
+#include "Engine/SkeletalMesh.h"
+#include "Engine/World.h"
+#include "Framework/Application/SlateApplication.h"
+#include "Framework/Commands/UIAction.h"
+#include "Framework/Commands/UICommandList.h"
+#include "Framework/MultiBox/MultiBoxBuilder.h"
+#include "Framework/MultiBox/MultiBoxDefs.h"
+#include "Framework/SlateDelegates.h"
+#include "GenericPlatform/ICursor.h"
+#include "HAL/PlatformCrt.h"
+#include "Input/DragAndDrop.h"
+#include "Internationalization/Text.h"
+#include "Layout/Children.h"
+#include "Layout/Visibility.h"
+#include "Materials/Material.h"
+#include "Materials/MaterialInterface.h"
+#include "Math/Transform.h"
+#include "Misc/AssertionMacros.h"
+#include "Misc/Attribute.h"
+#include "Misc/CString.h"
+#include "Misc/Optional.h"
+#include "MuCO/CustomizableObjectInstance.h"
 #include "MuCO/UnrealPortabilityHelpers.h"
+#include "MuCOE/CustomizableObjectEditorActions.h"
+#include "MuCOE/CustomizableObjectEditorUtilities.h"
+#include "MuCOE/CustomizableObjectEditorViewportClient.h"
+#include "MuCOE/CustomizableObjectEditorViewportLODCommands.h"
+#include "MuCOE/CustomizableObjectEditorViewportMenuCommands.h"
+#include "MuCOE/CustomizableObjectPreviewScene.h"
+#include "MuCOE/ICustomizableObjectInstanceEditor.h"
+#include "MuCOE/SCustomizableObjectEditorViewportToolBar.h"
+#include "MuCOE/SCustomizableObjectHighresScreenshot.h"
 #include "MuCOE/UnrealEditorPortabilityHelpers.h"
-#include "Editor/UnrealEdEngine.h"
+#include "PreviewScene.h"
+#include "Rendering/SkeletalMeshLODRenderData.h"
+#include "Rendering/SkeletalMeshRenderData.h"
+#include "Slate/SceneViewport.h"
+#include "SlotBase.h"
+#include "Styling/AppStyle.h"
+#include "Templates/Casts.h"
+#include "Textures/SlateIcon.h"
+#include "Types/ISlateMetaData.h"
+#include "UObject/ObjectPtr.h"
+#include "UObject/UObjectGlobals.h"
 #include "UnrealEdGlobals.h"
-#include "Editor/EditorEngine.h"
-#include "MuCOE/SCustomizableObjectHighresScreenshot.h"
-#include "MuCOE/CustomizableObjectPreviewScene.h"
-
-#include "EditorViewportCommands.h"
+#include "Viewports.h"
+#include "Widgets/Input/STextComboBox.h"
+#include "Widgets/Layout/SBorder.h"
+#include "Widgets/SBoxPanel.h"
+#include "Widgets/SWidget.h"
+
+class UCustomizableObject;
+class UCustomizableObjectNodeProjectorConstant;
+class UStaticMesh;
+struct FCustomizableObjectProjector;
+struct FGeometry;
 
 
 void SCustomizableObjectEditorViewport::Construct(const FArguments& InArgs, const FCustomizableObjectEditorViewportRequiredArgs& InRequiredArgs)
diff --git a/Engine/Plugins/Experimental/Mutable/Source/CustomizableObjectEditor/Private/MuCOE/SCustomizableObjectEditorViewport.h b/Engine/Plugins/Experimental/Mutable/Source/CustomizableObjectEditor/Private/MuCOE/SCustomizableObjectEditorViewport.h
--- a/Engine/Plugins/Experimental/Mutable/Source/CustomizableObjectEditor/Private/MuCOE/SCustomizableObjectEditorViewport.h
+++ b/Engine/Plugins/Experimental/Mutable/Source/CustomizableObjectEditor/Private/MuCOE/SCustomizableObjectEditorViewport.h
@@ -2,12 +2,38 @@
 
 #pragma once
 
-#include "PreviewScene.h"
-#include "SEditorViewport.h"
 #include "Components/SkeletalMeshComponent.h"
+#include "Containers/Array.h"
+#include "Containers/UnrealString.h"
+#include "Engine/EngineBaseTypes.h"
+#include "Input/Reply.h"
+#include "Math/BoxSphereBounds.h"
+#include "Math/Color.h"
+#include "Math/UnrealMathSSE.h"
+#include "Math/Vector.h"
+#include "SEditorViewport.h"
 #include "Settings/LevelEditorViewportSettings.h"
-#include "MuCOE/ICustomizableObjectInstanceEditor.h"
-#include "MuCOE/Nodes/CustomizableObjectNodeProjectorConstant.h"
+#include "Templates/SharedPointer.h"
+#include "Types/SlateEnums.h"
+#include "UObject/GCObject.h"
+#include "UnrealWidgetFwd.h"
+#include "Widgets/DeclarativeSyntaxSupport.h"
+#include "Widgets/SCompoundWidget.h"
+
+class FDragDropEvent;
+class FEditorViewportClient;
+class FReferenceCollector;
+class FSceneViewport;
+class FUICommandList;
+class STextComboBox;
+class SWidget;
+class UCustomizableObject;
+class UCustomizableObjectNodeProjectorConstant;
+class UCustomizableObjectNodeProjectorParameter;
+class UDebugSkelMeshComponent;
+class UStaticMesh;
+struct FCustomizableObjectProjector;
+struct FGeometry;
 
 
 struct FCustomizableObjectEditorViewportRequiredArgs
diff --git a/Engine/Plugins/Experimental/Mutable/Source/CustomizableObjectEditor/Private/MuCOE/SCustomizableObjectEditorViewportToolBar.cpp b/Engine/Plugins/Experimental/Mutable/Source/CustomizableObjectEditor/Private/MuCOE/SCustomizableObjectEditorViewportToolBar.cpp
--- a/Engine/Plugins/Experimental/Mutable/Source/CustomizableObjectEditor/Private/MuCOE/SCustomizableObjectEditorViewportToolBar.cpp
+++ b/Engine/Plugins/Experimental/Mutable/Source/CustomizableObjectEditor/Private/MuCOE/SCustomizableObjectEditorViewportToolBar.cpp
@@ -1,38 +1,68 @@
 // Copyright Epic Games, Inc. All Rights Reserved.
 
 #include "MuCOE/SCustomizableObjectEditorViewportToolBar.h"
-#include "MuCOE/SCustomizableObjectEditorViewport.h"
+
+#include "AssetViewerSettings.h"
+#include "Containers/UnrealString.h"
+#include "Delegates/Delegate.h"
+#include "Editor.h"
+#include "Editor/EditorEngine.h"
+#include "Editor/EditorPerProjectUserSettings.h"
+#include "Editor/UnrealEdEngine.h"
+#include "Editor/UnrealEdTypes.h"
 #include "EditorViewportCommands.h"
-#include "MuCOE/CustomizableObjectEditorViewportMenuCommands.h"
-#include "MuCOE/CustomizableObjectEditorViewportLODCommands.h"
+#include "Engine/Engine.h"
+#include "Fonts/SlateFontInfo.h"
+#include "Framework/Commands/UIAction.h"
+#include "Framework/Commands/UICommandInfo.h"
+#include "Framework/MultiBox/MultiBoxBuilder.h"
+#include "Framework/MultiBox/MultiBoxDefs.h"
+#include "Framework/MultiBox/MultiBoxExtender.h"
+#include "Framework/Text/TextLayout.h"
+#include "GenericPlatform/ICursor.h"
+#include "HAL/Platform.h"
+#include "HAL/PlatformCrt.h"
+#include "Internationalization/Internationalization.h"
+#include "Layout/Children.h"
+#include "Layout/Margin.h"
+#include "LevelEditor.h"
+#include "Math/Color.h"
+#include "Math/Rotator.h"
+#include "Misc/AssertionMacros.h"
+#include "Misc/Attribute.h"
+#include "Modules/ModuleManager.h"
+#include "MuCO/CustomizableObjectInstance.h"
 #include "MuCOE/CustomizableObjectEditorViewportClient.h"
+#include "MuCOE/CustomizableObjectEditorViewportLODCommands.h"
+#include "MuCOE/CustomizableObjectEditorViewportMenuCommands.h"
+#include "MuCOE/ICustomizableObjectInstanceEditor.h"
+#include "MuCOE/SCustomizableObjectEditorViewport.h"
+#include "Preferences/PersonaOptions.h"
 #include "SEditorViewportToolBarMenu.h"
-#include "STransformViewportToolbar.h"
-#include "Editor/EditorPerProjectUserSettings.h"
 #include "SEditorViewportViewMenu.h"
-#include "Widgets/Colors/SColorPicker.h"
-#include "Widgets/Input/SNumericEntryBox.h"
-#include "Widgets/Input/STextComboBox.h"
-#include "AssetViewerSettings.h"
-#include "ISkeletonTree.h"
-#include "IEditableSkeleton.h"
-#include "Engine/PreviewMeshCollection.h"
-#include "Widgets/Input/SButton.h"
-#include "Widgets/Input/SCheckBox.h"
-#include "MuCOE/CustomizableObjectEditorViewportClient.h"
-#include "Preferences/PersonaOptions.h"
-#include "MuCOE/SCustomizableObjectEditorTransformViewportToolbar.h"
-#include "Widgets/Layout/SScaleBox.h"
 #include "SViewportToolBarComboMenu.h"
-
-#include "Editor/UnrealEdEngine.h"
-#include "UnrealEdGlobals.h"
 #include "Settings/LevelEditorViewportSettings.h"
+#include "SlotBase.h"
+#include "Styling/AppStyle.h"
+#include "Textures/SlateIcon.h"
+#include "Types/ISlateMetaData.h"
+#include "Types/SlateStructs.h"
+#include "UObject/UObjectGlobals.h"
+#include "UObject/UnrealNames.h"
+#include "UnrealEdGlobals.h"
+#include "Widgets/Input/SButton.h"
+#include "Widgets/Input/SMenuAnchor.h"
+#include "Widgets/Input/SSpinBox.h"
+#include "Widgets/Layout/SBorder.h"
+#include "Widgets/Layout/SBox.h"
+#include "Widgets/Layout/SScaleBox.h"
 #include "Widgets/Layout/SUniformGridPanel.h"
-#include "Internationalization/Internationalization.h"
-#include "Editor/EditorEngine.h"
-#include "LevelEditor.h"
-#include "Modules/ModuleManager.h"
+#include "Widgets/SBoxPanel.h"
+#include "Widgets/SCompoundWidget.h"
+#include "Widgets/Text/STextBlock.h"
+
+class SWidget;
+struct FSlateBrush;
 
 #define LOCTEXT_NAMESPACE "CustomizableObjectEditorViewportToolBar"
 
diff --git a/Engine/Plugins/Experimental/Mutable/Source/CustomizableObjectEditor/Private/MuCOE/SCustomizableObjectHighresScreenshot.cpp b/Engine/Plugins/Experimental/Mutable/Source/CustomizableObjectEditor/Private/MuCOE/SCustomizableObjectHighresScreenshot.cpp
--- a/Engine/Plugins/Experimental/Mutable/Source/CustomizableObjectEditor/Private/MuCOE/SCustomizableObjectHighresScreenshot.cpp
+++ b/Engine/Plugins/Experimental/Mutable/Source/CustomizableObjectEditor/Private/MuCOE/SCustomizableObjectHighresScreenshot.cpp
@@ -1,20 +1,35 @@
 // Copyright Epic Games, Inc. All Rights Reserved.
 
 #include "MuCOE/SCustomizableObjectHighresScreenshot.h"
-#include "Framework/Docking/TabManager.h"
+
+#include "Components/SkeletalMeshComponent.h"
+#include "Containers/EnumAsByte.h"
+#include "CoreGlobals.h"
+#include "Delegates/Delegate.h"
+#include "Engine/Engine.h"
+#include "Engine/EngineBaseTypes.h"
+#include "Engine/EngineTypes.h"
+#include "Engine/World.h"
 #include "Framework/Application/SlateApplication.h"
-#include "Widgets/Layout/SBorder.h"
+#include "Framework/Docking/TabManager.h"
+#include "Internationalization/Internationalization.h"
+#include "Layout/Children.h"
+#include "Math/IntPoint.h"
+#include "Math/Vector2D.h"
+#include "Misc/Attribute.h"
+#include "MuCOE/CustomizableObjectEditorViewportClient.h"
+#include "MuCOE/CustomizableObjectPreviewScene.h"
+#include "ShowFlags.h"
+#include "SlotBase.h"
+#include "Styling/AppStyle.h"
+#include "UnrealClient.h"
 #include "Widgets/Images/SImage.h"
-#include "Widgets/Layout/SGridPanel.h"
 #include "Widgets/Input/SButton.h"
-#include "Widgets/Layout/SSplitter.h"
-#include "Widgets/Input/SSlider.h"
-#include "Styling/AppStyle.h"
 #include "Widgets/Input/SNumericEntryBox.h"
-#include "MuCOE/CustomizableObjectEditorViewportClient.h"
-#include "SEditorViewport.h"
-#include "Components/SkeletalMeshComponent.h"
-#include "MuCOE/CustomizableObjectPreviewScene.h"
+#include "Widgets/Input/SSlider.h"
+#include "Widgets/Layout/SBorder.h"
+#include "Widgets/Layout/SSplitter.h"
+#include "Widgets/SBoxPanel.h"
 
 TWeakPtr<class SWindow> SCustomizableObjectHighresScreenshot::CurrentWindow = NULL;
 TWeakPtr<class SCustomizableObjectHighresScreenshot> SCustomizableObjectHighresScreenshot::CurrentDialog = NULL;
diff --git a/Engine/Plugins/Experimental/Mutable/Source/CustomizableObjectEditor/Private/MuCOE/SCustomizableObjectHighresScreenshot.h b/Engine/Plugins/Experimental/Mutable/Source/CustomizableObjectEditor/Private/MuCOE/SCustomizableObjectHighresScreenshot.h
--- a/Engine/Plugins/Experimental/Mutable/Source/CustomizableObjectEditor/Private/MuCOE/SCustomizableObjectHighresScreenshot.h
+++ b/Engine/Plugins/Experimental/Mutable/Source/CustomizableObjectEditor/Private/MuCOE/SCustomizableObjectHighresScreenshot.h
@@ -2,20 +2,25 @@
 
 #pragma once
 
-#include "CoreMinimal.h"
-#include "Layout/Visibility.h"
+#include "HAL/PlatformCrt.h"
+#include "HighResScreenshot.h"
 #include "Input/Reply.h"
+#include "Layout/Visibility.h"
+#include "Math/IntRect.h"
+#include "Math/UnrealMathSSE.h"
+#include "Misc/Optional.h"
+#include "Slate/SceneViewport.h"
+#include "Styling/SlateTypes.h"
+#include "Templates/SharedPointer.h"
+#include "Types/SlateEnums.h"
 #include "Widgets/DeclarativeSyntaxSupport.h"
-#include "ShowFlags.h"
+#include "Widgets/Input/SCheckBox.h"
 #include "Widgets/SCompoundWidget.h"
-#include "Styling/SlateTypes.h"
 #include "Widgets/SWindow.h"
 #include "Widgets/Text/STextBlock.h"
-#include "Slate/SceneViewport.h"
-#include "Widgets/Input/SCheckBox.h"
-#include "HighResScreenshot.h"
 
 class SButton;
+class SHorizontalBox;
 class USkeletalMeshComponent;
 
 class SCustomizableObjectHighresScreenshot : public SCompoundWidget
diff --git a/Engine/Plugins/Experimental/Mutable/Source/CustomizableObjectEditor/Private/MuCOE/SCustomizableObjectLayoutGrid.cpp b/Engine/Plugins/Experimental/Mutable/Source/CustomizableObjectEditor/Private/MuCOE/SCustomizableObjectLayoutGrid.cpp
--- a/Engine/Plugins/Experimental/Mutable/Source/CustomizableObjectEditor/Private/MuCOE/SCustomizableObjectLayoutGrid.cpp
+++ b/Engine/Plugins/Experimental/Mutable/Source/CustomizableObjectEditor/Private/MuCOE/SCustomizableObjectLayoutGrid.cpp
@@ -1,23 +1,55 @@
 // Copyright Epic Games, Inc. All Rights Reserved.
 
 #include "MuCOE/SCustomizableObjectLayoutGrid.h"
-#include "MuCOE/CustomizableObjectEditorModule.h"
-#include "MuCOE/CustomizableObjectLayout.h"
-#include "Styling/AppStyle.h"
-#include "Async/TaskGraphInterfaces.h"
 
-#include "CanvasItem.h"
+#include "BatchedElements.h"
 #include "CanvasTypes.h"
-
-#include "Framework/MultiBox/MultiBoxBuilder.h"
+#include "Engine/World.h"
+#include "Framework/Application/MenuStack.h"
 #include "Framework/Application/SlateApplication.h"
-
+#include "Framework/Commands/UIAction.h"
+#include "Framework/MultiBox/MultiBoxBuilder.h"
+#include "GenericPlatform/GenericApplication.h"
+#include "GenericPlatform/ICursor.h"
+#include "HAL/PlatformCrt.h"
+#include "HitProxies.h"
 #include "Input/CursorReply.h"
-
-#include "Widgets/Input/SNumericEntryBox.h"
-
+#include "Input/Events.h"
+#include "InputCoreTypes.h"
+#include "Internationalization/Internationalization.h"
+#include "Internationalization/Text.h"
+#include "Layout/Geometry.h"
+#include "Layout/PaintGeometry.h"
+#include "Layout/WidgetPath.h"
+#include "Math/Box2D.h"
+#include "Math/IntRect.h"
+#include "Math/Vector.h"
+#include "Misc/Guid.h"
 // Required for engine branch preprocessor defines.
 #include "MuCO/UnrealPortabilityHelpers.h"
+#include "MuCOE/CustomizableObjectLayout.h"
+#include "RHI.h"
+#include "Rendering/DrawElements.h"
+#include "Rendering/RenderingCommon.h"
+#include "Rendering/SlateLayoutTransform.h"
+#include "Rendering/SlateRenderTransform.h"
+#include "Rendering/SlateRenderer.h"
+#include "RenderingThread.h"
+#include "Styling/AppStyle.h"
+#include "Styling/CoreStyle.h"
+#include "Styling/SlateTypes.h"
+#include "Textures/SlateIcon.h"
+#include "Types/SlateStructs.h"
+#include "UnrealClient.h"
+#include "Widgets/Input/SNumericEntryBox.h"
+#include "Widgets/Layout/SBox.h"
+
+class FExtender;
+class FPaintArgs;
+class FRHICommandListImmediate;
+class FSlateRect;
+class FWidgetStyle;
+struct FSlateBrush;
 
 #define LOCTEXT_NAMESPACE "CustomizableObjectEditor"
 
diff --git a/Engine/Plugins/Experimental/Mutable/Source/CustomizableObjectEditor/Private/MuCOE/SCustomizableObjectLayoutGrid.h b/Engine/Plugins/Experimental/Mutable/Source/CustomizableObjectEditor/Private/MuCOE/SCustomizableObjectLayoutGrid.h
--- a/Engine/Plugins/Experimental/Mutable/Source/CustomizableObjectEditor/Private/MuCOE/SCustomizableObjectLayoutGrid.h
+++ b/Engine/Plugins/Experimental/Mutable/Source/CustomizableObjectEditor/Private/MuCOE/SCustomizableObjectLayoutGrid.h
@@ -2,9 +2,30 @@
 
 #pragma once
 
-#include "MuCOE/Nodes/CustomizableObjectNodeLayoutBlocks.h"
-#include "Widgets/SCompoundWidget.h"
+#include "Containers/Array.h"
+#include "Containers/Map.h"
+#include "Delegates/Delegate.h"
+#include "Input/CursorReply.h"
+#include "Input/Reply.h"
+#include "Math/Color.h"
+#include "Math/IntPoint.h"
+#include "Math/UnrealMathSSE.h"
+#include "Math/Vector2D.h"
+#include "Misc/Attribute.h"
+#include "Misc/Optional.h"
+#include "MuCOE/CustomizableObjectLayout.h"
+#include "Templates/SharedPointer.h"
 #include "Widgets/DeclarativeSyntaxSupport.h"
+#include "Widgets/SCompoundWidget.h"
+
+class FPaintArgs;
+class FSlateRect;
+class FSlateWindowElementList;
+class FWidgetStyle;
+struct FGeometry;
+struct FGuid;
+struct FKeyEvent;
+struct FPointerEvent;
 //#include "PreviewScene.h"
 
 
