Ran IWYU on Mutable plugin

diff --git a/Engine/Plugins/Experimental/Mutable/Source/CustomizableObjectEditor/Private/MuCOE/SCustomizableObjectNodeLayoutBlocksEditor.cpp b/Engine/Plugins/Experimental/Mutable/Source/CustomizableObjectEditor/Private/MuCOE/SCustomizableObjectNodeLayoutBlocksEditor.cpp
--- a/Engine/Plugins/Experimental/Mutable/Source/CustomizableObjectEditor/Private/MuCOE/SCustomizableObjectNodeLayoutBlocksEditor.cpp
+++ b/Engine/Plugins/Experimental/Mutable/Source/CustomizableObjectEditor/Private/MuCOE/SCustomizableObjectNodeLayoutBlocksEditor.cpp
@@ -1,32 +1,44 @@
 // Copyright Epic Games, Inc. All Rights Reserved.
 
 #include "MuCOE/SCustomizableObjectNodeLayoutBlocksEditor.h"
-#include "MuCOE/CustomizableObjectEditorModule.h"
-#include "MuCOE/CustomizableObjectEditorActions.h"
-#include "MuCOE/Nodes/CustomizableObjectNodeLayoutBlocksDetails.h"
-#include "MuCOE/CustomizableObjectCompiler.h"
+
+#include "Delegates/Delegate.h"
+#include "Framework/Application/SlateApplication.h"
+#include "Framework/Commands/Commands.h"
+#include "Framework/Commands/InputChord.h"
+#include "Framework/Commands/UIAction.h"
+#include "Framework/Commands/UICommandInfo.h"
+#include "Framework/Commands/UICommandList.h"
+#include "Framework/MultiBox/MultiBoxBuilder.h"
+#include "Framework/MultiBox/MultiBoxDefs.h"
+#include "HAL/PlatformCrt.h"
+#include "Internationalization/Internationalization.h"
+#include "Layout/Children.h"
+#include "Layout/Margin.h"
+#include "Layout/Visibility.h"
+#include "Math/Color.h"
+#include "Math/IntPoint.h"
+#include "Math/IntRect.h"
+#include "Math/Vector2D.h"
+#include "Misc/Attribute.h"
+#include "Misc/Guid.h"
 #include "MuCOE/CustomizableObjectEditorStyle.h"
 #include "MuCOE/CustomizableObjectLayout.h"
-
-#include "MuCOE/GenerateMutableSource/GenerateMutableSource.h"
-
-#include "MuCOE/CustomizableObjectCompiler.h"
-
-#include "MouseDeltaTracker.h"
+#include "MuCOE/ICustomizableObjectInstanceEditor.h"
 #include "MuCOE/SCustomizableObjectLayoutGrid.h"
-
+#include "SlotBase.h"
+#include "Styling/AppStyle.h"
+#include "Styling/CoreStyle.h"
+#include "Styling/SlateColor.h"
+#include "UObject/UObjectGlobals.h"
+#include "UObject/UnrealNames.h"
 #include "Widgets/Input/STextComboBox.h"
-#include "Widgets/Input/SEditableTextBox.h"
-#include "Widgets/Input/SButton.h"
-#include "Framework/MultiBox/MultiBoxBuilder.h"
-#include "FbxMeshUtils.h"
-#include "BusyCursor.h"
-#include "MeshBuild.h"
-#include "ObjectTools.h"
-
-#include "MuT/NodeLayout.h"
+#include "Widgets/Layout/SBorder.h"
+#include "Widgets/SBoxPanel.h"
+#include "Widgets/Text/STextBlock.h"
 
-#include "../Private/GeomFitUtils.h"
+class ISlateStyle;
+class SWidget;
 
 
 #define LOCTEXT_NAMESPACE "CustomizableObjectEditor"
diff --git a/Engine/Plugins/Experimental/Mutable/Source/CustomizableObjectEditor/Private/MuCOE/SCustomizableObjectNodeLayoutBlocksEditor.h b/Engine/Plugins/Experimental/Mutable/Source/CustomizableObjectEditor/Private/MuCOE/SCustomizableObjectNodeLayoutBlocksEditor.h
--- a/Engine/Plugins/Experimental/Mutable/Source/CustomizableObjectEditor/Private/MuCOE/SCustomizableObjectNodeLayoutBlocksEditor.h
+++ b/Engine/Plugins/Experimental/Mutable/Source/CustomizableObjectEditor/Private/MuCOE/SCustomizableObjectNodeLayoutBlocksEditor.h
@@ -2,10 +2,22 @@
 
 #pragma once
 
+#include "Containers/Array.h"
+#include "Containers/UnrealString.h"
+#include "Math/UnrealMathSSE.h"
+#include "Templates/SharedPointer.h"
+#include "Types/SlateEnums.h"
+#include "UObject/GCObject.h"
+#include "UObject/NameTypes.h"
+#include "Widgets/DeclarativeSyntaxSupport.h"
 #include "Widgets/SCompoundWidget.h"
-#include "MuCOE/ICustomizableObjectInstanceEditor.h"
-#include "MuCOE/Nodes/CustomizableObjectNodeLayoutBlocks.h"
-#include "MuCOE/CustomizableObjectLayout.h"
+
+class FReferenceCollector;
+class ICustomizableObjectInstanceEditor;
+class ISlateStyle;
+class SWidget;
+struct FCustomizableObjectLayoutBlock;
+struct FGuid;
 
 
 /**
diff --git a/Engine/Plugins/Experimental/Mutable/Source/CustomizableObjectEditor/Private/MuCOE/SCustomizableObjectNodeLayoutBlocksSelector.cpp b/Engine/Plugins/Experimental/Mutable/Source/CustomizableObjectEditor/Private/MuCOE/SCustomizableObjectNodeLayoutBlocksSelector.cpp
--- a/Engine/Plugins/Experimental/Mutable/Source/CustomizableObjectEditor/Private/MuCOE/SCustomizableObjectNodeLayoutBlocksSelector.cpp
+++ b/Engine/Plugins/Experimental/Mutable/Source/CustomizableObjectEditor/Private/MuCOE/SCustomizableObjectNodeLayoutBlocksSelector.cpp
@@ -1,24 +1,48 @@
 // Copyright Epic Games, Inc. All Rights Reserved.
 
 #include "MuCOE/SCustomizableObjectNodeLayoutBlocksSelector.h"
-#include "MuCOE/CustomizableObjectEditorModule.h"
-#include "MuCOE/CustomizableObjectEditorActions.h"
-#include "MuCOE/CustomizableObjectEditorModule.h"
+
+#include "Delegates/Delegate.h"
+#include "EdGraph/EdGraph.h"
+#include "EdGraph/EdGraphPin.h"
+#include "Framework/Application/SlateApplication.h"
+#include "Framework/Commands/Commands.h"
+#include "Framework/Commands/InputChord.h"
+#include "Framework/Commands/UIAction.h"
+#include "Framework/Commands/UICommandInfo.h"
+#include "Framework/Commands/UICommandList.h"
+#include "Framework/MultiBox/MultiBoxBuilder.h"
+#include "Framework/MultiBox/MultiBoxDefs.h"
+#include "HAL/PlatformCrt.h"
+#include "Internationalization/Internationalization.h"
+#include "Internationalization/Text.h"
+#include "Layout/Children.h"
+#include "Math/Color.h"
+#include "Math/IntPoint.h"
+#include "Math/Vector2D.h"
+#include "Misc/Attribute.h"
+#include "Misc/Guid.h"
 #include "MuCOE/CustomizableObjectEditorStyle.h"
 #include "MuCOE/CustomizableObjectLayout.h"
-
-#include "MouseDeltaTracker.h"
-#include "MuCOE/SCustomizableObjectLayoutGrid.h"
-
-#include "CustomizableObjectClasses.h"
-#include "MuCOE/CustomizableObjectEditor.h"
-#include "FbxMeshUtils.h"
-#include "BusyCursor.h"
 #include "MuCOE/GraphTraversal.h"
-#include "MeshBuild.h"
-#include "ObjectTools.h"
-
-#include "../Private/GeomFitUtils.h"
+#include "MuCOE/ICustomizableObjectInstanceEditor.h"
+#include "MuCOE/Nodes/CustomizableObjectNodeEditLayoutBlocks.h"
+#include "MuCOE/Nodes/CustomizableObjectNodeMaterial.h"
+#include "MuCOE/Nodes/CustomizableObjectNodeMaterialBase.h"
+#include "MuCOE/Nodes/CustomizableObjectNodeSkeletalMesh.h"
+#include "MuCOE/Nodes/CustomizableObjectNodeTable.h"
+#include "MuCOE/SCustomizableObjectLayoutGrid.h"
+#include "SlotBase.h"
+#include "Styling/AppStyle.h"
+#include "Templates/Casts.h"
+#include "UObject/NameTypes.h"
+#include "UObject/UObjectGlobals.h"
+#include "UObject/UnrealNames.h"
+#include "Widgets/Layout/SBorder.h"
+#include "Widgets/SBoxPanel.h"
+#include "Widgets/Text/STextBlock.h"
+
+class SWidget;
 
 #define LOCTEXT_NAMESPACE "CustomizableObjectEditor"
 
diff --git a/Engine/Plugins/Experimental/Mutable/Source/CustomizableObjectEditor/Private/MuCOE/SCustomizableObjectNodeLayoutBlocksSelector.h b/Engine/Plugins/Experimental/Mutable/Source/CustomizableObjectEditor/Private/MuCOE/SCustomizableObjectNodeLayoutBlocksSelector.h
--- a/Engine/Plugins/Experimental/Mutable/Source/CustomizableObjectEditor/Private/MuCOE/SCustomizableObjectNodeLayoutBlocksSelector.h
+++ b/Engine/Plugins/Experimental/Mutable/Source/CustomizableObjectEditor/Private/MuCOE/SCustomizableObjectNodeLayoutBlocksSelector.h
@@ -2,9 +2,19 @@
 
 #pragma once
 
+#include "Containers/Array.h"
+#include "Containers/UnrealString.h"
+#include "Math/UnrealMathSSE.h"
+#include "Templates/SharedPointer.h"
+#include "UObject/GCObject.h"
+#include "Widgets/DeclarativeSyntaxSupport.h"
 #include "Widgets/SCompoundWidget.h"
-#include "MuCOE/ICustomizableObjectInstanceEditor.h"
-#include "MuCOE/Nodes/CustomizableObjectNodeLayoutBlocks.h"
+
+class FReferenceCollector;
+class ICustomizableObjectInstanceEditor;
+class SWidget;
+struct FCustomizableObjectLayoutBlock;
+struct FGuid;
 
 
 /**
diff --git a/Engine/Plugins/Experimental/Mutable/Source/CustomizableObjectEditor/Private/MuCOE/SCustomizableObjectNodeObjectRTMorphTargetOverride.cpp b/Engine/Plugins/Experimental/Mutable/Source/CustomizableObjectEditor/Private/MuCOE/SCustomizableObjectNodeObjectRTMorphTargetOverride.cpp
--- a/Engine/Plugins/Experimental/Mutable/Source/CustomizableObjectEditor/Private/MuCOE/SCustomizableObjectNodeObjectRTMorphTargetOverride.cpp
+++ b/Engine/Plugins/Experimental/Mutable/Source/CustomizableObjectEditor/Private/MuCOE/SCustomizableObjectNodeObjectRTMorphTargetOverride.cpp
@@ -2,16 +2,30 @@
 
 #include "MuCOE/SCustomizableObjectNodeObjectRTMorphTargetOverride.h"
 
-
-#include "EdGraph/EdGraphPin.h"
-
-#include "Widgets/Input/SSearchBox.h"
-#include "Widgets/Layout/SExpandableArea.h"
+#include "CoreTypes.h"
+#include "Framework/Views/ITypedTableView.h"
+#include "Internationalization/Internationalization.h"
+#include "Internationalization/Text.h"
+#include "Layout/Children.h"
+#include "Misc/AssertionMacros.h"
+#include "Misc/Attribute.h"
+#include "MuCOE/Nodes/CustomizableObjectNodeObject.h"
 #include "SEnumCombo.h"
 #include "ScopedTransaction.h"
-#include "MuCOE/Nodes/CustomizableObjectNodeObject.h"
-
+#include "SlotBase.h"
 #include "Templates/Decay.h"
+#include "UObject/UObjectGlobals.h"
+#include "Widgets/Input/SSearchBox.h"
+#include "Widgets/Layout/SBorder.h"
+#include "Widgets/Layout/SExpandableArea.h"
+#include "Widgets/SBoxPanel.h"
+#include "Widgets/SNullWidget.h"
+#include "Widgets/Text/STextBlock.h"
+#include "Widgets/Views/SHeaderRow.h"
+
+class ITableRow;
+class STableViewBase;
+class SWidget;
 
 #define LOCTEXT_NAMESPACE "CustomizableObjectEditor"
 
diff --git a/Engine/Plugins/Experimental/Mutable/Source/CustomizableObjectEditor/Private/MuCOE/SCustomizableObjectNodeObjectRTMorphTargetOverride.h b/Engine/Plugins/Experimental/Mutable/Source/CustomizableObjectEditor/Private/MuCOE/SCustomizableObjectNodeObjectRTMorphTargetOverride.h
--- a/Engine/Plugins/Experimental/Mutable/Source/CustomizableObjectEditor/Private/MuCOE/SCustomizableObjectNodeObjectRTMorphTargetOverride.h
+++ b/Engine/Plugins/Experimental/Mutable/Source/CustomizableObjectEditor/Private/MuCOE/SCustomizableObjectNodeObjectRTMorphTargetOverride.h
@@ -2,15 +2,33 @@
 
 #pragma once
 
+#include "Containers/Array.h"
+#include "Containers/BitArray.h"
+#include "Containers/Set.h"
+#include "Containers/SparseArray.h"
+#include "Containers/UnrealString.h"
+#include "Delegates/Delegate.h"
+#include "HAL/Platform.h"
+#include "HAL/PlatformCrt.h"
+#include "Misc/Optional.h"
+#include "Templates/SharedPointer.h"
+#include "Templates/UnrealTemplate.h"
+#include "Types/SlateEnums.h"
+#include "UObject/GCObject.h"
+#include "UObject/NameTypes.h"
+#include "Widgets/DeclarativeSyntaxSupport.h"
 #include "Widgets/SCompoundWidget.h"
 #include "Widgets/Views/SListView.h"
-#include "Widgets/Input/SSearchBox.h"
-#include "Templates/SharedPointer.h"
+#include "Widgets/Views/STableRow.h"
 
-// Forward declarations
+class FReferenceCollector;
+class FText;
+class ITableRow;
 class SCustomizableObjectNodeSkeletalMeshRTMorphTargetOverride;
+class SSearchBox;
+class STableViewBase;
+class SWidget;
 class UCustomizableObjectNodeObject;
-struct FRealTimeMorphSelectionOverride;
 
 enum class ECustomizableObjectSelectionOverride : uint8;
 
diff --git a/Engine/Plugins/Experimental/Mutable/Source/CustomizableObjectEditor/Private/MuCOE/SCustomizableObjectNodeObjectStatesView.cpp b/Engine/Plugins/Experimental/Mutable/Source/CustomizableObjectEditor/Private/MuCOE/SCustomizableObjectNodeObjectStatesView.cpp
--- a/Engine/Plugins/Experimental/Mutable/Source/CustomizableObjectEditor/Private/MuCOE/SCustomizableObjectNodeObjectStatesView.cpp
+++ b/Engine/Plugins/Experimental/Mutable/Source/CustomizableObjectEditor/Private/MuCOE/SCustomizableObjectNodeObjectStatesView.cpp
@@ -1,23 +1,34 @@
 // Copyright Epic Games, Inc. All Rights Reserved.
 
 #include "MuCOE/SCustomizableObjectNodeObjectStatesView.h"
-#include "MuCOE/CustomizableObjectEditorModule.h"
-#include "PropertyCustomizationHelpers.h"
-#include "MuCOE/CustomizableObjectEditorUtilities.h"
-#include "MuCO/CustomizableObject.h"
-#include "MuCOE/Nodes/CustomizableObjectNodeObjectGroup.h"
-#include "MuCOE/Nodes/CustomizableObjectNodeObjectDetails.h"
-
-#include "DetailLayoutBuilder.h"
-#include "DetailCategoryBuilder.h"
 
+#include "Framework/Views/ITypedTableView.h"
+#include "GenericPlatform/ICursor.h"
+#include "HAL/PlatformCrt.h"
+#include "Input/Reply.h"
+#include "Internationalization/Internationalization.h"
+#include "Layout/Children.h"
+#include "Misc/Attribute.h"
+#include "Misc/CString.h"
+#include "MuCO/CustomizableObject.h"
+#include "MuCOE/CustomizableObjectGraph.h"
+#include "MuCOE/Nodes/CustomizableObjectNodeObject.h"
+#include "SlotBase.h"
+#include "Styling/AppStyle.h"
+#include "Templates/Casts.h"
+#include "Widgets/Images/SImage.h"
 #include "Widgets/Input/SButton.h"
-#include "Widgets/Text/STextBlock.h"
-#include "Widgets/Input/STextComboBox.h"
-#include "Widgets/Input/SSearchBox.h"
 #include "Widgets/Input/SCheckBox.h"
+#include "Widgets/Input/SComboButton.h"
+#include "Widgets/Input/SSearchBox.h"
+#include "Widgets/SBoxPanel.h"
+#include "Widgets/SWidget.h"
+#include "Widgets/Text/STextBlock.h"
+#include "Widgets/Views/ITableRow.h"
+#include "Widgets/Views/STableRow.h"
 
-#include "Input/Reply.h"
+class STableViewBase;
+struct FSlateBrush;
 
 
 #define LOCTEXT_NAMESPACE "CustomizableObjectNodeObjectStatesView"
diff --git a/Engine/Plugins/Experimental/Mutable/Source/CustomizableObjectEditor/Private/MuCOE/SCustomizableObjectNodeObjectStatesView.h b/Engine/Plugins/Experimental/Mutable/Source/CustomizableObjectEditor/Private/MuCOE/SCustomizableObjectNodeObjectStatesView.h
--- a/Engine/Plugins/Experimental/Mutable/Source/CustomizableObjectEditor/Private/MuCOE/SCustomizableObjectNodeObjectStatesView.h
+++ b/Engine/Plugins/Experimental/Mutable/Source/CustomizableObjectEditor/Private/MuCOE/SCustomizableObjectNodeObjectStatesView.h
@@ -2,12 +2,34 @@
 
 #pragma once
 
+#include "Containers/Array.h"
+#include "Containers/BitArray.h"
+#include "Containers/Set.h"
+#include "Containers/SparseArray.h"
+#include "Containers/UnrealString.h"
+#include "Delegates/Delegate.h"
+#include "HAL/Platform.h"
+#include "Input/Reply.h"
+#include "Internationalization/Text.h"
+#include "Layout/Visibility.h"
+#include "Misc/Optional.h"
+#include "Styling/SlateTypes.h"
+#include "Templates/SharedPointer.h"
+#include "Types/SlateEnums.h"
+#include "Widgets/DeclarativeSyntaxSupport.h"
 #include "Widgets/SCompoundWidget.h"
-#include "SListViewSelectorDropdownMenu.h"
-#include "Widgets/Input/SSearchBox.h"
-#include "Widgets/Input/SComboBox.h"
-#include "Widgets/Input/SComboButton.h"
-#include "Widgets/Images/SImage.h"
+#include "Widgets/Views/SListView.h"
+
+class ITableRow;
+class SComboButton;
+class SImage;
+class SSearchBox;
+class STableViewBase;
+class SVerticalBox;
+class SWidget;
+class UCustomizableObjectNodeObject;
+struct FSlateBrush;
+template <typename OptionType> class SComboBox;
 
 class SCustomizableObjectRuntimeParameter : public SCompoundWidget
 {
diff --git a/Engine/Plugins/Experimental/Mutable/Source/CustomizableObjectEditor/Private/MuCOE/SCustomizableObjectNodeSkeletalMeshRTMorphSelector.cpp b/Engine/Plugins/Experimental/Mutable/Source/CustomizableObjectEditor/Private/MuCOE/SCustomizableObjectNodeSkeletalMeshRTMorphSelector.cpp
--- a/Engine/Plugins/Experimental/Mutable/Source/CustomizableObjectEditor/Private/MuCOE/SCustomizableObjectNodeSkeletalMeshRTMorphSelector.cpp
+++ b/Engine/Plugins/Experimental/Mutable/Source/CustomizableObjectEditor/Private/MuCOE/SCustomizableObjectNodeSkeletalMeshRTMorphSelector.cpp
@@ -1,12 +1,32 @@
 // Copyright Epic Games, Inc. All Rights Reserved.
 
 #include "MuCOE/SCustomizableObjectNodeSkeletalMeshRTMorphSelector.h"
-#include "EdGraph/EdGraphPin.h"
-#include "Widgets/Input/SSearchBox.h"
+
+#include "Animation/MorphTarget.h"
+#include "CoreTypes.h"
+#include "Engine/SkeletalMesh.h"
+#include "Framework/Views/ITypedTableView.h"
+#include "Internationalization/Internationalization.h"
+#include "Internationalization/Text.h"
+#include "Layout/Children.h"
+#include "Misc/AssertionMacros.h"
+#include "Misc/Attribute.h"
 #include "MuCOE/Nodes/CustomizableObjectNodeSkeletalMesh.h"
 #include "ScopedTransaction.h"
-
+#include "SlotBase.h"
 #include "Templates/Decay.h"
+#include "UObject/ObjectPtr.h"
+#include "UObject/UObjectGlobals.h"
+#include "Widgets/Input/SCheckBox.h"
+#include "Widgets/Input/SSearchBox.h"
+#include "Widgets/SBoxPanel.h"
+#include "Widgets/SNullWidget.h"
+#include "Widgets/Text/STextBlock.h"
+#include "Widgets/Views/SHeaderRow.h"
+
+class ITableRow;
+class STableViewBase;
+class SWidget;
 
 #define LOCTEXT_NAMESPACE "CustomizableObjectEditor"
 
diff --git a/Engine/Plugins/Experimental/Mutable/Source/CustomizableObjectEditor/Private/MuCOE/SCustomizableObjectNodeSkeletalMeshRTMorphSelector.h b/Engine/Plugins/Experimental/Mutable/Source/CustomizableObjectEditor/Private/MuCOE/SCustomizableObjectNodeSkeletalMeshRTMorphSelector.h
--- a/Engine/Plugins/Experimental/Mutable/Source/CustomizableObjectEditor/Private/MuCOE/SCustomizableObjectNodeSkeletalMeshRTMorphSelector.h
+++ b/Engine/Plugins/Experimental/Mutable/Source/CustomizableObjectEditor/Private/MuCOE/SCustomizableObjectNodeSkeletalMeshRTMorphSelector.h
@@ -2,20 +2,37 @@
 
 #pragma once
 
+#include "Containers/Array.h"
+#include "Containers/BitArray.h"
+#include "Containers/Set.h"
+#include "Containers/SparseArray.h"
+#include "Containers/UnrealString.h"
+#include "Delegates/Delegate.h"
+#include "HAL/Platform.h"
+#include "HAL/PlatformCrt.h"
+#include "Misc/Optional.h"
+#include "Styling/SlateTypes.h"
+#include "Templates/SharedPointer.h"
+#include "Templates/UnrealTemplate.h"
+#include "Types/SlateEnums.h"
+#include "UObject/GCObject.h"
+#include "UObject/NameTypes.h"
+#include "Widgets/DeclarativeSyntaxSupport.h"
 #include "Widgets/SCompoundWidget.h"
 #include "Widgets/Views/SListView.h"
-#include "Widgets/Input/SSearchBox.h"
-#include "Widgets/Input/SCheckBox.h"
-#include "Templates/SharedPointer.h"
-
+#include "Widgets/Views/STableRow.h"
 
-// Forward declarations
-class UCustomizableObjectNodeSkeletalMesh;
+class FReferenceCollector;
+class FText;
+class ITableRow;
 class SCustomizableObjectNodeSkeletalMeshRTMorphSelector;
+class SSearchBox;
+class STableViewBase;
+class SWidget;
+class UCustomizableObjectNodeSkeletalMesh;
 
 namespace CustomizableObjectNodeSkeletalMeshRTMorphSelector_Internal
 {
-    class SRowWidget;
     
     class FRowItemInfo
     {
diff --git a/Engine/Plugins/Experimental/Mutable/Source/CustomizableObjectEditor/Private/MuCOE/SCustomizableObjectProperties.cpp b/Engine/Plugins/Experimental/Mutable/Source/CustomizableObjectEditor/Private/MuCOE/SCustomizableObjectProperties.cpp
--- a/Engine/Plugins/Experimental/Mutable/Source/CustomizableObjectEditor/Private/MuCOE/SCustomizableObjectProperties.cpp
+++ b/Engine/Plugins/Experimental/Mutable/Source/CustomizableObjectEditor/Private/MuCOE/SCustomizableObjectProperties.cpp
@@ -1,13 +1,29 @@
 // Copyright Epic Games, Inc. All Rights Reserved.
 
 #include "MuCOE/SCustomizableObjectProperties.h"
-#include "MuCOE/CustomizableObjectEditorModule.h"
 
+#include "Containers/UnrealString.h"
+#include "Framework/Views/ITypedTableView.h"
+#include "HAL/PlatformMath.h"
+#include "Internationalization/Internationalization.h"
+#include "Internationalization/Text.h"
+#include "Layout/Children.h"
+#include "Layout/Margin.h"
+#include "Layout/Visibility.h"
+#include "Math/Color.h"
+#include "Misc/Attribute.h"
+#include "SlotBase.h"
+#include "Styling/SlateColor.h"
 #include "Widgets/Input/SComboButton.h"
-
-#include "PackageTools.h"
-#include "AssetRegistry/AssetRegistryModule.h"
-#include "AssetSelection.h"
+#include "Widgets/SBoxPanel.h"
+#include "Widgets/Text/STextBlock.h"
+#include "Widgets/Views/SExpanderArrow.h"
+#include "Widgets/Views/SHeaderRow.h"
+#include "Widgets/Views/STableRow.h"
+
+class ITableRow;
+class SWidget;
+class UCustomizableObject;
 
 #define LOCTEXT_NAMESPACE "SCustomizableObjectProperties"
 
diff --git a/Engine/Plugins/Experimental/Mutable/Source/CustomizableObjectEditor/Private/MuCOE/SCustomizableObjectProperties.h b/Engine/Plugins/Experimental/Mutable/Source/CustomizableObjectEditor/Private/MuCOE/SCustomizableObjectProperties.h
--- a/Engine/Plugins/Experimental/Mutable/Source/CustomizableObjectEditor/Private/MuCOE/SCustomizableObjectProperties.h
+++ b/Engine/Plugins/Experimental/Mutable/Source/CustomizableObjectEditor/Private/MuCOE/SCustomizableObjectProperties.h
@@ -3,9 +3,27 @@
 #pragma once
 
 
-#include "MuCOE/CustomizableObjectEditorModule.h"
+#include "Containers/Array.h"
+#include "Containers/BitArray.h"
+#include "Containers/Map.h"
+#include "Containers/Set.h"
+#include "Containers/SparseArray.h"
+#include "Delegates/Delegate.h"
+#include "HAL/PlatformCrt.h"
+#include "Misc/Optional.h"
+#include "Templates/SharedPointer.h"
+#include "Types/SlateConstants.h"
+#include "Types/SlateEnums.h"
+#include "Widgets/DeclarativeSyntaxSupport.h"
+#include "Widgets/SCompoundWidget.h"
+#include "Widgets/Views/STableViewBase.h"
 #include "Widgets/Views/STreeView.h"
 
+class FString;
+class ITableRow;
+class SWidget;
+class UCustomizableObject;
+
 
 //////////////////////////////////////////////////////////////////////////
 // SClassViewer
diff --git a/Engine/Plugins/Experimental/Mutable/Source/CustomizableObjectEditor/Private/MuCOE/SCustomizableObjectSystemVersion.cpp b/Engine/Plugins/Experimental/Mutable/Source/CustomizableObjectEditor/Private/MuCOE/SCustomizableObjectSystemVersion.cpp
--- a/Engine/Plugins/Experimental/Mutable/Source/CustomizableObjectEditor/Private/MuCOE/SCustomizableObjectSystemVersion.cpp
+++ b/Engine/Plugins/Experimental/Mutable/Source/CustomizableObjectEditor/Private/MuCOE/SCustomizableObjectSystemVersion.cpp
@@ -1,13 +1,18 @@
 // Copyright Epic Games, Inc. All Rights Reserved.
 
 #include "MuCOE/SCustomizableObjectSystemVersion.h"
-#include "MuCOE/CustomizableObjectEditorModule.h"
-#include "MuCO/ICustomizableObjectModule.h"
 
+#include "Containers/UnrealString.h"
+#include "Internationalization/Text.h"
+#include "Layout/Children.h"
+#include "Misc/Attribute.h"
+#include "MuCO/ICustomizableObjectModule.h"
+#include "SlotBase.h"
 #include "Widgets/Layout/SGridPanel.h"
+#include "Widgets/SBoxPanel.h"
 #include "Widgets/Text/STextBlock.h"
 
-#include "Interfaces/IPluginManager.h"
+struct FGeometry;
 
 
 void SCustomizableObjectSystemVersion::Construct( const FArguments& InArgs )
diff --git a/Engine/Plugins/Experimental/Mutable/Source/CustomizableObjectEditor/Private/MuCOE/SCustomizableObjectSystemVersion.h b/Engine/Plugins/Experimental/Mutable/Source/CustomizableObjectEditor/Private/MuCOE/SCustomizableObjectSystemVersion.h
--- a/Engine/Plugins/Experimental/Mutable/Source/CustomizableObjectEditor/Private/MuCOE/SCustomizableObjectSystemVersion.h
+++ b/Engine/Plugins/Experimental/Mutable/Source/CustomizableObjectEditor/Private/MuCOE/SCustomizableObjectSystemVersion.h
@@ -2,9 +2,10 @@
 
 #pragma once
 
-#include "MuCO/CustomizableObjectSystem.h"
-#include "Widgets/SCompoundWidget.h"
 #include "Widgets/DeclarativeSyntaxSupport.h"
+#include "Widgets/SCompoundWidget.h"
+
+struct FGeometry;
 
 
 class CUSTOMIZABLEOBJECTEDITOR_API SCustomizableObjectSystemVersion : public SCompoundWidget
diff --git a/Engine/Plugins/Experimental/Mutable/Source/CustomizableObjectEditor/Private/MuCOE/SMutableBoolViewer.cpp b/Engine/Plugins/Experimental/Mutable/Source/CustomizableObjectEditor/Private/MuCOE/SMutableBoolViewer.cpp
--- a/Engine/Plugins/Experimental/Mutable/Source/CustomizableObjectEditor/Private/MuCOE/SMutableBoolViewer.cpp
+++ b/Engine/Plugins/Experimental/Mutable/Source/CustomizableObjectEditor/Private/MuCOE/SMutableBoolViewer.cpp
@@ -2,9 +2,13 @@
 
 #include "MuCOE/SMutableBoolViewer.h"
 
-#include "Widgets/SBoxPanel.h"
+#include "Internationalization/Internationalization.h"
 #include "Internationalization/Text.h"
+#include "Layout/Children.h"
+#include "Misc/Attribute.h"
+#include "SlotBase.h"
 #include "Styling/SlateColor.h"
+#include "Widgets/SBoxPanel.h"
 #include "Widgets/Text/STextBlock.h"
 
 
diff --git a/Engine/Plugins/Experimental/Mutable/Source/CustomizableObjectEditor/Private/MuCOE/SMutableCodeViewer.cpp b/Engine/Plugins/Experimental/Mutable/Source/CustomizableObjectEditor/Private/MuCOE/SMutableCodeViewer.cpp
--- a/Engine/Plugins/Experimental/Mutable/Source/CustomizableObjectEditor/Private/MuCOE/SMutableCodeViewer.cpp
+++ b/Engine/Plugins/Experimental/Mutable/Source/CustomizableObjectEditor/Private/MuCOE/SMutableCodeViewer.cpp
@@ -2,38 +2,97 @@
 
 #include "MuCOE/SMutableCodeViewer.h"
 
+#include "Containers/StringConv.h"
+#include "CoreGlobals.h"
+#include "DesktopPlatformModule.h"
+#include "EditorDirectories.h"
+#include "Framework/Application/SlateApplication.h"
+#include "Framework/Commands/UIAction.h"
+#include "Framework/Commands/UICommandInfo.h"
+#include "Framework/MultiBox/MultiBoxBuilder.h"
+#include "Framework/MultiBox/MultiBoxDefs.h"
+#include "Framework/Text/TextLayout.h"
+#include "Framework/Views/ITypedTableView.h"
+#include "HAL/PlatformCrt.h"
+#include "HAL/UnrealMemory.h"
+#include "IDesktopPlatform.h"
+#include "Internationalization/Internationalization.h"
+#include "Internationalization/Text.h"
+#include "Layout/BasicLayoutWidgetSlot.h"
+#include "Layout/Children.h"
+#include "Layout/Margin.h"
+#include "Logging/LogCategory.h"
+#include "Logging/LogMacros.h"
+#include "Math/Color.h"
+#include "Math/IntPoint.h"
+#include "Math/UnrealMathSSE.h"
+#include "Math/Vector2D.h"
+#include "Misc/AssertionMacros.h"
+#include "Misc/Attribute.h"
+#include "MuCO/CustomizableObject.h"
+#include "MuCOE/SMutableBoolViewer.h"
+#include "MuCOE/SMutableColorPreviewBox.h"
+#include "MuCOE/SMutableColorViewer.h"
+#include "MuCOE/SMutableConstantsWidget.h"
+#include "MuCOE/SMutableCurveViewer.h"
 #include "MuCOE/SMutableImageViewer.h"
+#include "MuCOE/SMutableIntViewer.h"
+#include "MuCOE/SMutableLayoutViewer.h"
 #include "MuCOE/SMutableMeshViewer.h"
-#include "MuCOE/CustomizableObjectCompiler.h"
-#include "MuT/ErrorLogPrivate.h"
-#include "MuT/Streams.h"
-#include "MuCO/CustomizableInstancePrivateData.h"
-#include "MuR/ModelPrivate.h"
-#include "MuR/SystemPrivate.h"
-#include "MuR/Image.h"
+#include "MuCOE/SMutableParametersWidget.h"
+#include "MuCOE/SMutableProjectorViewer.h"
+#include "MuCOE/SMutableScalarViewer.h"
+#include "MuCOE/SMutableSkeletonViewer.h"
+#include "MuCOE/SMutableStringViewer.h"
 #include "MuR/Mesh.h"
+#include "MuR/ModelPrivate.h"
 #include "MuR/MutableString.h"
-
-#include "Framework/MultiBox/MultiBoxBuilder.h"
+#include "MuR/Ptr.h"
+#include "MuR/RefCounted.h"
+#include "MuR/Serialisation.h"
+#include "MuR/Settings.h"
+#include "MuR/SystemPrivate.h"
+#include "MuT/ErrorLogPrivate.h"
+#include "MuT/Streams.h"
+#include "MuT/Table.h"
+#include "SlotBase.h"
+#include "Styling/AppStyle.h"
+#include "Styling/CoreStyle.h"
+#include "Styling/ISlateStyle.h"
+#include "Styling/SlateColor.h"
+#include "Templates/UnrealTemplate.h"
+#include "Textures/SlateIcon.h"
+#include "Trace/Detail/Channel.h"
+#include "UObject/NameTypes.h"
+#include "UObject/UnrealNames.h"
+#include "Widgets/Images/SImage.h"
 #include "Widgets/Input/SButton.h"
+#include "Widgets/Input/SNumericEntryBox.h"
 #include "Widgets/Input/STextComboBox.h"
-#include "Widgets/Views/STreeView.h"
-#include "Widgets/Images/SImage.h"
-#include "Widgets/Layout/SSplitter.h"
+#include "Widgets/Layout/SBorder.h"
 #include "Widgets/Layout/SScrollBox.h"
-#include "Widgets/Input/SNumericEntryBox.h"
-#include "Widgets/Docking/SDockTab.h"
-#include "EditorDirectories.h"
-#include "DesktopPlatformModule.h"
-#include "SEnumCombo.h"
-#include "HAL/PlatformApplicationMisc.h"
-#include "Internationalization/Text.h"
-
-#include "MuCOE/SMutableColorPreviewBox.h"
-#include "MuCOE/SMutableConstantsWidget.h"
-
+#include "Widgets/Layout/SSplitter.h"
+#include "Widgets/SBoxPanel.h"
+#include "Widgets/SNullWidget.h"
+#include "Widgets/SOverlay.h"
+#include "Widgets/Text/STextBlock.h"
+#include "Widgets/Views/ITableRow.h"
+#include "Widgets/Views/SExpanderArrow.h"
 #include "Widgets/Views/SHeaderRow.h"
-#include "MuCOE/SMutableColorPreviewBox.h"
+#include "Widgets/Views/STableRow.h"
+#include "Widgets/Views/STreeView.h"
+
+#include <functional>
+
+class FExtender;
+class FReferenceCollector;
+class FUICommandList;
+class SWidget;
+namespace mu { struct Curve; }
+namespace mu { struct PROJECTOR; }
+namespace mu { struct SHAPE; }
+struct FGeometry;
+struct FSlateBrush;
 
 
 #define LOCTEXT_NAMESPACE "SMutableDebugger"
diff --git a/Engine/Plugins/Experimental/Mutable/Source/CustomizableObjectEditor/Private/MuCOE/SMutableCodeViewer.h b/Engine/Plugins/Experimental/Mutable/Source/CustomizableObjectEditor/Private/MuCOE/SMutableCodeViewer.h
--- a/Engine/Plugins/Experimental/Mutable/Source/CustomizableObjectEditor/Private/MuCOE/SMutableCodeViewer.h
+++ b/Engine/Plugins/Experimental/Mutable/Source/CustomizableObjectEditor/Private/MuCOE/SMutableCodeViewer.h
@@ -2,31 +2,61 @@
 
 #pragma once
 
-#include "MuCO/CustomizableObject.h"
-#include "MuCOE/SMutableBoolViewer.h"
-#include "MuCOE/SMutableColorViewer.h"
-#include "MuCOE/SMutableImageViewer.h"
-#include "MuCOE/SMutableIntViewer.h"
-#include "MuCOE/SMutableLayoutViewer.h"
-#include "MuCOE/SMutableMeshViewer.h"
-#include "MuCOE/SMutableParametersWidget.h"
-#include "MuCOE/SMutableProjectorViewer.h"
-#include "MuCOE/SMutableScalarViewer.h"
-#include "MuCOE/SMutableStringViewer.h"
-#include "MuCOE/SMutableSkeletonViewer.h"
-#include "MuCOE/SMutableCurveViewer.h"
-#include "Framework/Docking/TabManager.h"
-#include "MuR/Operations.h"
+#include "Containers/Array.h"
+#include "Containers/BitArray.h"
+#include "Containers/Map.h"
+#include "Containers/Set.h"
+#include "Containers/SparseArray.h"
+#include "Containers/UnrealString.h"
+#include "Delegates/Delegate.h"
+#include "HAL/Platform.h"
+#include "Input/Reply.h"
+#include "Internationalization/Text.h"
+#include "Layout/Visibility.h"
+#include "Misc/Optional.h"
+#include "MuR/Image.h"
+#include "MuR/MemoryPrivate.h"
+#include "MuR/Mesh.h"
 #include "MuR/Model.h"
+#include "MuR/ModelPrivate.h"
+#include "MuR/MutableMath.h"
+#include "MuR/Operations.h"
+#include "MuR/System.h"
+#include "Templates/SharedPointer.h"
+#include "Templates/TypeHash.h"
+#include "Types/SlateConstants.h"
+#include "Types/SlateEnums.h"
+#include "UObject/GCObject.h"
+#include "Widgets/DeclarativeSyntaxSupport.h"
 #include "Widgets/SCompoundWidget.h"
+#include "Widgets/Views/STableViewBase.h"
 #include "Widgets/Views/STreeView.h"
 
-#include "MuR/ModelPrivate.h"
-#include "MuR/ParametersPrivate.h"
+#include <utility>
 
-// Forward declarations
 class FMutableCodeTreeElement;
+class FReferenceCollector;
+class ITableRow;
+class SBorder;
+class SMutableBoolViewer;
+class SMutableColorViewer;
 class SMutableConstantsWidget;
+class SMutableCurveViewer;
+class SMutableImageViewer;
+class SMutableIntViewer;
+class SMutableLayoutViewer;
+class SMutableMeshViewer;
+class SMutableParametersWidget;
+class SMutableProjectorViewer;
+class SMutableScalarViewer;
+class SMutableSkeletonViewer;
+class SMutableStringViewer;
+class STextComboBox;
+class SWidget;
+namespace mu { struct Curve; }
+namespace mu { struct PROJECTOR; }
+namespace mu { struct SHAPE; }
+struct FGeometry;
 
 /** This widget shows the internal Mutable Code for debugging purposes. 
  * This is not the Unreal source graph in the UCustomizableObject, but the actual Mutable virtual machine graph.
diff --git a/Engine/Plugins/Experimental/Mutable/Source/CustomizableObjectEditor/Private/MuCOE/SMutableColorPreviewBox.cpp b/Engine/Plugins/Experimental/Mutable/Source/CustomizableObjectEditor/Private/MuCOE/SMutableColorPreviewBox.cpp
--- a/Engine/Plugins/Experimental/Mutable/Source/CustomizableObjectEditor/Private/MuCOE/SMutableColorPreviewBox.cpp
+++ b/Engine/Plugins/Experimental/Mutable/Source/CustomizableObjectEditor/Private/MuCOE/SMutableColorPreviewBox.cpp
@@ -2,10 +2,12 @@
 
 #include "MuCOE/SMutableColorPreviewBox.h"
 
-#include "Widgets/SBoxPanel.h"
-#include "Widgets/Layout/SBorder.h"
-#include "Styling/SlateColor.h"
+#include "Layout/Children.h"
+#include "SlotBase.h"
 #include "Styling/AppStyle.h"
+#include "Styling/SlateColor.h"
+#include "Widgets/Layout/SBorder.h"
+#include "Widgets/SBoxPanel.h"
 
 #define LOCTEXT_NAMESPACE "CustomizableObjectEditor"
 
diff --git a/Engine/Plugins/Experimental/Mutable/Source/CustomizableObjectEditor/Private/MuCOE/SMutableColorPreviewBox.h b/Engine/Plugins/Experimental/Mutable/Source/CustomizableObjectEditor/Private/MuCOE/SMutableColorPreviewBox.h
--- a/Engine/Plugins/Experimental/Mutable/Source/CustomizableObjectEditor/Private/MuCOE/SMutableColorPreviewBox.h
+++ b/Engine/Plugins/Experimental/Mutable/Source/CustomizableObjectEditor/Private/MuCOE/SMutableColorPreviewBox.h
@@ -2,9 +2,10 @@
 
 #pragma once
 
-#include "Widgets/SCompoundWidget.h"
-#include "Widgets/DeclarativeSyntaxSupport.h"
+#include "Misc/Attribute.h"
 #include "Styling/SlateColor.h"
+#include "Widgets/DeclarativeSyntaxSupport.h"
+#include "Widgets/SCompoundWidget.h"
 
 /** Custom widget designed to display a color on a box */
 class SMutableColorPreviewBox : public SCompoundWidget
diff --git a/Engine/Plugins/Experimental/Mutable/Source/CustomizableObjectEditor/Private/MuCOE/SMutableColorViewer.cpp b/Engine/Plugins/Experimental/Mutable/Source/CustomizableObjectEditor/Private/MuCOE/SMutableColorViewer.cpp
--- a/Engine/Plugins/Experimental/Mutable/Source/CustomizableObjectEditor/Private/MuCOE/SMutableColorViewer.cpp
+++ b/Engine/Plugins/Experimental/Mutable/Source/CustomizableObjectEditor/Private/MuCOE/SMutableColorViewer.cpp
@@ -1,14 +1,18 @@
 // Copyright Epic Games, Inc. All Rights Reserved.
 
 #include "MuCOE/SMutableColorViewer.h"
-#include "MuCOE/SMutableColorPreviewBox.h"
 
-#include "GenericPlatform/GenericPlatform.h"
-#include "Widgets/SBoxPanel.h"
+#include "HAL/Platform.h"
+#include "Internationalization/Internationalization.h"
 #include "Internationalization/Text.h"
-#include "Templates/SharedPointer.h"
+#include "Layout/Children.h"
 #include "Math/Color.h"
+#include "Misc/Attribute.h"
+#include "MuCOE/SMutableColorPreviewBox.h"
+#include "SlotBase.h"
 #include "Styling/SlateColor.h"
+#include "Templates/SharedPointer.h"
+#include "Widgets/SBoxPanel.h"
 #include "Widgets/Text/STextBlock.h"
 
 #define LOCTEXT_NAMESPACE "CustomizableObjectEditor"
diff --git a/Engine/Plugins/Experimental/Mutable/Source/CustomizableObjectEditor/Private/MuCOE/SMutableColorViewer.h b/Engine/Plugins/Experimental/Mutable/Source/CustomizableObjectEditor/Private/MuCOE/SMutableColorViewer.h
--- a/Engine/Plugins/Experimental/Mutable/Source/CustomizableObjectEditor/Private/MuCOE/SMutableColorViewer.h
+++ b/Engine/Plugins/Experimental/Mutable/Source/CustomizableObjectEditor/Private/MuCOE/SMutableColorViewer.h
@@ -2,12 +2,13 @@
 
 #pragma once
 
-#include "MuCOE/SMutableColorPreviewBox.h"
-#include "Widgets/SCompoundWidget.h"
-#include "Widgets/DeclarativeSyntaxSupport.h"
 #include "Internationalization/Text.h"
-#include "Templates/SharedPointer.h"
 #include "Styling/SlateColor.h"
+#include "Templates/SharedPointer.h"
+#include "Widgets/DeclarativeSyntaxSupport.h"
+#include "Widgets/SCompoundWidget.h"
+
+class SMutableColorPreviewBox;
 
 
 class SMutableColorViewer : public SCompoundWidget
diff --git a/Engine/Plugins/Experimental/Mutable/Source/CustomizableObjectEditor/Private/MuCOE/SMutableConstantsWidget.cpp b/Engine/Plugins/Experimental/Mutable/Source/CustomizableObjectEditor/Private/MuCOE/SMutableConstantsWidget.cpp
--- a/Engine/Plugins/Experimental/Mutable/Source/CustomizableObjectEditor/Private/MuCOE/SMutableConstantsWidget.cpp
+++ b/Engine/Plugins/Experimental/Mutable/Source/CustomizableObjectEditor/Private/MuCOE/SMutableConstantsWidget.cpp
@@ -2,29 +2,49 @@
 
 #include "MuCOE/SMutableConstantsWidget.h"
 
-#include "MuR/MutableString.h"
-#include "MuR/ModelPrivate.h"
-
+#include "Algo/StableSort.h"
+#include "Containers/UnrealString.h"
+#include "CoreGlobals.h"
+#include "Framework/Text/TextLayout.h"
+#include "Framework/Views/ITypedTableView.h"
+#include "HAL/PlatformCrt.h"
+#include "Internationalization/Internationalization.h"
+#include "Internationalization/Text.h"
+#include "Layout/BasicLayoutWidgetSlot.h"
+#include "Layout/Children.h"
+#include "Logging/LogCategory.h"
+#include "Logging/LogMacros.h"
+#include "Math/UnrealMathUtility.h"
+#include "Misc/AssertionMacros.h"
+#include "Misc/Attribute.h"
+#include "MuCOE/SMutableCodeViewer.h"
 #include "MuR/Image.h"
 #include "MuR/Layout.h"
 #include "MuR/Mesh.h"
-
-#include "MuCOE/SMutableImageViewer.h"
-#include "MuCOE/SMutableMeshViewer.h"
-#include "MuCOE/SMutableStringViewer.h"
-
-#include "MuCOE/SMutableCodeViewer.h"
-#include "Algo/StableSort.h"
-#include "Components/TextBlock.h"			
-#include "Internationalization/Text.h"		
+#include "MuR/ModelPrivate.h"
+#include "MuR/Operations.h"
+#include "MuR/ParametersPrivate.h"
+#include "MuR/Ptr.h"
+#include "MuR/Skeleton.h"
 #include "MuT/TypeInfo.h"
-#include "Widgets/Layout/SExpandableArea.h"	
-#include "Widgets/Layout/SScrollBox.h"		
+#include "SlotBase.h"
+#include "Templates/Tuple.h"
+#include "Templates/UnrealTemplate.h"
+#include "Trace/Detail/Channel.h"
+#include "Widgets/Input/SButton.h"
+#include "Widgets/Layout/SExpandableArea.h"
+#include "Widgets/Layout/SScrollBox.h"
 #include "Widgets/Layout/SSeparator.h"
+#include "Widgets/SBoxPanel.h"
+#include "Widgets/SNullWidget.h"
+#include "Widgets/Text/STextBlock.h"
 #include "Widgets/Views/SListView.h"
+#include "Widgets/Views/STableRow.h"
+#include "Widgets/Views/STableViewBase.h"
 #include "Widgets/Views/STileView.h"
-#include "Widgets/Input/SButton.h"
-#include "Math/UnrealMathUtility.h"
+
+class ITableRow;
+class SWidget;
 
 
 #define LOCTEXT_NAMESPACE "CustomizableObjectEditor"
diff --git a/Engine/Plugins/Experimental/Mutable/Source/CustomizableObjectEditor/Private/MuCOE/SMutableConstantsWidget.h b/Engine/Plugins/Experimental/Mutable/Source/CustomizableObjectEditor/Private/MuCOE/SMutableConstantsWidget.h
--- a/Engine/Plugins/Experimental/Mutable/Source/CustomizableObjectEditor/Private/MuCOE/SMutableConstantsWidget.h
+++ b/Engine/Plugins/Experimental/Mutable/Source/CustomizableObjectEditor/Private/MuCOE/SMutableConstantsWidget.h
@@ -2,25 +2,37 @@
 
 #pragma once
 
-#include "MuR/MutableString.h"
-#include "MuT/Table.h"
-#include "MuR/ModelPrivate.h"
-#include "MuR/MutableMath.h"
-#include "MuR/ParametersPrivate.h"
-#include "MuR/SystemPrivate.h"
-
+#include "Containers/Array.h"
+#include "Containers/BitArray.h"
+#include "Containers/Set.h"
+#include "Containers/SparseArray.h"
+#include "Delegates/Delegate.h"
+#include "HAL/Platform.h"
+#include "Input/Reply.h"
+#include "Internationalization/Text.h"
+#include "Misc/Optional.h"
 #include "MuR/Image.h"
-#include "MuR/Layout.h"
+#include "MuR/MemoryPrivate.h"
 #include "MuR/Mesh.h"
-
+#include "MuR/MutableMath.h"
 #include "Templates/SharedPointer.h"
+#include "Types/SlateEnums.h"
+#include "UObject/NameTypes.h"
+#include "Widgets/DeclarativeSyntaxSupport.h"
 #include "Widgets/SCompoundWidget.h"
-#include "Widgets/Layout/SExpandableArea.h"
+#include "Widgets/Views/SHeaderRow.h"
 #include "Widgets/Views/SListView.h"
 
-
+class ITableRow;
+class SExpandableArea;
 // Forward declarations
 class SMutableCodeViewer;
+class STableViewBase;
+class SWidget;
+namespace mu { struct Curve; }
+namespace mu { struct PROGRAM; }
+namespace mu { struct PROJECTOR; }
+namespace mu { struct SHAPE; }
 
 /**
  * Base Structure to define the different elements used by the lists on this object
diff --git a/Engine/Plugins/Experimental/Mutable/Source/CustomizableObjectEditor/Private/MuCOE/SMutableConstrainedBox.cpp b/Engine/Plugins/Experimental/Mutable/Source/CustomizableObjectEditor/Private/MuCOE/SMutableConstrainedBox.cpp
--- a/Engine/Plugins/Experimental/Mutable/Source/CustomizableObjectEditor/Private/MuCOE/SMutableConstrainedBox.cpp
+++ b/Engine/Plugins/Experimental/Mutable/Source/CustomizableObjectEditor/Private/MuCOE/SMutableConstrainedBox.cpp
@@ -2,6 +2,10 @@
 
 #include "MuCOE/SMutableConstrainedBox.h"
 
+#include "Layout/Children.h"
+#include "Math/UnrealMathUtility.h"
+#include "Widgets/SWidget.h"
+
 void SMutableConstrainedBox::Construct(const FArguments& InArgs)
 {
 	MinWidth = InArgs._MinWidth;
diff --git a/Engine/Plugins/Experimental/Mutable/Source/CustomizableObjectEditor/Private/MuCOE/SMutableConstrainedBox.h b/Engine/Plugins/Experimental/Mutable/Source/CustomizableObjectEditor/Private/MuCOE/SMutableConstrainedBox.h
--- a/Engine/Plugins/Experimental/Mutable/Source/CustomizableObjectEditor/Private/MuCOE/SMutableConstrainedBox.h
+++ b/Engine/Plugins/Experimental/Mutable/Source/CustomizableObjectEditor/Private/MuCOE/SMutableConstrainedBox.h
@@ -2,9 +2,15 @@
 
 #pragma once
 
+#include "Math/Vector2D.h"
+#include "Misc/Attribute.h"
+#include "Misc/Optional.h"
+#include "Templates/SharedPointer.h"
 #include "Widgets/DeclarativeSyntaxSupport.h"
 #include "Widgets/SCompoundWidget.h"
 
+class SWidget;
+
 /** Helper class to force a widget to fill in a space. Copied from SConstrainedBox.h */
 class SMutableConstrainedBox : public SCompoundWidget
 {
diff --git a/Engine/Plugins/Experimental/Mutable/Source/CustomizableObjectEditor/Private/MuCOE/SMutableCurveViewer.cpp b/Engine/Plugins/Experimental/Mutable/Source/CustomizableObjectEditor/Private/MuCOE/SMutableCurveViewer.cpp
--- a/Engine/Plugins/Experimental/Mutable/Source/CustomizableObjectEditor/Private/MuCOE/SMutableCurveViewer.cpp
+++ b/Engine/Plugins/Experimental/Mutable/Source/CustomizableObjectEditor/Private/MuCOE/SMutableCurveViewer.cpp
@@ -2,18 +2,44 @@
 
 #include "MuCOE/SMutableCurveViewer.h"
 
-#include "Math/Color.h"
-#include "MuR/ParametersPrivate.h"
-#include "MuT/TypeInfo.h"
-#include "Math/NumericLimits.h"
-
-#include "Widgets/Layout/SScrollBox.h"
-#include "Widgets/Views/SListView.h"
+#include "Containers/EnumAsByte.h"
+#include "Containers/Map.h"
+#include "Containers/UnrealString.h"
+#include "CurveDataAbstraction.h"
 #include "CurveEditor.h"
+#include "CurveEditorScreenSpace.h"
 #include "CurveEditorTypes.h"
 #include "CurveModel.h"
-#include "RichCurveEditorModel.h"
+#include "Curves/RealCurve.h"
+#include "Framework/Views/ITypedTableView.h"
+#include "HAL/PlatformCrt.h"
+#include "ICurveEditorModule.h"
+#include "Internationalization/Internationalization.h"
+#include "Internationalization/Text.h"
+#include "Layout/Children.h"
+#include "Math/Color.h"
+#include "Math/NumericLimits.h"
+#include "Math/UnrealMathSSE.h"
+#include "Misc/AssertionMacros.h"
+#include "Misc/Attribute.h"
+#include "MuR/MemoryPrivate.h"
+#include "MuR/ParametersPrivate.h"
+#include "MuT/TypeInfo.h"
 #include "SCurveEditorPanel.h"
+#include "SlotBase.h"
+#include "Templates/UniquePtr.h"
+#include "Templates/UnrealTemplate.h"
+#include "UObject/NameTypes.h"
+#include "Widgets/SBoxPanel.h"
+#include "Widgets/SNullWidget.h"
+#include "Widgets/Text/STextBlock.h"
+#include "Widgets/Views/SHeaderRow.h"
+#include "Widgets/Views/SListView.h"
+#include "Widgets/Views/STableRow.h"
+
+class ITableRow;
+class STableViewBase;
+class SWidget;
 
 #define LOCTEXT_NAMESPACE "SMutableCurveViewer"
 
diff --git a/Engine/Plugins/Experimental/Mutable/Source/CustomizableObjectEditor/Private/MuCOE/SMutableCurveViewer.h b/Engine/Plugins/Experimental/Mutable/Source/CustomizableObjectEditor/Private/MuCOE/SMutableCurveViewer.h
--- a/Engine/Plugins/Experimental/Mutable/Source/CustomizableObjectEditor/Private/MuCOE/SMutableCurveViewer.h
+++ b/Engine/Plugins/Experimental/Mutable/Source/CustomizableObjectEditor/Private/MuCOE/SMutableCurveViewer.h
@@ -2,16 +2,34 @@
 
 #pragma once
 
-#include "SCurveEditor.h"
+#include "Containers/Array.h"
+#include "Containers/ArrayView.h"
+#include "Containers/BitArray.h"
+#include "Containers/Set.h"
+#include "Containers/SparseArray.h"
+#include "CurveEditorTypes.h"
+#include "CurveModel.h"
+#include "Curves/KeyHandle.h"
+#include "Curves/RichCurve.h"
+#include "Delegates/Delegate.h"
+#include "HAL/Platform.h"
+#include "Misc/Optional.h"
 #include "MuR/ParametersPrivate.h"
-
+#include "Templates/SharedPointer.h"
+#include "Templates/Tuple.h"
+#include "UObject/UnrealType.h"
+#include "Widgets/DeclarativeSyntaxSupport.h"
+#include "Widgets/SCompoundWidget.h"
 #include "Widgets/Views/SListView.h"
 
-#include "CurveEditor.h"
-#include "CurveEditorTypes.h"
-#include "CurveModel.h"
-#include "RichCurveEditorModel.h"
-#include "SCurveEditorPanel.h"
+class FCurveEditor;
+class ITableRow;
+class SCurveEditorPanel;
+class STableViewBase;
+struct FCurveEditorScreenSpace;
+struct FKeyAttributes;
+struct FKeyDrawInfo;
+struct FKeyPosition;
 
 /** Custom curve model designed to be able to work with our mutable curve. The overriding of the FCurveModel
  * methods is mandatory.
diff --git a/Engine/Plugins/Experimental/Mutable/Source/CustomizableObjectEditor/Private/MuCOE/SMutableGraphViewer.cpp b/Engine/Plugins/Experimental/Mutable/Source/CustomizableObjectEditor/Private/MuCOE/SMutableGraphViewer.cpp
--- a/Engine/Plugins/Experimental/Mutable/Source/CustomizableObjectEditor/Private/MuCOE/SMutableGraphViewer.cpp
+++ b/Engine/Plugins/Experimental/Mutable/Source/CustomizableObjectEditor/Private/MuCOE/SMutableGraphViewer.cpp
@@ -2,21 +2,65 @@
 
 #include "MuCOE/SMutableGraphViewer.h"
 
+#include "Containers/StringConv.h"
+#include "DesktopPlatformModule.h"
+#include "EditorDirectories.h"
+#include "Framework/Application/SlateApplication.h"
+#include "Framework/Commands/UIAction.h"
+#include "Framework/Commands/UICommandInfo.h"
+#include "Framework/Docking/TabManager.h"
+#include "Framework/MultiBox/MultiBoxBuilder.h"
+#include "Framework/MultiBox/MultiBoxDefs.h"
+#include "Framework/SlateDelegates.h"
+#include "Framework/Views/ITypedTableView.h"
+#include "HAL/PlatformCrt.h"
+#include "HAL/UnrealMemory.h"
+#include "IDesktopPlatform.h"
+#include "Input/DragAndDrop.h"
+#include "Internationalization/Internationalization.h"
+#include "Internationalization/Text.h"
+#include "Layout/Children.h"
+#include "Layout/Margin.h"
+#include "Layout/Visibility.h"
+#include "Math/Color.h"
+#include "Math/UnrealMathSSE.h"
+#include "Misc/AssertionMacros.h"
+#include "Misc/Attribute.h"
+#include "Misc/Paths.h"
+#include "MuCOE/CustomizableObjectCompileRunnable.h"
 #include "MuCOE/CustomizableObjectEditorStyle.h"
-#include "MuCOE/CustomizableObjectCompiler.h"
 #include "MuCOE/SMutableCodeViewer.h"
+#include "MuR/Ptr.h"
+#include "MuR/RefCounted.h"
+#include "MuR/Serialisation.h"
 #include "MuT/Streams.h"
-
-#include "Misc/App.h"
-#include "Framework/MultiBox/MultiBoxBuilder.h"
-#include "Widgets/Input/SButton.h"
+#include "SlotBase.h"
+#include "Styling/AppStyle.h"
+#include "Styling/CoreStyle.h"
+#include "Styling/ISlateStyle.h"
+#include "Styling/SlateColor.h"
+#include "Textures/SlateIcon.h"
+#include "UObject/UnrealNames.h"
+#include "Widgets/Docking/SDockTab.h"
 #include "Widgets/Images/SImage.h"
-#include "Widgets/Views/STreeView.h"
 #include "Widgets/Input/STextComboBox.h"
-#include "Widgets/Docking/SDockTab.h"
-#include "EditorDirectories.h"
-#include "DesktopPlatformModule.h"
-#include "HAL/PlatformApplicationMisc.h"
+#include "Widgets/Layout/SBorder.h"
+#include "Widgets/Layout/SSplitter.h"
+#include "Widgets/SBoxPanel.h"
+#include "Widgets/SOverlay.h"
+#include "Widgets/Text/STextBlock.h"
+#include "Widgets/Views/SExpanderArrow.h"
+#include "Widgets/Views/SHeaderRow.h"
+#include "Widgets/Views/STableRow.h"
+#include "Widgets/Views/STreeView.h"
+
+class FExtender;
+class FReferenceCollector;
+class FUICommandList;
+class ITableRow;
+class SWidget;
+struct FGeometry;
+struct FSlateBrush;
 
 #define LOCTEXT_NAMESPACE "SMutableDebugger"
 
diff --git a/Engine/Plugins/Experimental/Mutable/Source/CustomizableObjectEditor/Private/MuCOE/SMutableGraphViewer.h b/Engine/Plugins/Experimental/Mutable/Source/CustomizableObjectEditor/Private/MuCOE/SMutableGraphViewer.h
--- a/Engine/Plugins/Experimental/Mutable/Source/CustomizableObjectEditor/Private/MuCOE/SMutableGraphViewer.h
+++ b/Engine/Plugins/Experimental/Mutable/Source/CustomizableObjectEditor/Private/MuCOE/SMutableGraphViewer.h
@@ -2,17 +2,38 @@
 
 #pragma once
 
-#include "Widgets/SCompoundWidget.h"
+#include "Containers/Array.h"
+#include "Containers/BitArray.h"
+#include "Containers/Map.h"
+#include "Containers/Set.h"
+#include "Containers/SparseArray.h"
+#include "Containers/UnrealString.h"
+#include "Delegates/Delegate.h"
+#include "HAL/Platform.h"
+#include "Input/Reply.h"
+#include "Misc/Optional.h"
 #include "MuCO/CustomizableObject.h"
+#include "MuCOE/CustomizableObjectCompiler.h"
 #include "MuT/Node.h"
+#include "Templates/SharedPointer.h"
+#include "Templates/TypeHash.h"
+#include "Types/SlateConstants.h"
+#include "Types/SlateEnums.h"
+#include "UObject/GCObject.h"
+#include "UObject/NameTypes.h"
+#include "Widgets/DeclarativeSyntaxSupport.h"
+#include "Widgets/SCompoundWidget.h"
+#include "Widgets/Views/STableViewBase.h"
 #include "Widgets/Views/STreeView.h"
-#include "Widgets/Input/STextComboBox.h"
-#include "MuCOE/CustomizableObjectCompiler.h"
-#include "Framework/Docking/TabManager.h"
-
 
-// Forward declarations
+class FDragDropEvent;
+class FReferenceCollector;
 class FMutableGraphTreeElement;
+class FTabManager;
+class ITableRow;
+class STextComboBox;
+class SWidget;
+struct FGeometry;
 
 
 /** This widget shows the internal Mutable Graph for debugging purposes. 
diff --git a/Engine/Plugins/Experimental/Mutable/Source/CustomizableObjectEditor/Private/MuCOE/SMutableImageViewer.cpp b/Engine/Plugins/Experimental/Mutable/Source/CustomizableObjectEditor/Private/MuCOE/SMutableImageViewer.cpp
--- a/Engine/Plugins/Experimental/Mutable/Source/CustomizableObjectEditor/Private/MuCOE/SMutableImageViewer.cpp
+++ b/Engine/Plugins/Experimental/Mutable/Source/CustomizableObjectEditor/Private/MuCOE/SMutableImageViewer.cpp
@@ -1,16 +1,43 @@
 // Copyright Epic Games, Inc. All Rights Reserved.
 
 #include "MuCOE/SMutableImageViewer.h"
-#include "MuCOE/CustomizableObjectEditorModule.h"
-#include "Styling/AppStyle.h"
-#include "Async/TaskGraphInterfaces.h"
+
 #include "Brushes/SlateImageBrush.h"
-#include "Widgets/Input/SNumericEntryBox.h"
-#include "Widgets/Input/SSegmentedControl.h"
+#include "Containers/Array.h"
+#include "Containers/EnumAsByte.h"
+#include "Containers/StringConv.h"
+#include "Engine/Texture.h"
+#include "Engine/Texture2D.h"
+#include "HAL/PlatformCrt.h"
+#include "Internationalization/Internationalization.h"
+#include "Layout/Children.h"
+#include "Layout/Geometry.h"
+#include "Layout/Margin.h"
+#include "Layout/PaintGeometry.h"
+#include "Math/Color.h"
 #include "MuCO/CustomizableInstancePrivateData.h"
-#include "CanvasItem.h"
-#include "CanvasTypes.h"
+#include "MuCO/CustomizableObject.h"
+#include "MuR/Ptr.h"
 #include "MuT/TypeInfo.h"
+#include "Rendering/DrawElements.h"
+#include "Rendering/RenderingCommon.h"
+#include "Rendering/SlateLayoutTransform.h"
+#include "SlotBase.h"
+#include "Styling/SlateBrush.h"
+#include "Styling/SlateColor.h"
+#include "Styling/WidgetStyle.h"
+#include "TextureResource.h"
+#include "Types/SlateEnums.h"
+#include "UObject/Class.h"
+#include "UObject/ObjectHandle.h"
+#include "UObject/UObjectGlobals.h"
+#include "Widgets/Input/SNumericEntryBox.h"
+#include "Widgets/Input/SSegmentedControl.h"
+#include "Widgets/SBoxPanel.h"
+#include "Widgets/Text/STextBlock.h"
+
+class FPaintArgs;
+class FSlateRect;
 
 
 #define LOCTEXT_NAMESPACE "CustomizableObjectEditor"
diff --git a/Engine/Plugins/Experimental/Mutable/Source/CustomizableObjectEditor/Private/MuCOE/SMutableImageViewer.h b/Engine/Plugins/Experimental/Mutable/Source/CustomizableObjectEditor/Private/MuCOE/SMutableImageViewer.h
--- a/Engine/Plugins/Experimental/Mutable/Source/CustomizableObjectEditor/Private/MuCOE/SMutableImageViewer.h
+++ b/Engine/Plugins/Experimental/Mutable/Source/CustomizableObjectEditor/Private/MuCOE/SMutableImageViewer.h
@@ -2,10 +2,29 @@
 
 #pragma once
 
-#include "MuCOE/SCustomizableObjectLayoutGrid.h"
-#include "Widgets/SCompoundWidget.h"
-#include "Widgets/DeclarativeSyntaxSupport.h"
+#include "Containers/UnrealString.h"
+#include "Internationalization/Text.h"
+#include "Layout/Visibility.h"
+#include "Math/IntPoint.h"
+#include "Math/UnrealMathSSE.h"
+#include "Math/Vector2D.h"
+#include "Misc/Attribute.h"
+#include "Misc/Optional.h"
 #include "MuR/Image.h"
+#include "Templates/SharedPointer.h"
+#include "UObject/GCObject.h"
+#include "UObject/ObjectPtr.h"
+#include "Widgets/DeclarativeSyntaxSupport.h"
+#include "Widgets/SCompoundWidget.h"
+
+class FPaintArgs;
+class FReferenceCollector;
+class FSlateRect;
+class FSlateWindowElementList;
+class FWidgetStyle;
+class UTexture;
+struct FGeometry;
+struct FSlateBrush;
 
 
 /** Describes different modes to show an image. */
diff --git a/Engine/Plugins/Experimental/Mutable/Source/CustomizableObjectEditor/Private/MuCOE/SMutableIntViewer.cpp b/Engine/Plugins/Experimental/Mutable/Source/CustomizableObjectEditor/Private/MuCOE/SMutableIntViewer.cpp
--- a/Engine/Plugins/Experimental/Mutable/Source/CustomizableObjectEditor/Private/MuCOE/SMutableIntViewer.cpp
+++ b/Engine/Plugins/Experimental/Mutable/Source/CustomizableObjectEditor/Private/MuCOE/SMutableIntViewer.cpp
@@ -2,9 +2,12 @@
 
 #include "MuCOE/SMutableIntViewer.h"
 
-#include "GenericPlatform/GenericPlatform.h"
-#include "Widgets/SBoxPanel.h"
+#include "Internationalization/Internationalization.h"
 #include "Internationalization/Text.h"
+#include "Layout/Children.h"
+#include "Misc/Attribute.h"
+#include "SlotBase.h"
+#include "Widgets/SBoxPanel.h"
 #include "Widgets/Text/STextBlock.h"
 
 #define LOCTEXT_NAMESPACE "CustomizableObjectEditor"
diff --git a/Engine/Plugins/Experimental/Mutable/Source/CustomizableObjectEditor/Private/MuCOE/SMutableIntViewer.h b/Engine/Plugins/Experimental/Mutable/Source/CustomizableObjectEditor/Private/MuCOE/SMutableIntViewer.h
--- a/Engine/Plugins/Experimental/Mutable/Source/CustomizableObjectEditor/Private/MuCOE/SMutableIntViewer.h
+++ b/Engine/Plugins/Experimental/Mutable/Source/CustomizableObjectEditor/Private/MuCOE/SMutableIntViewer.h
@@ -2,10 +2,10 @@
 
 #pragma once
 
-#include "Widgets/SCompoundWidget.h"
-#include "GenericPlatform/GenericPlatform.h"
-#include "Widgets/DeclarativeSyntaxSupport.h"
+#include "HAL/Platform.h"
 #include "Internationalization/Text.h"
+#include "Widgets/DeclarativeSyntaxSupport.h"
+#include "Widgets/SCompoundWidget.h"
 
 /*
 * Mutable viewer to show a representation of an int variable held by mutable.
diff --git a/Engine/Plugins/Experimental/Mutable/Source/CustomizableObjectEditor/Private/MuCOE/SMutableLayoutViewer.cpp b/Engine/Plugins/Experimental/Mutable/Source/CustomizableObjectEditor/Private/MuCOE/SMutableLayoutViewer.cpp
--- a/Engine/Plugins/Experimental/Mutable/Source/CustomizableObjectEditor/Private/MuCOE/SMutableLayoutViewer.cpp
+++ b/Engine/Plugins/Experimental/Mutable/Source/CustomizableObjectEditor/Private/MuCOE/SMutableLayoutViewer.cpp
@@ -1,16 +1,24 @@
 // Copyright Epic Games, Inc. All Rights Reserved.
 
 #include "MuCOE/SMutableLayoutViewer.h"
-#include "MuCOE/CustomizableObjectEditorModule.h"
-#include "Styling/AppStyle.h"
-#include "Async/TaskGraphInterfaces.h"
-#include "Brushes/SlateImageBrush.h"
-#include "Widgets/Input/SNumericEntryBox.h"
-#include "Widgets/Input/SSegmentedControl.h"
-#include "MuCO/CustomizableInstancePrivateData.h"
-#include "CanvasItem.h"
-#include "CanvasTypes.h"
-#include "MuT/TypeInfo.h"
+
+#include "Containers/Array.h"
+#include "Containers/UnrealString.h"
+#include "Layout/Children.h"
+#include "Layout/Margin.h"
+#include "Math/IntPoint.h"
+#include "Math/UnrealMathSSE.h"
+#include "Misc/Guid.h"
+#include "MuCOE/CustomizableObjectLayout.h"
+#include "MuCOE/SCustomizableObjectLayoutGrid.h"
+#include "MuR/Mesh.h"
+#include "MuR/Ptr.h"
+#include "SlotBase.h"
+#include "Types/SlateEnums.h"
+#include "Widgets/SBoxPanel.h"
+#include "Widgets/Text/STextBlock.h"
+
+struct FGeometry;
 
 
 #define LOCTEXT_NAMESPACE "CustomizableObjectEditor"
diff --git a/Engine/Plugins/Experimental/Mutable/Source/CustomizableObjectEditor/Private/MuCOE/SMutableLayoutViewer.h b/Engine/Plugins/Experimental/Mutable/Source/CustomizableObjectEditor/Private/MuCOE/SMutableLayoutViewer.h
--- a/Engine/Plugins/Experimental/Mutable/Source/CustomizableObjectEditor/Private/MuCOE/SMutableLayoutViewer.h
+++ b/Engine/Plugins/Experimental/Mutable/Source/CustomizableObjectEditor/Private/MuCOE/SMutableLayoutViewer.h
@@ -2,10 +2,16 @@
 
 #pragma once
 
-#include "MuCOE/SCustomizableObjectLayoutGrid.h"
-#include "Widgets/SCompoundWidget.h"
-#include "Widgets/DeclarativeSyntaxSupport.h"
+#include "Internationalization/Text.h"
+#include "Math/MathFwd.h"
+#include "Misc/Attribute.h"
 #include "MuR/Layout.h"
+#include "Templates/SharedPointer.h"
+#include "Widgets/DeclarativeSyntaxSupport.h"
+#include "Widgets/SCompoundWidget.h"
+
+class SCustomizableObjectLayoutGrid;
+struct FGeometry;
 
 
 /** Widget showing a Mutable layout, with its properties. */
diff --git a/Engine/Plugins/Experimental/Mutable/Source/CustomizableObjectEditor/Private/MuCOE/SMutableMeshViewer.cpp b/Engine/Plugins/Experimental/Mutable/Source/CustomizableObjectEditor/Private/MuCOE/SMutableMeshViewer.cpp
--- a/Engine/Plugins/Experimental/Mutable/Source/CustomizableObjectEditor/Private/MuCOE/SMutableMeshViewer.cpp
+++ b/Engine/Plugins/Experimental/Mutable/Source/CustomizableObjectEditor/Private/MuCOE/SMutableMeshViewer.cpp
@@ -1,23 +1,45 @@
 // Copyright Epic Games, Inc. All Rights Reserved.
 
 #include "MuCOE/SMutableMeshViewer.h"
-#include "MuCOE/CustomizableObjectEditorModule.h"
-#include "Styling/AppStyle.h"
-#include "Widgets/Views/SListView.h"
-#include "Widgets/SBoxPanel.h"
-#include "Widgets/Layout/SScrollBox.h"
-#include "Widgets/Layout/SSplitter.h"
-#include "Widgets/Layout/SBorder.h"
-#include "Widgets/Input/SNumericEntryBox.h"
-#include "MuCO/CustomizableInstancePrivateData.h"
-#include "CanvasItem.h"
+
+#include "AssetRegistry/AssetData.h"
+#include "Engine/SkeletalMesh.h"
+#include "Framework/Views/ITypedTableView.h"
+#include "HAL/PlatformCrt.h"
 #include "Internationalization/Text.h"
-#include "Widgets/Views/SHeaderRow.h"
-#include "Templates/SharedPointer.h"
-#include "PropertyCustomizationHelpers.h"
+#include "Layout/BasicLayoutWidgetSlot.h"
+#include "Layout/Children.h"
+#include "Layout/Visibility.h"
+#include "Misc/Attribute.h"
+#include "MuCOE/SMutableMeshViewport.h"
 #include "MuCOE/SMutableSkeletonViewer.h"
-#include "MuT/TypeInfo.h"
+#include "MuR/MeshBufferSet.h"
+#include "MuR/Ptr.h"
 #include "MuR/Skeleton.h"
+#include "MuT/TypeInfo.h"
+#include "PropertyCustomizationHelpers.h"
+#include "SlotBase.h"
+#include "Templates/Casts.h"
+#include "Templates/SharedPointer.h"
+#include "Templates/UnrealTemplate.h"
+#include "Types/SlateEnums.h"
+#include "UObject/NameTypes.h"
+#include "UObject/Object.h"
+#include "UObject/UObjectGlobals.h"
+#include "Widgets/Layout/SBorder.h"
+#include "Widgets/Layout/SScrollBox.h"
+#include "Widgets/Layout/SSplitter.h"
+#include "Widgets/SBoxPanel.h"
+#include "Widgets/SNullWidget.h"
+#include "Widgets/Text/STextBlock.h"
+#include "Widgets/Views/SHeaderRow.h"
+#include "Widgets/Views/SListView.h"
+#include "Widgets/Views/STableRow.h"
+
+class ITableRow;
+class STableViewBase;
+class SWidget;
+class USkeleton;
 
 #define LOCTEXT_NAMESPACE "CustomizableObjectEditor"
 
diff --git a/Engine/Plugins/Experimental/Mutable/Source/CustomizableObjectEditor/Private/MuCOE/SMutableMeshViewer.h b/Engine/Plugins/Experimental/Mutable/Source/CustomizableObjectEditor/Private/MuCOE/SMutableMeshViewer.h
--- a/Engine/Plugins/Experimental/Mutable/Source/CustomizableObjectEditor/Private/MuCOE/SMutableMeshViewer.h
+++ b/Engine/Plugins/Experimental/Mutable/Source/CustomizableObjectEditor/Private/MuCOE/SMutableMeshViewer.h
@@ -3,23 +3,37 @@
 #pragma once
 
 #include "Containers/Array.h"
+#include "Containers/BitArray.h"
+#include "Containers/Set.h"
+#include "Containers/SparseArray.h"
+#include "Containers/UnrealString.h"
+#include "Delegates/Delegate.h"
+#include "HAL/Platform.h"
+#include "Internationalization/Internationalization.h"
+#include "Internationalization/Text.h"
+#include "Math/Vector2D.h"
+#include "Misc/Optional.h"
 #include "MuR/Mesh.h"
+#include "Templates/SharedPointer.h"
+#include "UObject/GCObject.h"
 #include "Widgets/DeclarativeSyntaxSupport.h"
 #include "Widgets/SCompoundWidget.h"
 #include "Widgets/Views/SListView.h"
-#include "Widgets/Views/STableRow.h"
-#include "Widgets/Layout/SBorder.h"
-#include "AssetThumbnail.h"
-#include "UObject/GCObject.h"
-
-#include "Templates/SharedPointer.h"
-
-#include "MuCOE/SMutableMeshViewport.h"
-#include "Widgets/Input/SSearchBox.h"
-#include "Widgets/Layout/SScrollBox.h"
 
-#include "PropertyCustomizationHelpers.h"
-#include "MuCOE/SMutableSkeletonViewer.h"
+class FAssetThumbnail;
+class FAssetThumbnailPool;
+class FReferenceCollector;
+class ITableRow;
+class SBorder;
+class SMutableMeshViewport;
+class SMutableSkeletonViewer;
+class SScrollBox;
+class SSplitter;
+class STableViewBase;
+class SWidget;
+class USkeletalMesh;
+namespace mu { class FMeshBufferSet; }
+struct FAssetData;
 
 
 /** Container designed to hold the buffer channel data of a mutable mesh buffer to be later used by the UI */
diff --git a/Engine/Plugins/Experimental/Mutable/Source/CustomizableObjectEditor/Private/MuCOE/SMutableMeshViewport.cpp b/Engine/Plugins/Experimental/Mutable/Source/CustomizableObjectEditor/Private/MuCOE/SMutableMeshViewport.cpp
--- a/Engine/Plugins/Experimental/Mutable/Source/CustomizableObjectEditor/Private/MuCOE/SMutableMeshViewport.cpp
+++ b/Engine/Plugins/Experimental/Mutable/Source/CustomizableObjectEditor/Private/MuCOE/SMutableMeshViewport.cpp
@@ -2,18 +2,26 @@
 
 #include "MuCOE/SMutableMeshViewport.h"
 
-#include "MuCOE/FMutableViewportClient.h"
-#include "MuCO/CustomizableObjectInstance.h"
-#include "SEditorViewport.h"
-#include "Containers/UnrealString.h"
-#include "Internationalization/Text.h"
+#include "AdvancedPreviewScene.h"
+#include "Components/SkeletalMeshComponent.h"
+#include "Editor/UnrealEdTypes.h"
 #include "Engine/SkeletalMesh.h"
+#include "Math/MathFwd.h"
+#include "Math/Rotator.h"
+#include "Math/Transform.h"
+#include "Math/Vector.h"
 #include "MuCOE/FMutableViewportClient.h"
+#include "MuCOE/MutableMeshPreviewUtils.h"
 #include "MuR/Mesh.h"
-#include "MuR/MeshPrivate.h"
-#include "MuR/Parameters.h"
+#include "MuR/Ptr.h"
+#include "PreviewScene.h"
+#include "SEditorViewport.h"
+#include "UObject/Package.h"
+#include "UObject/UObjectGlobals.h"
+#include "UObject/UnrealNames.h"
 
-#include "MuCOE/MutableMeshPreviewUtils.h"
+class FEditorViewportClient;
+struct FGeometry;
 
 #define LOCTEXT_NAMESPACE "SMutableMeshViewerViewport"
 
diff --git a/Engine/Plugins/Experimental/Mutable/Source/CustomizableObjectEditor/Private/MuCOE/SMutableMeshViewport.h b/Engine/Plugins/Experimental/Mutable/Source/CustomizableObjectEditor/Private/MuCOE/SMutableMeshViewport.h
--- a/Engine/Plugins/Experimental/Mutable/Source/CustomizableObjectEditor/Private/MuCOE/SMutableMeshViewport.h
+++ b/Engine/Plugins/Experimental/Mutable/Source/CustomizableObjectEditor/Private/MuCOE/SMutableMeshViewport.h
@@ -2,16 +2,21 @@
 
 #pragma once
 
+#include "Containers/UnrealString.h"
+#include "HAL/Platform.h"
 #include "MuR/Mesh.h"
-#include "MuCOE/FMutableViewportClient.h"
-#include "Widgets/SCompoundWidget.h"
-#include "AdvancedPreviewScene.h"
-#include "Components/SkeletalMeshComponent.h"
 #include "SEditorViewport.h"
-#include "UnrealWidget.h"
+#include "Templates/SharedPointer.h"
 #include "UObject/GCObject.h"
-
-#include "Engine/SkeletalMesh.h"
+#include "Widgets/DeclarativeSyntaxSupport.h"
+
+class FAdvancedPreviewScene;
+class FEditorViewportClient;
+class FMutableMeshViewportClient;
+class FReferenceCollector;
+class USkeletalMesh;
+class USkeletalMeshComponent;
+struct FGeometry;
 
 
 /**
diff --git a/Engine/Plugins/Experimental/Mutable/Source/CustomizableObjectEditor/Private/MuCOE/SMutableObjectViewer.cpp b/Engine/Plugins/Experimental/Mutable/Source/CustomizableObjectEditor/Private/MuCOE/SMutableObjectViewer.cpp
--- a/Engine/Plugins/Experimental/Mutable/Source/CustomizableObjectEditor/Private/MuCOE/SMutableObjectViewer.cpp
+++ b/Engine/Plugins/Experimental/Mutable/Source/CustomizableObjectEditor/Private/MuCOE/SMutableObjectViewer.cpp
@@ -2,23 +2,55 @@
 
 #include "MuCOE/SMutableObjectViewer.h"
 
-#include "MuCOE/CustomizableObjectEditorStyle.h"
+#include "Framework/Commands/UIAction.h"
+#include "Framework/Commands/UICommandInfo.h"
+#include "Framework/Docking/TabManager.h"
+#include "Framework/MultiBox/MultiBoxBuilder.h"
+#include "Framework/MultiBox/MultiBoxDefs.h"
+#include "Framework/SlateDelegates.h"
+#include "Framework/Views/ITypedTableView.h"
+#include "HAL/PlatformCrt.h"
+#include "Interfaces/ITargetPlatform.h"
+#include "Interfaces/ITargetPlatformManagerModule.h"
+#include "Internationalization/Internationalization.h"
+#include "Internationalization/Text.h"
+#include "Layout/Children.h"
+#include "Layout/Margin.h"
+#include "Layout/Visibility.h"
+#include "Math/Color.h"
+#include "Math/UnrealMathSSE.h"
+#include "Misc/AssertionMacros.h"
+#include "Misc/Attribute.h"
+#include "Misc/CoreMisc.h"
+#include "MuCOE/CustomizableObjectCompileRunnable.h"
 #include "MuCOE/CustomizableObjectCompiler.h"
-#include "MuCOE/SMutableGraphViewer.h"
+#include "MuCOE/CustomizableObjectEditorStyle.h"
 #include "MuCOE/SMutableCodeViewer.h"
-#include "MuT/Streams.h"
-
-#include "Misc/App.h"
-#include "Framework/MultiBox/MultiBoxBuilder.h"
-#include "Widgets/Input/SButton.h"
+#include "MuCOE/SMutableGraphViewer.h"
+#include "MuT/Node.h"
+#include "SlotBase.h"
+#include "Styling/AppStyle.h"
+#include "Styling/CoreStyle.h"
+#include "Styling/ISlateStyle.h"
+#include "Styling/SlateColor.h"
+#include "Textures/SlateIcon.h"
+#include "UObject/UObjectGlobals.h"
+#include "UObject/UnrealNames.h"
+#include "Widgets/Docking/SDockTab.h"
 #include "Widgets/Images/SImage.h"
 #include "Widgets/Input/STextComboBox.h"
-#include "Widgets/Docking/SDockTab.h"
-#include "EditorDirectories.h"
-#include "DesktopPlatformModule.h"
-#include "HAL/PlatformApplicationMisc.h"
-#include "Interfaces/ITargetPlatform.h"
-#include "Interfaces/ITargetPlatformManagerModule.h"
+#include "Widgets/Layout/SBorder.h"
+#include "Widgets/SBoxPanel.h"
+#include "Widgets/SOverlay.h"
+#include "Widgets/Text/STextBlock.h"
+#include "Widgets/Views/SExpanderArrow.h"
+#include "Widgets/Views/STableRow.h"
+
+class FExtender;
+class FUICommandList;
+class ITableRow;
+class SWidget;
+struct FSlateBrush;
 
 #define LOCTEXT_NAMESPACE "SMutableDebugger"
 
diff --git a/Engine/Plugins/Experimental/Mutable/Source/CustomizableObjectEditor/Private/MuCOE/SMutableObjectViewer.h b/Engine/Plugins/Experimental/Mutable/Source/CustomizableObjectEditor/Private/MuCOE/SMutableObjectViewer.h
--- a/Engine/Plugins/Experimental/Mutable/Source/CustomizableObjectEditor/Private/MuCOE/SMutableObjectViewer.h
+++ b/Engine/Plugins/Experimental/Mutable/Source/CustomizableObjectEditor/Private/MuCOE/SMutableObjectViewer.h
@@ -2,15 +2,32 @@
 
 #pragma once
 
-#include "Widgets/SCompoundWidget.h"
-#include "Framework/Docking/TabManager.h"
-#include "Widgets/Views/STreeView.h"
+#include "Containers/Array.h"
+#include "Containers/BitArray.h"
+#include "Containers/Set.h"
+#include "Containers/SparseArray.h"
+#include "Containers/UnrealString.h"
+#include "Delegates/Delegate.h"
+#include "Misc/Optional.h"
 #include "MuCO/CustomizableObject.h"
 #include "MuCOE/CustomizableObjectCompiler.h"
-
+#include "Templates/SharedPointer.h"
+#include "Types/SlateConstants.h"
+#include "Types/SlateEnums.h"
+#include "UObject/GCObject.h"
+#include "UObject/NameTypes.h"
+#include "Widgets/DeclarativeSyntaxSupport.h"
+#include "Widgets/SCompoundWidget.h"
+#include "Widgets/Views/STableViewBase.h"
+#include "Widgets/Views/STreeView.h"
 
 // Forward declarations
 class FMutableObjectTreeElement;
+class FReferenceCollector;
+class FTabManager;
+class ITableRow;
+class STextComboBox;
+class SWidget;
 
 
 /** This widget shows the internals of a CustomizableObject for debugging purposes. 
diff --git a/Engine/Plugins/Experimental/Mutable/Source/CustomizableObjectEditor/Private/MuCOE/SMutableParametersWidget.cpp b/Engine/Plugins/Experimental/Mutable/Source/CustomizableObjectEditor/Private/MuCOE/SMutableParametersWidget.cpp
--- a/Engine/Plugins/Experimental/Mutable/Source/CustomizableObjectEditor/Private/MuCOE/SMutableParametersWidget.cpp
+++ b/Engine/Plugins/Experimental/Mutable/Source/CustomizableObjectEditor/Private/MuCOE/SMutableParametersWidget.cpp
@@ -1,17 +1,34 @@
 // Copyright Epic Games, Inc. All Rights Reserved.
 
 #include "MuCOE/SMutableParametersWidget.h"
-#include "MuCOE/CustomizableObjectEditorModule.h"
+
+#include "Containers/Array.h"
+#include "Containers/StringConv.h"
+#include "Containers/UnrealString.h"
+#include "Fonts/SlateFontInfo.h"
+#include "Framework/SlateDelegates.h"
+#include "Framework/Text/TextLayout.h"
+#include "HAL/PlatformCrt.h"
+#include "Input/Events.h"
+#include "InputCoreTypes.h"
+#include "Internationalization/Text.h"
+#include "Layout/Children.h"
+#include "Math/Vector2D.h"
+#include "Misc/Attribute.h"
+#include "MuCOE/SMutableTextSearchBox.h"
+#include "MuR/Ptr.h"
+#include "SlotBase.h"
 #include "Styling/AppStyle.h"
-#include "Async/TaskGraphInterfaces.h"
-#include "Brushes/SlateImageBrush.h"
-#include "Widgets/Input/SNumericEntryBox.h"
+#include "Widgets/Colors/SColorBlock.h"
+#include "Widgets/Colors/SColorPicker.h"
 #include "Widgets/Input/SCheckBox.h"
-#include "Widgets/Input/STextComboBox.h"
+#include "Widgets/Input/SNumericEntryBox.h"
 #include "Widgets/Input/SVectorInputBox.h"
-#include "Widgets/Colors/SColorPicker.h"
 #include "Widgets/Layout/SGridPanel.h"
-#include "MuT/TypeInfo.h"
+#include "Widgets/SBoxPanel.h"
+#include "Widgets/Text/STextBlock.h"
+
+struct FGeometry;
 
 
 #define LOCTEXT_NAMESPACE "CustomizableObjectEditor"
diff --git a/Engine/Plugins/Experimental/Mutable/Source/CustomizableObjectEditor/Private/MuCOE/SMutableParametersWidget.h b/Engine/Plugins/Experimental/Mutable/Source/CustomizableObjectEditor/Private/MuCOE/SMutableParametersWidget.h
--- a/Engine/Plugins/Experimental/Mutable/Source/CustomizableObjectEditor/Private/MuCOE/SMutableParametersWidget.h
+++ b/Engine/Plugins/Experimental/Mutable/Source/CustomizableObjectEditor/Private/MuCOE/SMutableParametersWidget.h
@@ -2,13 +2,26 @@
 
 #pragma once
 
-#include "MuCOE/SMutableTextSearchBox.h"
-#include "Widgets/SCompoundWidget.h"
-#include "Widgets/SBoxPanel.h"
-#include "Widgets/Input/SCheckBox.h"
-#include "Widgets/DeclarativeSyntaxSupport.h"
-#include "MuR/Parameters.h"
+#include "Delegates/Delegate.h"
+#include "Input/Reply.h"
+#include "Layout/Visibility.h"
 #include "Math/Axis.h"
+#include "Math/Color.h"
+#include "Math/UnrealMathSSE.h"
+#include "Math/Vector.h"
+#include "Misc/Optional.h"
+#include "MuR/Parameters.h"
+#include "Styling/SlateTypes.h"
+#include "Templates/SharedPointer.h"
+#include "Types/SlateEnums.h"
+#include "Widgets/DeclarativeSyntaxSupport.h"
+#include "Widgets/SCompoundWidget.h"
+
+class FString;
+class SMutableTextSearchBox;
+class SVerticalBox;
+struct FGeometry;
+struct FPointerEvent;
 
 
 /** Delegate that is executed when a parameter value changes */
diff --git a/Engine/Plugins/Experimental/Mutable/Source/CustomizableObjectEditor/Private/MuCOE/SMutableProjectorViewer.cpp b/Engine/Plugins/Experimental/Mutable/Source/CustomizableObjectEditor/Private/MuCOE/SMutableProjectorViewer.cpp
--- a/Engine/Plugins/Experimental/Mutable/Source/CustomizableObjectEditor/Private/MuCOE/SMutableProjectorViewer.cpp
+++ b/Engine/Plugins/Experimental/Mutable/Source/CustomizableObjectEditor/Private/MuCOE/SMutableProjectorViewer.cpp
@@ -2,10 +2,19 @@
 
 #include "MuCOE/SMutableProjectorViewer.h"
 
-#include "Widgets/Input/SVectorInputBox.h"
+#include "Containers/UnrealString.h"
+#include "HAL/PlatformCrt.h"
+#include "Internationalization/Internationalization.h"
 #include "Internationalization/Text.h"
+#include "Layout/Children.h"
+#include "Misc/AssertionMacros.h"
+#include "Misc/Attribute.h"
 #include "MuR/ParametersPrivate.h"
 #include "MuT/TypeInfo.h"
+#include "SlotBase.h"
+#include "Widgets/Input/SVectorInputBox.h"
+#include "Widgets/SBoxPanel.h"
+#include "Widgets/Text/STextBlock.h"
 
 #define LOCTEXT_NAMESPACE "MutableProjectorViewer"
 
diff --git a/Engine/Plugins/Experimental/Mutable/Source/CustomizableObjectEditor/Private/MuCOE/SMutableProjectorViewer.h b/Engine/Plugins/Experimental/Mutable/Source/CustomizableObjectEditor/Private/MuCOE/SMutableProjectorViewer.h
--- a/Engine/Plugins/Experimental/Mutable/Source/CustomizableObjectEditor/Private/MuCOE/SMutableProjectorViewer.h
+++ b/Engine/Plugins/Experimental/Mutable/Source/CustomizableObjectEditor/Private/MuCOE/SMutableProjectorViewer.h
@@ -2,9 +2,12 @@
 
 #pragma once
 
+#include "HAL/Platform.h"
 #include "Internationalization/Text.h"
+#include "Misc/Optional.h"
+#include "MuR/Parameters.h"
 #include "MuR/ParametersPrivate.h"
-#include "MuR/SystemPrivate.h"
+#include "Widgets/DeclarativeSyntaxSupport.h"
 #include "Widgets/SCompoundWidget.h"
 
 
diff --git a/Engine/Plugins/Experimental/Mutable/Source/CustomizableObjectEditor/Private/MuCOE/SMutableScalarViewer.cpp b/Engine/Plugins/Experimental/Mutable/Source/CustomizableObjectEditor/Private/MuCOE/SMutableScalarViewer.cpp
--- a/Engine/Plugins/Experimental/Mutable/Source/CustomizableObjectEditor/Private/MuCOE/SMutableScalarViewer.cpp
+++ b/Engine/Plugins/Experimental/Mutable/Source/CustomizableObjectEditor/Private/MuCOE/SMutableScalarViewer.cpp
@@ -2,8 +2,12 @@
 
 #include "MuCOE/SMutableScalarViewer.h"
 
-#include "Widgets/SBoxPanel.h"
+#include "Internationalization/Internationalization.h"
 #include "Internationalization/Text.h"
+#include "Layout/Children.h"
+#include "Misc/Attribute.h"
+#include "SlotBase.h"
+#include "Widgets/SBoxPanel.h"
 #include "Widgets/Text/STextBlock.h"
 
 #define LOCTEXT_NAMESPACE "CustomizableObjectEditor"
diff --git a/Engine/Plugins/Experimental/Mutable/Source/CustomizableObjectEditor/Private/MuCOE/SMutableSkeletonViewer.cpp b/Engine/Plugins/Experimental/Mutable/Source/CustomizableObjectEditor/Private/MuCOE/SMutableSkeletonViewer.cpp
--- a/Engine/Plugins/Experimental/Mutable/Source/CustomizableObjectEditor/Private/MuCOE/SMutableSkeletonViewer.cpp
+++ b/Engine/Plugins/Experimental/Mutable/Source/CustomizableObjectEditor/Private/MuCOE/SMutableSkeletonViewer.cpp
@@ -1,9 +1,24 @@
 // Copyright Epic Games, Inc. All Rights Reserved.
 
 #include "MuCOE/SMutableSkeletonViewer.h"
-#include "Widgets/Layout/SScrollBox.h"
-#include "MuT/TypeInfo.h"
+
+#include "Containers/UnrealString.h"
+#include "Framework/Views/ITypedTableView.h"
+#include "HAL/PlatformCrt.h"
+#include "Layout/BasicLayoutWidgetSlot.h"
+#include "Layout/Children.h"
+#include "Misc/Attribute.h"
+#include "MuR/Ptr.h"
 #include "MuR/Skeleton.h"
+#include "SlotBase.h"
+#include "Types/SlateEnums.h"
+#include "Widgets/Layout/SScrollBox.h"
+#include "Widgets/SBoxPanel.h"
+#include "Widgets/Text/STextBlock.h"
+#include "Widgets/Views/SExpanderArrow.h"
+#include "Widgets/Views/STableRow.h"
+
+class ITableRow;
 
 #define LOCTEXT_NAMESPACE "MutableSkeletonViewer"
 
diff --git a/Engine/Plugins/Experimental/Mutable/Source/CustomizableObjectEditor/Private/MuCOE/SMutableSkeletonViewer.h b/Engine/Plugins/Experimental/Mutable/Source/CustomizableObjectEditor/Private/MuCOE/SMutableSkeletonViewer.h
--- a/Engine/Plugins/Experimental/Mutable/Source/CustomizableObjectEditor/Private/MuCOE/SMutableSkeletonViewer.h
+++ b/Engine/Plugins/Experimental/Mutable/Source/CustomizableObjectEditor/Private/MuCOE/SMutableSkeletonViewer.h
@@ -2,10 +2,24 @@
 
 #pragma once
 
-#include "Widgets/SCompoundWidget.h"
+#include "Containers/Array.h"
+#include "Containers/BitArray.h"
+#include "Containers/Set.h"
+#include "Containers/SparseArray.h"
+#include "Delegates/Delegate.h"
+#include "HAL/Platform.h"
+#include "Internationalization/Internationalization.h"
+#include "Internationalization/Text.h"
+#include "Misc/Optional.h"
+#include "MuR/Mesh.h"
 #include "Templates/SharedPointer.h"
+#include "Types/SlateConstants.h"
+#include "Widgets/DeclarativeSyntaxSupport.h"
+#include "Widgets/SCompoundWidget.h"
+#include "Widgets/Views/STableViewBase.h"
 #include "Widgets/Views/STreeView.h"
-#include "MuR/Mesh.h"
+
+class ITableRow;
 
 /** Container with the information required to generate the tree of bones */
 struct FBoneDefinition
diff --git a/Engine/Plugins/Experimental/Mutable/Source/CustomizableObjectEditor/Private/MuCOE/SMutableStringViewer.cpp b/Engine/Plugins/Experimental/Mutable/Source/CustomizableObjectEditor/Private/MuCOE/SMutableStringViewer.cpp
--- a/Engine/Plugins/Experimental/Mutable/Source/CustomizableObjectEditor/Private/MuCOE/SMutableStringViewer.cpp
+++ b/Engine/Plugins/Experimental/Mutable/Source/CustomizableObjectEditor/Private/MuCOE/SMutableStringViewer.cpp
@@ -2,8 +2,11 @@
 
 #include "MuCOE/SMutableStringViewer.h"
 
-#include "Widgets/SBoxPanel.h"
+#include "Internationalization/Internationalization.h"
 #include "Internationalization/Text.h"
+#include "Layout/Children.h"
+#include "SlotBase.h"
+#include "Widgets/SBoxPanel.h"
 #include "Widgets/Text/STextBlock.h"
 
 #define LOCTEXT_NAMESPACE "CustomizableObjectEditor"
diff --git a/Engine/Plugins/Experimental/Mutable/Source/CustomizableObjectEditor/Private/MuCOE/SMutableStringViewer.h b/Engine/Plugins/Experimental/Mutable/Source/CustomizableObjectEditor/Private/MuCOE/SMutableStringViewer.h
--- a/Engine/Plugins/Experimental/Mutable/Source/CustomizableObjectEditor/Private/MuCOE/SMutableStringViewer.h
+++ b/Engine/Plugins/Experimental/Mutable/Source/CustomizableObjectEditor/Private/MuCOE/SMutableStringViewer.h
@@ -2,9 +2,10 @@
 
 #pragma once
 
-#include "Widgets/SCompoundWidget.h"
-#include "Widgets/DeclarativeSyntaxSupport.h"
 #include "Internationalization/Text.h"
+#include "Misc/Attribute.h"
+#include "Widgets/DeclarativeSyntaxSupport.h"
+#include "Widgets/SCompoundWidget.h"
 
 /**
 * Widget designed to display the value of a string.
diff --git a/Engine/Plugins/Experimental/Mutable/Source/CustomizableObjectEditor/Private/MuCOE/SMutableTextSearchBox.cpp b/Engine/Plugins/Experimental/Mutable/Source/CustomizableObjectEditor/Private/MuCOE/SMutableTextSearchBox.cpp
--- a/Engine/Plugins/Experimental/Mutable/Source/CustomizableObjectEditor/Private/MuCOE/SMutableTextSearchBox.cpp
+++ b/Engine/Plugins/Experimental/Mutable/Source/CustomizableObjectEditor/Private/MuCOE/SMutableTextSearchBox.cpp
@@ -2,19 +2,33 @@
 
 #include "MuCOE/SMutableTextSearchBox.h"
 
-#include "MuCO/CustomizableObject.h" // Required for LogMutable
-
-#include "Layout/WidgetPath.h"
+#include "CoreTypes.h"
 #include "Framework/Application/SlateApplication.h"
-#include "Widgets/Text/STextBlock.h"
+#include "Framework/Views/ITypedTableView.h"
+#include "HAL/PlatformCrt.h"
+#include "Input/Events.h"
+#include "InputCoreTypes.h"
+#include "Layout/Children.h"
+#include "Layout/Margin.h"
+#include "Layout/Visibility.h"
+#include "Layout/WidgetPath.h"
+#include "Math/Color.h"
+#include "Misc/AssertionMacros.h"
+#include "Styling/AppStyle.h"
+#include "Styling/SlateColor.h"
+#include "Types/SlateStructs.h"
+#include "Widgets/Input/SButton.h"
 #include "Widgets/Input/SMenuAnchor.h"
-#include "Widgets/Views/SListView.h"
 #include "Widgets/Input/SSearchBox.h"
-#include "Widgets/Input/SButton.h"
-#include "Widgets/Input/SComboButton.h"
-#include "Widgets/Images/SImage.h"
+#include "Widgets/Layout/SBorder.h"
+#include "Widgets/Layout/SBox.h"
+#include "Widgets/Text/STextBlock.h"
+#include "Widgets/Views/SListView.h"
+#include "Widgets/Views/STableRow.h"
 
-#include "Styling/AppStyle.h"
+class ITableRow;
+class STableViewBase;
+struct FGeometry;
 
 void SMutableTextSearchBox::Construct( const FArguments& InArgs )
 {
diff --git a/Engine/Plugins/Experimental/Mutable/Source/CustomizableObjectEditor/Private/MuCOE/SMutableTextSearchBox.h b/Engine/Plugins/Experimental/Mutable/Source/CustomizableObjectEditor/Private/MuCOE/SMutableTextSearchBox.h
--- a/Engine/Plugins/Experimental/Mutable/Source/CustomizableObjectEditor/Private/MuCOE/SMutableTextSearchBox.h
+++ b/Engine/Plugins/Experimental/Mutable/Source/CustomizableObjectEditor/Private/MuCOE/SMutableTextSearchBox.h
@@ -3,17 +3,30 @@
 #pragma once
 
 
-#include "CoreMinimal.h"
-#include "SlateFwd.h"
-#include "Misc/Attribute.h"
+#include "Containers/Array.h"
+#include "Containers/BitArray.h"
+#include "Containers/Set.h"
+#include "Containers/SparseArray.h"
+#include "Containers/UnrealString.h"
+#include "Delegates/Delegate.h"
+#include "Framework/SlateDelegates.h"
 #include "Input/Reply.h"
+#include "Internationalization/Text.h"
+#include "Misc/Attribute.h"
+#include "Misc/Optional.h"
+#include "Templates/SharedPointer.h"
+#include "Types/SlateEnums.h"
 #include "Widgets/DeclarativeSyntaxSupport.h"
 #include "Widgets/SCompoundWidget.h"
-#include "Widgets/Views/STableViewBase.h"
-#include "Widgets/Views/STableRow.h"
-#include "Widgets/Input/SSearchBox.h"
 
+class ITableRow;
 class SMenuAnchor;
+class SSearchBox;
+class STableViewBase;
+struct FFocusEvent;
+struct FGeometry;
+struct FKeyEvent;
+template <typename ItemType> class SListView;
 
 
 /**
diff --git a/Engine/Plugins/Experimental/Mutable/Source/CustomizableObjectEditor/Private/MuCOE/StatePerformanceTesting.cpp b/Engine/Plugins/Experimental/Mutable/Source/CustomizableObjectEditor/Private/MuCOE/StatePerformanceTesting.cpp
--- a/Engine/Plugins/Experimental/Mutable/Source/CustomizableObjectEditor/Private/MuCOE/StatePerformanceTesting.cpp
+++ b/Engine/Plugins/Experimental/Mutable/Source/CustomizableObjectEditor/Private/MuCOE/StatePerformanceTesting.cpp
@@ -2,12 +2,24 @@
 
 
 #include "MuCOE/StatePerformanceTesting.h"
-#include "MuCOE/CustomizableObjectEditor.h"
-#include "MuCOE/CustomizableObjectEditorModule.h"
-#include "CustomizableObjectClasses.h"
-#include "MuCO/CustomizableObjectInstance.h"
+
 #include "Components/SkeletalMeshComponent.h"
+#include "Editor.h"
+#include "Editor/EditorEngine.h"
+#include "Engine/EngineTypes.h"
+#include "HAL/PlatformCrt.h"
+#include "HAL/PlatformTime.h"
+#include "Internationalization/Internationalization.h"
+#include "Internationalization/Text.h"
 #include "Logging/MessageLog.h"
+#include "Logging/TokenizedMessage.h"
+#include "MuCO/CustomizableObject.h"
+#include "MuCO/CustomizableObjectInstance.h"
+#include "MuCO/CustomizableSkeletalComponent.h"
+#include "UObject/Class.h"
+#include "UObject/ObjectPtr.h"
+#include "UObject/Package.h"
+#include "UObject/UnrealNames.h"
 
 
 #define LOCTEXT_NAMESPACE "CustomizableObjectEditor"
diff --git a/Engine/Plugins/Experimental/Mutable/Source/CustomizableObjectEditor/Private/MuCOE/StatePerformanceTesting.h b/Engine/Plugins/Experimental/Mutable/Source/CustomizableObjectEditor/Private/MuCOE/StatePerformanceTesting.h
--- a/Engine/Plugins/Experimental/Mutable/Source/CustomizableObjectEditor/Private/MuCOE/StatePerformanceTesting.h
+++ b/Engine/Plugins/Experimental/Mutable/Source/CustomizableObjectEditor/Private/MuCOE/StatePerformanceTesting.h
@@ -2,18 +2,27 @@
 
 #pragma once
 
-#include "MuCOE/ICustomizableObjectEditor.h"
-#include "Misc/NotifyHook.h"
-#include "MuCOE/CustomizableObjectCompiler.h"
-#include "MuCO/CustomizableObjectInstance.h"
-#include "StatePerformanceTesting.generated.h"
+#include "Containers/Array.h"
+#include "Containers/Map.h"
+#include "Containers/Ticker.h"
+#include "Containers/UnrealString.h"
+#include "Delegates/Delegate.h"
+#include "Math/Color.h"
+#include "Math/UnrealMathSSE.h"
+#include "Math/Vector.h"
+#include "MuCO/CustomizableObjectParameterTypeDefinitions.h"
+#include "UObject/GCObject.h"
+#include "UObject/Object.h"
+#include "UObject/UObjectGlobals.h"
 
+#include "StatePerformanceTesting.generated.h"
 
-class UInstance;
+class UCustomizableObject;
+class UCustomizableObjectInstance;
 class UCustomizableSkeletalComponent;
-class UStaticMeshComponent;
 class USkeletalMeshComponent;
 class UUpdateClassHelper;
+struct FFrame;
 
 
 enum class TestType
diff --git a/Engine/Plugins/Experimental/Mutable/Source/CustomizableObjectEditor/Private/MuCOE/StressTest.cpp b/Engine/Plugins/Experimental/Mutable/Source/CustomizableObjectEditor/Private/MuCOE/StressTest.cpp
--- a/Engine/Plugins/Experimental/Mutable/Source/CustomizableObjectEditor/Private/MuCOE/StressTest.cpp
+++ b/Engine/Plugins/Experimental/Mutable/Source/CustomizableObjectEditor/Private/MuCOE/StressTest.cpp
@@ -1,15 +1,31 @@
 // Copyright Epic Games, Inc. All Rights Reserved.
 
 #include "MuCOE/StressTest.h"
-#include "UObject/Package.h"
-#include "MuCO/CustomizableObject.h"
-#include "MuCO/CustomizableObjectInstance.h"
-#include "Widgets/Notifications/SNotificationList.h"
+
+#include "Containers/IndirectArray.h"
 #include "Engine/SkeletalMesh.h"
-#include "Materials/Material.h"
+#include "Engine/Texture.h"
+#include "Engine/Texture2D.h"
+#include "HAL/PlatformCrt.h"
+#include "Logging/LogCategory.h"
+#include "Logging/LogMacros.h"
 #include "Materials/MaterialInstance.h"
-#include "EditorFramework/AssetImportData.h"
+#include "Materials/MaterialInterface.h"
+#include "Math/Color.h"
+#include "Math/NumericLimits.h"
+#include "Math/Vector.h"
+#include "MuCO/CustomizableObject.h"
+#include "MuCO/CustomizableObjectInstance.h"
+#include "MuCO/CustomizableSkeletalComponent.h"
 #include "MuCO/UnrealPortabilityHelpers.h"
+#include "Templates/Casts.h"
+#include "Trace/Detail/Channel.h"
+#include "UObject/ObjectPtr.h"
+#include "UObject/Package.h"
+#include "UObject/UnrealNames.h"
+
+class FSkeletalMeshLODRenderData;
+class FSkeletalMeshRenderData;
 
 
 void ULiveInstance::DelegatedCallback(UCustomizableObjectInstance* UpdatedInstance)
diff --git a/Engine/Plugins/Experimental/Mutable/Source/CustomizableObjectEditor/Private/MuCOE/StressTest.h b/Engine/Plugins/Experimental/Mutable/Source/CustomizableObjectEditor/Private/MuCOE/StressTest.h
--- a/Engine/Plugins/Experimental/Mutable/Source/CustomizableObjectEditor/Private/MuCOE/StressTest.h
+++ b/Engine/Plugins/Experimental/Mutable/Source/CustomizableObjectEditor/Private/MuCOE/StressTest.h
@@ -2,14 +2,25 @@
 
 #pragma once
 
-#include "Delegates/IDelegateInstance.h"
-#include "MuCOE/StatePerformanceTesting.h"
-#include "MuCO/CustomizableObject.h"
+#include "Containers/Array.h"
+#include "Containers/Map.h"
+#include "Containers/SparseArray.h"
 #include "Containers/Ticker.h"
-#include "StressTest.generated.h"
+#include "Containers/UnrealString.h"
+#include "Delegates/Delegate.h"
+#include "Math/IntVector.h"
+#include "Math/UnrealMathSSE.h"
+#include "MuCO/CustomizableObjectParameterTypeDefinitions.h"
+#include "MuCOE/StatePerformanceTesting.h"
+#include "PixelFormat.h"
+#include "UObject/Object.h"
+#include "UObject/UObjectGlobals.h"
 
+#include "StressTest.generated.h"
 
+class UCustomizableObject;
 class UCustomizableObjectInstance;
+struct FFrame;
 
 
 struct MaterialBriefInfo
diff --git a/Engine/Plugins/Experimental/Mutable/Source/CustomizableObjectEditor/Private/MuCOE/Thumbnails/CustomizableObjectInstanceThumbnailRenderer.cpp b/Engine/Plugins/Experimental/Mutable/Source/CustomizableObjectEditor/Private/MuCOE/Thumbnails/CustomizableObjectInstanceThumbnailRenderer.cpp
--- a/Engine/Plugins/Experimental/Mutable/Source/CustomizableObjectEditor/Private/MuCOE/Thumbnails/CustomizableObjectInstanceThumbnailRenderer.cpp
+++ b/Engine/Plugins/Experimental/Mutable/Source/CustomizableObjectEditor/Private/MuCOE/Thumbnails/CustomizableObjectInstanceThumbnailRenderer.cpp
@@ -1,12 +1,26 @@
 // Copyright Epic Games, Inc. All Rights Reserved.
 
 #include "MuCOE/Thumbnails/CustomizableObjectInstanceThumbnailRenderer.h"
-#include "MuCOE/CustomizableObjectEditorModule.h"
+
+#include "CanvasTypes.h"
+#include "Engine/Texture2D.h"
+#include "Math/Color.h"
+#include "MuCO/CustomizableObject.h"
 #include "MuCO/CustomizableObjectInstance.h"
+#include "MuCO/CustomizableSkeletalComponent.h"
 #include "MuCOE/CustomizableObjectCompiler.h"
+#include "SceneView.h"
+#include "ShowFlags.h"
+#include "Templates/Casts.h"
+#include "TextureResource.h"
+#include "ThumbnailHelpers.h"
+#include "ThumbnailRendering/ThumbnailRenderer.h"
+#include "UObject/Class.h"
 #include "UObject/ConstructorHelpers.h"
-#include "CanvasTypes.h"
-#include "EngineModule.h"
+#include "UObject/Object.h"
+
+class FRenderTarget;
+class USkeletalMesh;
 
 #define LOCTEXT_NAMESPACE "CustomizableObjectInstanceThumbnailRenderer"
 
diff --git a/Engine/Plugins/Experimental/Mutable/Source/CustomizableObjectEditor/Private/MuCOE/Thumbnails/CustomizableObjectThumbnailRenderer.cpp b/Engine/Plugins/Experimental/Mutable/Source/CustomizableObjectEditor/Private/MuCOE/Thumbnails/CustomizableObjectThumbnailRenderer.cpp
--- a/Engine/Plugins/Experimental/Mutable/Source/CustomizableObjectEditor/Private/MuCOE/Thumbnails/CustomizableObjectThumbnailRenderer.cpp
+++ b/Engine/Plugins/Experimental/Mutable/Source/CustomizableObjectEditor/Private/MuCOE/Thumbnails/CustomizableObjectThumbnailRenderer.cpp
@@ -1,15 +1,27 @@
 // Copyright Epic Games, Inc. All Rights Reserved.
 
 #include "MuCOE/Thumbnails/CustomizableObjectThumbnailRenderer.h"
-#include "MuCOE/CustomizableObjectEditorModule.h"
+
+#include "CanvasTypes.h"
+#include "Containers/Array.h"
+#include "Engine/Texture2D.h"
+#include "Math/Color.h"
 #include "MuCO/CustomizableObject.h"
-#include "MuCO/CustomizableObjectInstance.h"
+#include "MuCO/CustomizableSkeletalComponent.h"
+#include "MuCO/UnrealPortabilityHelpers.h"
 #include "MuCOE/CustomizableObjectCompiler.h"
+#include "SceneView.h"
+#include "ShowFlags.h"
+#include "Templates/Casts.h"
+#include "TextureResource.h"
+#include "ThumbnailHelpers.h"
+#include "ThumbnailRendering/ThumbnailRenderer.h"
+#include "UObject/Class.h"
 #include "UObject/ConstructorHelpers.h"
-#include "CanvasTypes.h"
-#include "EngineModule.h"
+#include "UObject/Object.h"
 
-#include "MuCO/UnrealPortabilityHelpers.h"
+class FRenderTarget;
+class USkeletalMesh;
 
 
 
diff --git a/Engine/Plugins/Experimental/Mutable/Source/CustomizableObjectEditor/Private/MuCOE/UnrealToMutableTextureConversionUtils.cpp b/Engine/Plugins/Experimental/Mutable/Source/CustomizableObjectEditor/Private/MuCOE/UnrealToMutableTextureConversionUtils.cpp
--- a/Engine/Plugins/Experimental/Mutable/Source/CustomizableObjectEditor/Private/MuCOE/UnrealToMutableTextureConversionUtils.cpp
+++ b/Engine/Plugins/Experimental/Mutable/Source/CustomizableObjectEditor/Private/MuCOE/UnrealToMutableTextureConversionUtils.cpp
@@ -2,13 +2,26 @@
 
 #include "MuCOE/UnrealToMutableTextureConversionUtils.h"
 
+#include "Containers/ArrayView.h"
+#include "Containers/ContainersFwd.h"
+#include "Containers/EnumAsByte.h"
+#include "Engine/Texture.h"
 #include "Engine/Texture2D.h"
-#include "RenderUtils.h"
+#include "Engine/TextureDefines.h"
+#include "HAL/PlatformCrt.h"
+#include "HAL/UnrealMemory.h"
 #include "ImageCore.h"
 #include "ImageCoreUtils.h"
-#include "Math/UnrealMathUtility.h"
 //#include "MuCO/CustomizableObject.h"	// For the LogMutable logging category
 #include "ImageUtils.h"
+#include "Math/Color.h"
+#include "Math/UnrealMathSSE.h"
+#include "Math/UnrealMathUtility.h"
+#include "Math/Vector.h"
+#include "Misc/AssertionMacros.h"
+#include "MuR/Ptr.h"
+#include "MuR/RefCounted.h"
+#include "Templates/UnrealTemplate.h"
 
 namespace UnrealToMutableImageConversion_Interanl
 {
diff --git a/Engine/Plugins/Experimental/Mutable/Source/CustomizableObjectEditor/Private/MuCOE/Widgets/SMutableExpandableTableRow.cpp b/Engine/Plugins/Experimental/Mutable/Source/CustomizableObjectEditor/Private/MuCOE/Widgets/SMutableExpandableTableRow.cpp
--- a/Engine/Plugins/Experimental/Mutable/Source/CustomizableObjectEditor/Private/MuCOE/Widgets/SMutableExpandableTableRow.cpp
+++ b/Engine/Plugins/Experimental/Mutable/Source/CustomizableObjectEditor/Private/MuCOE/Widgets/SMutableExpandableTableRow.cpp
@@ -2,6 +2,8 @@
 
 #include "MuCOE/Widgets/SMutableExpandableTableRow.h"
 
+#include "Math/Color.h"
+
 
 FSlateColor GetRowBackgroundColor(const int32 IndentLevel, const bool IsHovered)
 {
diff --git a/Engine/Plugins/Experimental/Mutable/Source/CustomizableObjectEditor/Private/MuCOE/Widgets/SMutableExpandableTableRow.h b/Engine/Plugins/Experimental/Mutable/Source/CustomizableObjectEditor/Private/MuCOE/Widgets/SMutableExpandableTableRow.h
--- a/Engine/Plugins/Experimental/Mutable/Source/CustomizableObjectEditor/Private/MuCOE/Widgets/SMutableExpandableTableRow.h
+++ b/Engine/Plugins/Experimental/Mutable/Source/CustomizableObjectEditor/Private/MuCOE/Widgets/SMutableExpandableTableRow.h
@@ -2,11 +2,40 @@
 
 #pragma once
 
+#include "Containers/IndirectArray.h"
+#include "Containers/Map.h"
+#include "Framework/Views/ITypedTableView.h"
+#include "HAL/Platform.h"
+#include "Input/Reply.h"
+#include "Layout/Clipping.h"
+#include "Layout/Margin.h"
+#include "Layout/Visibility.h"
+#include "Misc/AssertionMacros.h"
+#include "Misc/Attribute.h"
+#include "Misc/Optional.h"
 #include "MuCOE/SMutableConstrainedBox.h"
-#include "Widgets/Views/STableRow.h"
-#include "Widgets/Layout/SBorder.h"
+#include "SlotBase.h"
+#include "Styling/AppStyle.h"
+#include "Styling/CoreStyle.h"
+#include "Styling/ISlateStyle.h"
+#include "Styling/SlateColor.h"
+#include "Templates/SharedPointer.h"
+#include "Types/SlateEnums.h"
+#include "Types/SlateStructs.h"
+#include "UObject/NameTypes.h"
+#include "Widgets/DeclarativeSyntaxSupport.h"
 #include "Widgets/Images/SImage.h"
 #include "Widgets/Input/SButton.h"
+#include "Widgets/Layout/SBorder.h"
+#include "Widgets/Layout/SBox.h"
+#include "Widgets/SBoxPanel.h"
+#include "Widgets/SNullWidget.h"
+#include "Widgets/SWidget.h"
+#include "Widgets/Views/SHeaderRow.h"
+#include "Widgets/Views/STableRow.h"
+#include "Widgets/Views/STableViewBase.h"
+
+struct FSlateBrush;
 
 
 /** Get the background color of the Pin Viewer elements.
diff --git a/Engine/Plugins/Experimental/Mutable/Source/CustomizableObjectEditor/Public/MuCOE/CustomizableObjectDebuggerActions.h b/Engine/Plugins/Experimental/Mutable/Source/CustomizableObjectEditor/Public/MuCOE/CustomizableObjectDebuggerActions.h
--- a/Engine/Plugins/Experimental/Mutable/Source/CustomizableObjectEditor/Public/MuCOE/CustomizableObjectDebuggerActions.h
+++ b/Engine/Plugins/Experimental/Mutable/Source/CustomizableObjectEditor/Public/MuCOE/CustomizableObjectDebuggerActions.h
@@ -2,8 +2,10 @@
 
 #pragma once
 
-#include "Framework/Commands/UICommandInfo.h"
 #include "Framework/Commands/Commands.h"
+#include "Templates/SharedPointer.h"
+
+class FUICommandInfo;
 
 /**
  *
diff --git a/Engine/Plugins/Experimental/Mutable/Source/CustomizableObjectEditor/Public/MuCOE/CustomizableObjectEditorActions.h b/Engine/Plugins/Experimental/Mutable/Source/CustomizableObjectEditor/Public/MuCOE/CustomizableObjectEditorActions.h
--- a/Engine/Plugins/Experimental/Mutable/Source/CustomizableObjectEditor/Public/MuCOE/CustomizableObjectEditorActions.h
+++ b/Engine/Plugins/Experimental/Mutable/Source/CustomizableObjectEditor/Public/MuCOE/CustomizableObjectEditorActions.h
@@ -2,8 +2,10 @@
 
 #pragma once
 
-#include "Framework/Commands/UICommandInfo.h"
 #include "Framework/Commands/Commands.h"
+#include "Templates/SharedPointer.h"
+
+class FUICommandInfo;
 
 /**
  * 
diff --git a/Engine/Plugins/Experimental/Mutable/Source/CustomizableObjectEditor/Public/MuCOE/CustomizableObjectEditorLogger.h b/Engine/Plugins/Experimental/Mutable/Source/CustomizableObjectEditor/Public/MuCOE/CustomizableObjectEditorLogger.h
--- a/Engine/Plugins/Experimental/Mutable/Source/CustomizableObjectEditor/Public/MuCOE/CustomizableObjectEditorLogger.h
+++ b/Engine/Plugins/Experimental/Mutable/Source/CustomizableObjectEditor/Public/MuCOE/CustomizableObjectEditorLogger.h
@@ -2,10 +2,18 @@
 
 #pragma once
 
+#include "Containers/Array.h"
+#include "Containers/Map.h"
+#include "HAL/Platform.h"
+#include "Internationalization/Text.h"
+#include "Logging/TokenizedMessage.h"
+#include "Misc/Attribute.h"
 #include "MuCOE/Nodes/CustomizableObjectNode.h"
-
+#include "Templates/SharedPointer.h"
 
 class FCustomizableObjectEditorLogger;
+class FName;
+class SNotificationItem;
 
 /** Notification categories. */
 enum class ELoggerCategory
diff --git a/Engine/Plugins/Experimental/Mutable/Source/CustomizableObjectEditor/Public/MuCOE/CustomizableObjectEditorModule.h b/Engine/Plugins/Experimental/Mutable/Source/CustomizableObjectEditor/Public/MuCOE/CustomizableObjectEditorModule.h
--- a/Engine/Plugins/Experimental/Mutable/Source/CustomizableObjectEditor/Public/MuCOE/CustomizableObjectEditorModule.h
+++ b/Engine/Plugins/Experimental/Mutable/Source/CustomizableObjectEditor/Public/MuCOE/CustomizableObjectEditorModule.h
@@ -3,15 +3,19 @@
 #pragma once
 
 #include "AssetTypeCategories.h"
-#include "CustomizableObjectEditorClasses.h"
-
 #include "Modules/ModuleInterface.h"
 #include "Modules/ModuleManager.h"
-#include "Toolkits/AssetEditorToolkit.h"
+#include "Templates/SharedPointer.h"
 #include "Toolkits/IToolkit.h"
-#include "MuCOE/ICustomizableObjectEditor.h"
-#include "MuCOE/ICustomizableObjectDebugger.h"
-#include "MuCOE/CustomizableObjectEditorLogger.h"
+#include "UObject/NameTypes.h"
+
+class FCustomizableObjectEditorLogger;
+class ICustomizableObjectDebugger;
+class ICustomizableObjectEditor;
+class ICustomizableObjectInstanceEditor;
+class IToolkitHost;
+class UCustomizableObject;
+class UCustomizableObjectInstance;
 
 extern const FName CustomizableObjectEditorAppIdentifier;
 extern const FName CustomizableObjectInstanceEditorAppIdentifier;
diff --git a/Engine/Plugins/Experimental/Mutable/Source/CustomizableObjectEditor/Public/MuCOE/CustomizableObjectEditorStyle.h b/Engine/Plugins/Experimental/Mutable/Source/CustomizableObjectEditor/Public/MuCOE/CustomizableObjectEditorStyle.h
--- a/Engine/Plugins/Experimental/Mutable/Source/CustomizableObjectEditor/Public/MuCOE/CustomizableObjectEditorStyle.h
+++ b/Engine/Plugins/Experimental/Mutable/Source/CustomizableObjectEditor/Public/MuCOE/CustomizableObjectEditorStyle.h
@@ -1,7 +1,13 @@
 // Copyright Epic Games, Inc. All Rights Reserved.
 
 #pragma once
-#include "Styling/ISlateStyle.h"
+
+#include "Containers/UnrealString.h"
+#include "HAL/Platform.h"
+#include "Templates/SharedPointer.h"
+#include "UObject/NameTypes.h"
+
+class ISlateStyle;
 
 /** Manages the style which provides resources for niagara editor widgets. */
 class FCustomizableObjectEditorStyle
