Ran IWYU on Mutable plugin

diff --git a/Engine/Plugins/Experimental/Mutable/Source/CustomizableObjectEditor/Public/MuCOE/CustomizableObjectEditorViewportClient.h b/Engine/Plugins/Experimental/Mutable/Source/CustomizableObjectEditor/Public/MuCOE/CustomizableObjectEditorViewportClient.h
--- a/Engine/Plugins/Experimental/Mutable/Source/CustomizableObjectEditor/Public/MuCOE/CustomizableObjectEditorViewportClient.h
+++ b/Engine/Plugins/Experimental/Mutable/Source/CustomizableObjectEditor/Public/MuCOE/CustomizableObjectEditorViewportClient.h
@@ -2,15 +2,48 @@
 
 #pragma once
 
-#include "MuCOE/CustomizableObjectWidget.h"
-#include "EditorViewportClient.h"
 #include "Components/SkeletalMeshComponent.h"
-#include "Animation/AnimationAsset.h"
-#include "MuCOE/Nodes/CustomizableObjectNodeProjectorConstant.h"
-#include "MuCOE/Nodes/CustomizableObjectNodeProjectorParameter.h"
-#include "MuCOE/Nodes/CustomizableObjectNodeMeshClipMorph.h"
-#include "MuCOE/Nodes/CustomizableObjectNodeMeshClipWithMesh.h"
-#include "MuCOE/ICustomizableObjectInstanceEditor.h"
+#include "Containers/Array.h"
+#include "Containers/Map.h"
+#include "Containers/UnrealString.h"
+#include "Editor/UnrealEdTypes.h"
+#include "EditorViewportClient.h"
+#include "Math/Axis.h"
+#include "Math/BoxSphereBounds.h"
+#include "Math/Color.h"
+#include "Math/Matrix.h"
+#include "Math/Rotator.h"
+#include "Math/Sphere.h"
+#include "Math/UnrealMathSSE.h"
+#include "Math/Vector.h"
+#include "MuCO/CustomizableObjectParameterTypeDefinitions.h"
+#include "Templates/SharedPointer.h"
+#include "UObject/NameTypes.h"
+#include "UObject/WeakObjectPtrTemplates.h"
+#include "UnrealWidgetFwd.h"
+
+class FCanvas;
+class FCustomizableObjectWidget;
+class FMaterialRenderProxy;
+class FPreviewScene;
+class FPrimitiveDrawInterface;
+class FReferenceCollector;
+class FSceneView;
+class FViewport;
+class ICustomizableObjectInstanceEditor;
+class UAnimationAsset;
+class UCustomizableObject;
+class UCustomizableObjectNodeMeshClipMorph;
+class UCustomizableObjectNodeMeshClipWithMesh;
+class UCustomizableObjectNodeProjectorConstant;
+class UCustomizableObjectNodeProjectorParameter;
+class UDebugSkelMeshComponent;
+class UMaterial;
+class UMaterialInterface;
+class UObject;
+class UStaticMesh;
+class UStaticMeshComponent;
+struct FInputKeyEventArgs;
 
 
 extern void RemoveRestrictedChars(FString& String);
diff --git a/Engine/Plugins/Experimental/Mutable/Source/CustomizableObjectEditor/Public/MuCOE/CustomizableObjectEditorViewportLODCommands.h b/Engine/Plugins/Experimental/Mutable/Source/CustomizableObjectEditor/Public/MuCOE/CustomizableObjectEditorViewportLODCommands.h
--- a/Engine/Plugins/Experimental/Mutable/Source/CustomizableObjectEditor/Public/MuCOE/CustomizableObjectEditorViewportLODCommands.h
+++ b/Engine/Plugins/Experimental/Mutable/Source/CustomizableObjectEditor/Public/MuCOE/CustomizableObjectEditorViewportLODCommands.h
@@ -3,8 +3,14 @@
 #pragma once
 
 #include "Framework/Commands/Commands.h"
-#include "Styling/AppStyle.h"
+#include "HAL/Platform.h"
+#include "Internationalization/Internationalization.h"
 #include "MuCOE/CustomizableObjectEditorStyle.h"
+#include "Templates/SharedPointer.h"
+#include "UObject/NameTypes.h"
+#include "UObject/UnrealNames.h"
+
+class FUICommandInfo;
 
 /**
 * Class containing commands for persona viewport LOD actions
diff --git a/Engine/Plugins/Experimental/Mutable/Source/CustomizableObjectEditor/Public/MuCOE/CustomizableObjectEditorViewportMenuCommands.h b/Engine/Plugins/Experimental/Mutable/Source/CustomizableObjectEditor/Public/MuCOE/CustomizableObjectEditorViewportMenuCommands.h
--- a/Engine/Plugins/Experimental/Mutable/Source/CustomizableObjectEditor/Public/MuCOE/CustomizableObjectEditorViewportMenuCommands.h
+++ b/Engine/Plugins/Experimental/Mutable/Source/CustomizableObjectEditor/Public/MuCOE/CustomizableObjectEditorViewportMenuCommands.h
@@ -3,8 +3,14 @@
 #pragma once
 
 #include "Framework/Commands/Commands.h"
-#include "Styling/AppStyle.h"
+#include "HAL/Platform.h"
+#include "Internationalization/Internationalization.h"
 #include "MuCOE/CustomizableObjectEditorStyle.h"
+#include "Templates/SharedPointer.h"
+#include "UObject/NameTypes.h"
+#include "UObject/UnrealNames.h"
+
+class FUICommandInfo;
 
 /**
 * Class containing commands for viewport menu actions
diff --git a/Engine/Plugins/Experimental/Mutable/Source/CustomizableObjectEditor/Public/MuCOE/CustomizableObjectPin.h b/Engine/Plugins/Experimental/Mutable/Source/CustomizableObjectEditor/Public/MuCOE/CustomizableObjectPin.h
--- a/Engine/Plugins/Experimental/Mutable/Source/CustomizableObjectEditor/Public/MuCOE/CustomizableObjectPin.h
+++ b/Engine/Plugins/Experimental/Mutable/Source/CustomizableObjectEditor/Public/MuCOE/CustomizableObjectPin.h
@@ -2,8 +2,8 @@
 
 #pragma once
 
-#include "EdGraph/EdGraphPin.h"
 
+class UEdGraphPin;
 
 /** Returns true if the pin is orphan. */
 bool IsPinOrphan(const UEdGraphPin &Pin);
diff --git a/Engine/Plugins/Experimental/Mutable/Source/CustomizableObjectEditor/Public/MuCOE/CustomizableObjectWidget.h b/Engine/Plugins/Experimental/Mutable/Source/CustomizableObjectEditor/Public/MuCOE/CustomizableObjectWidget.h
--- a/Engine/Plugins/Experimental/Mutable/Source/CustomizableObjectEditor/Public/MuCOE/CustomizableObjectWidget.h
+++ b/Engine/Plugins/Experimental/Mutable/Source/CustomizableObjectEditor/Public/MuCOE/CustomizableObjectWidget.h
@@ -2,13 +2,25 @@
 
 #pragma once
 
-#include "UObject/GCObject.h"
-#include "SceneManagement.h"
+#include "Containers/UnrealString.h"
 #include "Editor/UnrealEdTypes.h"
-
+#include "GenericPlatform/ICursor.h"
+#include "HitProxies.h"
+#include "Math/Axis.h"
+#include "Math/Color.h"
+#include "Math/Matrix.h"
+#include "Math/UnrealMathSSE.h"
+#include "Math/Vector.h"
+#include "Math/Vector2D.h"
+#include "Math/Vector4.h"
 #include "SceneView.h"
-#include "SceneManagement.h"
-#include "Editor/UnrealEdTypes.h"
+#include "UObject/GCObject.h"
+
+class FPrimitiveDrawInterface;
+class FReferenceCollector;
+class UCustomizableObjectInstance;
+class UMaterialInstanceDynamic;
+class UMaterialInterface;
 
 
 class FCustomizableObjectWidget : public FGCObject
diff --git a/Engine/Plugins/Experimental/Mutable/Source/CustomizableObjectEditor/Public/MuCOE/GraphTraversal.h b/Engine/Plugins/Experimental/Mutable/Source/CustomizableObjectEditor/Public/MuCOE/GraphTraversal.h
--- a/Engine/Plugins/Experimental/Mutable/Source/CustomizableObjectEditor/Public/MuCOE/GraphTraversal.h
+++ b/Engine/Plugins/Experimental/Mutable/Source/CustomizableObjectEditor/Public/MuCOE/GraphTraversal.h
@@ -4,9 +4,9 @@
 
 #include "Containers/Array.h"
 
-#include "EdGraph/EdGraphPin.h"
-
-#include "MuCOE/Nodes/CustomizableObjectNodeObject.h"
+class UCustomizableObject;
+class UCustomizableObjectNodeObject;
+class UEdGraphPin;
 
 /** Follow the given input pin returning the output connected pin.
  *
diff --git a/Engine/Plugins/Experimental/Mutable/Source/CustomizableObjectEditor/Public/MuCOE/MutableMeshPreviewUtils.h b/Engine/Plugins/Experimental/Mutable/Source/CustomizableObjectEditor/Public/MuCOE/MutableMeshPreviewUtils.h
--- a/Engine/Plugins/Experimental/Mutable/Source/CustomizableObjectEditor/Public/MuCOE/MutableMeshPreviewUtils.h
+++ b/Engine/Plugins/Experimental/Mutable/Source/CustomizableObjectEditor/Public/MuCOE/MutableMeshPreviewUtils.h
@@ -2,8 +2,9 @@
 
 #pragma once
 
-#include "../Private/MuCO/CustomizableObjectSystemPrivate.h"
-#include "MuR/Mesh.h"
+#include "MuR/System.h"
+
+class USkeletalMesh;
 
 /**
  * Namespace designed to contain a series of methods used for the generation of unreal meshes for previewing purposes.
diff --git a/Engine/Plugins/Experimental/Mutable/Source/CustomizableObjectEditor/Public/MuCOE/MutableUtils.h b/Engine/Plugins/Experimental/Mutable/Source/CustomizableObjectEditor/Public/MuCOE/MutableUtils.h
--- a/Engine/Plugins/Experimental/Mutable/Source/CustomizableObjectEditor/Public/MuCOE/MutableUtils.h
+++ b/Engine/Plugins/Experimental/Mutable/Source/CustomizableObjectEditor/Public/MuCOE/MutableUtils.h
@@ -2,8 +2,10 @@
 
 #pragma once
 
+#include "Containers/Array.h"
 #include "Engine/SkeletalMesh.h"
 #include "Engine/StaticMesh.h"
+#include "HAL/Platform.h"
 #include "Math/Vector2D.h"
 
 
diff --git a/Engine/Plugins/Experimental/Mutable/Source/CustomizableObjectEditor/Public/MuCOE/RemapPins/CustomizableObjectNodeRemapPinsByName.h b/Engine/Plugins/Experimental/Mutable/Source/CustomizableObjectEditor/Public/MuCOE/RemapPins/CustomizableObjectNodeRemapPinsByName.h
--- a/Engine/Plugins/Experimental/Mutable/Source/CustomizableObjectEditor/Public/MuCOE/RemapPins/CustomizableObjectNodeRemapPinsByName.h
+++ b/Engine/Plugins/Experimental/Mutable/Source/CustomizableObjectEditor/Public/MuCOE/RemapPins/CustomizableObjectNodeRemapPinsByName.h
@@ -2,10 +2,16 @@
 
 #pragma once
 
+#include "Containers/Array.h"
+#include "Containers/Map.h"
 #include "MuCOE/RemapPins/CustomizableObjectNodeRemapPins.h"
+#include "UObject/UObjectGlobals.h"
 
 #include "CustomizableObjectNodeRemapPinsByName.generated.h"
 
+class UEdGraphPin;
+class UObject;
+
 
 /** 
  * Remap pins by name.
diff --git a/Engine/Plugins/Experimental/Mutable/Source/CustomizableObjectEditor/Public/MuCOE/RemapPins/CustomizableObjectNodeRemapPinsByNameDefaultPin.h b/Engine/Plugins/Experimental/Mutable/Source/CustomizableObjectEditor/Public/MuCOE/RemapPins/CustomizableObjectNodeRemapPinsByNameDefaultPin.h
--- a/Engine/Plugins/Experimental/Mutable/Source/CustomizableObjectEditor/Public/MuCOE/RemapPins/CustomizableObjectNodeRemapPinsByNameDefaultPin.h
+++ b/Engine/Plugins/Experimental/Mutable/Source/CustomizableObjectEditor/Public/MuCOE/RemapPins/CustomizableObjectNodeRemapPinsByNameDefaultPin.h
@@ -2,10 +2,16 @@
 
 #pragma once
 
+#include "Containers/Array.h"
+#include "Containers/Map.h"
 #include "MuCOE/RemapPins/CustomizableObjectNodeRemapPinsByName.h"
+#include "UObject/UObjectGlobals.h"
 
 #include "CustomizableObjectNodeRemapPinsByNameDefaultPin.generated.h"
 
+class UEdGraphPin;
+class UObject;
+
 
 /** If a default pin is given, remap it to the first pin. Otherwise, remap pins by name as usual. */
 UCLASS()
diff --git a/Engine/Plugins/Experimental/Mutable/Source/CustomizableObjectEditor/Public/MuCOE/SCustomizableInstanceProperties.h b/Engine/Plugins/Experimental/Mutable/Source/CustomizableObjectEditor/Public/MuCOE/SCustomizableInstanceProperties.h
--- a/Engine/Plugins/Experimental/Mutable/Source/CustomizableObjectEditor/Public/MuCOE/SCustomizableInstanceProperties.h
+++ b/Engine/Plugins/Experimental/Mutable/Source/CustomizableObjectEditor/Public/MuCOE/SCustomizableInstanceProperties.h
@@ -2,23 +2,31 @@
 
 #pragma once
 
-#include "MuCO/CustomizableObject.h"
-#include "MuCO/CustomizableObjectInstance.h"
-#include "Widgets/SCompoundWidget.h"
-#include "Widgets/Layout/SExpandableArea.h"
-#include "Widgets/DeclarativeSyntaxSupport.h"
+#include "Containers/Array.h"
+#include "Containers/Map.h"
+#include "Containers/UnrealString.h"
+#include "HAL/PlatformMath.h"
+#include "Input/Reply.h"
+#include "Internationalization/Text.h"
+#include "Math/Color.h"
 #include "Styling/SlateTypes.h"
-#include "Widgets/Input/SButton.h"
-#include "Widgets/Text/STextBlock.h"
-#include "Widgets/Input/SSearchBox.h"
+#include "Templates/SharedPointer.h"
+#include "Types/SlateEnums.h"
+#include "UObject/WeakObjectPtr.h"
+#include "UObject/WeakObjectPtrTemplates.h"
+#include "Widgets/DeclarativeSyntaxSupport.h"
 #include "Widgets/Input/SSpinBox.h"
-#include "MuCOE/CustomizableInstanceDetails.h"
+#include "Widgets/SCompoundWidget.h"
 
-class UCustomizableObjectInstance;
-class SSlider;
+class FCustomizableInstanceDetails;
+class SButton;
+class SExpandableArea;
+class SSearchBox;
+class STextBlock;
 class SVerticalBox;
-class UUpdateClassHelper;
-class UInstanceUpdatedHelper;
+class UCustomizableObjectInstance;
+struct FGeometry;
+struct FPointerEvent;
 
 struct FMutableParamExpandableArea
 {
diff --git a/Engine/Plugins/Experimental/Mutable/Source/CustomizableObjectEditor/Public/MuCOE/SCustomizableObjectEditorTransformViewportToolbar.h b/Engine/Plugins/Experimental/Mutable/Source/CustomizableObjectEditor/Public/MuCOE/SCustomizableObjectEditorTransformViewportToolbar.h
--- a/Engine/Plugins/Experimental/Mutable/Source/CustomizableObjectEditor/Public/MuCOE/SCustomizableObjectEditorTransformViewportToolbar.h
+++ b/Engine/Plugins/Experimental/Mutable/Source/CustomizableObjectEditor/Public/MuCOE/SCustomizableObjectEditorTransformViewportToolbar.h
@@ -2,8 +2,9 @@
 
 #pragma once
 
-#include "UnrealWidget.h"
 #include "SViewportToolBar.h"
+#include "Templates/SharedPointer.h"
+#include "Widgets/DeclarativeSyntaxSupport.h"
 
 class SCustomizableObjectEditorTransformViewportToolbar : public SViewportToolBar
 {
diff --git a/Engine/Plugins/Experimental/Mutable/Source/CustomizableObjectEditor/Public/MuCOE/SCustomizableObjectEditorViewportToolBar.h b/Engine/Plugins/Experimental/Mutable/Source/CustomizableObjectEditor/Public/MuCOE/SCustomizableObjectEditorViewportToolBar.h
--- a/Engine/Plugins/Experimental/Mutable/Source/CustomizableObjectEditor/Public/MuCOE/SCustomizableObjectEditorViewportToolBar.h
+++ b/Engine/Plugins/Experimental/Mutable/Source/CustomizableObjectEditor/Public/MuCOE/SCustomizableObjectEditorViewportToolBar.h
@@ -2,9 +2,24 @@
 
 #pragma once
 
+#include "Containers/Array.h"
+#include "Input/Reply.h"
+#include "Internationalization/Text.h"
+#include "Layout/Visibility.h"
+#include "Misc/Optional.h"
 #include "SViewportToolBar.h"
-#include "Framework/MultiBox/MultiBoxBuilder.h"
 #include "Settings/LevelEditorViewportSettings.h"
+#include "Styling/SlateColor.h"
+#include "Styling/SlateTypes.h"
+#include "Templates/SharedPointer.h"
+#include "Types/SlateEnums.h"
+#include "UObject/NameTypes.h"
+#include "Widgets/DeclarativeSyntaxSupport.h"
+
+class FMenuBuilder;
+class SMenuAnchor;
+class SWidget;
+struct FSlateBrush;
 
 
 /**
diff --git a/Engine/Plugins/Experimental/Mutable/Source/CustomizableObjectPopulation/Classes/MuCOP/CustomizableObjectPopulation.h b/Engine/Plugins/Experimental/Mutable/Source/CustomizableObjectPopulation/Classes/MuCOP/CustomizableObjectPopulation.h
--- a/Engine/Plugins/Experimental/Mutable/Source/CustomizableObjectPopulation/Classes/MuCOP/CustomizableObjectPopulation.h
+++ b/Engine/Plugins/Experimental/Mutable/Source/CustomizableObjectPopulation/Classes/MuCOP/CustomizableObjectPopulation.h
@@ -2,10 +2,22 @@
 
 #pragma once
 
-#include "MuCOP/CustomizableObjectPopulationClass.h"
-#include "MuCOP/CustomizableObjectPopulationGenerator.h"
+#include "Containers/Array.h"
+#include "Containers/UnrealString.h"
+#include "HAL/Platform.h"
+#include "Serialization/Archive.h"
+#include "UObject/Object.h"
+#include "UObject/ObjectPtr.h"
+#include "UObject/UObjectGlobals.h"
+
 #include "CustomizableObjectPopulation.generated.h"
 
+class ITargetPlatform;
+class UCustomizableObjectInstance;
+class UCustomizableObjectPopulationClass;
+class UCustomizableObjectPopulationGenerator;
+struct FFrame;
+
 
 USTRUCT()
 struct FClassWeightPair
diff --git a/Engine/Plugins/Experimental/Mutable/Source/CustomizableObjectPopulation/Classes/MuCOP/CustomizableObjectPopulationClass.h b/Engine/Plugins/Experimental/Mutable/Source/CustomizableObjectPopulation/Classes/MuCOP/CustomizableObjectPopulationClass.h
--- a/Engine/Plugins/Experimental/Mutable/Source/CustomizableObjectPopulation/Classes/MuCOP/CustomizableObjectPopulationClass.h
+++ b/Engine/Plugins/Experimental/Mutable/Source/CustomizableObjectPopulation/Classes/MuCOP/CustomizableObjectPopulationClass.h
@@ -2,9 +2,18 @@
 
 #pragma once
 
-#include "../Private/MuCOP/CustomizableObjectPopulationCharacteristic.h"
+#include "Containers/Array.h"
+#include "Containers/Map.h"
+#include "Containers/UnrealString.h"
+#include "UObject/Object.h"
+#include "UObject/ObjectPtr.h"
+#include "UObject/UObjectGlobals.h"
+
 #include "CustomizableObjectPopulationClass.generated.h"
 
+class FArchive;
+struct FCustomizableObjectPopulationCharacteristic;
+
 
 USTRUCT()
 struct FPopulationClassParameterOptions
diff --git a/Engine/Plugins/Experimental/Mutable/Source/CustomizableObjectPopulation/Classes/MuCOP/CustomizableObjectPopulationGenerator.h b/Engine/Plugins/Experimental/Mutable/Source/CustomizableObjectPopulation/Classes/MuCOP/CustomizableObjectPopulationGenerator.h
--- a/Engine/Plugins/Experimental/Mutable/Source/CustomizableObjectPopulation/Classes/MuCOP/CustomizableObjectPopulationGenerator.h
+++ b/Engine/Plugins/Experimental/Mutable/Source/CustomizableObjectPopulation/Classes/MuCOP/CustomizableObjectPopulationGenerator.h
@@ -2,18 +2,19 @@
 
 #pragma once
 
-#include "CoreTypes.h"
 #include "Containers/Array.h"
-#include "Containers/UnrealString.h"
-
-#include "Templates/PimplPtr.h"
-#include "Templates/Tuple.h"
-#include "Math/RandomStream.h"
+#include "CoreTypes.h"
+#include "Templates/SharedPointer.h"
+#include "UObject/Object.h"
+#include "UObject/ObjectPtr.h"
+#include "UObject/UObjectGlobals.h"
 
 //#include "MuCOP/CustomizableObjectPopulationSamplers.h"
 
 
 #include "CustomizableObjectPopulationGenerator.generated.h"
+
+class FArchive;
 /*
 struct FCustomizableObjectPopulationCharacteristic;
 class UCustomizableObjectPopulation;
@@ -59,10 +60,10 @@ public:
 };
 */
 
+class UCustomizableObject;
+class UCustomizableObjectInstance;
 // Forward declarations
 struct FCustomizableObjectPopulationGeneratorPrivate;
-class UCustomizableObjectInstance;
-class UCustomizableObject;
 
 UCLASS()
 class CUSTOMIZABLEOBJECTPOPULATION_API UCustomizableObjectPopulationGenerator : public UObject
diff --git a/Engine/Plugins/Experimental/Mutable/Source/CustomizableObjectPopulation/Private/MuCOP/CustomizableObjectPopulation.cpp b/Engine/Plugins/Experimental/Mutable/Source/CustomizableObjectPopulation/Private/MuCOP/CustomizableObjectPopulation.cpp
--- a/Engine/Plugins/Experimental/Mutable/Source/CustomizableObjectPopulation/Private/MuCOP/CustomizableObjectPopulation.cpp
+++ b/Engine/Plugins/Experimental/Mutable/Source/CustomizableObjectPopulation/Private/MuCOP/CustomizableObjectPopulation.cpp
@@ -1,15 +1,33 @@
 // Copyright Epic Games, Inc. All Rights Reserved.
 
 #include "MuCOP/CustomizableObjectPopulation.h"
+
+#include "Containers/Map.h"
+#include "Curves/RichCurve.h"
+#include "HAL/PlatformCrt.h"
+#include "Math/Color.h"
+#include "Math/UnrealMathSSE.h"
+#include "Misc/AssertionMacros.h"
 #include "MuCO/CustomizableObject.h"
+#include "MuCO/CustomizableObjectParameterTypeDefinitions.h"
+#include "MuCOP/CustomizableObjectPopulationCharacteristic.h"
+#include "MuCOP/CustomizableObjectPopulationClass.h"
+#include "MuCOP/CustomizableObjectPopulationConstraint.h"
 #include "MuCOP/CustomizableObjectPopulationCustomVersion.h"
+#include "MuCOP/CustomizableObjectPopulationGenerator.h"
+#include "Templates/Casts.h"
+#include "Templates/SharedPointer.h"
+#include "Templates/Tuple.h"
+
+class ITargetPlatform;
+class UCustomizableObjectInstance;
 
 #if WITH_EDITOR
-#include "MuCOP/CustomizableObjectPopulationSamplers.h"
-#include "MuCOP/CustomizableObjectPopulationGeneratorPrivate.h"
-#include "Curves/CurveLinearColor.h"
 #include "Curves/CurveFloat.h"
+#include "Curves/CurveLinearColor.h"
 #include "Curves/CurveVector.h"
+#include "MuCOP/CustomizableObjectPopulationGeneratorPrivate.h"
+#include "MuCOP/CustomizableObjectPopulationSamplers.h"
 #endif
 
 #define LOCTEXT_NAMESPACE "CustomizableObjectPopulation"
diff --git a/Engine/Plugins/Experimental/Mutable/Source/CustomizableObjectPopulation/Private/MuCOP/CustomizableObjectPopulationCharacteristic.cpp b/Engine/Plugins/Experimental/Mutable/Source/CustomizableObjectPopulation/Private/MuCOP/CustomizableObjectPopulationCharacteristic.cpp
--- a/Engine/Plugins/Experimental/Mutable/Source/CustomizableObjectPopulation/Private/MuCOP/CustomizableObjectPopulationCharacteristic.cpp
+++ b/Engine/Plugins/Experimental/Mutable/Source/CustomizableObjectPopulation/Private/MuCOP/CustomizableObjectPopulationCharacteristic.cpp
@@ -1,7 +1,6 @@
 // Copyright Epic Games, Inc. All Rights Reserved.
 
 #include "MuCOP/CustomizableObjectPopulationCharacteristic.h"
-#include "MuCO/CustomizableObject.h"
 
 #define LOCTEXT_NAMESPACE "CustomizableObjectPopulationCharacteristic"
 
diff --git a/Engine/Plugins/Experimental/Mutable/Source/CustomizableObjectPopulation/Private/MuCOP/CustomizableObjectPopulationCharacteristic.h b/Engine/Plugins/Experimental/Mutable/Source/CustomizableObjectPopulation/Private/MuCOP/CustomizableObjectPopulationCharacteristic.h
--- a/Engine/Plugins/Experimental/Mutable/Source/CustomizableObjectPopulation/Private/MuCOP/CustomizableObjectPopulationCharacteristic.h
+++ b/Engine/Plugins/Experimental/Mutable/Source/CustomizableObjectPopulation/Private/MuCOP/CustomizableObjectPopulationCharacteristic.h
@@ -1,7 +1,10 @@
 // Copyright Epic Games, Inc. All Rights Reserved.
 
 #pragma once
+#include "Containers/Array.h"
+#include "Containers/UnrealString.h"
 #include "CustomizableObjectPopulationConstraint.h"
+
 #include "CustomizableObjectPopulationCharacteristic.generated.h"
 
 USTRUCT()
diff --git a/Engine/Plugins/Experimental/Mutable/Source/CustomizableObjectPopulation/Private/MuCOP/CustomizableObjectPopulationClass.cpp b/Engine/Plugins/Experimental/Mutable/Source/CustomizableObjectPopulation/Private/MuCOP/CustomizableObjectPopulationClass.cpp
--- a/Engine/Plugins/Experimental/Mutable/Source/CustomizableObjectPopulation/Private/MuCOP/CustomizableObjectPopulationClass.cpp
+++ b/Engine/Plugins/Experimental/Mutable/Source/CustomizableObjectPopulation/Private/MuCOP/CustomizableObjectPopulationClass.cpp
@@ -1,8 +1,10 @@
 // Copyright Epic Games, Inc. All Rights Reserved.
 
 #include "MuCOP/CustomizableObjectPopulationClass.h"
+
 #include "MuCOP/CustomizableObjectPopulationCharacteristic.h"
 #include "MuCOP/CustomizableObjectPopulationCustomVersion.h"
+#include "Serialization/Archive.h"
 
 #define LOCTEXT_NAMESPACE "CustomizableObjectPopulationClass"
 
diff --git a/Engine/Plugins/Experimental/Mutable/Source/CustomizableObjectPopulation/Private/MuCOP/CustomizableObjectPopulationConstraint.h b/Engine/Plugins/Experimental/Mutable/Source/CustomizableObjectPopulation/Private/MuCOP/CustomizableObjectPopulationConstraint.h
--- a/Engine/Plugins/Experimental/Mutable/Source/CustomizableObjectPopulation/Private/MuCOP/CustomizableObjectPopulationConstraint.h
+++ b/Engine/Plugins/Experimental/Mutable/Source/CustomizableObjectPopulation/Private/MuCOP/CustomizableObjectPopulationConstraint.h
@@ -3,9 +3,16 @@
 #pragma once
 
 
-#include "Curves/CurveBase.h"
+#include "Containers/Array.h"
+#include "Containers/UnrealString.h"
+#include "CoreTypes.h"
+#include "Math/Color.h"
+#include "UObject/ObjectPtr.h"
+
 #include "CustomizableObjectPopulationConstraint.generated.h"
 
+class UCurveBase;
+
 UENUM()
 enum class EPopulationConstraintType : uint8
 {
diff --git a/Engine/Plugins/Experimental/Mutable/Source/CustomizableObjectPopulation/Private/MuCOP/CustomizableObjectPopulationCustomVersion.cpp b/Engine/Plugins/Experimental/Mutable/Source/CustomizableObjectPopulation/Private/MuCOP/CustomizableObjectPopulationCustomVersion.cpp
--- a/Engine/Plugins/Experimental/Mutable/Source/CustomizableObjectPopulation/Private/MuCOP/CustomizableObjectPopulationCustomVersion.cpp
+++ b/Engine/Plugins/Experimental/Mutable/Source/CustomizableObjectPopulation/Private/MuCOP/CustomizableObjectPopulationCustomVersion.cpp
@@ -1,6 +1,9 @@
 // Copyright Epic Games, Inc. All Rights Reserved.
 
 #include "MuCOP/CustomizableObjectPopulationCustomVersion.h"
+
+#include "HAL/Platform.h"
+#include "Misc/Guid.h"
 #include "Serialization/CustomVersion.h"
 
 using namespace CustomizableObjectPopulation;
diff --git a/Engine/Plugins/Experimental/Mutable/Source/CustomizableObjectPopulation/Private/MuCOP/CustomizableObjectPopulationGenerator.cpp b/Engine/Plugins/Experimental/Mutable/Source/CustomizableObjectPopulation/Private/MuCOP/CustomizableObjectPopulationGenerator.cpp
--- a/Engine/Plugins/Experimental/Mutable/Source/CustomizableObjectPopulation/Private/MuCOP/CustomizableObjectPopulationGenerator.cpp
+++ b/Engine/Plugins/Experimental/Mutable/Source/CustomizableObjectPopulation/Private/MuCOP/CustomizableObjectPopulationGenerator.cpp
@@ -1,13 +1,19 @@
 // Copyright Epic Games, Inc. All Rights Reserved.
 
 #include "MuCOP/CustomizableObjectPopulationGenerator.h"
-#include "MuCOP/CustomizableObjectPopulationCustomVersion.h"
-#include "MuCOP/CustomizableObjectPopulationGeneratorPrivate.h"
+
+#include "HAL/PlatformTime.h"
+#include "Math/Color.h"
+#include "Math/RandomStream.h"
+#include "Misc/AssertionMacros.h"
 #include "MuCO/CustomizableObject.h"
 #include "MuCO/CustomizableObjectInstance.h"
-#include "Algo/Transform.h"
-#include "Logging/MessageLog.h"
-#include "MuCOP/CustomizableObjectPopulationGenerator.h"
+#include "MuCOP/CustomizableObjectPopulationCustomVersion.h"
+#include "MuCOP/CustomizableObjectPopulationGeneratorPrivate.h"
+#include "MuCOP/CustomizableObjectPopulationSamplers.h"
+#include "Serialization/Archive.h"
+
+class FString;
 
 void FCustomizableObjectPopulationClassGenerator::GenerateParameters(
 	FRandomStream& Rand,
diff --git a/Engine/Plugins/Experimental/Mutable/Source/CustomizableObjectPopulation/Private/MuCOP/CustomizableObjectPopulationModule.cpp b/Engine/Plugins/Experimental/Mutable/Source/CustomizableObjectPopulation/Private/MuCOP/CustomizableObjectPopulationModule.cpp
--- a/Engine/Plugins/Experimental/Mutable/Source/CustomizableObjectPopulation/Private/MuCOP/CustomizableObjectPopulationModule.cpp
+++ b/Engine/Plugins/Experimental/Mutable/Source/CustomizableObjectPopulation/Private/MuCOP/CustomizableObjectPopulationModule.cpp
@@ -1,16 +1,27 @@
 // Copyright Epic Games, Inc. All Rights Reserved.
 
-#include "Modules/ModuleManager.h"
+#include "AssetRegistry/ARFilter.h"
+#include "AssetRegistry/AssetData.h"
+#include "AssetRegistry/AssetRegistryModule.h"
+#include "AssetRegistry/IAssetRegistry.h"
+#include "Containers/Array.h"
+#include "Containers/UnrealString.h"
+#include "HAL/Platform.h"
+#include "HAL/PlatformCrt.h"
 #include "Interfaces/IPluginManager.h"
-#include "Misc/ConfigCacheIni.h"
+#include "Modules/ModuleManager.h"
 #include "MuCO/CustomizableObject.h"
-
-#include "Runtime/Launch/Resources/Version.h"
 #include "MuCOE/ICustomizableObjectPopulationModule.h"
-
 #include "MuCOP/CustomizableObjectPopulation.h"
-
-#include "AssetRegistry/AssetRegistryModule.h"
+#include "MuCOP/CustomizableObjectPopulationGenerator.h"
+#include "PluginDescriptor.h"
+#include "Templates/Casts.h"
+#include "Templates/SharedPointer.h"
+#include "UObject/Class.h"
+#include "UObject/NameTypes.h"
+#include "UObject/Object.h"
+#include "UObject/TopLevelAssetPath.h"
+#include "UObject/UObjectGlobals.h"
 
 
 /**
diff --git a/Engine/Plugins/Experimental/Mutable/Source/CustomizableObjectPopulation/Private/MuCOP/CustomizableObjectPopulationSamplers.cpp b/Engine/Plugins/Experimental/Mutable/Source/CustomizableObjectPopulation/Private/MuCOP/CustomizableObjectPopulationSamplers.cpp
--- a/Engine/Plugins/Experimental/Mutable/Source/CustomizableObjectPopulation/Private/MuCOP/CustomizableObjectPopulationSamplers.cpp
+++ b/Engine/Plugins/Experimental/Mutable/Source/CustomizableObjectPopulation/Private/MuCOP/CustomizableObjectPopulationSamplers.cpp
@@ -1,6 +1,12 @@
 // Copyright Epic Games, Inc. All Rights Reserved.
 
 #include "MuCOP/CustomizableObjectPopulationSamplers.h"
+
+#include "Curves/CurveLinearColor.h"
+#include "Math/RandomStream.h"
+#include "Math/UnrealMathUtility.h"
+#include "Misc/AssertionMacros.h"
+#include "Serialization/Archive.h"
 //#include "MuCOP/CustomizableObjectPopulation.h"
 //#include "MuCOP/CustomizableObjectPopulationClass.h"
 
diff --git a/Engine/Plugins/Experimental/Mutable/Source/CustomizableObjectPopulation/Public/MuCOP/CustomizableObjectPopulationCustomVersion.h b/Engine/Plugins/Experimental/Mutable/Source/CustomizableObjectPopulation/Public/MuCOP/CustomizableObjectPopulationCustomVersion.h
--- a/Engine/Plugins/Experimental/Mutable/Source/CustomizableObjectPopulation/Public/MuCOP/CustomizableObjectPopulationCustomVersion.h
+++ b/Engine/Plugins/Experimental/Mutable/Source/CustomizableObjectPopulation/Public/MuCOP/CustomizableObjectPopulationCustomVersion.h
@@ -2,8 +2,9 @@
 
 #pragma once
 
-#include "CoreMinimal.h"
-#include "Misc/Guid.h"
+#include "HAL/Platform.h"
+
+struct FGuid;
 
 namespace CustomizableObjectPopulation
 {
diff --git a/Engine/Plugins/Experimental/Mutable/Source/CustomizableObjectPopulation/Public/MuCOP/CustomizableObjectPopulationSamplers.h b/Engine/Plugins/Experimental/Mutable/Source/CustomizableObjectPopulation/Public/MuCOP/CustomizableObjectPopulationSamplers.h
--- a/Engine/Plugins/Experimental/Mutable/Source/CustomizableObjectPopulation/Public/MuCOP/CustomizableObjectPopulationSamplers.h
+++ b/Engine/Plugins/Experimental/Mutable/Source/CustomizableObjectPopulation/Public/MuCOP/CustomizableObjectPopulationSamplers.h
@@ -2,17 +2,16 @@
 
 #pragma once
 
-#include "CoreTypes.h"
-
 #include "Containers/Array.h"
 #include "Containers/UnrealString.h"
-
+#include "CoreTypes.h"
+#include "Curves/RichCurve.h"
+#include "Math/Color.h"
 #include "Templates/Tuple.h"
-#include "Misc/Variant.h"
 
-#include "Math/RandomStream.h"
-#include "Curves/RichCurve.h"
-#include "Curves/CurveLinearColor.h"
+class FArchive;
+class UCurveLinearColor;
+struct FRandomStream;
 
 namespace CustomizableObjectPopulation
 {
diff --git a/Engine/Plugins/Experimental/Mutable/Source/CustomizableObjectPopulationEditor/Classes/MuCOPE/CustomizableObjectPopulationClassFactory.h b/Engine/Plugins/Experimental/Mutable/Source/CustomizableObjectPopulationEditor/Classes/MuCOPE/CustomizableObjectPopulationClassFactory.h
--- a/Engine/Plugins/Experimental/Mutable/Source/CustomizableObjectPopulationEditor/Classes/MuCOPE/CustomizableObjectPopulationClassFactory.h
+++ b/Engine/Plugins/Experimental/Mutable/Source/CustomizableObjectPopulationEditor/Classes/MuCOPE/CustomizableObjectPopulationClassFactory.h
@@ -3,9 +3,15 @@
 #pragma once
 
 #include "Factories/Factory.h"
+#include "UObject/NameTypes.h"
+#include "UObject/UObjectGlobals.h"
 
 #include "CustomizableObjectPopulationClassFactory.generated.h"
 
+class FFeedbackContext;
+class UClass;
+class UObject;
+
 
 UCLASS(MinimalAPI)
 class UCustomizableObjectPopulationClassFactory : public UFactory
diff --git a/Engine/Plugins/Experimental/Mutable/Source/CustomizableObjectPopulationEditor/CustomizableObjectPopulationEditor.Build.cs b/Engine/Plugins/Experimental/Mutable/Source/CustomizableObjectPopulationEditor/CustomizableObjectPopulationEditor.Build.cs
--- a/Engine/Plugins/Experimental/Mutable/Source/CustomizableObjectPopulationEditor/CustomizableObjectPopulationEditor.Build.cs
+++ b/Engine/Plugins/Experimental/Mutable/Source/CustomizableObjectPopulationEditor/CustomizableObjectPopulationEditor.Build.cs
@@ -77,5 +77,9 @@ public class CustomizableObjectPopulationEditor : ModuleRules
 				"CustomizableObjectPopulation",
 			}
 		);
+
+		PrivateIncludePaths.AddRange(new string[] {
+				"CustomizableObjectPopulation/Private",
+			});
 	}
 }
diff --git a/Engine/Plugins/Experimental/Mutable/Source/CustomizableObjectPopulationEditor/Private/MuCOPE/AssetTypeActions_CustomizableObjectPopulation.cpp b/Engine/Plugins/Experimental/Mutable/Source/CustomizableObjectPopulationEditor/Private/MuCOPE/AssetTypeActions_CustomizableObjectPopulation.cpp
--- a/Engine/Plugins/Experimental/Mutable/Source/CustomizableObjectPopulationEditor/Private/MuCOPE/AssetTypeActions_CustomizableObjectPopulation.cpp
+++ b/Engine/Plugins/Experimental/Mutable/Source/CustomizableObjectPopulationEditor/Private/MuCOPE/AssetTypeActions_CustomizableObjectPopulation.cpp
@@ -2,16 +2,22 @@
 
 #include "MuCOPE/AssetTypeActions_CustomizableObjectPopulation.h"
 
-#include "ContentBrowserModule.h"
-#include "IContentBrowserSingleton.h"
-#include "MuCOPE/CustomizableObjectPopulationEditorModule.h"
-#include "MuCOPE/CustomizableObjectPopulationEditor.h"
-#include "MuCO/CustomizableObject.h"// For log function
+#include "Delegates/Delegate.h"
+#include "Framework/Commands/UIAction.h"
+#include "Framework/MultiBox/MultiBoxBuilder.h"
+#include "HAL/PlatformCrt.h"
+#include "Internationalization/Internationalization.h"
+#include "Modules/ModuleManager.h"
 #include "MuCOE/CustomizableObjectEditorModule.h"
+#include "MuCOP/CustomizableObjectPopulation.h"
+#include "MuCOPE/CustomizableObjectPopulationEditorModule.h"
+#include "Templates/Casts.h"
+#include "Textures/SlateIcon.h"
+#include "Toolkits/IToolkit.h"
+#include "UObject/Object.h"
 
-#include "MuCOP/CustomizableObjectPopulationGenerator.h"
-
-#include "Framework/MultiBox/MultiBoxBuilder.h"
+class IToolkitHost;
+class UClass;
 
 #define LOCTEXT_NAMESPACE "AssetTypeActions_CustomizableObjectPopulation"
 
diff --git a/Engine/Plugins/Experimental/Mutable/Source/CustomizableObjectPopulationEditor/Private/MuCOPE/AssetTypeActions_CustomizableObjectPopulation.h b/Engine/Plugins/Experimental/Mutable/Source/CustomizableObjectPopulationEditor/Private/MuCOPE/AssetTypeActions_CustomizableObjectPopulation.h
--- a/Engine/Plugins/Experimental/Mutable/Source/CustomizableObjectPopulationEditor/Private/MuCOPE/AssetTypeActions_CustomizableObjectPopulation.h
+++ b/Engine/Plugins/Experimental/Mutable/Source/CustomizableObjectPopulationEditor/Private/MuCOPE/AssetTypeActions_CustomizableObjectPopulation.h
@@ -3,7 +3,18 @@
 #pragma once
 
 #include "AssetTypeActions_Base.h"
-#include "MuCOP/CustomizableObjectPopulation.h"
+#include "Containers/Array.h"
+#include "HAL/Platform.h"
+#include "Internationalization/Text.h"
+#include "Math/Color.h"
+#include "Templates/SharedPointer.h"
+#include "UObject/WeakObjectPtrTemplates.h"
+
+class FMenuBuilder;
+class IToolkitHost;
+class UClass;
+class UCustomizableObjectPopulation;
+class UObject;
 
 
 class FAssetTypeActions_CustomizableObjectPopulation : public FAssetTypeActions_Base
diff --git a/Engine/Plugins/Experimental/Mutable/Source/CustomizableObjectPopulationEditor/Private/MuCOPE/AssetTypeActions_CustomizableObjectPopulationClass.cpp b/Engine/Plugins/Experimental/Mutable/Source/CustomizableObjectPopulationEditor/Private/MuCOPE/AssetTypeActions_CustomizableObjectPopulationClass.cpp
--- a/Engine/Plugins/Experimental/Mutable/Source/CustomizableObjectPopulationEditor/Private/MuCOPE/AssetTypeActions_CustomizableObjectPopulationClass.cpp
+++ b/Engine/Plugins/Experimental/Mutable/Source/CustomizableObjectPopulationEditor/Private/MuCOPE/AssetTypeActions_CustomizableObjectPopulationClass.cpp
@@ -3,9 +3,19 @@
 
 #include "MuCOPE/AssetTypeActions_CustomizableObjectPopulationClass.h"
 
+#include "HAL/PlatformMath.h"
+#include "Internationalization/Internationalization.h"
+#include "Modules/ModuleManager.h"
 #include "MuCOE/CustomizableObjectEditorModule.h"
-#include "MuCOPE/CustomizableObjectPopulationClassEditor.h"
+#include "MuCOP/CustomizableObjectPopulationCharacteristic.h"
+#include "MuCOP/CustomizableObjectPopulationClass.h"
 #include "MuCOPE/CustomizableObjectPopulationEditorModule.h"
+#include "Templates/Casts.h"
+#include "Toolkits/IToolkit.h"
+#include "UObject/Object.h"
+
+class IToolkitHost;
+class UClass;
 
 #define LOCTEXT_NAMESPACE "AssetTypeActions_CustomizableObjectPopulationClass"
 
diff --git a/Engine/Plugins/Experimental/Mutable/Source/CustomizableObjectPopulationEditor/Private/MuCOPE/AssetTypeActions_CustomizableObjectPopulationClass.h b/Engine/Plugins/Experimental/Mutable/Source/CustomizableObjectPopulationEditor/Private/MuCOPE/AssetTypeActions_CustomizableObjectPopulationClass.h
--- a/Engine/Plugins/Experimental/Mutable/Source/CustomizableObjectPopulationEditor/Private/MuCOPE/AssetTypeActions_CustomizableObjectPopulationClass.h
+++ b/Engine/Plugins/Experimental/Mutable/Source/CustomizableObjectPopulationEditor/Private/MuCOPE/AssetTypeActions_CustomizableObjectPopulationClass.h
@@ -3,7 +3,15 @@
 #pragma once
 
 #include "AssetTypeActions_Base.h"
-#include "MuCOP/CustomizableObjectPopulationClass.h"
+#include "Containers/Array.h"
+#include "HAL/Platform.h"
+#include "Internationalization/Text.h"
+#include "Math/Color.h"
+#include "Templates/SharedPointer.h"
+
+class IToolkitHost;
+class UClass;
+class UObject;
 
 
 class FAssetTypeActions_CustomizableObjectPopulationClass : public FAssetTypeActions_Base
diff --git a/Engine/Plugins/Experimental/Mutable/Source/CustomizableObjectPopulationEditor/Private/MuCOPE/CustomizableObjectPopulationClassDetails.cpp b/Engine/Plugins/Experimental/Mutable/Source/CustomizableObjectPopulationEditor/Private/MuCOPE/CustomizableObjectPopulationClassDetails.cpp
--- a/Engine/Plugins/Experimental/Mutable/Source/CustomizableObjectPopulationEditor/Private/MuCOPE/CustomizableObjectPopulationClassDetails.cpp
+++ b/Engine/Plugins/Experimental/Mutable/Source/CustomizableObjectPopulationEditor/Private/MuCOPE/CustomizableObjectPopulationClassDetails.cpp
@@ -1,23 +1,59 @@
 // Copyright Epic Games, Inc. All Rights Reserved.
 
 #include "MuCOPE/CustomizableObjectPopulationClassDetails.h"
-#include "MuCOP/CustomizableObjectPopulationClass.h"
-#include "MuCO/CustomizableObject.h"
-#include "MuCO/CustomizableObjectInstance.h"
 
-#include "DetailLayoutBuilder.h"
+#include "AssetRegistry/AssetData.h"
+#include "Containers/EnumAsByte.h"
+#include "Curves/CurveBase.h"
+#include "Curves/CurveLinearColor.h"
+#include "CustomizableObjectPopulation/Private/MuCOP/CustomizableObjectPopulationCharacteristic.h"
 #include "DetailCategoryBuilder.h"
+#include "DetailLayoutBuilder.h"
+#include "DetailWidgetRow.h"
+#include "Engine/Texture.h"
+#include "Engine/Texture2D.h"
+#include "Framework/SlateDelegates.h"
+#include "Framework/Views/TableViewTypeTraits.h"
+#include "HAL/PlatformCrt.h"
+#include "HAL/PlatformMath.h"
+#include "IDetailsView.h"
+#include "Input/Events.h"
+#include "InputCoreTypes.h"
+#include "Internationalization/Internationalization.h"
+#include "Layout/Children.h"
+#include "Layout/Geometry.h"
+#include "Layout/Visibility.h"
+#include "Misc/Attribute.h"
+#include "MuCO/CustomizableObject.h"
+#include "MuCO/CustomizableObjectInstance.h"
+#include "MuCO/CustomizableObjectParameterTypeDefinitions.h"
+#include "MuCOP/CustomizableObjectPopulationClass.h"
 #include "PropertyCustomizationHelpers.h"
-
+#include "PropertyHandle.h"
+#include "Rendering/DrawElements.h"
+#include "Rendering/RenderingCommon.h"
+#include "Serialization/Archive.h"
+#include "SlotBase.h"
+#include "Styling/AppStyle.h"
+#include "Styling/SlateColor.h"
+#include "Templates/Casts.h"
+#include "Templates/UnrealTemplate.h"
+#include "UObject/ObjectPtr.h"
+#include "UObject/WeakObjectPtr.h"
+#include "UObject/WeakObjectPtrTemplates.h"
+#include "Widgets/Colors/SColorBlock.h"
+#include "Widgets/Colors/SColorPicker.h"
+#include "Widgets/Images/SImage.h"
+#include "Widgets/Input/SButton.h"
+#include "Widgets/Input/SSpinBox.h"
 #include "Widgets/Input/STextComboBox.h"
 #include "Widgets/SBoxPanel.h"
 #include "Widgets/Text/STextBlock.h"
-#include "Widgets/Input/SSpinBox.h"
-#include "Widgets/Input/SButton.h"
-#include "Slate/DeferredCleanupSlateBrush.h"
-#include "Curves/CurveLinearColor.h"
-#include "Widgets/Colors/SColorPicker.h"
-#include "Widgets/Colors/SColorBlock.h"
+
+class FPaintArgs;
+class FSlateRect;
+class FWidgetStyle;
+class SWidget;
 
 
 #define LOCTEXT_NAMESPACE "CustomizableObjectPopulationClassDetails"
diff --git a/Engine/Plugins/Experimental/Mutable/Source/CustomizableObjectPopulationEditor/Private/MuCOPE/CustomizableObjectPopulationClassDetails.h b/Engine/Plugins/Experimental/Mutable/Source/CustomizableObjectPopulationEditor/Private/MuCOPE/CustomizableObjectPopulationClassDetails.h
--- a/Engine/Plugins/Experimental/Mutable/Source/CustomizableObjectPopulationEditor/Private/MuCOPE/CustomizableObjectPopulationClassDetails.h
+++ b/Engine/Plugins/Experimental/Mutable/Source/CustomizableObjectPopulationEditor/Private/MuCOPE/CustomizableObjectPopulationClassDetails.h
@@ -2,18 +2,41 @@
 
 #pragma once
 
+#include "Containers/Array.h"
+#include "Containers/BitArray.h"
+#include "Containers/Map.h"
+#include "Containers/Set.h"
+#include "Containers/SparseArray.h"
+#include "Containers/UnrealString.h"
+#include "CustomizableObjectPopulation/Private/MuCOP/CustomizableObjectPopulationConstraint.h"
+#include "Delegates/Delegate.h"
+#include "HAL/Platform.h"
 #include "IDetailCustomization.h"
-
+#include "Input/Reply.h"
+#include "Internationalization/Text.h"
+#include "Math/Color.h"
+#include "Math/Vector2D.h"
+#include "Misc/Optional.h"
+#include "Styling/SlateBrush.h"
 #include "Templates/SharedPointer.h"
-#include "UObject/WeakObjectPtr.h"
-#include "Widgets/SCompoundWidget.h"
-#include "Engine/Texture.h"
+#include "Types/SlateEnums.h"
+#include "Widgets/DeclarativeSyntaxSupport.h"
 #include "Widgets/Input/SComboBox.h"
-#include "Widgets/Input/STextComboBox.h"
-#include "Components/VerticalBox.h"
-#include "MuCOP/CustomizableObjectPopulationClass.h"
+#include "Widgets/SCompoundWidget.h"
 
-class FStateDetailsNode;
+class FPaintArgs;
+class FSlateRect;
+class FSlateWindowElementList;
+class FWidgetStyle;
+class IDetailLayoutBuilder;
+class STextComboBox;
+class SVerticalBox;
+class SWidget;
+class UCustomizableObjectPopulationClass;
+class UTexture;
+struct FAssetData;
+struct FGeometry;
+struct FPointerEvent;
 
 // Widgt that draws a Square to represent a range
 class SRangeSquare : public SCompoundWidget
diff --git a/Engine/Plugins/Experimental/Mutable/Source/CustomizableObjectPopulationEditor/Private/MuCOPE/CustomizableObjectPopulationClassEditor.cpp b/Engine/Plugins/Experimental/Mutable/Source/CustomizableObjectPopulationEditor/Private/MuCOPE/CustomizableObjectPopulationClassEditor.cpp
--- a/Engine/Plugins/Experimental/Mutable/Source/CustomizableObjectPopulationEditor/Private/MuCOPE/CustomizableObjectPopulationClassEditor.cpp
+++ b/Engine/Plugins/Experimental/Mutable/Source/CustomizableObjectPopulationEditor/Private/MuCOPE/CustomizableObjectPopulationClassEditor.cpp
@@ -1,46 +1,91 @@
 // Copyright Epic Games, Inc. All Rights Reserved.
 
 #include "MuCOPE/CustomizableObjectPopulationClassEditor.h"
-#include "MuCOPE/CustomizableObjectPopulationClassEditorActions.h"
-#include "MuCOE/CustomizableObjectEditorModule.h"
-#include "MuCOPE/CustomizableObjectPopulationEditorModule.h"
-#include "MuCO/CustomizableObjectInstance.h"
-
-#include "WorkspaceMenuStructureModule.h"
-#include "WorkspaceMenuStructure.h"
 
+#include "AssetRegistry/ARFilter.h"
+#include "AssetRegistry/AssetData.h"
+#include "AssetRegistry/AssetRegistryModule.h"
+#include "AssetRegistry/IAssetRegistry.h"
+#include "AssetToolsModule.h"
+#include "Components/CapsuleComponent.h"
+#include "Components/SkeletalMeshComponent.h"
+#include "ContentBrowserDelegates.h"
+#include "ContentBrowserModule.h"
+#include "Curves/CurveBase.h"
+#include "DetailsViewArgs.h"
+#include "Editor.h"
+#include "Editor/EditorEngine.h"
+#include "Engine/EngineTypes.h"
+#include "Engine/SkeletalMesh.h"
 #include "FileHelpers.h"
-#include "MuCOE/UnrealEditorPortabilityHelpers.h"
-#include "Toolkits/IToolkitHost.h"
-
-#include "MuCOP/CustomizableObjectPopulationClass.h"
-#include "MuCOPE/SCustomizableObjectPopulationClassTagsTool.h"
-#include "MuCOPE/CustomizableObjectPopulationClassDetails.h"
-
+#include "Fonts/SlateFontInfo.h"
+#include "Framework/Commands/UIAction.h"
+#include "Framework/Commands/UICommandList.h"
+#include "Framework/Docking/TabManager.h"
+#include "Framework/MultiBox/MultiBoxBuilder.h"
+#include "Framework/MultiBox/MultiBoxExtender.h"
+#include "Framework/SlateDelegates.h"
+#include "HAL/PlatformCrt.h"
+#include "HAL/PlatformMisc.h"
+#include "IAssetTools.h"
+#include "IContentBrowserSingleton.h"
+#include "IDetailsView.h"
+#include "Input/Reply.h"
+#include "Internationalization/Internationalization.h"
+#include "Layout/BasicLayoutWidgetSlot.h"
+#include "Layout/Margin.h"
+#include "Math/Vector2D.h"
+#include "Misc/AssertionMacros.h"
+#include "Misc/Attribute.h"
+#include "Misc/MessageDialog.h"
+#include "Misc/PackageName.h"
+#include "Misc/Paths.h"
+#include "Modules/ModuleManager.h"
+#include "MuCO/CustomizableObject.h"
+#include "MuCO/CustomizableObjectInstance.h"
+#include "MuCO/CustomizableSkeletalComponent.h"
 #include "MuCOP/CustomizableObjectPopulation.h"
+#include "MuCOP/CustomizableObjectPopulationCharacteristic.h"
+#include "MuCOP/CustomizableObjectPopulationClass.h"
 #include "MuCOP/CustomizableObjectPopulationGenerator.h"
-#include "MuCOPE/SCustomizableObjectPopulationEditorViewport.h"
+#include "MuCOPE/CustomizableObjectPopulationClassEditorActions.h"
 #include "MuCOPE/CustomizableObjectPopulationEditorModule.h"
-
-#include "AssetToolsModule.h"
-#include "AssetRegistry/AssetRegistryModule.h"
-#include "ContentBrowserModule.h"
-#include "IContentBrowserSingleton.h"
-#include "Widgets/Input/SNumericEntryBox.h"
+#include "MuCOPE/SCustomizableObjectPopulationClassTagsTool.h"
+#include "MuCOPE/SCustomizableObjectPopulationEditorViewport.h"
+#include "PropertyEditorModule.h"
+#include "SCurveEditor.h"
+#include "SlotBase.h"
+#include "Styling/AppStyle.h"
+#include "Subsystems/AssetEditorSubsystem.h"
+#include "Templates/Casts.h"
+#include "Textures/SlateIcon.h"
+#include "Toolkits/AssetEditorToolkit.h"
+#include "Types/SlateEnums.h"
+#include "Types/SlateStructs.h"
+#include "UObject/Class.h"
+#include "UObject/Object.h"
+#include "UObject/ObjectPtr.h"
+#include "UObject/Package.h"
+#include "UObject/UObjectGlobals.h"
+#include "UObject/UnrealNames.h"
+#include "Widgets/DeclarativeSyntaxSupport.h"
+#include "Widgets/Docking/SDockTab.h"
 #include "Widgets/Input/SButton.h"
+#include "Widgets/Input/SEditableTextBox.h"
+#include "Widgets/Input/SNumericEntryBox.h"
+#include "Widgets/Layout/SBorder.h"
+#include "Widgets/Layout/SBox.h"
+#include "Widgets/Layout/SScrollBar.h"
 #include "Widgets/Layout/SScrollBox.h"
-#include "SCurveEditor.h"
-#include "Framework/MultiBox/MultiBoxBuilder.h"
 #include "Widgets/Layout/SUniformGridPanel.h"
-#include "Widgets/Docking/SDockTab.h"
-#include "Misc/MessageDialog.h"
-
-#include "Animation/DebugSkelMeshComponent.h"
-
-#include "Components/BoxComponent.h"
-#include "Components/CapsuleComponent.h"
+#include "Widgets/SBoxPanel.h"
+#include "Widgets/SWindow.h"
+#include "Widgets/Text/STextBlock.h"
+#include "WorkspaceMenuStructure.h"
+#include "WorkspaceMenuStructureModule.h"
 
-#include "Framework/MultiBox/MultiBoxExtender.h" 
+class IToolkitHost;
+class SWidget;
 
 #define LOCTEXT_NAMESPACE "CustomizableObjectPopulationClassEditor"
 
diff --git a/Engine/Plugins/Experimental/Mutable/Source/CustomizableObjectPopulationEditor/Private/MuCOPE/CustomizableObjectPopulationClassEditor.h b/Engine/Plugins/Experimental/Mutable/Source/CustomizableObjectPopulationEditor/Private/MuCOPE/CustomizableObjectPopulationClassEditor.h
--- a/Engine/Plugins/Experimental/Mutable/Source/CustomizableObjectPopulationEditor/Private/MuCOPE/CustomizableObjectPopulationClassEditor.h
+++ b/Engine/Plugins/Experimental/Mutable/Source/CustomizableObjectPopulationEditor/Private/MuCOPE/CustomizableObjectPopulationClassEditor.h
@@ -2,15 +2,33 @@
 
 #pragma once 
 
-#include "MuCOPE/ICustomizableObjectPopulationClassEditor.h"
-#include "MuCO/CustomizableObjectInstance.h"
+#include "Containers/Array.h"
+#include "Containers/UnrealString.h"
+#include "Delegates/Delegate.h"
+#include "EditorUndoClient.h"
+#include "HAL/Platform.h"
+#include "Internationalization/Text.h"
+#include "Math/Color.h"
 #include "Misc/NotifyHook.h"
+#include "Misc/Optional.h"
+#include "MuCOPE/ICustomizableObjectPopulationClassEditor.h"
+#include "Stats/Stats2.h"
+#include "Templates/SharedPointer.h"
 #include "TickableEditorObject.h"
-#include "EditorUndoClient.h"
-#include "Toolkits/AssetEditorToolkit.h"
-#include "Curves/CurveBase.h"
-
+#include "Toolkits/IToolkit.h"
+#include "UObject/GCObject.h"
+#include "UObject/NameTypes.h"
+
+class FReferenceCollector;
+class FSpawnTabArgs;
+class FUICommandList;
+class IToolkitHost;
+class SDockTab;
+class SWidget;
+class UCurveBase;
+class UCustomizableObjectInstance;
 class UCustomizableObjectPopulationClass;
+template <typename NumericType> class SNumericEntryBox;
 
 class FCustomizableObjectPopulationClassEditor : 
 	public ICustomizableObjectPopulationClassEditor,
diff --git a/Engine/Plugins/Experimental/Mutable/Source/CustomizableObjectPopulationEditor/Private/MuCOPE/CustomizableObjectPopulationClassEditorActions.cpp b/Engine/Plugins/Experimental/Mutable/Source/CustomizableObjectPopulationEditor/Private/MuCOPE/CustomizableObjectPopulationClassEditorActions.cpp
--- a/Engine/Plugins/Experimental/Mutable/Source/CustomizableObjectPopulationEditor/Private/MuCOPE/CustomizableObjectPopulationClassEditorActions.cpp
+++ b/Engine/Plugins/Experimental/Mutable/Source/CustomizableObjectPopulationEditor/Private/MuCOPE/CustomizableObjectPopulationClassEditorActions.cpp
@@ -2,9 +2,12 @@
 
 #include "MuCOPE/CustomizableObjectPopulationClassEditorActions.h"
 
-#include "MuCOPE/CustomizableObjectPopulationEditorStyle.h"
-#include "MuCOP/CustomizableObjectPopulationClass.h"
 #include "Framework/Commands/InputChord.h"
+#include "Framework/Commands/UICommandInfo.h"
+#include "Internationalization/Internationalization.h"
+#include "MuCOPE/CustomizableObjectPopulationEditorStyle.h"
+#include "UObject/NameTypes.h"
+#include "UObject/UnrealNames.h"
 
 #define LOCTEXT_NAMESPACE "CustomizableObjectEditor"
 
diff --git a/Engine/Plugins/Experimental/Mutable/Source/CustomizableObjectPopulationEditor/Private/MuCOPE/CustomizableObjectPopulationClassFactory.cpp b/Engine/Plugins/Experimental/Mutable/Source/CustomizableObjectPopulationEditor/Private/MuCOPE/CustomizableObjectPopulationClassFactory.cpp
--- a/Engine/Plugins/Experimental/Mutable/Source/CustomizableObjectPopulationEditor/Private/MuCOPE/CustomizableObjectPopulationClassFactory.cpp
+++ b/Engine/Plugins/Experimental/Mutable/Source/CustomizableObjectPopulationEditor/Private/MuCOPE/CustomizableObjectPopulationClassFactory.cpp
@@ -1,7 +1,14 @@
 // Copyright Epic Games, Inc. All Rights Reserved.
 
 #include "MuCOPE/CustomizableObjectPopulationClassFactory.h"
+
+#include "MuCOP/CustomizableObjectPopulationCharacteristic.h"
 #include "MuCOP/CustomizableObjectPopulationClass.h"
+#include "Templates/SubclassOf.h"
+
+class FFeedbackContext;
+class UClass;
+class UObject;
 
 
 #define LOCTEXT_NAMESPACE "CustomizableObjectPopulationClassFactory"
diff --git a/Engine/Plugins/Experimental/Mutable/Source/CustomizableObjectPopulationEditor/Private/MuCOPE/CustomizableObjectPopulationEditor.cpp b/Engine/Plugins/Experimental/Mutable/Source/CustomizableObjectPopulationEditor/Private/MuCOPE/CustomizableObjectPopulationEditor.cpp
--- a/Engine/Plugins/Experimental/Mutable/Source/CustomizableObjectPopulationEditor/Private/MuCOPE/CustomizableObjectPopulationEditor.cpp
+++ b/Engine/Plugins/Experimental/Mutable/Source/CustomizableObjectPopulationEditor/Private/MuCOPE/CustomizableObjectPopulationEditor.cpp
@@ -1,43 +1,83 @@
 // Copyright Epic Games, Inc. All Rights Reserved.
 
 #include "MuCOPE/CustomizableObjectPopulationEditor.h"
-#include "MuCOE/CustomizableObjectEditorModule.h"
-#include "MuCOPE/CustomizableObjectPopulationEditorModule.h"
-#include "MuCOPE/CustomizableObjectPopulationEditorActions.h"
-#include "CustomizableObjectClasses.h"
-
-#include "WorkspaceMenuStructureModule.h"
-#include "WorkspaceMenuStructure.h"
 
+#include "AssetRegistry/AssetData.h"
+#include "AssetRegistry/AssetRegistryModule.h"
+#include "AssetRegistry/IAssetRegistry.h"
+#include "AssetToolsModule.h"
+#include "Components/CapsuleComponent.h"
+#include "Components/SkeletalMeshComponent.h"
+#include "ContentBrowserDelegates.h"
+#include "ContentBrowserModule.h"
+#include "DetailsViewArgs.h"
+#include "Editor.h"
+#include "Editor/EditorEngine.h"
+#include "Engine/EngineTypes.h"
+#include "Engine/SkeletalMesh.h"
 #include "FileHelpers.h"
-#include "MuCOE/UnrealEditorPortabilityHelpers.h"
-#include "Toolkits/IToolkitHost.h"
-
+#include "Fonts/SlateFontInfo.h"
+#include "Framework/Commands/UIAction.h"
+#include "Framework/Commands/UICommandList.h"
+#include "Framework/Docking/TabManager.h"
+#include "Framework/MultiBox/MultiBoxBuilder.h"
+#include "Framework/MultiBox/MultiBoxExtender.h"
+#include "Framework/SlateDelegates.h"
+#include "HAL/PlatformCrt.h"
+#include "HAL/PlatformMisc.h"
+#include "IAssetTools.h"
+#include "IContentBrowserSingleton.h"
+#include "IDetailsView.h"
+#include "Input/Reply.h"
+#include "Internationalization/Internationalization.h"
+#include "Layout/Margin.h"
+#include "Logging/LogCategory.h"
+#include "Logging/LogMacros.h"
+#include "Math/Vector2D.h"
+#include "Misc/AssertionMacros.h"
+#include "Misc/Attribute.h"
+#include "Misc/MessageDialog.h"
+#include "Misc/PackageName.h"
+#include "Misc/Paths.h"
+#include "Modules/ModuleManager.h"
+#include "MuCO/CustomizableObject.h"
+#include "MuCO/CustomizableObjectInstance.h"
+#include "MuCO/CustomizableSkeletalComponent.h"
 #include "MuCOP/CustomizableObjectPopulation.h"
 #include "MuCOP/CustomizableObjectPopulationGenerator.h"
-#include "MuCOPE/SCustomizableObjectPopulationEditorViewport.h"
+#include "MuCOPE/CustomizableObjectPopulationEditorActions.h"
 #include "MuCOPE/CustomizableObjectPopulationEditorModule.h"
-
-#include "AssetTypeActions_Base.h"
-#include "AssetRegistry/AssetRegistryModule.h"
-#include "ContentBrowserModule.h"
-#include "IContentBrowserSingleton.h"
+#include "MuCOPE/SCustomizableObjectPopulationEditorViewport.h"
+#include "PropertyEditorModule.h"
+#include "SlotBase.h"
+#include "Styling/AppStyle.h"
+#include "Subsystems/AssetEditorSubsystem.h"
+#include "Templates/Casts.h"
+#include "Textures/SlateIcon.h"
+#include "Toolkits/AssetEditorToolkit.h"
+#include "Trace/Detail/Channel.h"
+#include "Types/SlateEnums.h"
+#include "UObject/Class.h"
+#include "UObject/Object.h"
+#include "UObject/ObjectPtr.h"
+#include "UObject/Package.h"
+#include "UObject/UObjectGlobals.h"
+#include "UObject/UnrealNames.h"
+#include "Widgets/DeclarativeSyntaxSupport.h"
+#include "Widgets/Docking/SDockTab.h"
 #include "Widgets/Input/SButton.h"
-#include "Widgets/Input/STextComboBox.h"
-#include "Framework/MultiBox/MultiBoxBuilder.h"
+#include "Widgets/Input/SEditableTextBox.h"
 #include "Widgets/Input/SNumericEntryBox.h"
-#include "Widgets/Docking/SDockTab.h"
+#include "Widgets/Layout/SBorder.h"
 #include "Widgets/Layout/SUniformGridPanel.h"
+#include "Widgets/SBoxPanel.h"
+#include "Widgets/SWindow.h"
 #include "Widgets/Text/STextBlock.h"
-#include "Widgets/Input/SEditableTextBox.h"
-#include "Misc/MessageDialog.h"
-
-#include "Animation/DebugSkelMeshComponent.h"
-
-#include "Components/BoxComponent.h"
-#include "Components/CapsuleComponent.h"
+#include "WorkspaceMenuStructure.h"
+#include "WorkspaceMenuStructureModule.h"
 
-#include "Framework/MultiBox/MultiBoxExtender.h" 
+class IToolkitHost;
+class SWidget;
 
 
 #define LOCTEXT_NAMESPACE "CustomizableObjectPopulationEditor"
diff --git a/Engine/Plugins/Experimental/Mutable/Source/CustomizableObjectPopulationEditor/Private/MuCOPE/CustomizableObjectPopulationEditor.h b/Engine/Plugins/Experimental/Mutable/Source/CustomizableObjectPopulationEditor/Private/MuCOPE/CustomizableObjectPopulationEditor.h
--- a/Engine/Plugins/Experimental/Mutable/Source/CustomizableObjectPopulationEditor/Private/MuCOPE/CustomizableObjectPopulationEditor.h
+++ b/Engine/Plugins/Experimental/Mutable/Source/CustomizableObjectPopulationEditor/Private/MuCOPE/CustomizableObjectPopulationEditor.h
@@ -2,14 +2,32 @@
 
 #pragma once 
 
-#include "MuCOPE/ICustomizableObjectPopulationEditor.h"
+#include "Containers/Array.h"
+#include "Containers/UnrealString.h"
+#include "Delegates/Delegate.h"
+#include "EditorUndoClient.h"
+#include "HAL/Platform.h"
+#include "Internationalization/Text.h"
+#include "Math/Color.h"
 #include "Misc/NotifyHook.h"
+#include "Misc/Optional.h"
+#include "MuCOPE/ICustomizableObjectPopulationEditor.h"
+#include "Stats/Stats2.h"
+#include "Templates/SharedPointer.h"
 #include "TickableEditorObject.h"
-#include "EditorUndoClient.h"
-#include "Toolkits/AssetEditorToolkit.h"
-#include "MuCO/CustomizableObjectInstance.h"
-
+#include "Toolkits/IToolkit.h"
+#include "UObject/GCObject.h"
+#include "UObject/NameTypes.h"
+
+class FReferenceCollector;
+class FSpawnTabArgs;
+class FUICommandList;
+class IToolkitHost;
+class SDockTab;
+class SWidget;
+class UCustomizableObjectInstance;
 class UCustomizableObjectPopulation;
+template <typename NumericType> class SNumericEntryBox;
 
 class FCustomizableObjectPopulationEditor :
 	public ICustomizableObjectPopulationEditor,
diff --git a/Engine/Plugins/Experimental/Mutable/Source/CustomizableObjectPopulationEditor/Private/MuCOPE/CustomizableObjectPopulationEditorActions.cpp b/Engine/Plugins/Experimental/Mutable/Source/CustomizableObjectPopulationEditor/Private/MuCOPE/CustomizableObjectPopulationEditorActions.cpp
--- a/Engine/Plugins/Experimental/Mutable/Source/CustomizableObjectPopulationEditor/Private/MuCOPE/CustomizableObjectPopulationEditorActions.cpp
+++ b/Engine/Plugins/Experimental/Mutable/Source/CustomizableObjectPopulationEditor/Private/MuCOPE/CustomizableObjectPopulationEditorActions.cpp
@@ -1,9 +1,13 @@
 // Copyright Epic Games, Inc. All Rights Reserved.
 
 #include "MuCOPE/CustomizableObjectPopulationEditorActions.h"
-#include "MuCOPE/CustomizableObjectPopulationEditorStyle.h"
-#include "MuCOP/CustomizableObjectPopulation.h"
+
 #include "Framework/Commands/InputChord.h"
+#include "Framework/Commands/UICommandInfo.h"
+#include "Internationalization/Internationalization.h"
+#include "MuCOPE/CustomizableObjectPopulationEditorStyle.h"
+#include "UObject/NameTypes.h"
+#include "UObject/UnrealNames.h"
 
 #define LOCTEXT_NAMESPACE "CustomizableObjectPopulationEditor"
 
diff --git a/Engine/Plugins/Experimental/Mutable/Source/CustomizableObjectPopulationEditor/Private/MuCOPE/CustomizableObjectPopulationEditorModule.cpp b/Engine/Plugins/Experimental/Mutable/Source/CustomizableObjectPopulationEditor/Private/MuCOPE/CustomizableObjectPopulationEditorModule.cpp
--- a/Engine/Plugins/Experimental/Mutable/Source/CustomizableObjectPopulationEditor/Private/MuCOPE/CustomizableObjectPopulationEditorModule.cpp
+++ b/Engine/Plugins/Experimental/Mutable/Source/CustomizableObjectPopulationEditor/Private/MuCOPE/CustomizableObjectPopulationEditorModule.cpp
@@ -1,35 +1,33 @@
 // Copyright Epic Games, Inc. All Rights Reserved.
 
 #include "MuCOPE/CustomizableObjectPopulationEditorModule.h"
-#include "MuCOPE/CustomizableObjectPopulationEditor.h"
-#include "MuCO/CustomizableObject.h"	// For the LogMutable log category
-#include "MuCO/CustomizableObjectSystem.h"		// For defines related to memory function replacements.
-#include "MuCOE/CustomizableObjectIdentifierCustomization.h"
-#include "Modules/ModuleManager.h"
-#include "Interfaces/IPluginManager.h"
+
 #include "AssetToolsModule.h"
-#include "Styling/SlateStyle.h"
-#include "Brushes/SlateImageBrush.h"
-#include "MessageLogModule.h"
-#include "Logging/MessageLog.h"
-#include "MuCOE/CustomizableObjectEditorSettings.h"
-#include "ISettingsModule.h"
-#include "ISettingsSection.h"
-#include "AssetRegistry/AssetRegistryModule.h"
-#include "HAL/FileManager.h"
-#include "CoreGlobals.h"
-#include "Settings/ProjectPackagingSettings.h"
-#include "MuCOE/UnrealEditorPortabilityHelpers.h"
+#include "Containers/StringConv.h"
+#include "Delegates/Delegate.h"
+#include "HAL/PlatformMath.h"
+#include "HAL/UnrealMemory.h"
+#include "IAssetTools.h"
+#include "Logging/LogCategory.h"
+#include "Logging/LogMacros.h"
+#include "Modules/ModuleManager.h"
+#include "MuCO/CustomizableObject.h"	// For the LogMutable log category
 #include "MuCOPE/AssetTypeActions_CustomizableObjectPopulation.h"
 #include "MuCOPE/AssetTypeActions_CustomizableObjectPopulationClass.h"
+#include "MuCOPE/CustomizableObjectPopulationClassDetails.h"
 #include "MuCOPE/CustomizableObjectPopulationClassEditor.h"
 #include "MuCOPE/CustomizableObjectPopulationEditor.h"
 #include "MuCOPE/CustomizableObjectPopulationEditorStyle.h"
-
-#include "MuCOPE/CustomizableObjectPopulationClassDetails.h"
-
-// This is necessary to register the cook delegate.
-#include "CookOnTheSide/CookOnTheFlyServer.h"
+#include "PropertyEditorDelegates.h"
+#include "PropertyEditorModule.h"
+#include "Toolkits/AssetEditorToolkit.h"
+#include "Trace/Detail/Channel.h"
+
+class ICustomizableObjectPopulationClassEditor;
+class ICustomizableObjectPopulationEditor;
+class IToolkitHost;
+class UCustomizableObjectPopulation;
+class UCustomizableObjectPopulationClass;
 
 
 const FName CustomizableObjectPopulationEditorAppIdentifier = FName(TEXT("CustomizableObjectPopulationEditorApp"));
diff --git a/Engine/Plugins/Experimental/Mutable/Source/CustomizableObjectPopulationEditor/Private/MuCOPE/CustomizableObjectPopulationEditorStyle.cpp b/Engine/Plugins/Experimental/Mutable/Source/CustomizableObjectPopulationEditor/Private/MuCOPE/CustomizableObjectPopulationEditorStyle.cpp
--- a/Engine/Plugins/Experimental/Mutable/Source/CustomizableObjectPopulationEditor/Private/MuCOPE/CustomizableObjectPopulationEditorStyle.cpp
+++ b/Engine/Plugins/Experimental/Mutable/Source/CustomizableObjectPopulationEditor/Private/MuCOPE/CustomizableObjectPopulationEditorStyle.cpp
@@ -1,16 +1,16 @@
 // Copyright Epic Games, Inc. All Rights Reserved.
 #include "MuCOPE/CustomizableObjectPopulationEditorStyle.h"
 
-#include "Framework/Application/SlateApplication.h"
-#include "Styling/AppStyle.h"
-#include "Slate/SlateGameResources.h"
-#include "Styling/SlateStyleRegistry.h"
-#include "Styling/SlateTypes.h"
-#include "Styling/CoreStyle.h"
+#include "Brushes/SlateImageBrush.h"
+#include "HAL/PlatformMath.h"
 #include "Interfaces/IPluginManager.h"
-#include "Settings/EditorStyleSettings.h"
-#include "Styling/StarshipCoreStyle.h"
-#include "Styling/StyleColors.h"
+#include "Math/Vector2D.h"
+#include "Misc/AssertionMacros.h"
+#include "Misc/Paths.h"
+#include "Styling/SlateStyle.h"
+#include "Styling/SlateStyleRegistry.h"
+
+class ISlateStyle;
 
 TSharedPtr<FSlateStyleSet> FCustomizableObjectPopulationEditorStyle::CustomizableObjectPopulationEditorStyleInstance = NULL;
 
diff --git a/Engine/Plugins/Experimental/Mutable/Source/CustomizableObjectPopulationEditor/Private/MuCOPE/CustomizableObjectPopulationEditorStyle.h b/Engine/Plugins/Experimental/Mutable/Source/CustomizableObjectPopulationEditor/Private/MuCOPE/CustomizableObjectPopulationEditorStyle.h
--- a/Engine/Plugins/Experimental/Mutable/Source/CustomizableObjectPopulationEditor/Private/MuCOPE/CustomizableObjectPopulationEditorStyle.h
+++ b/Engine/Plugins/Experimental/Mutable/Source/CustomizableObjectPopulationEditor/Private/MuCOPE/CustomizableObjectPopulationEditorStyle.h
@@ -1,7 +1,13 @@
 // Copyright Epic Games, Inc. All Rights Reserved.
 
 #pragma once
-#include "Styling/ISlateStyle.h"
+
+#include "Containers/UnrealString.h"
+#include "HAL/Platform.h"
+#include "Templates/SharedPointer.h"
+#include "UObject/NameTypes.h"
+
+class ISlateStyle;
 
 /** Manages the style which provides resources for niagara editor widgets. */
 class FCustomizableObjectPopulationEditorStyle
diff --git a/Engine/Plugins/Experimental/Mutable/Source/CustomizableObjectPopulationEditor/Private/MuCOPE/CustomizableObjectPopulationEditorViewportClient.cpp b/Engine/Plugins/Experimental/Mutable/Source/CustomizableObjectPopulationEditor/Private/MuCOPE/CustomizableObjectPopulationEditorViewportClient.cpp
--- a/Engine/Plugins/Experimental/Mutable/Source/CustomizableObjectPopulationEditor/Private/MuCOPE/CustomizableObjectPopulationEditorViewportClient.cpp
+++ b/Engine/Plugins/Experimental/Mutable/Source/CustomizableObjectPopulationEditor/Private/MuCOPE/CustomizableObjectPopulationEditorViewportClient.cpp
@@ -1,10 +1,30 @@
 // Copyright Epic Games, Inc. All Rights Reserved.
 
 #include "MuCOPE/CustomizableObjectPopulationEditorViewportClient.h"
-#include "Components/BoxComponent.h"
+
+#include "AdvancedPreviewScene.h"
+#include "CollisionQueryParams.h"
 #include "Components/CapsuleComponent.h"
+#include "Components/SkeletalMeshComponent.h"
+#include "CoreGlobals.h"
+#include "Editor.h"
+#include "Engine/HitResult.h"
+#include "Engine/NetSerialization.h"
+#include "Engine/World.h"
+#include "EngineDefines.h"
+#include "Math/BoxSphereBounds.h"
+#include "Math/Color.h"
+#include "Math/UnrealMathSSE.h"
+#include "Math/Vector.h"
+#include "PreviewScene.h"
+#include "UObject/NameTypes.h"
+#include "UObject/ObjectPtr.h"
+#include "UObject/UObjectGlobals.h"
 #include "UnrealWidget.h"
 
+class FSceneView;
+class HHitProxy;
+
 #define LOCTEXT_NAMESPACE "CustomizableObjectPopulationEditorViewportClient"
 
 FCustomizableObjectPopulationEditorViewportClient::FCustomizableObjectPopulationEditorViewportClient(const TSharedRef<FAdvancedPreviewScene>& InPreviewScene)
diff --git a/Engine/Plugins/Experimental/Mutable/Source/CustomizableObjectPopulationEditor/Private/MuCOPE/CustomizableObjectPopulationEditorViewportClient.h b/Engine/Plugins/Experimental/Mutable/Source/CustomizableObjectPopulationEditor/Private/MuCOPE/CustomizableObjectPopulationEditorViewportClient.h
--- a/Engine/Plugins/Experimental/Mutable/Source/CustomizableObjectPopulationEditor/Private/MuCOPE/CustomizableObjectPopulationEditorViewportClient.h
+++ b/Engine/Plugins/Experimental/Mutable/Source/CustomizableObjectPopulationEditor/Private/MuCOPE/CustomizableObjectPopulationEditorViewportClient.h
@@ -2,9 +2,18 @@
 
 #pragma once
 
+#include "Containers/Array.h"
 #include "EditorViewportClient.h"
-#include "AdvancedPreviewScene.h"
-#include "Components/SkeletalMeshComponent.h"
+#include "Engine/EngineBaseTypes.h"
+#include "HAL/Platform.h"
+#include "InputCoreTypes.h"
+#include "Templates/SharedPointer.h"
+
+class FAdvancedPreviewScene;
+class FReferenceCollector;
+class FSceneView;
+class HHitProxy;
+class USkeletalMeshComponent;
 
 
 class FCustomizableObjectPopulationEditorViewportClient : public FEditorViewportClient
diff --git a/Engine/Plugins/Experimental/Mutable/Source/CustomizableObjectPopulationEditor/Private/MuCOPE/CustomizableObjectPopulationFactory.cpp b/Engine/Plugins/Experimental/Mutable/Source/CustomizableObjectPopulationEditor/Private/MuCOPE/CustomizableObjectPopulationFactory.cpp
--- a/Engine/Plugins/Experimental/Mutable/Source/CustomizableObjectPopulationEditor/Private/MuCOPE/CustomizableObjectPopulationFactory.cpp
+++ b/Engine/Plugins/Experimental/Mutable/Source/CustomizableObjectPopulationEditor/Private/MuCOPE/CustomizableObjectPopulationFactory.cpp
@@ -1,7 +1,13 @@
 // Copyright Epic Games, Inc. All Rights Reserved.
 
 #include "MuCOPE/CustomizableObjectPopulationFactory.h"
+
 #include "MuCOP/CustomizableObjectPopulation.h"
+#include "Templates/SubclassOf.h"
+
+class FFeedbackContext;
+class UClass;
+class UObject;
 
 
 #define LOCTEXT_NAMESPACE "CustomizableObjectPopulationFactory"
diff --git a/Engine/Plugins/Experimental/Mutable/Source/CustomizableObjectPopulationEditor/Private/MuCOPE/SCustomizableObjectPopulationClassTagManager.cpp b/Engine/Plugins/Experimental/Mutable/Source/CustomizableObjectPopulationEditor/Private/MuCOPE/SCustomizableObjectPopulationClassTagManager.cpp
--- a/Engine/Plugins/Experimental/Mutable/Source/CustomizableObjectPopulationEditor/Private/MuCOPE/SCustomizableObjectPopulationClassTagManager.cpp
+++ b/Engine/Plugins/Experimental/Mutable/Source/CustomizableObjectPopulationEditor/Private/MuCOPE/SCustomizableObjectPopulationClassTagManager.cpp
@@ -1,13 +1,27 @@
 // Copyright Epic Games, Inc. All Rights Reserved.
 
 #include "MuCOPE/SCustomizableObjectPopulationClassTagManager.h"
-#include "MuCO/CustomizableObject.h"
 
-#include "Widgets/Input/STextComboBox.h"
-#include "Widgets/Input/SEditableTextBox.h"
+#include "Containers/UnrealString.h"
+#include "Fonts/SlateFontInfo.h"
+#include "HAL/PlatformCrt.h"
+#include "Internationalization/Internationalization.h"
+#include "Internationalization/Text.h"
+#include "Layout/BasicLayoutWidgetSlot.h"
+#include "Layout/Children.h"
+#include "Misc/Attribute.h"
+#include "MuCO/CustomizableObject.h"
+#include "SlotBase.h"
+#include "Styling/AppStyle.h"
+#include "Styling/CoreStyle.h"
+#include "Widgets/Images/SImage.h"
 #include "Widgets/Input/SButton.h"
+#include "Widgets/Input/SEditableTextBox.h"
+#include "Widgets/Input/STextComboBox.h"
+#include "Widgets/Layout/SScrollBar.h"
 #include "Widgets/Layout/SScrollBox.h"
-#include "Widgets/Images/SImage.h"
+#include "Widgets/SBoxPanel.h"
+#include "Widgets/Text/STextBlock.h"
 
 #define LOCTEXT_NAMESPACE "CustomizableObjectPopulationClassTagsTool"
 
diff --git a/Engine/Plugins/Experimental/Mutable/Source/CustomizableObjectPopulationEditor/Private/MuCOPE/SCustomizableObjectPopulationClassTagManager.h b/Engine/Plugins/Experimental/Mutable/Source/CustomizableObjectPopulationEditor/Private/MuCOPE/SCustomizableObjectPopulationClassTagManager.h
--- a/Engine/Plugins/Experimental/Mutable/Source/CustomizableObjectPopulationEditor/Private/MuCOPE/SCustomizableObjectPopulationClassTagManager.h
+++ b/Engine/Plugins/Experimental/Mutable/Source/CustomizableObjectPopulationEditor/Private/MuCOPE/SCustomizableObjectPopulationClassTagManager.h
@@ -2,13 +2,17 @@
 
 #pragma once
 
+#include "Containers/Array.h"
+#include "HAL/Platform.h"
+#include "Input/Reply.h"
 #include "Templates/SharedPointer.h"
-#include "UObject/WeakObjectPtr.h"
+#include "Types/SlateEnums.h"
+#include "Widgets/DeclarativeSyntaxSupport.h"
 #include "Widgets/SCompoundWidget.h"
-#include "Components/VerticalBox.h"
-#include "Widgets/Layout/SScrollBox.h"
-
 
+class FString;
+class FText;
+class SVerticalBox;
 class UCustomizableObject;
 
 //Tool to manage the population class tags of a customizable object
diff --git a/Engine/Plugins/Experimental/Mutable/Source/CustomizableObjectPopulationEditor/Private/MuCOPE/SCustomizableObjectPopulationClassTagsTool.cpp b/Engine/Plugins/Experimental/Mutable/Source/CustomizableObjectPopulationEditor/Private/MuCOPE/SCustomizableObjectPopulationClassTagsTool.cpp
--- a/Engine/Plugins/Experimental/Mutable/Source/CustomizableObjectPopulationEditor/Private/MuCOPE/SCustomizableObjectPopulationClassTagsTool.cpp
+++ b/Engine/Plugins/Experimental/Mutable/Source/CustomizableObjectPopulationEditor/Private/MuCOPE/SCustomizableObjectPopulationClassTagsTool.cpp
@@ -1,16 +1,42 @@
 // Copyright Epic Games, Inc. All Rights Reserved.
 
 #include "MuCOPE/SCustomizableObjectPopulationClassTagsTool.h"
-#include "MuCOPE/CustomizableObjectPopulationClassDetails.h"
-#include "MuCOP/CustomizableObjectPopulationClass.h"
-#include "MuCO/CustomizableObject.h"
 
-#include "Widgets/Input/STextComboBox.h"
-#include "Widgets/Input/SSearchBox.h"
-#include "Widgets/Input/SButton.h"
-#include "Widgets/Layout/SScrollBox.h"
+#include "Containers/Map.h"
+#include "CoreTypes.h"
+#include "Fonts/SlateFontInfo.h"
+#include "HAL/PlatformCrt.h"
+#include "Internationalization/Internationalization.h"
+#include "Internationalization/Text.h"
+#include "Layout/BasicLayoutWidgetSlot.h"
+#include "Layout/Children.h"
+#include "Layout/Margin.h"
+#include "Math/Color.h"
+#include "Math/Vector2D.h"
+#include "Misc/Attribute.h"
+#include "MuCO/CustomizableObject.h"
+#include "MuCOP/CustomizableObjectPopulationCharacteristic.h"
+#include "MuCOP/CustomizableObjectPopulationClass.h"
+#include "MuCOPE/CustomizableObjectPopulationClassDetails.h"
+#include "SlotBase.h"
+#include "Styling/AppStyle.h"
+#include "Styling/CoreStyle.h"
+#include "Styling/SlateColor.h"
+#include "Types/SlateStructs.h"
+#include "UObject/ObjectPtr.h"
 #include "Widgets/Images/SImage.h"
+#include "Widgets/Input/SButton.h"
+#include "Widgets/Input/SSearchBox.h"
+#include "Widgets/Input/STextComboBox.h"
+#include "Widgets/Layout/SBox.h"
 #include "Widgets/Layout/SExpandableArea.h"
+#include "Widgets/Layout/SScrollBar.h"
+#include "Widgets/Layout/SScrollBox.h"
+#include "Widgets/SBoxPanel.h"
+#include "Widgets/Text/STextBlock.h"
+
+struct FGeometry;
+struct FSlateBrush;
 
 
 #define LOCTEXT_NAMESPACE "CustomizableObjectPopulationClassTagsTool"
diff --git a/Engine/Plugins/Experimental/Mutable/Source/CustomizableObjectPopulationEditor/Private/MuCOPE/SCustomizableObjectPopulationClassTagsTool.h b/Engine/Plugins/Experimental/Mutable/Source/CustomizableObjectPopulationEditor/Private/MuCOPE/SCustomizableObjectPopulationClassTagsTool.h
--- a/Engine/Plugins/Experimental/Mutable/Source/CustomizableObjectPopulationEditor/Private/MuCOPE/SCustomizableObjectPopulationClassTagsTool.h
+++ b/Engine/Plugins/Experimental/Mutable/Source/CustomizableObjectPopulationEditor/Private/MuCOPE/SCustomizableObjectPopulationClassTagsTool.h
@@ -2,20 +2,25 @@
 
 #pragma once
 
+#include "Containers/Array.h"
+#include "Containers/UnrealString.h"
+#include "HAL/Platform.h"
+#include "Input/Reply.h"
 #include "Templates/SharedPointer.h"
-#include "UObject/WeakObjectPtr.h"
+#include "Types/SlateEnums.h"
+#include "Widgets/DeclarativeSyntaxSupport.h"
 #include "Widgets/SCompoundWidget.h"
-#include "Components/HorizontalBox.h"
-#include "Components/VerticalBox.h"
-#include "Widgets/Input/STextComboBox.h"
-#include "Widgets/Views/STableViewBase.h"
-#include "Widgets/Input/SComboButton.h"
-#include "Widgets/Input/SSearchBox.h"
-#include "Widgets/Views/SListView.h"
-#include "MuCOP/CustomizableObjectPopulationClass.h"
-#include "MuCO/CustomizableObject.h"
 
+class FText;
+class SExpandableArea;
+class SHorizontalBox;
+class SSearchBox;
+class STextComboBox;
+class SVerticalBox;
+class UCustomizableObject;
 class UCustomizableObjectPopulationClass;
+struct FGeometry;
+struct FSlateBrush;
 
 // Widget to manage a single tag 
 class SSingleTagWidget : public SCompoundWidget
diff --git a/Engine/Plugins/Experimental/Mutable/Source/CustomizableObjectPopulationEditor/Private/MuCOPE/SCustomizableObjectPopulationEditorViewport.cpp b/Engine/Plugins/Experimental/Mutable/Source/CustomizableObjectPopulationEditor/Private/MuCOPE/SCustomizableObjectPopulationEditorViewport.cpp
--- a/Engine/Plugins/Experimental/Mutable/Source/CustomizableObjectPopulationEditor/Private/MuCOPE/SCustomizableObjectPopulationEditorViewport.cpp
+++ b/Engine/Plugins/Experimental/Mutable/Source/CustomizableObjectPopulationEditor/Private/MuCOPE/SCustomizableObjectPopulationEditorViewport.cpp
@@ -1,10 +1,26 @@
 // Copyright Epic Games, Inc. All Rights Reserved.
 
 #include "MuCOPE/SCustomizableObjectPopulationEditorViewport.h"
-#include "MuCOPE/SCustomizableObjectPopulationEditorViewportToolBar.h"
-#include "MuCOPE/CustomizableObjectPopulationEditorViewportClient.h"
-#include "Components/BoxComponent.h"
+
+#include "AdvancedPreviewScene.h"
 #include "Components/CapsuleComponent.h"
+#include "Components/SkeletalMeshComponent.h"
+#include "Editor/UnrealEdTypes.h"
+#include "GenericPlatform/ICursor.h"
+#include "Math/Transform.h"
+#include "Math/UnrealMathSSE.h"
+#include "Math/Vector.h"
+#include "Misc/Attribute.h"
+#include "Misc/Optional.h"
+#include "MuCOPE/CustomizableObjectPopulationEditorViewportClient.h"
+#include "MuCOPE/SCustomizableObjectPopulationEditorViewportToolBar.h"
+#include "PreviewScene.h"
+#include "UObject/UObjectGlobals.h"
+#include "Viewports.h"
+
+class FEditorViewportClient;
+class SWidget;
+struct FGeometry;
 
 
 #define LOCTEXT_NAMESPACE "CustomizableObjectPopulationEditorViewport"
diff --git a/Engine/Plugins/Experimental/Mutable/Source/CustomizableObjectPopulationEditor/Private/MuCOPE/SCustomizableObjectPopulationEditorViewport.h b/Engine/Plugins/Experimental/Mutable/Source/CustomizableObjectPopulationEditor/Private/MuCOPE/SCustomizableObjectPopulationEditorViewport.h
--- a/Engine/Plugins/Experimental/Mutable/Source/CustomizableObjectPopulationEditor/Private/MuCOPE/SCustomizableObjectPopulationEditorViewport.h
+++ b/Engine/Plugins/Experimental/Mutable/Source/CustomizableObjectPopulationEditor/Private/MuCOPE/SCustomizableObjectPopulationEditorViewport.h
@@ -2,9 +2,21 @@
 
 #pragma once
 
+#include "Containers/Array.h"
+#include "Containers/UnrealString.h"
+#include "HAL/Platform.h"
 #include "SEditorViewport.h"
-#include "AdvancedPreviewScene.h"
-#include "Components/SkeletalMeshComponent.h"
+#include "Templates/SharedPointer.h"
+#include "UObject/GCObject.h"
+#include "Widgets/DeclarativeSyntaxSupport.h"
+
+class FAdvancedPreviewScene;
+class FEditorViewportClient;
+class FReferenceCollector;
+class SWidget;
+class UCapsuleComponent;
+class USkeletalMeshComponent;
+struct FGeometry;
 
 
 class SCustomizableObjectPopulationEditorViewport : public SEditorViewport, public FGCObject
diff --git a/Engine/Plugins/Experimental/Mutable/Source/CustomizableObjectPopulationEditor/Private/MuCOPE/SCustomizableObjectPopulationEditorViewportToolBar.cpp b/Engine/Plugins/Experimental/Mutable/Source/CustomizableObjectPopulationEditor/Private/MuCOPE/SCustomizableObjectPopulationEditorViewportToolBar.cpp
--- a/Engine/Plugins/Experimental/Mutable/Source/CustomizableObjectPopulationEditor/Private/MuCOPE/SCustomizableObjectPopulationEditorViewportToolBar.cpp
+++ b/Engine/Plugins/Experimental/Mutable/Source/CustomizableObjectPopulationEditor/Private/MuCOPE/SCustomizableObjectPopulationEditorViewportToolBar.cpp
@@ -1,34 +1,42 @@
 // Copyright Epic Games, Inc. All Rights Reserved.
 
 #include "MuCOPE/SCustomizableObjectPopulationEditorViewportToolBar.h"
-#include "MuCOPE/SCustomizableObjectPopulationEditorViewport.h"
+
+#include "Containers/Array.h"
+#include "Delegates/Delegate.h"
+#include "EditorViewportClient.h"
 #include "EditorViewportCommands.h"
-#include "MuCOE/CustomizableObjectEditorViewportMenuCommands.h"
-#include "SEditorViewportToolBarMenu.h"
-#include "STransformViewportToolbar.h"
-#include "Editor/EditorPerProjectUserSettings.h"
-#include "SEditorViewportViewMenu.h"
-#include "Widgets/Colors/SColorPicker.h"
-#include "Widgets/Input/SNumericEntryBox.h"
-#include "Widgets/Input/STextComboBox.h"
-#include "AssetViewerSettings.h"
-#include "ISkeletonTree.h"
-#include "IEditableSkeleton.h"
-#include "Engine/PreviewMeshCollection.h"
-#include "Widgets/Input/SButton.h"
-#include "Widgets/Input/SCheckBox.h"
-#include "Preferences/PersonaOptions.h"
-#include "Widgets/Layout/SScaleBox.h"
-#include "SViewportToolBarComboMenu.h"
-
-#include "Editor/UnrealEdEngine.h"
-#include "UnrealEdGlobals.h"
-#include "Settings/LevelEditorViewportSettings.h"
-#include "Widgets/Layout/SUniformGridPanel.h"
+#include "Fonts/SlateFontInfo.h"
+#include "Framework/MultiBox/MultiBoxBuilder.h"
+#include "Framework/MultiBox/MultiBoxExtender.h"
+#include "GenericPlatform/ICursor.h"
+#include "HAL/Platform.h"
+#include "HAL/PlatformCrt.h"
 #include "Internationalization/Internationalization.h"
-#include "Editor/EditorEngine.h"
+#include "Layout/Children.h"
+#include "Layout/Margin.h"
 #include "LevelEditor.h"
+#include "Math/Color.h"
+#include "Misc/Attribute.h"
+#include "Misc/Optional.h"
 #include "Modules/ModuleManager.h"
+#include "MuCOE/CustomizableObjectEditorViewportMenuCommands.h"
+#include "MuCOPE/SCustomizableObjectPopulationEditorViewport.h"
+#include "SEditorViewportToolBarMenu.h"
+#include "SEditorViewportViewMenu.h"
+#include "SlotBase.h"
+#include "Styling/AppStyle.h"
+#include "Types/ISlateMetaData.h"
+#include "Types/SlateStructs.h"
+#include "UObject/NameTypes.h"
+#include "Widgets/Input/SMenuAnchor.h"
+#include "Widgets/Input/SSpinBox.h"
+#include "Widgets/Layout/SBorder.h"
+#include "Widgets/Layout/SBox.h"
+#include "Widgets/SBoxPanel.h"
+#include "Widgets/SCompoundWidget.h"
+
+class SWidget;
 
 #define LOCTEXT_NAMESPACE "CustomizableObjectPopulationEditorViewportToolBar"
 
diff --git a/Engine/Plugins/Experimental/Mutable/Source/CustomizableObjectPopulationEditor/Private/MuCOPE/SCustomizableObjectPopulationEditorViewportToolBar.h b/Engine/Plugins/Experimental/Mutable/Source/CustomizableObjectPopulationEditor/Private/MuCOPE/SCustomizableObjectPopulationEditorViewportToolBar.h
--- a/Engine/Plugins/Experimental/Mutable/Source/CustomizableObjectPopulationEditor/Private/MuCOPE/SCustomizableObjectPopulationEditorViewportToolBar.h
+++ b/Engine/Plugins/Experimental/Mutable/Source/CustomizableObjectPopulationEditor/Private/MuCOPE/SCustomizableObjectPopulationEditorViewportToolBar.h
@@ -2,9 +2,15 @@
 
 #pragma once
 
+#include "Input/Reply.h"
 #include "SViewportToolBar.h"
-#include "Framework/MultiBox/MultiBoxBuilder.h"
-#include "Settings/LevelEditorViewportSettings.h"
+#include "Styling/SlateColor.h"
+#include "Templates/SharedPointer.h"
+#include "Types/SlateEnums.h"
+#include "Widgets/DeclarativeSyntaxSupport.h"
+
+class SMenuAnchor;
+class SWidget;
 
 
 /**
diff --git a/Engine/Plugins/Experimental/Mutable/Source/CustomizableObjectPopulationEditor/Public/MuCOPE/CustomizableObjectPopulationClassEditorActions.h b/Engine/Plugins/Experimental/Mutable/Source/CustomizableObjectPopulationEditor/Public/MuCOPE/CustomizableObjectPopulationClassEditorActions.h
--- a/Engine/Plugins/Experimental/Mutable/Source/CustomizableObjectPopulationEditor/Public/MuCOPE/CustomizableObjectPopulationClassEditorActions.h
+++ b/Engine/Plugins/Experimental/Mutable/Source/CustomizableObjectPopulationEditor/Public/MuCOPE/CustomizableObjectPopulationClassEditorActions.h
@@ -2,9 +2,10 @@
 
 #pragma once
 
-#include "Styling/AppStyle.h"
-#include "Framework/Commands/UICommandInfo.h"
 #include "Framework/Commands/Commands.h"
+#include "Templates/SharedPointer.h"
+
+class FUICommandInfo;
 
 class FCustomizableObjectPopulationClassEditorCommands : public TCommands<FCustomizableObjectPopulationClassEditorCommands>
 {
diff --git a/Engine/Plugins/Experimental/Mutable/Source/CustomizableObjectPopulationEditor/Public/MuCOPE/CustomizableObjectPopulationEditorActions.h b/Engine/Plugins/Experimental/Mutable/Source/CustomizableObjectPopulationEditor/Public/MuCOPE/CustomizableObjectPopulationEditorActions.h
--- a/Engine/Plugins/Experimental/Mutable/Source/CustomizableObjectPopulationEditor/Public/MuCOPE/CustomizableObjectPopulationEditorActions.h
+++ b/Engine/Plugins/Experimental/Mutable/Source/CustomizableObjectPopulationEditor/Public/MuCOPE/CustomizableObjectPopulationEditorActions.h
@@ -2,9 +2,10 @@
 
 #pragma once
 
-#include "Styling/AppStyle.h"
-#include "Framework/Commands/UICommandInfo.h"
 #include "Framework/Commands/Commands.h"
+#include "Templates/SharedPointer.h"
+
+class FUICommandInfo;
 
 class FCustomizableObjectPopulationEditorCommands : public TCommands<FCustomizableObjectPopulationEditorCommands>
 {
diff --git a/Engine/Plugins/Experimental/Mutable/Source/CustomizableObjectPopulationEditor/Public/MuCOPE/CustomizableObjectPopulationEditorModule.h b/Engine/Plugins/Experimental/Mutable/Source/CustomizableObjectPopulationEditor/Public/MuCOPE/CustomizableObjectPopulationEditorModule.h
--- a/Engine/Plugins/Experimental/Mutable/Source/CustomizableObjectPopulationEditor/Public/MuCOPE/CustomizableObjectPopulationEditorModule.h
+++ b/Engine/Plugins/Experimental/Mutable/Source/CustomizableObjectPopulationEditor/Public/MuCOPE/CustomizableObjectPopulationEditorModule.h
@@ -4,11 +4,15 @@
 
 #include "Modules/ModuleInterface.h"
 #include "Modules/ModuleManager.h"
-#include "Toolkits/AssetEditorToolkit.h"
+#include "Templates/SharedPointer.h"
 #include "Toolkits/IToolkit.h"
-#include "Logging/TokenizedMessage.h"
-#include "MuCOPE/ICustomizableObjectPopulationEditor.h"
-#include "MuCOPE/ICustomizableObjectPopulationClassEditor.h"
+#include "UObject/NameTypes.h"
+
+class ICustomizableObjectPopulationClassEditor;
+class ICustomizableObjectPopulationEditor;
+class IToolkitHost;
+class UCustomizableObjectPopulation;
+class UCustomizableObjectPopulationClass;
 
 extern const FName CustomizableObjectEditorAppIdentifier;
 extern const FName CustomizableObjectInstanceEditorAppIdentifier;
diff --git a/Engine/Plugins/Experimental/Mutable/Source/MutableRuntime/Private/MuR/BlockCompression/Miro/Miro.cpp b/Engine/Plugins/Experimental/Mutable/Source/MutableRuntime/Private/MuR/BlockCompression/Miro/Miro.cpp
--- a/Engine/Plugins/Experimental/Mutable/Source/MutableRuntime/Private/MuR/BlockCompression/Miro/Miro.cpp
+++ b/Engine/Plugins/Experimental/Mutable/Source/MutableRuntime/Private/MuR/BlockCompression/Miro/Miro.cpp
@@ -2,12 +2,13 @@
 
 #include "MuR/BlockCompression/Miro/Miro.h"
 
-#include "MuR/Platform.h"
-#include "Misc/AssertionMacros.h"
 #include "Async/ParallelFor.h"
+#include "HAL/UnrealMemory.h"
+#include "Math/UnrealMathSSE.h"
+#include "Misc/AssertionMacros.h"
 
-#include <cstdlib>
 #include <cmath>
+#include <cstdlib>
 #include <memory>
 
 // debug
diff --git a/Engine/Plugins/Experimental/Mutable/Source/MutableRuntime/Private/MuR/CodeRunner.cpp b/Engine/Plugins/Experimental/Mutable/Source/MutableRuntime/Private/MuR/CodeRunner.cpp
--- a/Engine/Plugins/Experimental/Mutable/Source/MutableRuntime/Private/MuR/CodeRunner.cpp
+++ b/Engine/Plugins/Experimental/Mutable/Source/MutableRuntime/Private/MuR/CodeRunner.cpp
@@ -2,75 +2,87 @@
 
 #include "MuR/CodeRunner.h"
 
-#include "MuR/Platform.h"
-
-#include "MuR/SystemPrivate.h"
-#include "MuR/SettingsPrivate.h"
-#include "MuR/ModelPrivate.h"
-#include "MuR/InstancePrivate.h"
+#include "HAL/UnrealMemory.h"
+#include "Logging/LogCategory.h"
+#include "Logging/LogMacros.h"
+#include "Math/IntPoint.h"
+#include "Math/UnrealMathSSE.h"
+#include "Math/UnrealMathUtility.h"
+#include "Math/Vector2D.h"
+#include "Misc/AssertionMacros.h"
 #include "MuR/ImagePrivate.h"
-#include "MuR/MeshPrivate.h"
+#include "MuR/Instance.h"
+#include "MuR/InstancePrivate.h"
 #include "MuR/MemoryPrivate.h"
-
-#include "MuR/Operations.h"
-
-#include "MuR/OpImageResize.h"
-#include "MuR/OpImagePixelFormat.h"
-#include "MuR/OpImageMipmap.h"
-#include "MuR/OpImageSoftLight.h"
-#include "MuR/OpImageHardLight.h"
-#include "MuR/OpImageBurn.h"
-#include "MuR/OpImageDodge.h"
-#include "MuR/OpImageScreen.h"
-#include "MuR/OpImageOverlay.h"
+#include "MuR/Mesh.h"
+#include "MuR/MeshBufferSet.h"
+#include "MuR/Model.h"
+#include "MuR/ModelPrivate.h"
+#include "MuR/MutableMath.h"
+#include "MuR/MutableString.h"
+#include "MuR/MutableTrace.h"
+#include "MuR/OpEvalCurve.h"
 #include "MuR/OpImageAlphaOverlay.h"
-#include "MuR/OpImageCompose.h"
+#include "MuR/OpImageApplyComposite.h"
+#include "MuR/OpImageBinarise.h"
 #include "MuR/OpImageBlend.h"
-#include "MuR/OpImageDifference.h"
-#include "MuR/OpImageInterpolate.h"
-#include "MuR/OpImageMultiply.h"
-#include "MuR/OpImageSaturate.h"
-#include "MuR/OpImageLuminance.h"
-#include "MuR/OpImageSwizzle.h"
-#include "MuR/OpImageSelectColour.h"
+#include "MuR/OpImageBurn.h"
 #include "MuR/OpImageColourMap.h"
-#include "MuR/OpImageGradient.h"
-#include "MuR/OpImageBinarise.h"
+#include "MuR/OpImageCompose.h"
 #include "MuR/OpImageCrop.h"
-#include "MuR/OpImageRasterMesh.h"
-#include "MuR/OpImageGrow.h"
-#include "MuR/OpImageProject.h"
+#include "MuR/OpImageDifference.h"
 #include "MuR/OpImageDisplace.h"
+#include "MuR/OpImageDodge.h"
+#include "MuR/OpImageGradient.h"
+#include "MuR/OpImageHardLight.h"
+#include "MuR/OpImageInterpolate.h"
 #include "MuR/OpImageInvert.h"
+#include "MuR/OpImageLuminance.h"
+#include "MuR/OpImageMipmap.h"
+#include "MuR/OpImageMultiply.h"
 #include "MuR/OpImageNormalCombine.h"
-#include "MuR/OpImageApplyComposite.h"
+#include "MuR/OpImageOverlay.h"
+#include "MuR/OpImagePixelFormat.h"
+#include "MuR/OpImageProject.h"
+#include "MuR/OpImageRasterMesh.h"
+#include "MuR/OpImageResize.h"
+#include "MuR/OpImageSaturate.h"
+#include "MuR/OpImageScreen.h"
+#include "MuR/OpImageSelectColour.h"
+#include "MuR/OpImageSoftLight.h"
 #include "MuR/OpImageTransform.h"
-
+#include "MuR/OpLayoutPack.h"
+#include "MuR/OpLayoutRemoveBlocks.h"
 #include "MuR/OpMeshApplyLayout.h"
+#include "MuR/OpMeshApplyPose.h"
+#include "MuR/OpMeshBind.h"
+#include "MuR/OpMeshClipDeform.h"
+#include "MuR/OpMeshClipMorphPlane.h"
+#include "MuR/OpMeshClipWithMesh.h"
 #include "MuR/OpMeshDifference.h"
-#include "MuR/OpMeshMorph.h"
-#include "MuR/OpMeshInterpolate.h"
-#include "MuR/OpMeshChartDifference.h"
-#include "MuR/OpMeshSubtract.h"
-#include "MuR/OpMeshRemoveChart.h"
+#include "MuR/OpMeshExtractLayoutBlock.h"
 #include "MuR/OpMeshFormat.h"
+#include "MuR/OpMeshGeometryOperation.h"
 #include "MuR/OpMeshMerge.h"
-#include "MuR/OpMeshExtractLayoutBlock.h"
-#include "MuR/OpMeshTransform.h"
-#include "MuR/OpMeshClipMorphPlane.h"
-#include "MuR/OpMeshClipWithMesh.h"
-#include "MuR/OpMeshApplyPose.h"
+#include "MuR/OpMeshMorph.h"
 #include "MuR/OpMeshRemapIndices.h"
-#include "MuR/OpMeshGeometryOperation.h"
+#include "MuR/OpMeshRemoveChart.h"
 #include "MuR/OpMeshReshape.h"
-#include "MuR/OpMeshClipDeform.h"
+#include "MuR/OpMeshSubtract.h"
+#include "MuR/OpMeshTransform.h"
+#include "MuR/Operations.h"
+#include "MuR/Parameters.h"
+#include "MuR/ParametersPrivate.h"
+#include "MuR/PhysicsBody.h"
+#include "MuR/Platform.h"
+#include "MuR/SettingsPrivate.h"
+#include "MuR/Skeleton.h"
+#include "MuR/SystemPrivate.h"
+#include "Templates/Tuple.h"
+#include "Trace/Detail/Channel.h"
 
-#include "MuR/OpLayoutPack.h"
-#include "MuR/OpLayoutRemoveBlocks.h"
+#include <memory>
+#include <utility>
 
-#include "MuR/OpEvalCurve.h"
-
-#include "Math/UnrealMathUtility.h"
-
 namespace mu
 {
diff --git a/Engine/Plugins/Experimental/Mutable/Source/MutableRuntime/Private/MuR/CodeRunner.h b/Engine/Plugins/Experimental/Mutable/Source/MutableRuntime/Private/MuR/CodeRunner.h
--- a/Engine/Plugins/Experimental/Mutable/Source/MutableRuntime/Private/MuR/CodeRunner.h
+++ b/Engine/Plugins/Experimental/Mutable/Source/MutableRuntime/Private/MuR/CodeRunner.h
@@ -2,26 +2,38 @@
 
 #pragma once
 
-#include "MuR/System.h"
-
-#include "MuR/Operations.h"
-#include "MuR/ImagePrivate.h"
+#include "Async/TaskGraphInterfaces.h"
+#include "Containers/Array.h"
+#include "Containers/ContainerAllocationPolicies.h"
+#include "HAL/PlatformCrt.h"
+#include "HAL/PlatformMath.h"
+#include "MuR/Image.h"
 #include "MuR/Layout.h"
-#include "MuR/InstancePrivate.h"
-#include "MuR/ParametersPrivate.h"
+#include "MuR/MemoryPrivate.h"
+#include "MuR/MutableMemory.h"
+#include "MuR/Operations.h"
+#include "MuR/Ptr.h"
+#include "MuR/RefCounted.h"
+#include "MuR/Serialisation.h"
+#include "MuR/Settings.h"
+#include "MuR/System.h"
 #include "MuR/SystemPrivate.h"
-#include "MuR/ModelPrivate.h"
+#include "MuR/Types.h"
+#include "Templates/RefCounting.h"
+#include "Templates/SharedPointer.h"
 
 // This define could come from MuR/System.h
 #ifdef MUTABLE_USE_NEW_TASKGRAPH
 	#include "Tasks/Task.h"
 #else
-	#include "Async/Async.h"
 #endif
 
 
 namespace  mu
 {
+class Model;
+class Parameters;
+class RangeIndex;
 
 
     //---------------------------------------------------------------------------------------------
diff --git a/Engine/Plugins/Experimental/Mutable/Source/MutableRuntime/Private/MuR/CodeRunnerTasked.cpp b/Engine/Plugins/Experimental/Mutable/Source/MutableRuntime/Private/MuR/CodeRunnerTasked.cpp
--- a/Engine/Plugins/Experimental/Mutable/Source/MutableRuntime/Private/MuR/CodeRunnerTasked.cpp
+++ b/Engine/Plugins/Experimental/Mutable/Source/MutableRuntime/Private/MuR/CodeRunnerTasked.cpp
@@ -1,48 +1,58 @@
 // Copyright Epic Games, Inc. All Rights Reserved.
 
+#include "Async/TaskGraphInterfaces.h"
+#include "Containers/Array.h"
+#include "Containers/UnrealString.h"
+#include "HAL/PlatformCrt.h"
+#include "HAL/UnrealMemory.h"
+#include "Logging/LogCategory.h"
+#include "Logging/LogMacros.h"
+#include "Math/IntPoint.h"
+#include "Math/UnrealMathSSE.h"
+#include "Misc/AssertionMacros.h"
 #include "MuR/CodeRunner.h"
-
-#include "MuR/SystemPrivate.h"
-#include "MuR/SettingsPrivate.h"
-#include "MuR/ModelPrivate.h"
-#include "MuR/InstancePrivate.h"
+#include "MuR/Image.h"
 #include "MuR/ImagePrivate.h"
+#include "MuR/Layout.h"
 #include "MuR/MemoryPrivate.h"
-#include "MuR/MeshPrivate.h"
-
-#include "MuR/Operations.h"
-#include "MuR/OpImageResize.h"
-#include "MuR/OpImagePixelFormat.h"
-#include "MuR/OpImageMipmap.h"
-#include "MuR/OpImageSoftLight.h"
-#include "MuR/OpImageHardLight.h"
-#include "MuR/OpImageBurn.h"
-#include "MuR/OpImageDodge.h"
-#include "MuR/OpImageScreen.h"
-#include "MuR/OpImageOverlay.h"
+#include "MuR/Mesh.h"
+#include "MuR/Model.h"
+#include "MuR/ModelPrivate.h"
+#include "MuR/MutableMath.h"
+#include "MuR/MutableTrace.h"
 #include "MuR/OpImageAlphaOverlay.h"
-#include "MuR/OpImageCompose.h"
 #include "MuR/OpImageBlend.h"
-#include "MuR/OpImageDifference.h"
-#include "MuR/OpImageInterpolate.h"
+#include "MuR/OpImageBurn.h"
+#include "MuR/OpImageCompose.h"
+#include "MuR/OpImageDodge.h"
+#include "MuR/OpImageHardLight.h"
+#include "MuR/OpImageMipmap.h"
 #include "MuR/OpImageMultiply.h"
-#include "MuR/OpImageSaturate.h"
-#include "MuR/OpImageLuminance.h"
-#include "MuR/OpImageSwizzle.h"
-#include "MuR/OpImageSelectColour.h"
-#include "MuR/OpImageColourMap.h"
-#include "MuR/OpImageGradient.h"
-#include "MuR/OpImageBinarise.h"
-#include "MuR/OpImageCrop.h"
-#include "MuR/OpImageRasterMesh.h"
-#include "MuR/OpImageGrow.h"
-#include "MuR/OpImageProject.h"
-#include "MuR/OpImageDisplace.h"
-#include "MuR/OpImageInvert.h"
 #include "MuR/OpImageNormalCombine.h"
-
-#include "MuR/CodeVisitor.h"
+#include "MuR/OpImageOverlay.h"
+#include "MuR/OpImagePixelFormat.h"
+#include "MuR/OpImageResize.h"
+#include "MuR/OpImageScreen.h"
+#include "MuR/OpImageSoftLight.h"
+#include "MuR/OpImageSwizzle.h"
+#include "MuR/Operations.h"
 #include "MuR/Platform.h"
+#include "MuR/Ptr.h"
+#include "MuR/RefCounted.h"
+#include "MuR/Serialisation.h"
+#include "MuR/SerialisationPrivate.h"
+#include "MuR/Settings.h"
+#include "MuR/SettingsPrivate.h"
+#include "MuR/SystemPrivate.h"
+#include "Stats/Stats2.h"
+#include "Templates/RefCounting.h"
+#include "Templates/SharedPointer.h"
+#include "Templates/Tuple.h"
+#include "Trace/Detail/Channel.h"
+
+#include <functional>
+#include <memory>
+#include <utility>
 
 
 DECLARE_CYCLE_STAT(TEXT("MutableCoreTask"), STAT_MutableCoreTask, STATGROUP_Game);
diff --git a/Engine/Plugins/Experimental/Mutable/Source/MutableRuntime/Private/MuR/CodeVisitor.cpp b/Engine/Plugins/Experimental/Mutable/Source/MutableRuntime/Private/MuR/CodeVisitor.cpp
--- a/Engine/Plugins/Experimental/Mutable/Source/MutableRuntime/Private/MuR/CodeVisitor.cpp
+++ b/Engine/Plugins/Experimental/Mutable/Source/MutableRuntime/Private/MuR/CodeVisitor.cpp
@@ -3,7 +3,13 @@
 
 #include "MuR/CodeVisitor.h"
 
+#include "HAL/PlatformCrt.h"
 #include "MuR/MutableTrace.h"
+#include "map"
+
+#include <functional>
+#include <memory>
+#include <utility>
 
 namespace mu
 {
diff --git a/Engine/Plugins/Experimental/Mutable/Source/MutableRuntime/Private/MuR/CodeVisitor.h b/Engine/Plugins/Experimental/Mutable/Source/MutableRuntime/Private/MuR/CodeVisitor.h
--- a/Engine/Plugins/Experimental/Mutable/Source/MutableRuntime/Private/MuR/CodeVisitor.h
+++ b/Engine/Plugins/Experimental/Mutable/Source/MutableRuntime/Private/MuR/CodeVisitor.h
@@ -2,14 +2,22 @@
 
 #pragma once
 
-#include "MuR/Operations.h"
+#include "HAL/PlatformMath.h"
+#include "HAL/UnrealMemory.h"
+#include "Misc/AssertionMacros.h"
+#include "MuR/MemoryPrivate.h"
+#include "MuR/Model.h"
 #include "MuR/ModelPrivate.h"
+#include "MuR/MutableMemory.h"
+#include "MuR/Operations.h"
+#include "MuR/ParametersPrivate.h"
+#include "MuR/Ptr.h"
+#include "MuR/System.h"
 #include "MuR/SystemPrivate.h"
 
-#include <algorithm>
-
 namespace mu
 {
+class Parameters;
 
     //---------------------------------------------------------------------------------------------
     //! Decide what operations are an "add resource" since they are handled differently sometimes.
diff --git a/Engine/Plugins/Experimental/Mutable/Source/MutableRuntime/Private/MuR/Image.cpp b/Engine/Plugins/Experimental/Mutable/Source/MutableRuntime/Private/MuR/Image.cpp
--- a/Engine/Plugins/Experimental/Mutable/Source/MutableRuntime/Private/MuR/Image.cpp
+++ b/Engine/Plugins/Experimental/Mutable/Source/MutableRuntime/Private/MuR/Image.cpp
@@ -1,12 +1,15 @@
 // Copyright Epic Games, Inc. All Rights Reserved.
 
 #include "MuR/Image.h"
+
+#include "HAL/UnrealMemory.h"
+#include "Math/UnrealMathSSE.h"
 #include "MuR/ImagePrivate.h"
+#include "MuR/MutableMath.h"
 #include "MuR/MutableTrace.h"
+#include "MuR/OpImagePixelFormat.h"
 #include "MuR/OpImageResize.h"
 
-#include "MuR/Platform.h"
-
 namespace mu
 {
 
diff --git a/Engine/Plugins/Experimental/Mutable/Source/MutableRuntime/Private/MuR/ImageRLE.cpp b/Engine/Plugins/Experimental/Mutable/Source/MutableRuntime/Private/MuR/ImageRLE.cpp
--- a/Engine/Plugins/Experimental/Mutable/Source/MutableRuntime/Private/MuR/ImageRLE.cpp
+++ b/Engine/Plugins/Experimental/Mutable/Source/MutableRuntime/Private/MuR/ImageRLE.cpp
@@ -1,7 +1,13 @@
 // Copyright Epic Games, Inc. All Rights Reserved.
 
 #include "MuR/ImageRLE.h"
-#include "MuR/Platform.h"
+
+#include "HAL/UnrealMemory.h"
+#include "Misc/AssertionMacros.h"
+#include "MuR/MemoryPrivate.h"
+
+#include <memory>
+#include <utility>
 
 namespace mu
 {
diff --git a/Engine/Plugins/Experimental/Mutable/Source/MutableRuntime/Private/MuR/ImageRLE.h b/Engine/Plugins/Experimental/Mutable/Source/MutableRuntime/Private/MuR/ImageRLE.h
--- a/Engine/Plugins/Experimental/Mutable/Source/MutableRuntime/Private/MuR/ImageRLE.h
+++ b/Engine/Plugins/Experimental/Mutable/Source/MutableRuntime/Private/MuR/ImageRLE.h
@@ -2,8 +2,9 @@
 
 #pragma once
 
-#include "MuR/ImagePrivate.h"
-
+#include "Containers/Array.h"
+#include "Math/UnrealMathSSE.h"
+#include "MuR/Image.h"
 /** This define enabled additional checks when using RLE compression. These checks have a lot of
 * of verhead so it should usually be disabled.
 */
diff --git a/Engine/Plugins/Experimental/Mutable/Source/MutableRuntime/Private/MuR/Instance.cpp b/Engine/Plugins/Experimental/Mutable/Source/MutableRuntime/Private/MuR/Instance.cpp
--- a/Engine/Plugins/Experimental/Mutable/Source/MutableRuntime/Private/MuR/Instance.cpp
+++ b/Engine/Plugins/Experimental/Mutable/Source/MutableRuntime/Private/MuR/Instance.cpp
@@ -3,11 +3,14 @@
 
 #include "MuR/Instance.h"
 
-#include "MuR/Platform.h"
-
+#include "HAL/LowLevelMemTracker.h"
+#include "Misc/AssertionMacros.h"
 #include "MuR/InstancePrivate.h"
-#include "MuR/ImagePrivate.h"
-#include "MuR/MeshPrivate.h"
+#include "MuR/MemoryPrivate.h"
+#include "MuR/MutableMath.h"
+
+#include <memory>
+#include <utility>
 
 
 namespace mu
diff --git a/Engine/Plugins/Experimental/Mutable/Source/MutableRuntime/Private/MuR/Layout.cpp b/Engine/Plugins/Experimental/Mutable/Source/MutableRuntime/Private/MuR/Layout.cpp
--- a/Engine/Plugins/Experimental/Mutable/Source/MutableRuntime/Private/MuR/Layout.cpp
+++ b/Engine/Plugins/Experimental/Mutable/Source/MutableRuntime/Private/MuR/Layout.cpp
@@ -3,7 +3,9 @@
 
 #include "MuR/Layout.h"
 
-#include "MuR/Platform.h"
+#include "HAL/LowLevelMemTracker.h"
+#include "Math/IntPoint.h"
+#include "MuR/MutableMath.h"
 
 
 namespace mu {
diff --git a/Engine/Plugins/Experimental/Mutable/Source/MutableRuntime/Private/MuR/Memory.cpp b/Engine/Plugins/Experimental/Mutable/Source/MutableRuntime/Private/MuR/Memory.cpp
--- a/Engine/Plugins/Experimental/Mutable/Source/MutableRuntime/Private/MuR/Memory.cpp
+++ b/Engine/Plugins/Experimental/Mutable/Source/MutableRuntime/Private/MuR/Memory.cpp
@@ -1,16 +1,10 @@
 // Copyright Epic Games, Inc. All Rights Reserved.
 
+#include "MuR/BlockCompression/Miro/Miro.h"
 #include "MuR/MutableMemory.h"
 
-#include "MuR/Platform.h"
-
-#include "MuR/MemoryPrivate.h"
-
-#include "MuR/BlockCompression/BlockCompression.h"
-
-#include <atomic>
+#include <stdint.h>
 #include <cstdlib>
-#include <cassert>
 
 namespace mu
 {
