Ran IWYU on Mutable plugin

diff --git a/Engine/Plugins/Experimental/Mutable/Source/MutableRuntime/Private/MuR/Mesh.cpp b/Engine/Plugins/Experimental/Mutable/Source/MutableRuntime/Private/MuR/Mesh.cpp
--- a/Engine/Plugins/Experimental/Mutable/Source/MutableRuntime/Private/MuR/Mesh.cpp
+++ b/Engine/Plugins/Experimental/Mutable/Source/MutableRuntime/Private/MuR/Mesh.cpp
@@ -3,10 +3,15 @@
 
 #include "MuR/Mesh.h"
 
+#include "Containers/UnrealString.h"
+#include "HAL/LowLevelMemTracker.h"
+#include "HAL/PlatformCrt.h"
+#include "HAL/UnrealMemory.h"
 #include "MuR/MeshPrivate.h"
-#include "MuR/Platform.h"
 #include "MuR/MutableTrace.h"
 
+#include <utility>
+
 namespace mu
 {
 
diff --git a/Engine/Plugins/Experimental/Mutable/Source/MutableRuntime/Private/MuR/MeshBufferSet.cpp b/Engine/Plugins/Experimental/Mutable/Source/MutableRuntime/Private/MuR/MeshBufferSet.cpp
--- a/Engine/Plugins/Experimental/Mutable/Source/MutableRuntime/Private/MuR/MeshBufferSet.cpp
+++ b/Engine/Plugins/Experimental/Mutable/Source/MutableRuntime/Private/MuR/MeshBufferSet.cpp
@@ -1,8 +1,12 @@
 // Copyright Epic Games, Inc. All Rights Reserved.
 
 #include "MuR/MeshBufferSet.h"
-#include "MuR/MeshPrivate.h"
-#include "MuR/Platform.h"
+
+#include "HAL/LowLevelMemTracker.h"
+#include "HAL/PlatformCrt.h"
+#include "HAL/UnrealMemory.h"
+#include "Math/UnrealMathSSE.h"
+#include "Misc/AssertionMacros.h"
 #include "MuR/MutableTrace.h"
 
 
diff --git a/Engine/Plugins/Experimental/Mutable/Source/MutableRuntime/Private/MuR/Model.cpp b/Engine/Plugins/Experimental/Mutable/Source/MutableRuntime/Private/MuR/Model.cpp
--- a/Engine/Plugins/Experimental/Mutable/Source/MutableRuntime/Private/MuR/Model.cpp
+++ b/Engine/Plugins/Experimental/Mutable/Source/MutableRuntime/Private/MuR/Model.cpp
@@ -3,15 +3,26 @@
 
 #include "MuR/Model.h"
 
-#include "MuR/Platform.h"
-
+#include "Containers/Array.h"
+#include "Containers/Map.h"
+#include "HAL/LowLevelMemTracker.h"
+#include "HAL/PlatformMath.h"
+#include "Misc/AssertionMacros.h"
+#include "MuR/Image.h"
+#include "MuR/MemoryPrivate.h"
+#include "MuR/Mesh.h"
 #include "MuR/ModelPrivate.h"
 #include "MuR/Operations.h"
-#include "MuR/CodeVisitor.h"
-
-#include "Hash/CityHash.h"
-
-#include <algorithm>
+#include "MuR/ParametersPrivate.h"
+#include "MuR/Serialisation.h"
+#include "MuR/SerialisationPrivate.h"
+#include "MuR/System.h"
+#include "MuR/Types.h"
+#include "Templates/Tuple.h"
+
+#include <memory>
+#include <string>
+#include <utility>
 
 namespace mu
 {
diff --git a/Engine/Plugins/Experimental/Mutable/Source/MutableRuntime/Private/MuR/MutableRuntimeModule.cpp b/Engine/Plugins/Experimental/Mutable/Source/MutableRuntime/Private/MuR/MutableRuntimeModule.cpp
--- a/Engine/Plugins/Experimental/Mutable/Source/MutableRuntime/Private/MuR/MutableRuntimeModule.cpp
+++ b/Engine/Plugins/Experimental/Mutable/Source/MutableRuntime/Private/MuR/MutableRuntimeModule.cpp
@@ -1,12 +1,14 @@
 // Copyright Epic Games, Inc. All Rights Reserved.
 
 #include "MuR/MutableRuntimeModule.h"
-#include "CoreMinimal.h"
-#include "Modules/ModuleManager.h"
+
+#include "HAL/PlatformMath.h"
+#include "HAL/UnrealMemory.h"
 #include "Logging/LogMacros.h"
-#include "MuR/Platform.h"
+#include "Misc/AssertionMacros.h"
+#include "Modules/ModuleManager.h"
 #include "MuR/MutableMemory.h"
-#include "MuR/System.h"
+#include "MuR/Platform.h"
 
 IMPLEMENT_MODULE(FMutableRuntimeModule, MutableRuntime);
 
diff --git a/Engine/Plugins/Experimental/Mutable/Source/MutableRuntime/Private/MuR/MutableString.cpp b/Engine/Plugins/Experimental/Mutable/Source/MutableRuntime/Private/MuR/MutableString.cpp
--- a/Engine/Plugins/Experimental/Mutable/Source/MutableRuntime/Private/MuR/MutableString.cpp
+++ b/Engine/Plugins/Experimental/Mutable/Source/MutableRuntime/Private/MuR/MutableString.cpp
@@ -2,7 +2,6 @@
 
 
 #include "MuR/MutableString.h"
-#include "MuR/Platform.h"
 
 
 namespace mu {
diff --git a/Engine/Plugins/Experimental/Mutable/Source/MutableRuntime/Private/MuR/OpImageHardLight.h b/Engine/Plugins/Experimental/Mutable/Source/MutableRuntime/Private/MuR/OpImageHardLight.h
--- a/Engine/Plugins/Experimental/Mutable/Source/MutableRuntime/Private/MuR/OpImageHardLight.h
+++ b/Engine/Plugins/Experimental/Mutable/Source/MutableRuntime/Private/MuR/OpImageHardLight.h
@@ -3,6 +3,7 @@
 #pragma once
 
 #include "MuR/ImagePrivate.h"
+#include "MuR/OpImageLayer.h"
 
 namespace mu
 {
diff --git a/Engine/Plugins/Experimental/Mutable/Source/MutableRuntime/Private/MuR/OpImagePixelFormat.cpp b/Engine/Plugins/Experimental/Mutable/Source/MutableRuntime/Private/MuR/OpImagePixelFormat.cpp
--- a/Engine/Plugins/Experimental/Mutable/Source/MutableRuntime/Private/MuR/OpImagePixelFormat.cpp
+++ b/Engine/Plugins/Experimental/Mutable/Source/MutableRuntime/Private/MuR/OpImagePixelFormat.cpp
@@ -3,10 +3,15 @@
 
 #include "MuR/OpImagePixelFormat.h"
 
+#include "Containers/Array.h"
+#include "HAL/UnrealMemory.h"
+#include "Math/UnrealMathSSE.h"
+#include "Misc/AssertionMacros.h"
+#include "MuR/BlockCompression/Miro/Miro.h"
 #include "MuR/ImageRLE.h"
-#include "MuR/BlockCompression/BlockCompression.h"
-
+#include "MuR/MutableMath.h"
 #include "MuR/MutableTrace.h"
+#include "MuR/Ptr.h"
 
 
 namespace mu
diff --git a/Engine/Plugins/Experimental/Mutable/Source/MutableRuntime/Private/MuR/OpImagePixelFormat.h b/Engine/Plugins/Experimental/Mutable/Source/MutableRuntime/Private/MuR/OpImagePixelFormat.h
--- a/Engine/Plugins/Experimental/Mutable/Source/MutableRuntime/Private/MuR/OpImagePixelFormat.h
+++ b/Engine/Plugins/Experimental/Mutable/Source/MutableRuntime/Private/MuR/OpImagePixelFormat.h
@@ -2,8 +2,8 @@
 
 #pragma once
 
-#include "MuR/ImagePrivate.h"
 
+#include "MuR/Image.h"
 
 namespace mu
 {
diff --git a/Engine/Plugins/Experimental/Mutable/Source/MutableRuntime/Private/MuR/OpImageProject.cpp b/Engine/Plugins/Experimental/Mutable/Source/MutableRuntime/Private/MuR/OpImageProject.cpp
--- a/Engine/Plugins/Experimental/Mutable/Source/MutableRuntime/Private/MuR/OpImageProject.cpp
+++ b/Engine/Plugins/Experimental/Mutable/Source/MutableRuntime/Private/MuR/OpImageProject.cpp
@@ -2,19 +2,32 @@
 
 
 #include "MuR/OpImageProject.h"
-#include "MuR/OpMeshClipWithMesh.h"
 
 #include "Async/ParallelFor.h"
-#include "MuR/Platform.h"
+#include "Containers/Array.h"
+#include "Math/IntPoint.h"
+#include "Math/UnrealMathSSE.h"
+#include "Misc/AssertionMacros.h"
+#include "MuR/ConvertData.h"
+#include "MuR/Image.h"
 #include "MuR/ImagePrivate.h"
+#include "MuR/Layout.h"
+#include "MuR/MeshBufferSet.h"
 #include "MuR/MeshPrivate.h"
-#include "MuR/Raster.h"
-#include "MuR/ConvertData.h"
 #include "MuR/MutableMath.h"
+#include "MuR/MutableTrace.h"
 #include "MuR/Octree.h"
+#include "MuR/OpMeshClipWithMesh.h"
+#include "MuR/Parameters.h"
+#include "MuR/ParametersPrivate.h"
+#include "MuR/Ptr.h"
+#include "MuR/Raster.h"
 
+#include <functional>
+#include <memory>
 #include <queue>
 #include <unordered_set>
+#include <utility>
 
 
 namespace mu
diff --git a/Engine/Plugins/Experimental/Mutable/Source/MutableRuntime/Private/MuR/OpImageProject.h b/Engine/Plugins/Experimental/Mutable/Source/MutableRuntime/Private/MuR/OpImageProject.h
--- a/Engine/Plugins/Experimental/Mutable/Source/MutableRuntime/Private/MuR/OpImageProject.h
+++ b/Engine/Plugins/Experimental/Mutable/Source/MutableRuntime/Private/MuR/OpImageProject.h
@@ -2,15 +2,17 @@
 
 #pragma once
 
-#include "MuR/Platform.h"
-#include "MuR/Raster.h"
-#include "MuR/ParametersPrivate.h"
-#include "MuR/Image.h"
+#include "MuR/MemoryPrivate.h"
 #include "MuR/Mesh.h"
 
+#include <stdint.h>
+
 
 namespace mu
 {
+class Image;
+struct PROJECTOR;
+template <int NUM_INTERPOLATORS> class RasterVertex;
 
     struct SCRATCH_IMAGE_PROJECT
     {
diff --git a/Engine/Plugins/Experimental/Mutable/Source/MutableRuntime/Private/MuR/OpMeshFormat.cpp b/Engine/Plugins/Experimental/Mutable/Source/MutableRuntime/Private/MuR/OpMeshFormat.cpp
--- a/Engine/Plugins/Experimental/Mutable/Source/MutableRuntime/Private/MuR/OpMeshFormat.cpp
+++ b/Engine/Plugins/Experimental/Mutable/Source/MutableRuntime/Private/MuR/OpMeshFormat.cpp
@@ -2,11 +2,25 @@
 
 
 #include "MuR/OpMeshFormat.h"
-#include "MuR/Platform.h"
-#include "MuR/MutableMath.h"
-#include "MuR/MeshPrivate.h"
+
+#include "Containers/Array.h"
+#include "Containers/Map.h"
+#include "HAL/PlatformCrt.h"
+#include "HAL/PlatformMath.h"
+#include "HAL/UnrealMemory.h"
+#include "Misc/AssertionMacros.h"
 #include "MuR/ConvertData.h"
+#include "MuR/Layout.h"
+#include "MuR/MemoryPrivate.h"
+#include "MuR/MeshBufferSet.h"
+#include "MuR/MeshPrivate.h"
+#include "MuR/MutableMath.h"
 #include "MuR/MutableTrace.h"
+#include "MuR/Ptr.h"
+#include "MuR/RefCounted.h"
+
+#include <memory>
+#include <utility>
 
 
 namespace mu
diff --git a/Engine/Plugins/Experimental/Mutable/Source/MutableRuntime/Private/MuR/Operations.cpp b/Engine/Plugins/Experimental/Mutable/Source/MutableRuntime/Private/MuR/Operations.cpp
--- a/Engine/Plugins/Experimental/Mutable/Source/MutableRuntime/Private/MuR/Operations.cpp
+++ b/Engine/Plugins/Experimental/Mutable/Source/MutableRuntime/Private/MuR/Operations.cpp
@@ -2,7 +2,7 @@
 
 
 #include "MuR/Operations.h"
-#include "MuR/Platform.h"
+
 #include "MuR/ModelPrivate.h"
 
 
diff --git a/Engine/Plugins/Experimental/Mutable/Source/MutableRuntime/Private/MuR/Operations.h b/Engine/Plugins/Experimental/Mutable/Source/MutableRuntime/Private/MuR/Operations.h
--- a/Engine/Plugins/Experimental/Mutable/Source/MutableRuntime/Private/MuR/Operations.h
+++ b/Engine/Plugins/Experimental/Mutable/Source/MutableRuntime/Private/MuR/Operations.h
@@ -2,10 +2,12 @@
 
 #pragma once
 
-#include "MuR/MutableMath.h"
-#include "MuR/SerialisationPrivate.h"
+#include "HAL/PlatformMath.h"
+#include "HAL/UnrealMemory.h"
+#include "Misc/AssertionMacros.h"
 #include "MuR/Image.h"
-#include "MuR/Types.h"
+#include "MuR/Platform.h"
+#include "MuR/SerialisationPrivate.h"
 
 #include <functional>
 
diff --git a/Engine/Plugins/Experimental/Mutable/Source/MutableRuntime/Private/MuR/Parameters.cpp b/Engine/Plugins/Experimental/Mutable/Source/MutableRuntime/Private/MuR/Parameters.cpp
--- a/Engine/Plugins/Experimental/Mutable/Source/MutableRuntime/Private/MuR/Parameters.cpp
+++ b/Engine/Plugins/Experimental/Mutable/Source/MutableRuntime/Private/MuR/Parameters.cpp
@@ -2,12 +2,21 @@
 
 
 #include "MuR/Parameters.h"
-#include "MuR/Platform.h"
 
+#include "HAL/LowLevelMemTracker.h"
+#include "HAL/PlatformMath.h"
+#include "HAL/UnrealMemory.h"
+#include "Misc/AssertionMacros.h"
+#include "MuR/MemoryPrivate.h"
+#include "MuR/Model.h"
+#include "MuR/ModelPrivate.h"
 #include "MuR/ParametersPrivate.h"
+#include "MuR/Serialisation.h"
+#include "MuR/SerialisationPrivate.h"
+#include "MuR/System.h"
 
-#include "MuR/ModelPrivate.h"
-#include "MuR/Operations.h"
+#include <memory>
+#include <utility>
 
 
 namespace mu
diff --git a/Engine/Plugins/Experimental/Mutable/Source/MutableRuntime/Private/MuR/PhysicsBody.cpp b/Engine/Plugins/Experimental/Mutable/Source/MutableRuntime/Private/MuR/PhysicsBody.cpp
--- a/Engine/Plugins/Experimental/Mutable/Source/MutableRuntime/Private/MuR/PhysicsBody.cpp
+++ b/Engine/Plugins/Experimental/Mutable/Source/MutableRuntime/Private/MuR/PhysicsBody.cpp
@@ -2,6 +2,8 @@
 
 #include "MuR/PhysicsBody.h"
 
+#include "MuR/SerialisationPrivate.h"
+
 namespace mu
 {
 
diff --git a/Engine/Plugins/Experimental/Mutable/Source/MutableRuntime/Private/MuR/Serialisation.cpp b/Engine/Plugins/Experimental/Mutable/Source/MutableRuntime/Private/MuR/Serialisation.cpp
--- a/Engine/Plugins/Experimental/Mutable/Source/MutableRuntime/Private/MuR/Serialisation.cpp
+++ b/Engine/Plugins/Experimental/Mutable/Source/MutableRuntime/Private/MuR/Serialisation.cpp
@@ -2,11 +2,17 @@
 
 
 #include "MuR/Serialisation.h"
-#include "MuR/Platform.h"
 
+#include "Containers/Array.h"
+#include "HAL/PlatformCrt.h"
+#include "HAL/PlatformMath.h"
+#include "Misc/AssertionMacros.h"
+#include "MuR/Image.h"
+#include "MuR/MemoryPrivate.h"
 #include "MuR/SerialisationPrivate.h"
 
-#include "MuR/Image.h"
+#include <memory>
+#include <utility>
 
 
 namespace mu
diff --git a/Engine/Plugins/Experimental/Mutable/Source/MutableRuntime/Private/MuR/Settings.cpp b/Engine/Plugins/Experimental/Mutable/Source/MutableRuntime/Private/MuR/Settings.cpp
--- a/Engine/Plugins/Experimental/Mutable/Source/MutableRuntime/Private/MuR/Settings.cpp
+++ b/Engine/Plugins/Experimental/Mutable/Source/MutableRuntime/Private/MuR/Settings.cpp
@@ -1,10 +1,9 @@
 // Copyright Epic Games, Inc. All Rights Reserved.
 
 #include "MuR/Settings.h"
-#include "MuR/Platform.h"
 
+#include "Misc/AssertionMacros.h"
 #include "MuR/SettingsPrivate.h"
-#include "MuR/Config.h"
 
 namespace mu
 {
diff --git a/Engine/Plugins/Experimental/Mutable/Source/MutableRuntime/Private/MuR/Skeleton.cpp b/Engine/Plugins/Experimental/Mutable/Source/MutableRuntime/Private/MuR/Skeleton.cpp
--- a/Engine/Plugins/Experimental/Mutable/Source/MutableRuntime/Private/MuR/Skeleton.cpp
+++ b/Engine/Plugins/Experimental/Mutable/Source/MutableRuntime/Private/MuR/Skeleton.cpp
@@ -3,8 +3,8 @@
 
 #include "MuR/Skeleton.h"
 
-#include "MuR/Platform.h"
-#include "MuR/MutableTrace.h"
+#include "MuR/MemoryPrivate.h"
+#include "MuR/SerialisationPrivate.h"
 
 
 namespace mu
diff --git a/Engine/Plugins/Experimental/Mutable/Source/MutableRuntime/Private/MuR/System.cpp b/Engine/Plugins/Experimental/Mutable/Source/MutableRuntime/Private/MuR/System.cpp
--- a/Engine/Plugins/Experimental/Mutable/Source/MutableRuntime/Private/MuR/System.cpp
+++ b/Engine/Plugins/Experimental/Mutable/Source/MutableRuntime/Private/MuR/System.cpp
@@ -1,23 +1,41 @@
 // Copyright Epic Games, Inc. All Rights Reserved.
 
 #include "MuR/System.h"
-#include "MuR/Platform.h"
 
+#include "Containers/Array.h"
+#include "Containers/Map.h"
+#include "Containers/Set.h"
+#include "HAL/IConsoleManager.h"
+#include "HAL/LowLevelMemTracker.h"
+#include "HAL/PlatformCrt.h"
+#include "HAL/PlatformMath.h"
+#include "Logging/LogCategory.h"
+#include "Logging/LogMacros.h"
+#include "Misc/AssertionMacros.h"
 #include "MuR/CodeRunner.h"
-#include "MuR/SystemPrivate.h"
-
-#include "MuR/ModelPrivate.h"
+#include "MuR/CodeVisitor.h"
 #include "MuR/InstancePrivate.h"
-#include "MuR/ImagePrivate.h"
-#include "MuR/MeshPrivate.h"
 #include "MuR/MemoryPrivate.h"
+#include "MuR/Mesh.h"
+#include "MuR/Model.h"
 #include "MuR/ModelPrivate.h"
-#include "MuR/SettingsPrivate.h"
+#include "MuR/MutableMath.h"
+#include "MuR/MutableString.h"
+#include "MuR/MutableTrace.h"
 #include "MuR/Operations.h"
-#include "MuR/CodeVisitor.h"
+#include "MuR/Parameters.h"
+#include "MuR/ParametersPrivate.h"
+#include "MuR/Platform.h"
+#include "MuR/Serialisation.h"
+#include "MuR/SettingsPrivate.h"
+#include "MuR/SystemPrivate.h"
+#include "Templates/SharedPointer.h"
+#include "Templates/Tuple.h"
+#include "Trace/Detail/Channel.h"
 
-#include "MuR/Config.h"
-#include "MuR/MutableTrace.h"
+#include <functional>
+#include <memory>
+#include <utility>
 
 namespace mu
 {
diff --git a/Engine/Plugins/Experimental/Mutable/Source/MutableRuntime/Public/MuR/Image.h b/Engine/Plugins/Experimental/Mutable/Source/MutableRuntime/Public/MuR/Image.h
--- a/Engine/Plugins/Experimental/Mutable/Source/MutableRuntime/Public/MuR/Image.h
+++ b/Engine/Plugins/Experimental/Mutable/Source/MutableRuntime/Public/MuR/Image.h
@@ -2,18 +2,23 @@
 
 #pragma once
 
-#include "MuR/Types.h"
+#include "../Private/MuR/MutableMath.h"
+#include "Containers/Array.h"
+#include "HAL/LowLevelMemTracker.h"
+#include "HAL/Platform.h"
+#include "Misc/AssertionMacros.h"
+#include "MuR/MutableTrace.h"
 #include "MuR/Ptr.h"
 #include "MuR/RefCounted.h"
 #include "MuR/Serialisation.h"
-#include "MuR/MutableTrace.h"
-#include "../Private/MuR/MutableMath.h"
+#include "MuR/SerialisationPrivate.h"
 
 namespace mu
 {
 
 	// Forward declarations
 	class Image;
+
 	typedef Ptr<Image> ImagePtr;
 	typedef Ptr<const Image> ImagePtrConst;
 
diff --git a/Engine/Plugins/Experimental/Mutable/Source/MutableRuntime/Public/MuR/Instance.h b/Engine/Plugins/Experimental/Mutable/Source/MutableRuntime/Public/MuR/Instance.h
--- a/Engine/Plugins/Experimental/Mutable/Source/MutableRuntime/Public/MuR/Instance.h
+++ b/Engine/Plugins/Experimental/Mutable/Source/MutableRuntime/Public/MuR/Instance.h
@@ -2,18 +2,16 @@
 
 #pragma once
 
-#include "MuR/Types.h"
-#include "MuR/RefCounted.h"
+#include "HAL/PlatformMath.h"
 #include "MuR/Ptr.h"
+#include "MuR/RefCounted.h"
+#include "MuR/Types.h"
 
 
 namespace mu
 {
-	// Forward references
-	class Image;
-	class Mesh;
-
 	class Instance;
+
 	typedef Ptr<Instance> InstancePtr;
 	typedef Ptr<const Instance> InstancePtrConst;
 
@@ -209,6 +207,7 @@ namespace mu
 		// Interface pattern
 		//-----------------------------------------------------------------------------------------
 		class Private;
+
 		Private* GetPrivate() const;
 
 	protected:
diff --git a/Engine/Plugins/Experimental/Mutable/Source/MutableRuntime/Public/MuR/Layout.h b/Engine/Plugins/Experimental/Mutable/Source/MutableRuntime/Public/MuR/Layout.h
--- a/Engine/Plugins/Experimental/Mutable/Source/MutableRuntime/Public/MuR/Layout.h
+++ b/Engine/Plugins/Experimental/Mutable/Source/MutableRuntime/Public/MuR/Layout.h
@@ -2,18 +2,21 @@
 
 #pragma once
 
-#include "MuR/Types.h"
+#include "../Private/MuR/MutableMath.h"
+#include "Containers/Array.h"
+#include "Math/UnrealMathSSE.h"
+#include "Misc/AssertionMacros.h"
 #include "MuR/Ptr.h"
 #include "MuR/RefCounted.h"
 #include "MuR/Serialisation.h"
-#include "../Private/MuR/MutableMath.h"
-#include "Math/IntPoint.h"
+#include "MuR/SerialisationPrivate.h"
 
 namespace mu
 {
 
 	// Forward references
 	class Layout;
+
 	typedef Ptr<Layout> LayoutPtr;
 	typedef Ptr<const Layout> LayoutPtrConst;
 
diff --git a/Engine/Plugins/Experimental/Mutable/Source/MutableRuntime/Public/MuR/Mesh.h b/Engine/Plugins/Experimental/Mutable/Source/MutableRuntime/Public/MuR/Mesh.h
--- a/Engine/Plugins/Experimental/Mutable/Source/MutableRuntime/Public/MuR/Mesh.h
+++ b/Engine/Plugins/Experimental/Mutable/Source/MutableRuntime/Public/MuR/Mesh.h
@@ -2,17 +2,28 @@
 
 #pragma once
 
-#include "MuR/Types.h"
+#include "Containers/Array.h"
+#include "Containers/Map.h"
+#include "HAL/PlatformMath.h"
+#include "Math/Transform.h"
+#include "Math/VectorRegister.h"
+#include "Misc/AssertionMacros.h"
+#include "Misc/EnumClassFlags.h"
+#include "MuR/Layout.h"
+#include "MuR/MemoryPrivate.h"
+#include "MuR/MeshBufferSet.h"
+#include "MuR/MutableMath.h"
+#include "MuR/PhysicsBody.h"
 #include "MuR/Ptr.h"
 #include "MuR/RefCounted.h"
 #include "MuR/Serialisation.h"
-#include "MuR/MeshBufferSet.h"
+#include "MuR/SerialisationPrivate.h"
 #include "MuR/Skeleton.h"
-#include "MuR/PhysicsBody.h"
-#include "MuR/Layout.h"
-#include "Math/Transform.h"
+#include "Templates/Tuple.h"
 
-#include "Misc/EnumClassFlags.h"
+#include <memory>
+
+class FString;
 
 namespace mu
 {
diff --git a/Engine/Plugins/Experimental/Mutable/Source/MutableRuntime/Public/MuR/MeshBufferSet.h b/Engine/Plugins/Experimental/Mutable/Source/MutableRuntime/Public/MuR/MeshBufferSet.h
--- a/Engine/Plugins/Experimental/Mutable/Source/MutableRuntime/Public/MuR/MeshBufferSet.h
+++ b/Engine/Plugins/Experimental/Mutable/Source/MutableRuntime/Public/MuR/MeshBufferSet.h
@@ -2,11 +2,11 @@
 
 #pragma once
 
-#include "MuR/Types.h"
-#include "MuR/Ptr.h"
-#include "MuR/RefCounted.h"
 #include "../Private/MuR/SerialisationPrivate.h"
-#include "Math/Transform.h"
+#include "Containers/Array.h"
+#include "HAL/PlatformMath.h"
+#include "MuR/MutableMemory.h"
+#include "MuR/Serialisation.h"
 
 
 namespace mu
diff --git a/Engine/Plugins/Experimental/Mutable/Source/MutableRuntime/Public/MuR/Model.h b/Engine/Plugins/Experimental/Mutable/Source/MutableRuntime/Public/MuR/Model.h
--- a/Engine/Plugins/Experimental/Mutable/Source/MutableRuntime/Public/MuR/Model.h
+++ b/Engine/Plugins/Experimental/Mutable/Source/MutableRuntime/Public/MuR/Model.h
@@ -2,14 +2,17 @@
 
 #pragma once
 
-#include "MuR/RefCounted.h"
-#include "MuR/Ptr.h"
-#include "MuR/Types.h"
-#include "MuR/Serialisation.h"
 #include "MuR/Parameters.h"
+#include "MuR/Ptr.h"
+#include "MuR/RefCounted.h"
 
+#include <stdint.h>
 #include <functional>
 
+namespace mu { class InputArchive; }
+namespace mu { class ModelStreamer; }
+namespace mu { class OutputArchive; }
+
 
 //! This version number changes whenever there is a compatibility-breaking change in the Model
 //! data structures. Compiled models are not necessarily compatible when the runtime is updated,
@@ -22,20 +25,21 @@
 namespace mu
 {
 
-	//!
-	class Parameters;
 	typedef Ptr<Parameters> ParametersPtr;
 	typedef Ptr<const Parameters> ParametersPtrConst;
 
     class Model;
+
     typedef Ptr<Model> ModelPtr;
     typedef Ptr<const Model> ModelPtrConst;
 
     class ModelParametersGenerator;
+
     typedef Ptr<ModelParametersGenerator> ModelParametersGeneratorPtr;
     typedef Ptr<const ModelParametersGenerator> ModelParametersGeneratorPtrConst;
 
     class System;
+
     typedef Ptr<System> SystemPtr;
     typedef Ptr<const System> SystemPtrConst;
 
@@ -109,6 +113,7 @@ namespace mu
 		// Interface pattern
 		//-----------------------------------------------------------------------------------------
 		class Private;
+
 		Private* GetPrivate() const;
 
 	protected:
@@ -161,6 +166,7 @@ namespace mu
         // Interface pattern
         //-----------------------------------------------------------------------------------------
         class Private;
+
         Private* GetPrivate() const;
 
     protected:
diff --git a/Engine/Plugins/Experimental/Mutable/Source/MutableRuntime/Public/MuR/MutableString.h b/Engine/Plugins/Experimental/Mutable/Source/MutableRuntime/Public/MuR/MutableString.h
--- a/Engine/Plugins/Experimental/Mutable/Source/MutableRuntime/Public/MuR/MutableString.h
+++ b/Engine/Plugins/Experimental/Mutable/Source/MutableRuntime/Public/MuR/MutableString.h
@@ -2,13 +2,15 @@
 
 #pragma once
 
-#include "MuR/Types.h"
+#include "Misc/AssertionMacros.h"
+#include "MuR/MemoryPrivate.h"
 #include "MuR/Ptr.h"
 #include "MuR/RefCounted.h"
 #include "MuR/Serialisation.h"
 #include "MuR/SerialisationPrivate.h"
-#include "MuR/Platform.h"
-#include "MuR/MemoryPrivate.h"
+
+#include <stdint.h>
+#include <memory>
 
 namespace mu
 {
diff --git a/Engine/Plugins/Experimental/Mutable/Source/MutableRuntime/Public/MuR/MutableTrace.h b/Engine/Plugins/Experimental/Mutable/Source/MutableRuntime/Public/MuR/MutableTrace.h
--- a/Engine/Plugins/Experimental/Mutable/Source/MutableRuntime/Public/MuR/MutableTrace.h
+++ b/Engine/Plugins/Experimental/Mutable/Source/MutableRuntime/Public/MuR/MutableTrace.h
@@ -4,6 +4,8 @@
 
 #include "Trace/Trace.h"
 
+namespace UE::Trace { class FChannel; }
+
 UE_TRACE_CHANNEL_EXTERN(MutableChannel, MUTABLERUNTIME_API)
 
 /** Custom Mutable profiler scope. */
diff --git a/Engine/Plugins/Experimental/Mutable/Source/MutableRuntime/Public/MuR/Parameters.h b/Engine/Plugins/Experimental/Mutable/Source/MutableRuntime/Public/MuR/Parameters.h
--- a/Engine/Plugins/Experimental/Mutable/Source/MutableRuntime/Public/MuR/Parameters.h
+++ b/Engine/Plugins/Experimental/Mutable/Source/MutableRuntime/Public/MuR/Parameters.h
@@ -2,10 +2,13 @@
 
 #pragma once
 
-#include "MuR/RefCounted.h"
+#include "HAL/Platform.h"
 #include "MuR/Ptr.h"
+#include "MuR/RefCounted.h"
 #include "MuR/Types.h"
-#include "MuR/Serialisation.h"
+
+namespace mu { class InputArchive; }
+namespace mu { class OutputArchive; }
 
 #define MUTABLE_MAX_STRING_PARAM_LENGTH     40
 
@@ -14,14 +17,15 @@ namespace mu
 
 	//!
 	class Parameters;
+
     using ParametersPtr = Ptr<Parameters>;
     using ParametersPtrConst = Ptr<const Parameters>;
 
     class RangeIndex;
+
 	using RangeIndexPtr = Ptr<RangeIndex>;
 	using RangeIndexPtrConst = Ptr<const RangeIndex>;
 
-    class Image;
 	class Model;
 
 
@@ -138,6 +142,7 @@ namespace mu
         // Interface pattern
         //-----------------------------------------------------------------------------------------
         class Private;
+
         Private* GetPrivate() const;
 
     protected:
@@ -389,6 +394,7 @@ namespace mu
 		// Interface pattern
 		//-----------------------------------------------------------------------------------------
 		class Private;
+
 		Private* GetPrivate() const;
 
 	protected:
diff --git a/Engine/Plugins/Experimental/Mutable/Source/MutableRuntime/Public/MuR/PhysicsBody.h b/Engine/Plugins/Experimental/Mutable/Source/MutableRuntime/Public/MuR/PhysicsBody.h
--- a/Engine/Plugins/Experimental/Mutable/Source/MutableRuntime/Public/MuR/PhysicsBody.h
+++ b/Engine/Plugins/Experimental/Mutable/Source/MutableRuntime/Public/MuR/PhysicsBody.h
@@ -2,11 +2,20 @@
 
 #pragma once
 
-#include "MuR/Types.h"
+#include "../Private/MuR/SerialisationPrivate.h"
+#include "Containers/Array.h"
+#include "Math/Quat.h"
+#include "Math/Transform.h"
+#include "Math/UnrealMathSSE.h"
+#include "Math/Vector.h"
+#include "Math/VectorRegister.h"
+#include "Misc/AssertionMacros.h"
+#include "MuR/MemoryPrivate.h"
 #include "MuR/Ptr.h"
 #include "MuR/RefCounted.h"
 #include "MuR/Serialisation.h"
-#include "../Private/MuR/SerialisationPrivate.h"
+
+#include <memory>
 
 namespace mu
 {
diff --git a/Engine/Plugins/Experimental/Mutable/Source/MutableRuntime/Public/MuR/Serialisation.h b/Engine/Plugins/Experimental/Mutable/Source/MutableRuntime/Public/MuR/Serialisation.h
--- a/Engine/Plugins/Experimental/Mutable/Source/MutableRuntime/Public/MuR/Serialisation.h
+++ b/Engine/Plugins/Experimental/Mutable/Source/MutableRuntime/Public/MuR/Serialisation.h
@@ -2,7 +2,7 @@
 
 #pragma once
 
-#include "MuR/Types.h"
+#include "HAL/Platform.h"
 #include "MuR/MutableMemory.h"
 #include "MuR/Ptr.h"
 #include "MuR/RefCounted.h"
@@ -179,6 +179,7 @@ namespace mu
 
         // Interface pattern
         class Private;
+
         Private* GetPrivate() const;
 
     private:
@@ -203,6 +204,7 @@ namespace mu
 
         // Interface pattern
         class Private;
+
         Private* GetPrivate() const;
 
     private:
@@ -238,6 +240,7 @@ namespace mu
     private:
 
         class Private;
+
         Private* m_pD;
     };
 
@@ -278,6 +281,7 @@ namespace mu
     private:
 
         class Private;
+
         Private* m_pD;
 
     };
diff --git a/Engine/Plugins/Experimental/Mutable/Source/MutableRuntime/Public/MuR/Settings.h b/Engine/Plugins/Experimental/Mutable/Source/MutableRuntime/Public/MuR/Settings.h
--- a/Engine/Plugins/Experimental/Mutable/Source/MutableRuntime/Public/MuR/Settings.h
+++ b/Engine/Plugins/Experimental/Mutable/Source/MutableRuntime/Public/MuR/Settings.h
@@ -2,17 +2,17 @@
 
 #pragma once
 
-#include "MuR/RefCounted.h"
 #include "MuR/Ptr.h"
-#include "MuR/Types.h"
-#include "MuR/Serialisation.h"
-#include "MuR/Instance.h"
+#include "MuR/RefCounted.h"
+
+#include <stdint.h>
 
 
 namespace mu
 {
 	// Forward references
     class Settings;
+
     using SettingsPtr=Ptr<Settings>;
     using SettingsPtrConst=Ptr<const Settings>;
 
@@ -56,6 +56,7 @@ namespace mu
         // Interface pattern
         //-----------------------------------------------------------------------------------------
         class Private;
+
         Private* GetPrivate() const;
 
         Settings( const Settings& ) = delete;
diff --git a/Engine/Plugins/Experimental/Mutable/Source/MutableRuntime/Public/MuR/Skeleton.h b/Engine/Plugins/Experimental/Mutable/Source/MutableRuntime/Public/MuR/Skeleton.h
--- a/Engine/Plugins/Experimental/Mutable/Source/MutableRuntime/Public/MuR/Skeleton.h
+++ b/Engine/Plugins/Experimental/Mutable/Source/MutableRuntime/Public/MuR/Skeleton.h
@@ -2,14 +2,19 @@
 
 #pragma once
 
-#include "MuR/Types.h"
-#include "MuR/Ptr.h"
-#include "MuR/RefCounted.h"
-#include "MuR/Serialisation.h"
-#include "../Private/MuR/SerialisationPrivate.h"
 //#include "MuR/MutableMath.h"
 #include "../Private/MuR/MemoryPrivate.h"
+#include "../Private/MuR/SerialisationPrivate.h"
+#include "Containers/Array.h"
+#include "HAL/PlatformMath.h"
 #include "Math/Transform.h"
+#include "Math/VectorRegister.h"
+#include "Misc/AssertionMacros.h"
+#include "MuR/Ptr.h"
+#include "MuR/RefCounted.h"
+#include "MuR/Serialisation.h"
+
+#include <memory>
 
 
 namespace mu
@@ -17,6 +22,7 @@ namespace mu
 
 	// Forward references
     class Skeleton;
+
     typedef Ptr<Skeleton> SkeletonPtr;
     typedef Ptr<const Skeleton> SkeletonPtrConst;
 
diff --git a/Engine/Plugins/Experimental/Mutable/Source/MutableRuntime/Public/MuR/System.h b/Engine/Plugins/Experimental/Mutable/Source/MutableRuntime/Public/MuR/System.h
--- a/Engine/Plugins/Experimental/Mutable/Source/MutableRuntime/Public/MuR/System.h
+++ b/Engine/Plugins/Experimental/Mutable/Source/MutableRuntime/Public/MuR/System.h
@@ -2,13 +2,14 @@
 
 #pragma once
 
-#include "MuR/RefCounted.h"
-#include "MuR/Ptr.h"
-#include "MuR/Types.h"
-#include "MuR/Serialisation.h"
-#include "MuR/Instance.h"
+#include "HAL/Platform.h"
 #include "MuR/Image.h"
+#include "MuR/Instance.h"
+#include "MuR/MutableMemory.h"
+#include "MuR/Ptr.h"
+#include "MuR/RefCounted.h"
 #include "MuR/Settings.h"
+#include "MuR/Types.h"
 
 // This define will use the newer task graph interface to manage mutable concurrency. 
 // This is currently broken in Switch and maybe other consoles, for some unknown reason.
@@ -19,18 +20,23 @@ namespace mu
 {
 	// Forward references
 	class Model;
+class ModelStreamer;
+
     using ModelPtr=Ptr<Model>;
     using ModelPtrConst=Ptr<const Model>;
 
 	class Parameters;
+
     using ParametersPtr=Ptr<Parameters>;
     using ParametersPtrConst=Ptr<const Parameters>;
 
     class Mesh;
+
     using MeshPtr=Ptr<Mesh>;
     using MeshPtrConst=Ptr<const Mesh>;
 
     class System;
+
     using SystemPtr=Ptr<System>;
     using SystemPtrConst=Ptr<const System>;
 
@@ -219,6 +225,7 @@ namespace mu
 		// Interface pattern
 		//-----------------------------------------------------------------------------------------
 		class Private;
+
 		Private* GetPrivate() const;
 
     public:
diff --git a/Engine/Plugins/Experimental/Mutable/Source/MutableTools/Private/MuT/AST.cpp b/Engine/Plugins/Experimental/Mutable/Source/MutableTools/Private/MuT/AST.cpp
--- a/Engine/Plugins/Experimental/Mutable/Source/MutableTools/Private/MuT/AST.cpp
+++ b/Engine/Plugins/Experimental/Mutable/Source/MutableTools/Private/MuT/AST.cpp
@@ -1,13 +1,24 @@
 // Copyright Epic Games, Inc. All Rights Reserved.
 
 #include "MuT/AST.h"
+
+#include "Containers/Set.h"
+#include "Logging/LogCategory.h"
+#include "Logging/LogMacros.h"
+#include "MuR/Mesh.h"
+#include "MuR/MutableTrace.h"
+#include "MuR/Platform.h"
 #include "MuT/ASTOpConstantResource.h"
 #include "MuT/ASTOpMeshRemoveMask.h"
-
 #include "MuT/ErrorLogPrivate.h"
-#include "MuT/CodeOptimiser.h"
 #include "MuT/StreamsPrivate.h"
-#include "MuR/MutableTrace.h"
+#include "Trace/Detail/Channel.h"
+#include "array"
+#include "atomic"
+#include "cstdint"
+#include "iosfwd"
+#include "unordered_map"
+#include "unordered_set"
 
 #include <queue>
 
diff --git a/Engine/Plugins/Experimental/Mutable/Source/MutableTools/Private/MuT/AST.h b/Engine/Plugins/Experimental/Mutable/Source/MutableTools/Private/MuT/AST.h
--- a/Engine/Plugins/Experimental/Mutable/Source/MutableTools/Private/MuT/AST.h
+++ b/Engine/Plugins/Experimental/Mutable/Source/MutableTools/Private/MuT/AST.h
@@ -2,18 +2,33 @@
 
 #pragma once
 
+#include "Containers/Array.h"
+#include "Containers/ContainerAllocationPolicies.h"
 #include "Containers/Map.h"
-#include "MuR/Operations.h"
-#include "MuR/MutableMath.h"
+#include "HAL/PlatformCrt.h"
+#include "HAL/PlatformMath.h"
+#include "Misc/AssertionMacros.h"
+#include "MuR/Image.h"
+#include "MuR/MemoryPrivate.h"
 #include "MuR/ModelPrivate.h"
+#include "MuR/MutableMath.h"
+#include "MuR/MutableMemory.h"
+#include "MuR/Operations.h"
 #include "MuR/ParametersPrivate.h"
-#include "MuR/Image.h"
+#include "MuR/Ptr.h"
+#include "MuR/RefCounted.h"
+#include "Templates/TypeHash.h"
 
-#include <unordered_set>
-#include <unordered_map>
-#include <memory>
 #include <array>
 #include <atomic>
+#include <functional>
+#include <memory>
+#include <unordered_map>
+#include <unordered_set>
+#include <utility>
+
+namespace mu { class ASTOp; }
+namespace mu { class ASTOpFixed; }
 
 namespace std
 {
@@ -156,7 +171,6 @@ namespace mu
     };
 
 
-    class ASTOp;
     typedef vector<Ptr<ASTOp>> ASTOpList;
     typedef std::unordered_set<Ptr<ASTOp>> ASTOpSet;
 
diff --git a/Engine/Plugins/Experimental/Mutable/Source/MutableTools/Private/MuT/ASTOpAddLOD.cpp b/Engine/Plugins/Experimental/Mutable/Source/MutableTools/Private/MuT/ASTOpAddLOD.cpp
--- a/Engine/Plugins/Experimental/Mutable/Source/MutableTools/Private/MuT/ASTOpAddLOD.cpp
+++ b/Engine/Plugins/Experimental/Mutable/Source/MutableTools/Private/MuT/ASTOpAddLOD.cpp
@@ -1,10 +1,15 @@
 // Copyright Epic Games, Inc. All Rights Reserved.
 
 #include "MuT/ASTOpAddLOD.h"
-#include "MuT/ErrorLogPrivate.h"
-#include "MuT/CodeOptimiser.h"
+
+#include "HAL/PlatformMath.h"
+#include "MuR/ModelPrivate.h"
+#include "MuR/RefCounted.h"
 #include "MuT/StreamsPrivate.h"
 
+#include <memory>
+#include <utility>
+
 
 namespace mu
 {
diff --git a/Engine/Plugins/Experimental/Mutable/Source/MutableTools/Private/MuT/ASTOpAddLOD.h b/Engine/Plugins/Experimental/Mutable/Source/MutableTools/Private/MuT/ASTOpAddLOD.h
--- a/Engine/Plugins/Experimental/Mutable/Source/MutableTools/Private/MuT/ASTOpAddLOD.h
+++ b/Engine/Plugins/Experimental/Mutable/Source/MutableTools/Private/MuT/ASTOpAddLOD.h
@@ -2,11 +2,18 @@
 
 #pragma once
 
+#include "HAL/Platform.h"
+#include "MuR/MemoryPrivate.h"
+#include "MuR/Operations.h"
+#include "MuR/Ptr.h"
 #include "MuT/AST.h"
 
+#include <functional>
+
 
 namespace mu
 {
+struct PROGRAM;
 
 
 	//---------------------------------------------------------------------------------------------
diff --git a/Engine/Plugins/Experimental/Mutable/Source/MutableTools/Private/MuT/ASTOpConditional.cpp b/Engine/Plugins/Experimental/Mutable/Source/MutableTools/Private/MuT/ASTOpConditional.cpp
--- a/Engine/Plugins/Experimental/Mutable/Source/MutableTools/Private/MuT/ASTOpConditional.cpp
+++ b/Engine/Plugins/Experimental/Mutable/Source/MutableTools/Private/MuT/ASTOpConditional.cpp
@@ -1,12 +1,20 @@
 // Copyright Epic Games, Inc. All Rights Reserved.
 
 #include "MuT/ASTOpConditional.h"
+
+#include "Containers/Map.h"
+#include "HAL/PlatformMath.h"
+#include "Misc/AssertionMacros.h"
+#include "MuR/MemoryPrivate.h"
+#include "MuR/ModelPrivate.h"
+#include "MuR/MutableMath.h"
+#include "MuR/RefCounted.h"
+#include "MuR/Types.h"
 #include "MuT/ASTOpConstantBool.h"
-#include "MuT/ErrorLogPrivate.h"
-#include "MuT/CodeOptimiser.h"
 #include "MuT/StreamsPrivate.h"
 
-#include <queue>
+#include <memory>
+#include <utility>
 
 
 using namespace mu;
diff --git a/Engine/Plugins/Experimental/Mutable/Source/MutableTools/Private/MuT/ASTOpConditional.h b/Engine/Plugins/Experimental/Mutable/Source/MutableTools/Private/MuT/ASTOpConditional.h
--- a/Engine/Plugins/Experimental/Mutable/Source/MutableTools/Private/MuT/ASTOpConditional.h
+++ b/Engine/Plugins/Experimental/Mutable/Source/MutableTools/Private/MuT/ASTOpConditional.h
@@ -2,11 +2,18 @@
 
 #pragma once
 
+#include "HAL/Platform.h"
+#include "MuR/Image.h"
+#include "MuR/Operations.h"
+#include "MuR/Ptr.h"
 #include "MuT/AST.h"
 
+#include <functional>
+
 
 namespace mu
 {
+struct PROGRAM;
 
 	//---------------------------------------------------------------------------------------------
 	//! Conditional operation.
diff --git a/Engine/Plugins/Experimental/Mutable/Source/MutableTools/Private/MuT/ASTOpConstantBool.cpp b/Engine/Plugins/Experimental/Mutable/Source/MutableTools/Private/MuT/ASTOpConstantBool.cpp
--- a/Engine/Plugins/Experimental/Mutable/Source/MutableTools/Private/MuT/ASTOpConstantBool.cpp
+++ b/Engine/Plugins/Experimental/Mutable/Source/MutableTools/Private/MuT/ASTOpConstantBool.cpp
@@ -1,11 +1,17 @@
 // Copyright Epic Games, Inc. All Rights Reserved.
 
 #include "MuT/ASTOpConstantBool.h"
-#include "MuT/ErrorLogPrivate.h"
-#include "MuT/CodeOptimiser.h"
-#include "Hash/CityHash.h"
+
+#include "HAL/PlatformMath.h"
+#include "MuR/MemoryPrivate.h"
+#include "MuR/ModelPrivate.h"
+#include "MuR/RefCounted.h"
+#include "MuR/Types.h"
 #include "MuT/StreamsPrivate.h"
 
+#include <memory>
+#include <utility>
+
 
 using namespace mu;
 
diff --git a/Engine/Plugins/Experimental/Mutable/Source/MutableTools/Private/MuT/ASTOpConstantBool.h b/Engine/Plugins/Experimental/Mutable/Source/MutableTools/Private/MuT/ASTOpConstantBool.h
--- a/Engine/Plugins/Experimental/Mutable/Source/MutableTools/Private/MuT/ASTOpConstantBool.h
+++ b/Engine/Plugins/Experimental/Mutable/Source/MutableTools/Private/MuT/ASTOpConstantBool.h
@@ -2,11 +2,17 @@
 
 #pragma once
 
+#include "HAL/Platform.h"
+#include "MuR/Operations.h"
+#include "MuR/Ptr.h"
 #include "MuT/AST.h"
 
+#include <functional>
+
 
 namespace mu
 {
+struct PROGRAM;
 
 
 	//---------------------------------------------------------------------------------------------
diff --git a/Engine/Plugins/Experimental/Mutable/Source/MutableTools/Private/MuT/ASTOpConstantProjector.cpp b/Engine/Plugins/Experimental/Mutable/Source/MutableTools/Private/MuT/ASTOpConstantProjector.cpp
--- a/Engine/Plugins/Experimental/Mutable/Source/MutableTools/Private/MuT/ASTOpConstantProjector.cpp
+++ b/Engine/Plugins/Experimental/Mutable/Source/MutableTools/Private/MuT/ASTOpConstantProjector.cpp
@@ -1,10 +1,16 @@
 // Copyright Epic Games, Inc. All Rights Reserved.
 
 #include "MuT/ASTOpConstantProjector.h"
-#include "MuT/ErrorLogPrivate.h"
-#include "MuT/CodeOptimiser.h"
+
+#include "HAL/PlatformMath.h"
+#include "MuR/MemoryPrivate.h"
+#include "MuR/ModelPrivate.h"
+#include "MuR/RefCounted.h"
 #include "MuT/StreamsPrivate.h"
 
+#include <memory>
+#include <utility>
+
 
 namespace mu
 {
diff --git a/Engine/Plugins/Experimental/Mutable/Source/MutableTools/Private/MuT/ASTOpConstantProjector.h b/Engine/Plugins/Experimental/Mutable/Source/MutableTools/Private/MuT/ASTOpConstantProjector.h
--- a/Engine/Plugins/Experimental/Mutable/Source/MutableTools/Private/MuT/ASTOpConstantProjector.h
+++ b/Engine/Plugins/Experimental/Mutable/Source/MutableTools/Private/MuT/ASTOpConstantProjector.h
@@ -2,11 +2,18 @@
 
 #pragma once
 
+#include "HAL/Platform.h"
+#include "MuR/Operations.h"
+#include "MuR/ParametersPrivate.h"
+#include "MuR/Ptr.h"
 #include "MuT/AST.h"
 
+#include <functional>
+
 
 namespace mu
 {
+struct PROGRAM;
 
 
 	//---------------------------------------------------------------------------------------------
diff --git a/Engine/Plugins/Experimental/Mutable/Source/MutableTools/Private/MuT/ASTOpConstantResource.cpp b/Engine/Plugins/Experimental/Mutable/Source/MutableTools/Private/MuT/ASTOpConstantResource.cpp
--- a/Engine/Plugins/Experimental/Mutable/Source/MutableTools/Private/MuT/ASTOpConstantResource.cpp
+++ b/Engine/Plugins/Experimental/Mutable/Source/MutableTools/Private/MuT/ASTOpConstantResource.cpp
@@ -1,15 +1,24 @@
 // Copyright Epic Games, Inc. All Rights Reserved.
 
 #include "MuT/ASTOpConstantResource.h"
-#include "MuT/ErrorLogPrivate.h"
-#include "MuT/CodeOptimiser.h"
-#include "MuT/StreamsPrivate.h"
 
+#include "Containers/Array.h"
+#include "HAL/PlatformMath.h"
 #include "Hash/CityHash.h"
+#include "Misc/AssertionMacros.h"
+#include "MuR/Layout.h"
+#include "MuR/MemoryPrivate.h"
+#include "MuR/Mesh.h"
+#include "MuR/ModelPrivate.h"
+#include "MuR/MutableMath.h"
+#include "MuR/PhysicsBody.h"
+#include "MuR/Serialisation.h"
+#include "MuR/Skeleton.h"
+#include "MuR/Types.h"
+#include "MuT/StreamsPrivate.h"
 
-#include "MuR/OpMeshMerge.h"
-
-#include <queue>
+#include <memory>
+#include <utility>
 
 
 using namespace mu;
diff --git a/Engine/Plugins/Experimental/Mutable/Source/MutableTools/Private/MuT/ASTOpConstantResource.h b/Engine/Plugins/Experimental/Mutable/Source/MutableTools/Private/MuT/ASTOpConstantResource.h
--- a/Engine/Plugins/Experimental/Mutable/Source/MutableTools/Private/MuT/ASTOpConstantResource.h
+++ b/Engine/Plugins/Experimental/Mutable/Source/MutableTools/Private/MuT/ASTOpConstantResource.h
@@ -2,11 +2,20 @@
 
 #pragma once
 
+#include "HAL/Platform.h"
+#include "MuR/Image.h"
+#include "MuR/Operations.h"
+#include "MuR/Ptr.h"
+#include "MuR/RefCounted.h"
 #include "MuT/AST.h"
 
+#include <functional>
+
 
 namespace mu
 {
+struct PROGRAM;
+template <class SCALAR> class vec4;
 
 
 	//---------------------------------------------------------------------------------------------
diff --git a/Engine/Plugins/Experimental/Mutable/Source/MutableTools/Private/MuT/ASTOpConstantString.cpp b/Engine/Plugins/Experimental/Mutable/Source/MutableTools/Private/MuT/ASTOpConstantString.cpp
--- a/Engine/Plugins/Experimental/Mutable/Source/MutableTools/Private/MuT/ASTOpConstantString.cpp
+++ b/Engine/Plugins/Experimental/Mutable/Source/MutableTools/Private/MuT/ASTOpConstantString.cpp
@@ -1,9 +1,15 @@
 // Copyright Epic Games, Inc. All Rights Reserved.
 
 #include "MuT/ASTOpConstantString.h"
-#include "MuT/ErrorLogPrivate.h"
-#include "MuT/CodeOptimiser.h"
-#include "MuT/StreamsPrivate.h"
+
+#include "HAL/PlatformMath.h"
+#include "MuR/ModelPrivate.h"
+#include "MuR/RefCounted.h"
+#include "MuR/Types.h"
+
+#include <memory>
+#include <string>
+#include <utility>
 
 
 using namespace mu;
diff --git a/Engine/Plugins/Experimental/Mutable/Source/MutableTools/Private/MuT/ASTOpConstantString.h b/Engine/Plugins/Experimental/Mutable/Source/MutableTools/Private/MuT/ASTOpConstantString.h
--- a/Engine/Plugins/Experimental/Mutable/Source/MutableTools/Private/MuT/ASTOpConstantString.h
+++ b/Engine/Plugins/Experimental/Mutable/Source/MutableTools/Private/MuT/ASTOpConstantString.h
@@ -2,11 +2,18 @@
 
 #pragma once
 
+#include "HAL/Platform.h"
+#include "MuR/MemoryPrivate.h"
+#include "MuR/Operations.h"
+#include "MuR/Ptr.h"
 #include "MuT/AST.h"
 
+#include <functional>
+
 
 namespace mu
 {
+struct PROGRAM;
 
 
 	//---------------------------------------------------------------------------------------------
diff --git a/Engine/Plugins/Experimental/Mutable/Source/MutableTools/Private/MuT/ASTOpImageCompose.cpp b/Engine/Plugins/Experimental/Mutable/Source/MutableTools/Private/MuT/ASTOpImageCompose.cpp
--- a/Engine/Plugins/Experimental/Mutable/Source/MutableTools/Private/MuT/ASTOpImageCompose.cpp
+++ b/Engine/Plugins/Experimental/Mutable/Source/MutableTools/Private/MuT/ASTOpImageCompose.cpp
@@ -1,12 +1,23 @@
 // Copyright Epic Games, Inc. All Rights Reserved.
 
 #include "MuT/ASTOpImageCompose.h"
-#include "MuT/ErrorLogPrivate.h"
-#include "MuT/CodeOptimiser.h"
-#include "MuT/StreamsPrivate.h"
-#include "MuT/ASTOpImagePixelFormat.h"
-#include "MuT/ASTOpConstantResource.h"
+
+#include "Containers/Array.h"
+#include "Containers/Map.h"
+#include "HAL/PlatformMath.h"
 #include "MuR/Layout.h"
+#include "MuR/MemoryPrivate.h"
+#include "MuR/ModelPrivate.h"
+#include "MuR/MutableMath.h"
+#include "MuR/MutableTrace.h"
+#include "MuR/RefCounted.h"
+#include "MuR/Types.h"
+#include "MuT/ASTOpConstantResource.h"
+#include "MuT/ASTOpImagePixelFormat.h"
+#include "MuT/StreamsPrivate.h"
+
+#include <memory>
+#include <utility>
 
 
 using namespace mu;
diff --git a/Engine/Plugins/Experimental/Mutable/Source/MutableTools/Private/MuT/ASTOpImageCompose.h b/Engine/Plugins/Experimental/Mutable/Source/MutableTools/Private/MuT/ASTOpImageCompose.h
--- a/Engine/Plugins/Experimental/Mutable/Source/MutableTools/Private/MuT/ASTOpImageCompose.h
+++ b/Engine/Plugins/Experimental/Mutable/Source/MutableTools/Private/MuT/ASTOpImageCompose.h
@@ -2,11 +2,19 @@
 
 #pragma once
 
+#include "HAL/Platform.h"
+#include "MuR/Image.h"
+#include "MuR/Operations.h"
+#include "MuR/Ptr.h"
 #include "MuT/AST.h"
 
+#include <functional>
+
 
 namespace mu
 {
+struct PROGRAM;
+template <class SCALAR> class vec4;
 
 	class ASTOpImageCompose : public ASTOp
 	{
diff --git a/Engine/Plugins/Experimental/Mutable/Source/MutableTools/Private/MuT/ASTOpImageMipmap.cpp b/Engine/Plugins/Experimental/Mutable/Source/MutableTools/Private/MuT/ASTOpImageMipmap.cpp
--- a/Engine/Plugins/Experimental/Mutable/Source/MutableTools/Private/MuT/ASTOpImageMipmap.cpp
+++ b/Engine/Plugins/Experimental/Mutable/Source/MutableTools/Private/MuT/ASTOpImageMipmap.cpp
@@ -2,16 +2,26 @@
 
 #include "MuT/ASTOpImageMipmap.h"
 
-#include "MuT/ASTOpImagePixelFormat.h"
+#include "Containers/Map.h"
+#include "HAL/UnrealMemory.h"
+#include "Math/UnrealMathSSE.h"
+#include "Misc/AssertionMacros.h"
+#include "MuR/ImagePrivate.h"
+#include "MuR/MemoryPrivate.h"
+#include "MuR/ModelPrivate.h"
+#include "MuR/MutableMath.h"
+#include "MuR/RefCounted.h"
+#include "MuR/Types.h"
+#include "MuT/ASTOpConditional.h"
 #include "MuT/ASTOpImageCompose.h"
 #include "MuT/ASTOpImagePatch.h"
+#include "MuT/ASTOpImagePixelFormat.h"
 #include "MuT/ASTOpSwitch.h"
-#include "MuT/ASTOpConditional.h"
-
-#include "MuT/ErrorLogPrivate.h"
-#include "MuT/CodeOptimiser.h"
 #include "MuT/StreamsPrivate.h"
 
+#include <memory>
+#include <utility>
+
 
 using namespace mu;
 
diff --git a/Engine/Plugins/Experimental/Mutable/Source/MutableTools/Private/MuT/ASTOpImageMipmap.h b/Engine/Plugins/Experimental/Mutable/Source/MutableTools/Private/MuT/ASTOpImageMipmap.h
--- a/Engine/Plugins/Experimental/Mutable/Source/MutableTools/Private/MuT/ASTOpImageMipmap.h
+++ b/Engine/Plugins/Experimental/Mutable/Source/MutableTools/Private/MuT/ASTOpImageMipmap.h
@@ -2,11 +2,19 @@
 
 #pragma once
 
+#include "HAL/PlatformMath.h"
+#include "MuR/Image.h"
+#include "MuR/Operations.h"
+#include "MuR/Ptr.h"
 #include "MuT/AST.h"
 
+#include <functional>
+
 
 namespace mu
 {
+struct PROGRAM;
+template <class SCALAR> class vec4;
 
 	class ASTOpImageMipmap : public ASTOp
 	{
diff --git a/Engine/Plugins/Experimental/Mutable/Source/MutableTools/Private/MuT/ASTOpImageMultiLayer.cpp b/Engine/Plugins/Experimental/Mutable/Source/MutableTools/Private/MuT/ASTOpImageMultiLayer.cpp
--- a/Engine/Plugins/Experimental/Mutable/Source/MutableTools/Private/MuT/ASTOpImageMultiLayer.cpp
+++ b/Engine/Plugins/Experimental/Mutable/Source/MutableTools/Private/MuT/ASTOpImageMultiLayer.cpp
@@ -1,10 +1,18 @@
 // Copyright Epic Games, Inc. All Rights Reserved.
 
 #include "MuT/ASTOpImageMultiLayer.h"
-#include "MuT/ErrorLogPrivate.h"
-#include "MuT/CodeOptimiser.h"
+
+#include "Containers/Map.h"
+#include "HAL/PlatformMath.h"
+#include "MuR/MemoryPrivate.h"
+#include "MuR/ModelPrivate.h"
+#include "MuR/RefCounted.h"
+#include "MuR/Types.h"
 #include "MuT/StreamsPrivate.h"
 
+#include <memory>
+#include <utility>
+
 
 using namespace mu;
 
diff --git a/Engine/Plugins/Experimental/Mutable/Source/MutableTools/Private/MuT/ASTOpImageMultiLayer.h b/Engine/Plugins/Experimental/Mutable/Source/MutableTools/Private/MuT/ASTOpImageMultiLayer.h
--- a/Engine/Plugins/Experimental/Mutable/Source/MutableTools/Private/MuT/ASTOpImageMultiLayer.h
+++ b/Engine/Plugins/Experimental/Mutable/Source/MutableTools/Private/MuT/ASTOpImageMultiLayer.h
@@ -2,11 +2,18 @@
 
 #pragma once
 
+#include "HAL/Platform.h"
+#include "MuR/Image.h"
+#include "MuR/Operations.h"
+#include "MuR/Ptr.h"
 #include "MuT/AST.h"
 
+#include <functional>
+
 
 namespace mu
 {
+struct PROGRAM;
 
 
 	//---------------------------------------------------------------------------------------------
diff --git a/Engine/Plugins/Experimental/Mutable/Source/MutableTools/Private/MuT/ASTOpImageNormalComposite.cpp b/Engine/Plugins/Experimental/Mutable/Source/MutableTools/Private/MuT/ASTOpImageNormalComposite.cpp
--- a/Engine/Plugins/Experimental/Mutable/Source/MutableTools/Private/MuT/ASTOpImageNormalComposite.cpp
+++ b/Engine/Plugins/Experimental/Mutable/Source/MutableTools/Private/MuT/ASTOpImageNormalComposite.cpp
@@ -1,10 +1,19 @@
 // Copyright Epic Games, Inc. All Rights Reserved.
 
 #include "MuT/ASTOpImageNormalComposite.h"
-#include "MuT/ErrorLogPrivate.h"
-#include "MuT/CodeOptimiser.h"
+
+#include "Containers/Map.h"
+#include "HAL/PlatformMath.h"
+#include "HAL/UnrealMemory.h"
+#include "MuR/MemoryPrivate.h"
+#include "MuR/ModelPrivate.h"
+#include "MuR/RefCounted.h"
+#include "MuR/Types.h"
 #include "MuT/StreamsPrivate.h"
 
+#include <memory>
+#include <utility>
+
 using namespace mu;
 
 
diff --git a/Engine/Plugins/Experimental/Mutable/Source/MutableTools/Private/MuT/ASTOpImagePatch.cpp b/Engine/Plugins/Experimental/Mutable/Source/MutableTools/Private/MuT/ASTOpImagePatch.cpp
--- a/Engine/Plugins/Experimental/Mutable/Source/MutableTools/Private/MuT/ASTOpImagePatch.cpp
+++ b/Engine/Plugins/Experimental/Mutable/Source/MutableTools/Private/MuT/ASTOpImagePatch.cpp
@@ -1,10 +1,17 @@
 // Copyright Epic Games, Inc. All Rights Reserved.
 
 #include "MuT/ASTOpImagePatch.h"
-#include "MuT/ErrorLogPrivate.h"
-#include "MuT/CodeOptimiser.h"
+
+#include "Containers/Map.h"
+#include "MuR/MemoryPrivate.h"
+#include "MuR/ModelPrivate.h"
+#include "MuR/RefCounted.h"
+#include "MuR/Types.h"
 #include "MuT/StreamsPrivate.h"
 
+#include <memory>
+#include <utility>
+
 
 using namespace mu;
 
diff --git a/Engine/Plugins/Experimental/Mutable/Source/MutableTools/Private/MuT/ASTOpImagePatch.h b/Engine/Plugins/Experimental/Mutable/Source/MutableTools/Private/MuT/ASTOpImagePatch.h
--- a/Engine/Plugins/Experimental/Mutable/Source/MutableTools/Private/MuT/ASTOpImagePatch.h
+++ b/Engine/Plugins/Experimental/Mutable/Source/MutableTools/Private/MuT/ASTOpImagePatch.h
@@ -2,11 +2,19 @@
 
 #pragma once
 
+#include "HAL/PlatformMath.h"
+#include "MuR/Image.h"
+#include "MuR/MutableMath.h"
+#include "MuR/Operations.h"
+#include "MuR/Ptr.h"
 #include "MuT/AST.h"
 
+#include <functional>
+
 
 namespace mu
 {
+struct PROGRAM;
 
 	class ASTOpImagePatch : public ASTOp
 	{
diff --git a/Engine/Plugins/Experimental/Mutable/Source/MutableTools/Private/MuT/ASTOpImagePixelFormat.cpp b/Engine/Plugins/Experimental/Mutable/Source/MutableTools/Private/MuT/ASTOpImagePixelFormat.cpp
--- a/Engine/Plugins/Experimental/Mutable/Source/MutableTools/Private/MuT/ASTOpImagePixelFormat.cpp
+++ b/Engine/Plugins/Experimental/Mutable/Source/MutableTools/Private/MuT/ASTOpImagePixelFormat.cpp
@@ -1,15 +1,28 @@
 // Copyright Epic Games, Inc. All Rights Reserved.
 
 #include "MuT/ASTOpImagePixelFormat.h"
-#include "MuT/ASTOpImageMipmap.h"
+
+#include "Containers/Map.h"
+#include "HAL/PlatformMath.h"
+#include "HAL/UnrealMemory.h"
+#include "Misc/AssertionMacros.h"
+#include "MuR/ImagePrivate.h"
+#include "MuR/MemoryPrivate.h"
+#include "MuR/ModelPrivate.h"
+#include "MuR/RefCounted.h"
+#include "MuR/Types.h"
+#include "MuT/ASTOpConditional.h"
 #include "MuT/ASTOpImageCompose.h"
+#include "MuT/ASTOpImageMipmap.h"
 #include "MuT/ASTOpImagePatch.h"
-#include "MuT/ASTOpConditional.h"
 #include "MuT/ASTOpSwitch.h"
-#include "MuT/ErrorLogPrivate.h"
-#include "MuT/CodeOptimiser.h"
 #include "MuT/StreamsPrivate.h"
 
+#include <memory>
+#include <utility>
+
+namespace mu { template <class SCALAR> class vec4; }
+
 
 using namespace mu;
 
diff --git a/Engine/Plugins/Experimental/Mutable/Source/MutableTools/Private/MuT/ASTOpImagePixelFormat.h b/Engine/Plugins/Experimental/Mutable/Source/MutableTools/Private/MuT/ASTOpImagePixelFormat.h
--- a/Engine/Plugins/Experimental/Mutable/Source/MutableTools/Private/MuT/ASTOpImagePixelFormat.h
+++ b/Engine/Plugins/Experimental/Mutable/Source/MutableTools/Private/MuT/ASTOpImagePixelFormat.h
@@ -2,11 +2,19 @@
 
 #pragma once
 
+#include "HAL/Platform.h"
+#include "MuR/Image.h"
+#include "MuR/Operations.h"
+#include "MuR/Ptr.h"
 #include "MuT/AST.h"
 
+#include <functional>
+
 
 namespace mu
 {
+struct PROGRAM;
+template <class SCALAR> class vec4;
 
 	class ASTOpImagePixelFormat : public ASTOp
 	{
diff --git a/Engine/Plugins/Experimental/Mutable/Source/MutableTools/Private/MuT/ASTOpImageTransform.cpp b/Engine/Plugins/Experimental/Mutable/Source/MutableTools/Private/MuT/ASTOpImageTransform.cpp
--- a/Engine/Plugins/Experimental/Mutable/Source/MutableTools/Private/MuT/ASTOpImageTransform.cpp
+++ b/Engine/Plugins/Experimental/Mutable/Source/MutableTools/Private/MuT/ASTOpImageTransform.cpp
@@ -1,10 +1,18 @@
 // Copyright Epic Games, Inc. All Rights Reserved.
 
 #include "MuT/ASTOpImageTransform.h"
-#include "MuT/ErrorLogPrivate.h"
-#include "MuT/CodeOptimiser.h"
+
+#include "Containers/Map.h"
+#include "HAL/PlatformMath.h"
+#include "MuR/MemoryPrivate.h"
+#include "MuR/ModelPrivate.h"
+#include "MuR/RefCounted.h"
+#include "MuR/Types.h"
 #include "MuT/StreamsPrivate.h"
 
+#include <memory>
+#include <utility>
+
 
 using namespace mu;
 
diff --git a/Engine/Plugins/Experimental/Mutable/Source/MutableTools/Private/MuT/ASTOpInstanceAdd.cpp b/Engine/Plugins/Experimental/Mutable/Source/MutableTools/Private/MuT/ASTOpInstanceAdd.cpp
--- a/Engine/Plugins/Experimental/Mutable/Source/MutableTools/Private/MuT/ASTOpInstanceAdd.cpp
+++ b/Engine/Plugins/Experimental/Mutable/Source/MutableTools/Private/MuT/ASTOpInstanceAdd.cpp
@@ -1,11 +1,18 @@
 // Copyright Epic Games, Inc. All Rights Reserved.
 
 #include "MuT/ASTOpInstanceAdd.h"
-#include "MuT/ErrorLogPrivate.h"
+
+#include "Misc/AssertionMacros.h"
+#include "MuR/ModelPrivate.h"
+#include "MuR/ParametersPrivate.h"
+#include "MuR/RefCounted.h"
+#include "MuR/Types.h"
 #include "MuT/CodeOptimiser.h"
 #include "MuT/StreamsPrivate.h"
 
-#include <queue>
+#include <algorithm>
+#include <memory>
+#include <utility>
 
 
 using namespace mu;
diff --git a/Engine/Plugins/Experimental/Mutable/Source/MutableTools/Private/MuT/ASTOpInstanceAdd.h b/Engine/Plugins/Experimental/Mutable/Source/MutableTools/Private/MuT/ASTOpInstanceAdd.h
--- a/Engine/Plugins/Experimental/Mutable/Source/MutableTools/Private/MuT/ASTOpInstanceAdd.h
+++ b/Engine/Plugins/Experimental/Mutable/Source/MutableTools/Private/MuT/ASTOpInstanceAdd.h
@@ -2,11 +2,18 @@
 
 #pragma once
 
+#include "HAL/PlatformMath.h"
+#include "MuR/MemoryPrivate.h"
+#include "MuR/Operations.h"
+#include "MuR/Ptr.h"
 #include "MuT/AST.h"
 
+#include <functional>
+
 
 namespace mu
 {
+struct PROGRAM;
 
 
 	//---------------------------------------------------------------------------------------------
diff --git a/Engine/Plugins/Experimental/Mutable/Source/MutableTools/Private/MuT/ASTOpMeshApplyPose.cpp b/Engine/Plugins/Experimental/Mutable/Source/MutableTools/Private/MuT/ASTOpMeshApplyPose.cpp
--- a/Engine/Plugins/Experimental/Mutable/Source/MutableTools/Private/MuT/ASTOpMeshApplyPose.cpp
+++ b/Engine/Plugins/Experimental/Mutable/Source/MutableTools/Private/MuT/ASTOpMeshApplyPose.cpp
@@ -1,10 +1,17 @@
 // Copyright Epic Games, Inc. All Rights Reserved.
 
 #include "MuT/ASTOpMeshApplyPose.h"
-#include "MuT/ErrorLogPrivate.h"
-#include "MuT/CodeOptimiser.h"
+
+#include "HAL/PlatformMath.h"
+#include "MuR/MemoryPrivate.h"
+#include "MuR/ModelPrivate.h"
+#include "MuR/RefCounted.h"
+#include "MuR/Types.h"
 #include "MuT/StreamsPrivate.h"
 
+#include <memory>
+#include <utility>
+
 
 
 using namespace mu;
diff --git a/Engine/Plugins/Experimental/Mutable/Source/MutableTools/Private/MuT/ASTOpMeshApplyPose.h b/Engine/Plugins/Experimental/Mutable/Source/MutableTools/Private/MuT/ASTOpMeshApplyPose.h
--- a/Engine/Plugins/Experimental/Mutable/Source/MutableTools/Private/MuT/ASTOpMeshApplyPose.h
+++ b/Engine/Plugins/Experimental/Mutable/Source/MutableTools/Private/MuT/ASTOpMeshApplyPose.h
@@ -2,11 +2,17 @@
 
 #pragma once
 
+#include "HAL/Platform.h"
+#include "MuR/Operations.h"
+#include "MuR/Ptr.h"
 #include "MuT/AST.h"
 
+#include <functional>
+
 
 namespace mu
 {
+struct PROGRAM;
 
 
 	//---------------------------------------------------------------------------------------------
diff --git a/Engine/Plugins/Experimental/Mutable/Source/MutableTools/Private/MuT/ASTOpMeshApplyShape.cpp b/Engine/Plugins/Experimental/Mutable/Source/MutableTools/Private/MuT/ASTOpMeshApplyShape.cpp
--- a/Engine/Plugins/Experimental/Mutable/Source/MutableTools/Private/MuT/ASTOpMeshApplyShape.cpp
+++ b/Engine/Plugins/Experimental/Mutable/Source/MutableTools/Private/MuT/ASTOpMeshApplyShape.cpp
@@ -1,10 +1,17 @@
 // Copyright Epic Games, Inc. All Rights Reserved.
 
 #include "MuT/ASTOpMeshApplyShape.h"
-#include "MuT/ErrorLogPrivate.h"
-#include "MuT/CodeOptimiser.h"
+
+#include "HAL/PlatformMath.h"
+#include "MuR/MemoryPrivate.h"
+#include "MuR/ModelPrivate.h"
+#include "MuR/RefCounted.h"
+#include "MuR/Types.h"
 #include "MuT/StreamsPrivate.h"
 
+#include <memory>
+#include <utility>
+
 using namespace mu;
 
 
diff --git a/Engine/Plugins/Experimental/Mutable/Source/MutableTools/Private/MuT/ASTOpMeshBindShape.cpp b/Engine/Plugins/Experimental/Mutable/Source/MutableTools/Private/MuT/ASTOpMeshBindShape.cpp
--- a/Engine/Plugins/Experimental/Mutable/Source/MutableTools/Private/MuT/ASTOpMeshBindShape.cpp
+++ b/Engine/Plugins/Experimental/Mutable/Source/MutableTools/Private/MuT/ASTOpMeshBindShape.cpp
@@ -1,11 +1,18 @@
 // Copyright Epic Games, Inc. All Rights Reserved.
 
 #include "MuT/ASTOpMeshBindShape.h"
+
+#include "HAL/PlatformMath.h"
+#include "HAL/UnrealMemory.h"
+#include "MuR/ModelPrivate.h"
+#include "MuR/RefCounted.h"
+#include "MuR/Types.h"
 #include "MuT/ASTOpSwitch.h"
-#include "MuT/ErrorLogPrivate.h"
-#include "MuT/CodeOptimiser.h"
 #include "MuT/StreamsPrivate.h"
 
+#include <memory>
+#include <utility>
+
 using namespace mu;
 
 
diff --git a/Engine/Plugins/Experimental/Mutable/Source/MutableTools/Private/MuT/ASTOpMeshBindShape.h b/Engine/Plugins/Experimental/Mutable/Source/MutableTools/Private/MuT/ASTOpMeshBindShape.h
--- a/Engine/Plugins/Experimental/Mutable/Source/MutableTools/Private/MuT/ASTOpMeshBindShape.h
+++ b/Engine/Plugins/Experimental/Mutable/Source/MutableTools/Private/MuT/ASTOpMeshBindShape.h
@@ -2,11 +2,19 @@
 
 #pragma once
 
+#include "Containers/Array.h"
+#include "HAL/Platform.h"
+#include "MuR/MemoryPrivate.h"
+#include "MuR/Operations.h"
+#include "MuR/Ptr.h"
 #include "MuT/AST.h"
 
+#include <functional>
+
 
 namespace mu
 {
+struct PROGRAM;
 
 
 	//---------------------------------------------------------------------------------------------
diff --git a/Engine/Plugins/Experimental/Mutable/Source/MutableTools/Private/MuT/ASTOpMeshClipMorphPlane.cpp b/Engine/Plugins/Experimental/Mutable/Source/MutableTools/Private/MuT/ASTOpMeshClipMorphPlane.cpp
--- a/Engine/Plugins/Experimental/Mutable/Source/MutableTools/Private/MuT/ASTOpMeshClipMorphPlane.cpp
+++ b/Engine/Plugins/Experimental/Mutable/Source/MutableTools/Private/MuT/ASTOpMeshClipMorphPlane.cpp
@@ -1,10 +1,16 @@
 // Copyright Epic Games, Inc. All Rights Reserved.
 
 #include "MuT/ASTOpMeshClipMorphPlane.h"
-#include "MuT/ErrorLogPrivate.h"
-#include "MuT/CodeOptimiser.h"
+
+#include "HAL/PlatformMath.h"
+#include "MuR/ModelPrivate.h"
+#include "MuR/RefCounted.h"
+#include "MuR/Types.h"
 #include "MuT/StreamsPrivate.h"
 
+#include <memory>
+#include <utility>
+
 
 using namespace mu;
 
diff --git a/Engine/Plugins/Experimental/Mutable/Source/MutableTools/Private/MuT/ASTOpMeshClipMorphPlane.h b/Engine/Plugins/Experimental/Mutable/Source/MutableTools/Private/MuT/ASTOpMeshClipMorphPlane.h
--- a/Engine/Plugins/Experimental/Mutable/Source/MutableTools/Private/MuT/ASTOpMeshClipMorphPlane.h
+++ b/Engine/Plugins/Experimental/Mutable/Source/MutableTools/Private/MuT/ASTOpMeshClipMorphPlane.h
@@ -2,11 +2,19 @@
 
 #pragma once
 
+#include "HAL/Platform.h"
+#include "MuR/MemoryPrivate.h"
+#include "MuR/Operations.h"
+#include "MuR/ParametersPrivate.h"
+#include "MuR/Ptr.h"
 #include "MuT/AST.h"
 
+#include <functional>
+
 
 namespace mu
 {
+struct PROGRAM;
 
 	//---------------------------------------------------------------------------------------------
 	//!
diff --git a/Engine/Plugins/Experimental/Mutable/Source/MutableTools/Private/MuT/ASTOpMeshExtractLayoutBlocks.cpp b/Engine/Plugins/Experimental/Mutable/Source/MutableTools/Private/MuT/ASTOpMeshExtractLayoutBlocks.cpp
--- a/Engine/Plugins/Experimental/Mutable/Source/MutableTools/Private/MuT/ASTOpMeshExtractLayoutBlocks.cpp
+++ b/Engine/Plugins/Experimental/Mutable/Source/MutableTools/Private/MuT/ASTOpMeshExtractLayoutBlocks.cpp
@@ -1,9 +1,14 @@
 // Copyright Epic Games, Inc. All Rights Reserved.
 
 #include "MuT/ASTOpMeshExtractLayoutBlocks.h"
-#include "MuT/ErrorLogPrivate.h"
-#include "MuT/CodeOptimiser.h"
-#include "MuT/StreamsPrivate.h"
+
+#include "Misc/AssertionMacros.h"
+#include "MuR/ModelPrivate.h"
+#include "MuR/RefCounted.h"
+#include "MuR/Types.h"
+
+#include <memory>
+#include <utility>
 
 
 using namespace mu;
diff --git a/Engine/Plugins/Experimental/Mutable/Source/MutableTools/Private/MuT/ASTOpMeshExtractLayoutBlocks.h b/Engine/Plugins/Experimental/Mutable/Source/MutableTools/Private/MuT/ASTOpMeshExtractLayoutBlocks.h
--- a/Engine/Plugins/Experimental/Mutable/Source/MutableTools/Private/MuT/ASTOpMeshExtractLayoutBlocks.h
+++ b/Engine/Plugins/Experimental/Mutable/Source/MutableTools/Private/MuT/ASTOpMeshExtractLayoutBlocks.h
@@ -2,11 +2,18 @@
 
 #pragma once
 
+#include "HAL/PlatformMath.h"
+#include "MuR/MemoryPrivate.h"
+#include "MuR/Operations.h"
+#include "MuR/Ptr.h"
 #include "MuT/AST.h"
 
+#include <functional>
+
 
 namespace mu
 {
+struct PROGRAM;
 
 
 	//---------------------------------------------------------------------------------------------
diff --git a/Engine/Plugins/Experimental/Mutable/Source/MutableTools/Private/MuT/ASTOpMeshFormat.cpp b/Engine/Plugins/Experimental/Mutable/Source/MutableTools/Private/MuT/ASTOpMeshFormat.cpp
--- a/Engine/Plugins/Experimental/Mutable/Source/MutableTools/Private/MuT/ASTOpMeshFormat.cpp
+++ b/Engine/Plugins/Experimental/Mutable/Source/MutableTools/Private/MuT/ASTOpMeshFormat.cpp
@@ -1,15 +1,27 @@
 // Copyright Epic Games, Inc. All Rights Reserved.
 
 #include "MuT/ASTOpMeshFormat.h"
-#include "MuT/ASTOpConstantResource.h"
-#include "MuT/ASTOpSwitch.h"
+
+#include "HAL/UnrealMemory.h"
+#include "Misc/AssertionMacros.h"
+#include "MuR/MemoryPrivate.h"
+#include "MuR/Mesh.h"
+#include "MuR/MeshBufferSet.h"
+#include "MuR/ModelPrivate.h"
+#include "MuR/MutableTrace.h"
+#include "MuR/RefCounted.h"
+#include "MuR/System.h"
+#include "MuR/Types.h"
 #include "MuT/ASTOpConditional.h"
-#include "MuT/ASTOpMeshRemoveMask.h"
+#include "MuT/ASTOpConstantResource.h"
 #include "MuT/ASTOpMeshClipMorphPlane.h"
-#include "MuT/ErrorLogPrivate.h"
-#include "MuT/CodeOptimiser.h"
+#include "MuT/ASTOpMeshRemoveMask.h"
+#include "MuT/ASTOpSwitch.h"
 #include "MuT/StreamsPrivate.h"
 
+#include <memory>
+#include <utility>
+
 
 using namespace mu;
 
diff --git a/Engine/Plugins/Experimental/Mutable/Source/MutableTools/Private/MuT/ASTOpMeshFormat.h b/Engine/Plugins/Experimental/Mutable/Source/MutableTools/Private/MuT/ASTOpMeshFormat.h
--- a/Engine/Plugins/Experimental/Mutable/Source/MutableTools/Private/MuT/ASTOpMeshFormat.h
+++ b/Engine/Plugins/Experimental/Mutable/Source/MutableTools/Private/MuT/ASTOpMeshFormat.h
@@ -2,11 +2,17 @@
 
 #pragma once
 
+#include "HAL/PlatformMath.h"
+#include "MuR/Operations.h"
+#include "MuR/Ptr.h"
 #include "MuT/AST.h"
 
+#include <functional>
+
 
 namespace mu
 {
+struct PROGRAM;
 	
 	class ASTOpMeshFormat : public ASTOp
 	{
diff --git a/Engine/Plugins/Experimental/Mutable/Source/MutableTools/Private/MuT/ASTOpMeshGeometryOperation.cpp b/Engine/Plugins/Experimental/Mutable/Source/MutableTools/Private/MuT/ASTOpMeshGeometryOperation.cpp
--- a/Engine/Plugins/Experimental/Mutable/Source/MutableTools/Private/MuT/ASTOpMeshGeometryOperation.cpp
+++ b/Engine/Plugins/Experimental/Mutable/Source/MutableTools/Private/MuT/ASTOpMeshGeometryOperation.cpp
@@ -1,10 +1,16 @@
 // Copyright Epic Games, Inc. All Rights Reserved.
 
 #include "MuT/ASTOpMeshGeometryOperation.h"
-#include "MuT/ErrorLogPrivate.h"
-#include "MuT/CodeOptimiser.h"
+
+#include "HAL/PlatformMath.h"
+#include "MuR/MemoryPrivate.h"
+#include "MuR/ModelPrivate.h"
+#include "MuR/RefCounted.h"
 #include "MuT/StreamsPrivate.h"
 
+#include <memory>
+#include <utility>
+
 namespace mu
 {
 
diff --git a/Engine/Plugins/Experimental/Mutable/Source/MutableTools/Private/MuT/ASTOpMeshMaskClipMesh.cpp b/Engine/Plugins/Experimental/Mutable/Source/MutableTools/Private/MuT/ASTOpMeshMaskClipMesh.cpp
--- a/Engine/Plugins/Experimental/Mutable/Source/MutableTools/Private/MuT/ASTOpMeshMaskClipMesh.cpp
+++ b/Engine/Plugins/Experimental/Mutable/Source/MutableTools/Private/MuT/ASTOpMeshMaskClipMesh.cpp
@@ -1,13 +1,20 @@
 // Copyright Epic Games, Inc. All Rights Reserved.
 
 #include "MuT/ASTOpMeshMaskClipMesh.h"
-#include "MuT/ASTOpSwitch.h"
+
+#include "HAL/PlatformMath.h"
+#include "MuR/MemoryPrivate.h"
+#include "MuR/ModelPrivate.h"
+#include "MuR/RefCounted.h"
+#include "MuR/Types.h"
 #include "MuT/ASTOpConditional.h"
 #include "MuT/ASTOpMeshRemoveMask.h"
-#include "MuT/ErrorLogPrivate.h"
-#include "MuT/CodeOptimiser.h"
+#include "MuT/ASTOpSwitch.h"
 #include "MuT/StreamsPrivate.h"
 
+#include <memory>
+#include <utility>
+
 
 using namespace mu;
 
diff --git a/Engine/Plugins/Experimental/Mutable/Source/MutableTools/Private/MuT/ASTOpMeshRemoveMask.cpp b/Engine/Plugins/Experimental/Mutable/Source/MutableTools/Private/MuT/ASTOpMeshRemoveMask.cpp
--- a/Engine/Plugins/Experimental/Mutable/Source/MutableTools/Private/MuT/ASTOpMeshRemoveMask.cpp
+++ b/Engine/Plugins/Experimental/Mutable/Source/MutableTools/Private/MuT/ASTOpMeshRemoveMask.cpp
@@ -1,10 +1,15 @@
 // Copyright Epic Games, Inc. All Rights Reserved.
 
 #include "MuT/ASTOpMeshRemoveMask.h"
-#include "MuT/ErrorLogPrivate.h"
-#include "MuT/CodeOptimiser.h"
+
+#include "HAL/PlatformMath.h"
+#include "MuR/ModelPrivate.h"
+#include "MuR/RefCounted.h"
+#include "MuR/Types.h"
 #include "MuT/StreamsPrivate.h"
 
+#include <memory>
+
 
 using namespace mu;
 
diff --git a/Engine/Plugins/Experimental/Mutable/Source/MutableTools/Private/MuT/ASTOpMeshRemoveMask.h b/Engine/Plugins/Experimental/Mutable/Source/MutableTools/Private/MuT/ASTOpMeshRemoveMask.h
--- a/Engine/Plugins/Experimental/Mutable/Source/MutableTools/Private/MuT/ASTOpMeshRemoveMask.h
+++ b/Engine/Plugins/Experimental/Mutable/Source/MutableTools/Private/MuT/ASTOpMeshRemoveMask.h
@@ -2,11 +2,19 @@
 
 #pragma once
 
+#include "HAL/Platform.h"
+#include "MuR/MemoryPrivate.h"
+#include "MuR/Operations.h"
+#include "MuR/Ptr.h"
 #include "MuT/AST.h"
 
+#include <functional>
+#include <utility>
+
 
 namespace mu
 {
+struct PROGRAM;
 
 
 	//---------------------------------------------------------------------------------------------
diff --git a/Engine/Plugins/Experimental/Mutable/Source/MutableTools/Private/MuT/ASTOpMeshTransform.h b/Engine/Plugins/Experimental/Mutable/Source/MutableTools/Private/MuT/ASTOpMeshTransform.h
--- a/Engine/Plugins/Experimental/Mutable/Source/MutableTools/Private/MuT/ASTOpMeshTransform.h
+++ b/Engine/Plugins/Experimental/Mutable/Source/MutableTools/Private/MuT/ASTOpMeshTransform.h
@@ -2,11 +2,18 @@
 
 #pragma once
 
+#include "HAL/Platform.h"
+#include "MuR/MutableMath.h"
+#include "MuR/Operations.h"
+#include "MuR/Ptr.h"
 #include "MuT/AST.h"
 
+#include <functional>
+
 
 namespace mu
 {
+struct PROGRAM;
 
 
 	//---------------------------------------------------------------------------------------------
diff --git a/Engine/Plugins/Experimental/Mutable/Source/MutableTools/Private/MuT/ASTOpParameter.cpp b/Engine/Plugins/Experimental/Mutable/Source/MutableTools/Private/MuT/ASTOpParameter.cpp
--- a/Engine/Plugins/Experimental/Mutable/Source/MutableTools/Private/MuT/ASTOpParameter.cpp
+++ b/Engine/Plugins/Experimental/Mutable/Source/MutableTools/Private/MuT/ASTOpParameter.cpp
@@ -1,10 +1,18 @@
 // Copyright Epic Games, Inc. All Rights Reserved.
 
 #include "MuT/ASTOpParameter.h"
-#include "MuT/ErrorLogPrivate.h"
-#include "MuT/CodeOptimiser.h"
+
+#include "HAL/PlatformMath.h"
+#include "Misc/AssertionMacros.h"
+#include "MuR/ModelPrivate.h"
+#include "MuR/Parameters.h"
+#include "MuR/RefCounted.h"
+#include "MuR/Types.h"
 #include "MuT/StreamsPrivate.h"
 
+#include <memory>
+#include <utility>
+
 
 using namespace mu;
 
diff --git a/Engine/Plugins/Experimental/Mutable/Source/MutableTools/Private/MuT/ASTOpParameter.h b/Engine/Plugins/Experimental/Mutable/Source/MutableTools/Private/MuT/ASTOpParameter.h
--- a/Engine/Plugins/Experimental/Mutable/Source/MutableTools/Private/MuT/ASTOpParameter.h
+++ b/Engine/Plugins/Experimental/Mutable/Source/MutableTools/Private/MuT/ASTOpParameter.h
@@ -2,11 +2,20 @@
 
 #pragma once
 
+#include "HAL/Platform.h"
+#include "MuR/Image.h"
+#include "MuR/MemoryPrivate.h"
+#include "MuR/Operations.h"
+#include "MuR/ParametersPrivate.h"
+#include "MuR/Ptr.h"
 #include "MuT/AST.h"
 
+#include <functional>
+
 
 namespace mu
 {
+struct PROGRAM;
 
 
 	//---------------------------------------------------------------------------------------------
diff --git a/Engine/Plugins/Experimental/Mutable/Source/MutableTools/Private/MuT/ASTOpScalarCurve.h b/Engine/Plugins/Experimental/Mutable/Source/MutableTools/Private/MuT/ASTOpScalarCurve.h
--- a/Engine/Plugins/Experimental/Mutable/Source/MutableTools/Private/MuT/ASTOpScalarCurve.h
+++ b/Engine/Plugins/Experimental/Mutable/Source/MutableTools/Private/MuT/ASTOpScalarCurve.h
@@ -2,11 +2,18 @@
 
 #pragma once
 
+#include "HAL/Platform.h"
+#include "MuR/Operations.h"
+#include "MuR/ParametersPrivate.h"
+#include "MuR/Ptr.h"
 #include "MuT/AST.h"
 
+#include <functional>
+
 
 namespace mu
 {
+struct PROGRAM;
 
 
 	//---------------------------------------------------------------------------------------------
diff --git a/Engine/Plugins/Experimental/Mutable/Source/MutableTools/Private/MuT/ASTOpSwitch.cpp b/Engine/Plugins/Experimental/Mutable/Source/MutableTools/Private/MuT/ASTOpSwitch.cpp
--- a/Engine/Plugins/Experimental/Mutable/Source/MutableTools/Private/MuT/ASTOpSwitch.cpp
+++ b/Engine/Plugins/Experimental/Mutable/Source/MutableTools/Private/MuT/ASTOpSwitch.cpp
@@ -1,10 +1,19 @@
 // Copyright Epic Games, Inc. All Rights Reserved.
 
 #include "MuT/ASTOpSwitch.h"
-#include "MuT/ErrorLogPrivate.h"
-#include "MuT/CodeOptimiser.h"
+
+#include "Containers/Map.h"
+#include "Misc/AssertionMacros.h"
+#include "MuR/ImagePrivate.h"
+#include "MuR/ModelPrivate.h"
+#include "MuR/MutableMath.h"
+#include "MuR/RefCounted.h"
+#include "MuR/Types.h"
 #include "MuT/StreamsPrivate.h"
 
+#include <memory>
+#include <utility>
+
 
 using namespace mu;
 
diff --git a/Engine/Plugins/Experimental/Mutable/Source/MutableTools/Private/MuT/ASTOpSwitch.h b/Engine/Plugins/Experimental/Mutable/Source/MutableTools/Private/MuT/ASTOpSwitch.h
--- a/Engine/Plugins/Experimental/Mutable/Source/MutableTools/Private/MuT/ASTOpSwitch.h
+++ b/Engine/Plugins/Experimental/Mutable/Source/MutableTools/Private/MuT/ASTOpSwitch.h
@@ -2,11 +2,20 @@
 
 #pragma once
 
+#include "HAL/PlatformMath.h"
+#include "MuR/Image.h"
+#include "MuR/MemoryPrivate.h"
+#include "MuR/Operations.h"
+#include "MuR/Ptr.h"
 #include "MuT/AST.h"
 
+#include <functional>
+
 
 namespace mu
 {
+struct PROGRAM;
+template <class SCALAR> class vec4;
 
 
 	//---------------------------------------------------------------------------------------------
diff --git a/Engine/Plugins/Experimental/Mutable/Source/MutableTools/Private/MuT/CodeGenerator.cpp b/Engine/Plugins/Experimental/Mutable/Source/MutableTools/Private/MuT/CodeGenerator.cpp
--- a/Engine/Plugins/Experimental/Mutable/Source/MutableTools/Private/MuT/CodeGenerator.cpp
+++ b/Engine/Plugins/Experimental/Mutable/Source/MutableTools/Private/MuT/CodeGenerator.cpp
@@ -2,71 +2,82 @@
 
 #include "MuT/CodeGenerator.h"
 
+#include "Containers/Array.h"
+#include "Logging/LogCategory.h"
+#include "Logging/LogMacros.h"
+#include "Math/IntPoint.h"
+#include "Math/UnrealMathSSE.h"
+#include "MuR/ImagePrivate.h"
+#include "MuR/Layout.h"
+#include "MuR/MutableTrace.h"
+#include "MuR/Operations.h"
+#include "MuR/ParametersPrivate.h"
+#include "MuR/Platform.h"
+#include "MuT/ASTOpAddLOD.h"
+#include "MuT/ASTOpConditional.h"
+#include "MuT/ASTOpConstantBool.h"
+#include "MuT/ASTOpConstantResource.h"
+#include "MuT/ASTOpImageCompose.h"
+#include "MuT/ASTOpImageMipmap.h"
+#include "MuT/ASTOpImagePixelFormat.h"
+#include "MuT/ASTOpInstanceAdd.h"
+#include "MuT/ASTOpMeshBindShape.h"
+#include "MuT/ASTOpMeshClipDeform.h"
+#include "MuT/ASTOpMeshClipMorphPlane.h"
+#include "MuT/ASTOpMeshMaskClipMesh.h"
+#include "MuT/ASTOpMeshRemoveMask.h"
+#include "MuT/ASTOpParameter.h"
 #include "MuT/CodeGenerator_SecondPass.h"
 #include "MuT/CodeOptimiser.h"
-#include "MuR/Platform.h"
 #include "MuT/CompilerPrivate.h"
-#include "MuT/ImageDescGenerator.h"
-#include "MuT/TablePrivate.h"
-#include "MuT/NodeLayoutPrivate.h"
-
-#include "MuT/NodeBoolPrivate.h"
-
-#include "MuT/NodeRangeFromScalarPrivate.h"
-
-#include "MuT/NodeProjectorPrivate.h"
-
-#include "MuT/NodeMeshTablePrivate.h"
-#include "MuT/NodeImageTablePrivate.h"
-
-#include "MuT/NodeSurfaceNewPrivate.h"
-#include "MuT/NodeSurfaceEditPrivate.h"
-#include "MuT/NodeSurfaceVariationPrivate.h"
-
-#include "MuT/NodeComponentPrivate.h"
+#include "MuT/ErrorLogPrivate.h"
+#include "MuT/NodeColour.h"
+#include "MuT/NodeComponent.h"
 #include "MuT/NodeComponentNewPrivate.h"
-#include "MuT/NodeComponentEditPrivate.h"
-
-#include "MuT/NodeLODPrivate.h"
-#include "MuT/NodeObjectNewPrivate.h"
-#include "MuT/NodeObjectStatePrivate.h"
-#include "MuT/NodeObjectGroupPrivate.h"
-
+#include "MuT/NodeImage.h"
+#include "MuT/NodeImageFormat.h"
 #include "MuT/NodeImageFormatPrivate.h"
+#include "MuT/NodeImageMipmap.h"
 #include "MuT/NodeImageMipmapPrivate.h"
-
+#include "MuT/NodeLODPrivate.h"
+#include "MuT/NodeMesh.h"
 #include "MuT/NodeMeshClipMorphPlane.h"
 #include "MuT/NodeMeshClipWithMesh.h"
-#include "MuT/NodeMeshClipDeform.h"
+#include "MuT/NodeMeshConstant.h"
+#include "MuT/NodeMeshFormat.h"
+#include "MuT/NodeMeshFragment.h"
 #include "MuT/NodeMeshGeometryOperation.h"
+#include "MuT/NodeMeshInterpolate.h"
+#include "MuT/NodeMeshMorph.h"
 #include "MuT/NodeMeshReshape.h"
-
-#include "MuT/NodeModifierPrivate.h"
+#include "MuT/NodeModifier.h"
+#include "MuT/NodeModifierMeshClipDeform.h"
+#include "MuT/NodeModifierMeshClipDeformPrivate.h"
+#include "MuT/NodeModifierMeshClipMorphPlane.h"
 #include "MuT/NodeModifierMeshClipMorphPlanePrivate.h"
+#include "MuT/NodeModifierMeshClipWithMesh.h"
 #include "MuT/NodeModifierMeshClipWithMeshPrivate.h"
-#include "MuT/NodeModifierMeshClipDeformPrivate.h"
-
-#include "MuT/ASTOpImagePixelFormat.h"
-#include "MuT/ASTOpImageMipmap.h"
-#include "MuT/ASTOpImageCompose.h"
-#include "MuT/ASTOpImagePatch.h"
-#include "MuT/ASTOpConditional.h"
-#include "MuT/ASTOpInstanceAdd.h"
-#include "MuT/ASTOpConstantBool.h"
-#include "MuT/ASTOpConstantResource.h"
-#include "MuT/ASTOpMeshRemoveMask.h"
-#include "MuT/ASTOpMeshClipDeform.h"
-#include "MuT/ASTOpMeshBindShape.h"
-#include "MuT/ASTOpMeshClipMorphPlane.h"
-#include "MuT/ASTOpMeshMaskClipMesh.h"
-#include "MuT/ASTOpAddLOD.h"
-
-#include "MuR/ModelPrivate.h"
-#include "MuR/Operations.h"
-#include "MuR/OpImageCrop.h"
-#include "MuR/SystemPrivate.h"
-
-#include "MuR/MutableTrace.h"
+#include "MuT/NodeModifierPrivate.h"
+#include "MuT/NodeObject.h"
+#include "MuT/NodeObjectGroupPrivate.h"
+#include "MuT/NodeObjectNewPrivate.h"
+#include "MuT/NodeObjectStatePrivate.h"
+#include "MuT/NodePatchImagePrivate.h"
+#include "MuT/NodePatchMesh.h"
+#include "MuT/NodePrivate.h"
+#include "MuT/NodeRange.h"
+#include "MuT/NodeRangeFromScalar.h"
+#include "MuT/NodeScalar.h"
+#include "MuT/NodeScalarConstant.h"
+#include "MuT/NodeSurface.h"
+#include "MuT/NodeSurfaceEdit.h"
+#include "MuT/NodeSurfaceEditPrivate.h"
+#include "MuT/NodeSurfaceNewPrivate.h"
+#include "MuT/NodeSurfaceVariation.h"
+#include "MuT/NodeSurfaceVariationPrivate.h"
+#include "MuT/TablePrivate.h"
+#include "MuT/TaskManager.h"
+#include "Trace/Detail/Channel.h"
 
 
 namespace mu
diff --git a/Engine/Plugins/Experimental/Mutable/Source/MutableTools/Private/MuT/CodeGenerator.h b/Engine/Plugins/Experimental/Mutable/Source/MutableTools/Private/MuT/CodeGenerator.h
--- a/Engine/Plugins/Experimental/Mutable/Source/MutableTools/Private/MuT/CodeGenerator.h
+++ b/Engine/Plugins/Experimental/Mutable/Source/MutableTools/Private/MuT/CodeGenerator.h
@@ -2,128 +2,118 @@
 
 #pragma once
 
-#include "MuT/Visitor.h"
-
+#include "Containers/Map.h"
+#include "HAL/PlatformMath.h"
+#include "Misc/AssertionMacros.h"
+#include "MuR/Image.h"
+#include "MuR/MemoryPrivate.h"
+#include "MuR/Mesh.h"
+#include "MuR/MutableMath.h"
+#include "MuR/MutableMemory.h"
+#include "MuR/Operations.h"
+#include "MuR/Parameters.h"
+#include "MuR/Ptr.h"
+#include "MuR/RefCounted.h"
+#include "MuR/System.h"
+#include "MuT/AST.h"
+#include "MuT/ASTOpSwitch.h"
 #include "MuT/CodeGenerator_FirstPass.h"
-#include "MuT/CompilerPrivate.h"
-#include "MuT/NodeLayoutPrivate.h"
-#include "MuT/ErrorLogPrivate.h"
-#include "MuT/NodeObjectNewPrivate.h"
-#include "MuT/TaskManager.h"
-
-#include "MuT/TablePrivate.h"
-
+#include "MuT/Compiler.h"
+#include "MuT/ErrorLog.h"
+#include "MuT/Node.h"
 #include "MuT/NodeBool.h"
-
-#include "MuT/NodeScalar.h"
-#include "MuT/NodeScalarConstant.h"
-#include "MuT/NodeScalarParameter.h"
-#include "MuT/NodeScalarSwitch.h"
-#include "MuT/NodeScalarEnumParameter.h"
-#include "MuT/NodeScalarCurve.h"
-#include "MuT/NodeScalarArithmeticOperation.h"
-#include "MuT/NodeScalarVariation.h"
-#include "MuT/NodeRangeFromScalar.h"
-#include "MuT/NodeScalarTable.h"
-
-#include "MuT/NodeString.h"
-#include "MuT/NodeStringConstant.h"
-#include "MuT/NodeStringParameter.h"
-
-#include "MuT/NodeColour.h"
-#include "MuT/NodeColourConstant.h"
-#include "MuT/NodeColourParameter.h"
-#include "MuT/NodeColourSwitch.h"
 #include "MuT/NodeColourSampleImage.h"
-#include "MuT/NodeColourFromScalars.h"
-#include "MuT/NodeColourTable.h"
-#include "MuT/NodeColourArithmeticOperation.h"
-#include "MuT/NodeColourVariation.h"
-
-#include "MuT/NodeProjector.h"
-
-#include "MuT/NodeImage.h"
-#include "MuT/NodeImageConstant.h"
-#include "MuT/NodeImageParameter.h"
-#include "MuT/NodeImageDifference.h"
-#include "MuT/NodeImageInterpolate.h"
-#include "MuT/NodeImageInterpolate3.h"
-#include "MuT/NodeImageSaturate.h"
-#include "MuT/NodeImageLuminance.h"
-#include "MuT/NodeImageLayer.h"
-#include "MuT/NodeImageLayerColour.h"
-#include "MuT/NodeImageMultiLayer.h"
-#include "MuT/NodeImageTable.h"
-#include "MuT/NodeImageSwizzle.h"
-#include "MuT/NodeImageFormat.h"
-#include "MuT/NodeImageSelectColour.h"
-#include "MuT/NodeImageColourMap.h"
-#include "MuT/NodeImageGradient.h"
-#include "MuT/NodeImageBinarise.h"
-#include "MuT/NodeImageResize.h"
-#include "MuT/NodeImagePlainColour.h"
-#include "MuT/NodeImageSwitch.h"
-#include "MuT/NodeImageProject.h"
-#include "MuT/NodeImageMipmap.h"
-#include "MuT/NodeImageConditional.h"
-#include "MuT/NodeImageInvert.h"
-#include "MuT/NodeImageVariation.h"
-#include "MuT/NodeImageNormalComposite.h"
-#include "MuT/NodeImageTransform.h"
-
-#include "MuT/NodeMesh.h"
-#include "MuT/NodeMeshConstant.h"
-#include "MuT/NodeMeshSwitch.h"
-#include "MuT/NodeMeshTable.h"
-#include "MuT/NodeMeshInterpolate.h"
-#include "MuT/NodeMeshMorph.h"
-#include "MuT/NodeMeshMakeMorph.h"
-#include "MuT/NodePatchMesh.h"
-#include "MuT/NodeMeshSubtract.h"
-#include "MuT/NodeMeshFormat.h"
-#include "MuT/NodeMeshFragment.h"
-#include "MuT/NodeMeshTransform.h"
-#include "MuT/NodeMeshClipMorphPlane.h"
-#include "MuT/NodeMeshClipWithMesh.h"
-#include "MuT/NodeMeshApplyPose.h"
-#include "MuT/NodeMeshVariation.h"
-#include "MuT/NodeMeshGeometryOperation.h"
-#include "MuT/NodeMeshReshape.h"
-#include "MuT/NodeMeshClipDeform.h"
-
-#include "MuT/NodeComponent.h"
-#include "MuT/NodeComponentNew.h"
 #include "MuT/NodeComponentEdit.h"
-
-#include "MuT/NodeModifier.h"
-#include "MuT/NodeModifierMeshClipMorphPlane.h"
-#include "MuT/NodeModifierMeshClipWithMesh.h"
-
-#include "MuT/NodeSurface.h"
-#include "MuT/NodeSurfaceNew.h"
-#include "MuT/NodeSurfaceEdit.h"
-
-#include "MuT/NodePatchImage.h"
+#include "MuT/NodeComponentNew.h"
+#include "MuT/NodeImageProject.h"
 #include "MuT/NodeLOD.h"
+#include "MuT/NodeLayout.h"
+#include "MuT/NodeMeshGeometryOperation.h"
+#include "MuT/NodeMeshTable.h"
+#include "MuT/NodeObjectGroup.h"
 #include "MuT/NodeObjectNew.h"
 #include "MuT/NodeObjectState.h"
-#include "MuT/NodeObjectGroup.h"
-
-#include "MuT/AST.h"
-#include "MuT/ASTOpParameter.h"
-#include "MuT/ASTOpSwitch.h"
-
-#include "MuR/Operations.h"
-#include "MuR/ModelPrivate.h"
-#include "MuR/ImagePrivate.h"
-
-#include "Containers/Map.h"
+#include "MuT/NodePatchImage.h"
+#include "MuT/NodeProjector.h"
+#include "MuT/NodeString.h"
+#include "MuT/NodeSurfaceNew.h"
+#include "MuT/Table.h"
+#include "MuT/TablePrivate.h"
+#include "MuT/Visitor.h"
+#include "Templates/TypeHash.h"
 
+#include <functional>
+#include <memory>
 #include <shared_mutex>
+#include <utility>
 
 
 namespace mu
 {
+class ASTOpParameter;
+class Layout;
+class NodeColourArithmeticOperation;
+class NodeColourConstant;
+class NodeColourFromScalars;
+class NodeColourParameter;
+class NodeColourSwitch;
+class NodeColourTable;
+class NodeColourVariation;
+class NodeImageBinarise;
+class NodeImageColourMap;
+class NodeImageConditional;
+class NodeImageConstant;
+class NodeImageDifference;
+class NodeImageFormat;
+class NodeImageGradient;
+class NodeImageInterpolate3;
+class NodeImageInterpolate;
+class NodeImageInvert;
+class NodeImageLayer;
+class NodeImageLayerColour;
+class NodeImageLuminance;
+class NodeImageMipmap;
+class NodeImageMultiLayer;
+class NodeImageNormalComposite;
+class NodeImageParameter;
+class NodeImagePlainColour;
+class NodeImageResize;
+class NodeImageSaturate;
+class NodeImageSelectColour;
+class NodeImageSwitch;
+class NodeImageSwizzle;
+class NodeImageTable;
+class NodeImageTransform;
+class NodeImageVariation;
+class NodeMeshApplyPose;
+class NodeMeshClipDeform;
+class NodeMeshClipMorphPlane;
+class NodeMeshClipWithMesh;
+class NodeMeshConstant;
+class NodeMeshFormat;
+class NodeMeshFragment;
+class NodeMeshInterpolate;
+class NodeMeshMakeMorph;
+class NodeMeshMorph;
+class NodeMeshReshape;
+class NodeMeshSubtract;
+class NodeMeshSwitch;
+class NodeMeshTransform;
+class NodeMeshVariation;
+class NodeRange;
+class NodeScalarArithmeticOperation;
+class NodeScalarConstant;
+class NodeScalarCurve;
+class NodeScalarEnumParameter;
+class NodeScalarParameter;
+class NodeScalarSwitch;
+class NodeScalarTable;
+class NodeScalarVariation;
+class NodeStringConstant;
+class NodeStringParameter;
+class TaskManager;
+struct OBJECT_STATE;
+struct PROGRAM;
 
     //---------------------------------------------------------------------------------------------
     //! Code generator
diff --git a/Engine/Plugins/Experimental/Mutable/Source/MutableTools/Private/MuT/CodeGenerator_Bool.cpp b/Engine/Plugins/Experimental/Mutable/Source/MutableTools/Private/MuT/CodeGenerator_Bool.cpp
--- a/Engine/Plugins/Experimental/Mutable/Source/MutableTools/Private/MuT/CodeGenerator_Bool.cpp
+++ b/Engine/Plugins/Experimental/Mutable/Source/MutableTools/Private/MuT/CodeGenerator_Bool.cpp
@@ -1,14 +1,25 @@
 // Copyright Epic Games, Inc. All Rights Reserved.
 
+#include "HAL/PlatformCrt.h"
+#include "Misc/AssertionMacros.h"
+#include "MuR/MemoryPrivate.h"
+#include "MuR/Operations.h"
+#include "MuR/Parameters.h"
+#include "MuR/ParametersPrivate.h"
+#include "MuR/Ptr.h"
+#include "MuR/RefCounted.h"
+#include "MuR/Types.h"
+#include "MuT/AST.h"
+#include "MuT/ASTOpConstantBool.h"
+#include "MuT/ASTOpParameter.h"
 #include "MuT/CodeGenerator.h"
-
-#include "MuR/Platform.h"
-#include "MuT/CompilerPrivate.h"
+#include "MuT/NodeBool.h"
 #include "MuT/NodeBoolPrivate.h"
-#include "MuT/ASTOpConstantBool.h"
+#include "MuT/NodeRange.h"
+#include "map"
 
-#include "MuR/ModelPrivate.h"
-#include "MuR/Operations.h"
+#include <memory>
+#include <utility>
 
 
 using namespace mu;
diff --git a/Engine/Plugins/Experimental/Mutable/Source/MutableTools/Private/MuT/CodeGenerator_Color.cpp b/Engine/Plugins/Experimental/Mutable/Source/MutableTools/Private/MuT/CodeGenerator_Color.cpp
--- a/Engine/Plugins/Experimental/Mutable/Source/MutableTools/Private/MuT/CodeGenerator_Color.cpp
+++ b/Engine/Plugins/Experimental/Mutable/Source/MutableTools/Private/MuT/CodeGenerator_Color.cpp
@@ -1,24 +1,56 @@
 // Copyright Epic Games, Inc. All Rights Reserved.
 
-#include "MuT/CodeGenerator.h"
-
+#include "HAL/PlatformMath.h"
+#include "Misc/AssertionMacros.h"
+#include "MuR/Image.h"
+#include "MuR/MemoryPrivate.h"
+#include "MuR/MutableMath.h"
+#include "MuR/MutableTrace.h"
+#include "MuR/Operations.h"
+#include "MuR/Parameters.h"
+#include "MuR/ParametersPrivate.h"
 #include "MuR/Platform.h"
-#include "MuT/CompilerPrivate.h"
-
-#include "MuT/NodeColourPrivate.h"
+#include "MuR/Ptr.h"
+#include "MuR/RefCounted.h"
+#include "MuR/Types.h"
+#include "MuT/AST.h"
+#include "MuT/ASTOpConditional.h"
+#include "MuT/ASTOpParameter.h"
+#include "MuT/ASTOpSwitch.h"
+#include "MuT/CodeGenerator.h"
+#include "MuT/CodeGenerator_FirstPass.h"
+#include "MuT/ErrorLog.h"
+#include "MuT/ErrorLogPrivate.h"
+#include "MuT/Node.h"
+#include "MuT/NodeColour.h"
+#include "MuT/NodeColourArithmeticOperation.h"
+#include "MuT/NodeColourArithmeticOperationPrivate.h"
+#include "MuT/NodeColourConstant.h"
 #include "MuT/NodeColourConstantPrivate.h"
+#include "MuT/NodeColourFromScalars.h"
+#include "MuT/NodeColourFromScalarsPrivate.h"
+#include "MuT/NodeColourParameter.h"
 #include "MuT/NodeColourParameterPrivate.h"
-#include "MuT/NodeColourSwitchPrivate.h"
+#include "MuT/NodeColourSampleImage.h"
 #include "MuT/NodeColourSampleImagePrivate.h"
-#include "MuT/NodeColourFromScalarsPrivate.h"
+#include "MuT/NodeColourSwitch.h"
+#include "MuT/NodeColourSwitchPrivate.h"
+#include "MuT/NodeColourTable.h"
 #include "MuT/NodeColourTablePrivate.h"
-#include "MuT/NodeColourArithmeticOperationPrivate.h"
+#include "MuT/NodeColourVariation.h"
 #include "MuT/NodeColourVariationPrivate.h"
-
-#include "MuT/ASTOpConditional.h"
-
-#include "MuR/ModelPrivate.h"
-#include "MuR/Operations.h"
+#include "MuT/NodeImage.h"
+#include "MuT/NodeImageProject.h"
+#include "MuT/NodeMeshGeometryOperation.h"
+#include "MuT/NodeMeshTable.h"
+#include "MuT/NodeRange.h"
+#include "MuT/NodeScalar.h"
+#include "MuT/NodeScalarConstant.h"
+#include "MuT/Table.h"
+#include "MuT/TablePrivate.h"
+
+#include <memory>
+#include <utility>
 
 
 using namespace mu;
diff --git a/Engine/Plugins/Experimental/Mutable/Source/MutableTools/Private/MuT/CodeGenerator_FirstPass.cpp b/Engine/Plugins/Experimental/Mutable/Source/MutableTools/Private/MuT/CodeGenerator_FirstPass.cpp
--- a/Engine/Plugins/Experimental/Mutable/Source/MutableTools/Private/MuT/CodeGenerator_FirstPass.cpp
+++ b/Engine/Plugins/Experimental/Mutable/Source/MutableTools/Private/MuT/CodeGenerator_FirstPass.cpp
@@ -2,39 +2,35 @@
 
 #include "MuT/CodeGenerator_FirstPass.h"
 
+#include "HAL/PlatformMath.h"
+#include "Misc/AssertionMacros.h"
+#include "MuR/MutableTrace.h"
+#include "MuR/Operations.h"
+#include "MuR/Parameters.h"
+#include "MuR/ParametersPrivate.h"
 #include "MuR/Platform.h"
+#include "MuT/AST.h"
+#include "MuT/ASTOpConstantBool.h"
+#include "MuT/ASTOpParameter.h"
 #include "MuT/CompilerPrivate.h"
-
-#include "MuT/NodeSurfaceNewPrivate.h"
-#include "MuT/NodeSurfaceEditPrivate.h"
-#include "MuT/NodeSurfaceVariationPrivate.h"
-
-#include "MuT/NodeComponentPrivate.h"
-#include "MuT/NodeComponentNewPrivate.h"
+#include "MuT/ErrorLog.h"
+#include "MuT/ErrorLogPrivate.h"
+#include "MuT/NodeComponent.h"
 #include "MuT/NodeComponentEditPrivate.h"
-
+#include "MuT/NodeComponentNewPrivate.h"
 #include "MuT/NodeLODPrivate.h"
-#include "MuT/NodeObjectNewPrivate.h"
-#include "MuT/NodeObjectStatePrivate.h"
-#include "MuT/NodeObjectGroupPrivate.h"
-
-#include "MuT/NodeMeshClipMorphPlane.h"
-#include "MuT/NodeMeshClipWithMesh.h"
-
-#include "MuT/NodeModifierPrivate.h"
+#include "MuT/NodeModifierMeshClipDeformPrivate.h"
 #include "MuT/NodeModifierMeshClipMorphPlanePrivate.h"
 #include "MuT/NodeModifierMeshClipWithMeshPrivate.h"
-#include "MuT/NodeModifierMeshClipDeformPrivate.h"
-
-#include "MuT/AST.h"
-#include "MuT/ASTOpConstantBool.h"
-#include "MuT/ASTOpParameter.h"
-
-#include "MuR/ModelPrivate.h"
-#include "MuR/MutableTrace.h"
-#include "MuR/Operations.h"
-#include "MuR/SystemPrivate.h"
-
+#include "MuT/NodeObject.h"
+#include "MuT/NodeObjectGroupPrivate.h"
+#include "MuT/NodeObjectNewPrivate.h"
+#include "MuT/NodeObjectStatePrivate.h"
+#include "MuT/NodePrivate.h"
+#include "MuT/NodeSurface.h"
+#include "MuT/NodeSurfaceEditPrivate.h"
+#include "MuT/NodeSurfaceNewPrivate.h"
+#include "MuT/NodeSurfaceVariationPrivate.h"
 
 namespace mu
 {
diff --git a/Engine/Plugins/Experimental/Mutable/Source/MutableTools/Private/MuT/CodeGenerator_FirstPass.h b/Engine/Plugins/Experimental/Mutable/Source/MutableTools/Private/MuT/CodeGenerator_FirstPass.h
--- a/Engine/Plugins/Experimental/Mutable/Source/MutableTools/Private/MuT/CodeGenerator_FirstPass.h
+++ b/Engine/Plugins/Experimental/Mutable/Source/MutableTools/Private/MuT/CodeGenerator_FirstPass.h
@@ -2,47 +2,37 @@
 
 #pragma once
 
-#include "MuT/Visitor.h"
-
+#include "MuR/MemoryPrivate.h"
+#include "MuR/MutableMemory.h"
+#include "MuR/Ptr.h"
+#include "MuR/RefCounted.h"
 #include "MuT/AST.h"
-#include "MuT/CompilerPrivate.h"
-#include "MuT/NodeLayoutPrivate.h"
-#include "MuT/ErrorLogPrivate.h"
-#include "MuT/NodeObjectNewPrivate.h"
-
-#include "MuT/Table.h"
-
-#include "MuT/NodeBool.h"
-
-#include "MuT/NodeComponent.h"
-#include "MuT/NodeComponentNew.h"
+#include "MuT/Compiler.h"
+#include "MuT/Node.h"
 #include "MuT/NodeComponentEdit.h"
-
+#include "MuT/NodeComponentNew.h"
+#include "MuT/NodeLOD.h"
 #include "MuT/NodeModifier.h"
+#include "MuT/NodeModifierMeshClipDeform.h"
 #include "MuT/NodeModifierMeshClipMorphPlane.h"
 #include "MuT/NodeModifierMeshClipWithMesh.h"
-#include "MuT/NodeModifierMeshClipDeform.h"
-
-#include "MuT/NodeSurface.h"
-#include "MuT/NodeSurfaceNew.h"
-#include "MuT/NodeSurfaceEdit.h"
-#include "MuT/NodeSurfaceVariation.h"
-
-#include "MuT/NodePatchImage.h"
-#include "MuT/NodeLOD.h"
+#include "MuT/NodeObjectGroup.h"
 #include "MuT/NodeObjectNew.h"
 #include "MuT/NodeObjectState.h"
-#include "MuT/NodeObjectGroup.h"
 #include "MuT/NodePatchMesh.h"
+#include "MuT/NodeSurfaceEdit.h"
+#include "MuT/NodeSurfaceNew.h"
+#include "MuT/NodeSurfaceVariation.h"
+#include "MuT/Visitor.h"
 
-#include "MuR/Operations.h"
-#include "MuR/ModelPrivate.h"
-#include "MuR/ImagePrivate.h"
-
+#include <stdint.h>
 #include <memory>
+#include <utility>
 
 namespace mu
 {
+class Layout;
+struct OBJECT_STATE;
 
 
     //! Store the results of the code generation of a mesh.
diff --git a/Engine/Plugins/Experimental/Mutable/Source/MutableTools/Private/MuT/CodeGenerator_Image.cpp b/Engine/Plugins/Experimental/Mutable/Source/MutableTools/Private/MuT/CodeGenerator_Image.cpp
--- a/Engine/Plugins/Experimental/Mutable/Source/MutableTools/Private/MuT/CodeGenerator_Image.cpp
+++ b/Engine/Plugins/Experimental/Mutable/Source/MutableTools/Private/MuT/CodeGenerator_Image.cpp
@@ -1,57 +1,113 @@
 // Copyright Epic Games, Inc. All Rights Reserved.
 
-#include "MuT/CodeGenerator.h"
-
+#include "HAL/PlatformMath.h"
+#include "Misc/AssertionMacros.h"
+#include "MuR/Image.h"
+#include "MuR/ImagePrivate.h"
+#include "MuR/MemoryPrivate.h"
+#include "MuR/Mesh.h"
+#include "MuR/MutableMath.h"
+#include "MuR/MutableTrace.h"
+#include "MuR/OpImageCrop.h"
+#include "MuR/OpImageProject.h"
+#include "MuR/Operations.h"
+#include "MuR/Parameters.h"
+#include "MuR/ParametersPrivate.h"
 #include "MuR/Platform.h"
+#include "MuR/Ptr.h"
+#include "MuR/RefCounted.h"
+#include "MuR/Serialisation.h"
+#include "MuR/System.h"
+#include "MuT/AST.h"
+#include "MuT/ASTOpConditional.h"
+#include "MuT/ASTOpConstantResource.h"
+#include "MuT/ASTOpImageMipmap.h"
+#include "MuT/ASTOpImageMultiLayer.h"
+#include "MuT/ASTOpImageNormalComposite.h"
+#include "MuT/ASTOpImagePixelFormat.h"
+#include "MuT/ASTOpImageTransform.h"
+#include "MuT/ASTOpMeshExtractLayoutBlocks.h"
+#include "MuT/ASTOpMeshFormat.h"
+#include "MuT/ASTOpParameter.h"
+#include "MuT/ASTOpSwitch.h"
+#include "MuT/CodeGenerator.h"
+#include "MuT/CodeGenerator_FirstPass.h"
+#include "MuT/Compiler.h"
 #include "MuT/CompilerPrivate.h"
+#include "MuT/ErrorLog.h"
+#include "MuT/ErrorLogPrivate.h"
 #include "MuT/ImageDescGenerator.h"
-#include "MuT/TablePrivate.h"
-#include "MuT/NodeLayoutPrivate.h"
-
-#include "MuT/NodeImagePrivate.h"
+#include "MuT/Node.h"
+#include "MuT/NodeBool.h"
+#include "MuT/NodeColour.h"
+#include "MuT/NodeColourFromScalars.h"
+#include "MuT/NodeColourSampleImage.h"
+#include "MuT/NodeImage.h"
+#include "MuT/NodeImageBinarise.h"
+#include "MuT/NodeImageBinarisePrivate.h"
+#include "MuT/NodeImageColourMap.h"
+#include "MuT/NodeImageColourMapPrivate.h"
+#include "MuT/NodeImageConditional.h"
+#include "MuT/NodeImageConditionalPrivate.h"
+#include "MuT/NodeImageConstant.h"
 #include "MuT/NodeImageConstantPrivate.h"
+#include "MuT/NodeImageDifference.h"
 #include "MuT/NodeImageDifferencePrivate.h"
-#include "MuT/NodeImageInterpolatePrivate.h"
+#include "MuT/NodeImageFormat.h"
+#include "MuT/NodeImageFormatPrivate.h"
+#include "MuT/NodeImageGradient.h"
+#include "MuT/NodeImageGradientPrivate.h"
+#include "MuT/NodeImageInterpolate.h"
+#include "MuT/NodeImageInterpolate3.h"
 #include "MuT/NodeImageInterpolate3Private.h"
-#include "MuT/NodeImageSaturatePrivate.h"
-#include "MuT/NodeImageLuminancePrivate.h"
-#include "MuT/NodeImageLayerPrivate.h"
+#include "MuT/NodeImageInterpolatePrivate.h"
+#include "MuT/NodeImageInvert.h"
+#include "MuT/NodeImageInvertPrivate.h"
+#include "MuT/NodeImageLayer.h"
+#include "MuT/NodeImageLayerColour.h"
 #include "MuT/NodeImageLayerColourPrivate.h"
+#include "MuT/NodeImageLayerPrivate.h"
+#include "MuT/NodeImageLuminance.h"
+#include "MuT/NodeImageLuminancePrivate.h"
+#include "MuT/NodeImageMipmap.h"
+#include "MuT/NodeImageMipmapPrivate.h"
+#include "MuT/NodeImageMultiLayer.h"
 #include "MuT/NodeImageMultiLayerPrivate.h"
-#include "MuT/NodeImageTablePrivate.h"
-#include "MuT/NodeImageSwizzlePrivate.h"
-#include "MuT/NodeImageFormatPrivate.h"
-#include "MuT/NodeImageSelectColourPrivate.h"
-#include "MuT/NodeImageColourMapPrivate.h"
-#include "MuT/NodeImageGradientPrivate.h"
-#include "MuT/NodeImageBinarisePrivate.h"
-#include "MuT/NodeImageResizePrivate.h"
+#include "MuT/NodeImageNormalComposite.h"
+#include "MuT/NodeImageNormalCompositePrivate.h"
+#include "MuT/NodeImageParameter.h"
+#include "MuT/NodeImageParameterPrivate.h"
+#include "MuT/NodeImagePlainColour.h"
 #include "MuT/NodeImagePlainColourPrivate.h"
-#include "MuT/NodeImageSwitchPrivate.h"
+#include "MuT/NodeImageProject.h"
 #include "MuT/NodeImageProjectPrivate.h"
-#include "MuT/NodeImageMipmapPrivate.h"
-#include "MuT/NodeImageConditionalPrivate.h"
-#include "MuT/NodeImageParameterPrivate.h"
-#include "MuT/NodeImageInvertPrivate.h"
-#include "MuT/NodeImageVariationPrivate.h"
-#include "MuT/NodeImageNormalCompositePrivate.h"
+#include "MuT/NodeImageResize.h"
+#include "MuT/NodeImageResizePrivate.h"
+#include "MuT/NodeImageSaturate.h"
+#include "MuT/NodeImageSaturatePrivate.h"
+#include "MuT/NodeImageSelectColour.h"
+#include "MuT/NodeImageSelectColourPrivate.h"
+#include "MuT/NodeImageSwitch.h"
+#include "MuT/NodeImageSwitchPrivate.h"
+#include "MuT/NodeImageSwizzle.h"
+#include "MuT/NodeImageSwizzlePrivate.h"
+#include "MuT/NodeImageTable.h"
+#include "MuT/NodeImageTablePrivate.h"
+#include "MuT/NodeImageTransform.h"
 #include "MuT/NodeImageTransformPrivate.h"
+#include "MuT/NodeImageVariation.h"
+#include "MuT/NodeImageVariationPrivate.h"
+#include "MuT/NodeMeshGeometryOperation.h"
+#include "MuT/NodeMeshTable.h"
+#include "MuT/NodeProjector.h"
+#include "MuT/NodeRange.h"
+#include "MuT/NodeScalar.h"
+#include "MuT/NodeScalarConstant.h"
+#include "MuT/Table.h"
+#include "MuT/TablePrivate.h"
 
-#include "MuT/ASTOpMeshFormat.h"
-#include "MuT/ASTOpImagePixelFormat.h"
-#include "MuT/ASTOpImageMipmap.h"
-#include "MuT/ASTOpConstantResource.h"
-#include "MuT/ASTOpConditional.h"
-#include "MuT/ASTOpImageMultiLayer.h"
-#include "MuT/ASTOpImageTransform.h"
-#include "MuT/ASTOpImageNormalComposite.h"
-#include "MuT/ASTOpMeshExtractLayoutBlocks.h"
-
-#include "MuR/ModelPrivate.h"
-#include "MuR/MutableTrace.h"
-#include "MuR/Operations.h"
-#include "MuR/OpImageCrop.h"
-#include "MuR/OpImageProject.h"
+#include <memory>
+#include <utility>
 
 
 namespace mu
