Ran IWYU on Mutable plugin

diff --git a/Engine/Plugins/Experimental/Mutable/Source/MutableTools/Private/MuT/CodeGenerator_Mesh.cpp b/Engine/Plugins/Experimental/Mutable/Source/MutableTools/Private/MuT/CodeGenerator_Mesh.cpp
--- a/Engine/Plugins/Experimental/Mutable/Source/MutableTools/Private/MuT/CodeGenerator_Mesh.cpp
+++ b/Engine/Plugins/Experimental/Mutable/Source/MutableTools/Private/MuT/CodeGenerator_Mesh.cpp
@@ -1,53 +1,88 @@
 // Copyright Epic Games, Inc. All Rights Reserved.
 
-#include "MuT/CodeGenerator.h"
-
+#include "Containers/Array.h"
+#include "Math/IntPoint.h"
+#include "Math/UnrealMathSSE.h"
+#include "Misc/AssertionMacros.h"
+#include "MuR/Layout.h"
+#include "MuR/MemoryPrivate.h"
+#include "MuR/Mesh.h"
+#include "MuR/MeshBufferSet.h"
+#include "MuR/MutableMath.h"
+#include "MuR/Operations.h"
+#include "MuR/ParametersPrivate.h"
 #include "MuR/Platform.h"
+#include "MuR/Ptr.h"
+#include "MuR/RefCounted.h"
+#include "MuR/System.h"
+#include "MuT/AST.h"
+#include "MuT/ASTOpConditional.h"
+#include "MuT/ASTOpConstantResource.h"
+#include "MuT/ASTOpMeshApplyPose.h"
+#include "MuT/ASTOpMeshApplyShape.h"
+#include "MuT/ASTOpMeshBindShape.h"
+#include "MuT/ASTOpMeshClipDeform.h"
+#include "MuT/ASTOpMeshClipMorphPlane.h"
+#include "MuT/ASTOpMeshExtractLayoutBlocks.h"
+#include "MuT/ASTOpMeshFormat.h"
+#include "MuT/ASTOpMeshGeometryOperation.h"
+#include "MuT/ASTOpMeshMorphReshape.h"
+#include "MuT/ASTOpMeshRemapIndices.h"
+#include "MuT/ASTOpMeshTransform.h"
+#include "MuT/ASTOpSwitch.h"
+#include "MuT/CodeGenerator.h"
+#include "MuT/CodeGenerator_FirstPass.h"
 #include "MuT/CompilerPrivate.h"
-#include "MuT/TablePrivate.h"
+#include "MuT/ErrorLog.h"
+#include "MuT/ErrorLogPrivate.h"
+#include "MuT/NodeImageProject.h"
+#include "MuT/NodeLayout.h"
 #include "MuT/NodeLayoutPrivate.h"
-
-#include "MuT/NodeMeshPrivate.h"
+#include "MuT/NodeMesh.h"
+#include "MuT/NodeMeshApplyPose.h"
+#include "MuT/NodeMeshApplyPosePrivate.h"
+#include "MuT/NodeMeshClipDeform.h"
+#include "MuT/NodeMeshClipDeformPrivate.h"
+#include "MuT/NodeMeshClipMorphPlane.h"
+#include "MuT/NodeMeshClipMorphPlanePrivate.h"
+#include "MuT/NodeMeshClipWithMesh.h"
+#include "MuT/NodeMeshClipWithMeshPrivate.h"
+#include "MuT/NodeMeshConstant.h"
 #include "MuT/NodeMeshConstantPrivate.h"
-#include "MuT/NodeMeshSwitchPrivate.h"
-#include "MuT/NodeMeshTablePrivate.h"
-#include "MuT/NodeMeshInterpolatePrivate.h"
-#include "MuT/NodeMeshMorphPrivate.h"
-#include "MuT/NodeMeshMakeMorphPrivate.h"
-#include "MuT/NodePatchMeshPrivate.h"
-#include "MuT/NodeMeshSubtractPrivate.h"
+#include "MuT/NodeMeshFormat.h"
 #include "MuT/NodeMeshFormatPrivate.h"
+#include "MuT/NodeMeshFragment.h"
 #include "MuT/NodeMeshFragmentPrivate.h"
-#include "MuT/NodeMeshTransformPrivate.h"
-#include "MuT/NodeMeshClipMorphPlanePrivate.h"
-#include "MuT/NodeMeshClipWithMeshPrivate.h"
-#include "MuT/NodeMeshApplyPosePrivate.h"
-#include "MuT/NodeMeshVariationPrivate.h"
+#include "MuT/NodeMeshGeometryOperation.h"
 #include "MuT/NodeMeshGeometryOperationPrivate.h"
+#include "MuT/NodeMeshInterpolate.h"
+#include "MuT/NodeMeshInterpolatePrivate.h"
+#include "MuT/NodeMeshMakeMorph.h"
+#include "MuT/NodeMeshMakeMorphPrivate.h"
+#include "MuT/NodeMeshMorph.h"
+#include "MuT/NodeMeshMorphPrivate.h"
+#include "MuT/NodeMeshReshape.h"
 #include "MuT/NodeMeshReshapePrivate.h"
-#include "MuT/NodeMeshClipDeformPrivate.h"
+#include "MuT/NodeMeshSubtract.h"
+#include "MuT/NodeMeshSwitch.h"
+#include "MuT/NodeMeshSwitchPrivate.h"
+#include "MuT/NodeMeshTable.h"
+#include "MuT/NodeMeshTablePrivate.h"
+#include "MuT/NodeMeshTransform.h"
+#include "MuT/NodeMeshTransformPrivate.h"
+#include "MuT/NodeMeshVariation.h"
+#include "MuT/NodeMeshVariationPrivate.h"
+#include "MuT/NodeScalar.h"
+#include "MuT/Table.h"
+#include "MuT/TablePrivate.h"
 
-#include "MuT/ASTOpSwitch.h"
-#include "MuT/ASTOpConditional.h"
-#include "MuT/ASTOpConstantResource.h"
-#include "MuT/ASTOpMeshFormat.h"
-#include "MuT/ASTOpMeshRemapIndices.h"
-#include "MuT/ASTOpMeshMorphReshape.h"
-#include "MuT/ASTOpMeshApplyShape.h"
-#include "MuT/ASTOpMeshBindShape.h"
-#include "MuT/ASTOpMeshExtractLayoutBlocks.h"
-#include "MuT/ASTOpMeshTransform.h"
-#include "MuT/ASTOpMeshClipMorphPlane.h"
-#include "MuT/ASTOpMeshGeometryOperation.h"
-#include "MuT/ASTOpMeshApplyPose.h"
-#include "MuT/ASTOpMeshClipDeform.h"
-
-#include "MuR/ModelPrivate.h"
-#include "MuR/Operations.h"
+#include <memory>
+#include <utility>
 
 
 namespace mu
 {
+class Node;
 
     //---------------------------------------------------------------------------------------------
     void CodeGenerator::PrepareForLayout( LayoutPtrConst pSourceLayout,
diff --git a/Engine/Plugins/Experimental/Mutable/Source/MutableTools/Private/MuT/CodeGenerator_Projector.cpp b/Engine/Plugins/Experimental/Mutable/Source/MutableTools/Private/MuT/CodeGenerator_Projector.cpp
--- a/Engine/Plugins/Experimental/Mutable/Source/MutableTools/Private/MuT/CodeGenerator_Projector.cpp
+++ b/Engine/Plugins/Experimental/Mutable/Source/MutableTools/Private/MuT/CodeGenerator_Projector.cpp
@@ -1,15 +1,28 @@
 // Copyright Epic Games, Inc. All Rights Reserved.
 
+#include "HAL/PlatformCrt.h"
+#include "Misc/AssertionMacros.h"
+#include "MuR/MemoryPrivate.h"
+#include "MuR/MutableMath.h"
+#include "MuR/Operations.h"
+#include "MuR/Parameters.h"
+#include "MuR/ParametersPrivate.h"
+#include "MuR/Ptr.h"
+#include "MuR/RefCounted.h"
+#include "MuR/Types.h"
+#include "MuT/AST.h"
+#include "MuT/ASTOpConstantProjector.h"
+#include "MuT/ASTOpParameter.h"
 #include "MuT/CodeGenerator.h"
-
-#include "MuR/Platform.h"
-#include "MuT/CompilerPrivate.h"
-
+#include "MuT/ErrorLog.h"
+#include "MuT/ErrorLogPrivate.h"
+#include "MuT/NodeProjector.h"
 #include "MuT/NodeProjectorPrivate.h"
-#include "MuT/ASTOpConstantProjector.h"
+#include "MuT/NodeRange.h"
+#include "map"
 
-#include "MuR/ModelPrivate.h"
-#include "MuR/Operations.h"
+#include <memory>
+#include <utility>
 
 
 using namespace mu;
diff --git a/Engine/Plugins/Experimental/Mutable/Source/MutableTools/Private/MuT/CodeGenerator_Scalar.cpp b/Engine/Plugins/Experimental/Mutable/Source/MutableTools/Private/MuT/CodeGenerator_Scalar.cpp
--- a/Engine/Plugins/Experimental/Mutable/Source/MutableTools/Private/MuT/CodeGenerator_Scalar.cpp
+++ b/Engine/Plugins/Experimental/Mutable/Source/MutableTools/Private/MuT/CodeGenerator_Scalar.cpp
@@ -1,25 +1,54 @@
 // Copyright Epic Games, Inc. All Rights Reserved.
 
-#include "MuT/CodeGenerator.h"
-
+#include "HAL/PlatformMath.h"
+#include "Misc/AssertionMacros.h"
+#include "MuR/Image.h"
+#include "MuR/MemoryPrivate.h"
+#include "MuR/MutableMath.h"
+#include "MuR/Operations.h"
+#include "MuR/Parameters.h"
+#include "MuR/ParametersPrivate.h"
 #include "MuR/Platform.h"
-#include "MuT/CompilerPrivate.h"
-
-#include "MuT/NodeScalarPrivate.h"
+#include "MuR/Ptr.h"
+#include "MuR/RefCounted.h"
+#include "MuR/Types.h"
+#include "MuT/AST.h"
+#include "MuT/ASTOpConditional.h"
+#include "MuT/ASTOpParameter.h"
+#include "MuT/ASTOpScalarCurve.h"
+#include "MuT/ASTOpSwitch.h"
+#include "MuT/CodeGenerator.h"
+#include "MuT/CodeGenerator_FirstPass.h"
+#include "MuT/ErrorLog.h"
+#include "MuT/ErrorLogPrivate.h"
+#include "MuT/NodeImage.h"
+#include "MuT/NodeMeshGeometryOperation.h"
+#include "MuT/NodeMeshTable.h"
+#include "MuT/NodeRange.h"
+#include "MuT/NodeScalar.h"
+#include "MuT/NodeScalarArithmeticOperation.h"
+#include "MuT/NodeScalarArithmeticOperationPrivate.h"
+#include "MuT/NodeScalarConstant.h"
 #include "MuT/NodeScalarConstantPrivate.h"
+#include "MuT/NodeScalarCurve.h"
+#include "MuT/NodeScalarCurvePrivate.h"
+#include "MuT/NodeScalarEnumParameter.h"
+#include "MuT/NodeScalarEnumParameterPrivate.h"
+#include "MuT/NodeScalarParameter.h"
 #include "MuT/NodeScalarParameterPrivate.h"
+#include "MuT/NodeScalarSwitch.h"
 #include "MuT/NodeScalarSwitchPrivate.h"
-#include "MuT/NodeScalarEnumParameterPrivate.h"
-#include "MuT/NodeScalarCurvePrivate.h"
-#include "MuT/NodeScalarArithmeticOperationPrivate.h"
-#include "MuT/NodeScalarVariationPrivate.h"
+#include "MuT/NodeScalarTable.h"
 #include "MuT/NodeScalarTablePrivate.h"
+#include "MuT/NodeScalarVariation.h"
+#include "MuT/NodeScalarVariationPrivate.h"
+#include "MuT/Table.h"
+#include "MuT/TablePrivate.h"
 
-#include "MuT/ASTOpScalarCurve.h"
-#include "MuT/ASTOpConditional.h"
+#include <memory>
+#include <utility>
 
-#include "MuR/ModelPrivate.h"
-#include "MuR/Operations.h"
+namespace mu { class Node; }
 
 
 using namespace mu;
diff --git a/Engine/Plugins/Experimental/Mutable/Source/MutableTools/Private/MuT/CodeGenerator_SecondPass.cpp b/Engine/Plugins/Experimental/Mutable/Source/MutableTools/Private/MuT/CodeGenerator_SecondPass.cpp
--- a/Engine/Plugins/Experimental/Mutable/Source/MutableTools/Private/MuT/CodeGenerator_SecondPass.cpp
+++ b/Engine/Plugins/Experimental/Mutable/Source/MutableTools/Private/MuT/CodeGenerator_SecondPass.cpp
@@ -1,39 +1,14 @@
 // Copyright Epic Games, Inc. All Rights Reserved.
 
 #include "MuT/CodeGenerator_SecondPass.h"
-#include "MuT/CodeGenerator_FirstPass.h"
-#include "MuT/CodeGenerator.h"
-
-#include "MuR/Platform.h"
-#include "MuT/CompilerPrivate.h"
-
-#include "MuT/NodeSurfaceNewPrivate.h"
-#include "MuT/NodeSurfaceEditPrivate.h"
-
-#include "MuT/NodeComponentPrivate.h"
-#include "MuT/NodeComponentNewPrivate.h"
-#include "MuT/NodeComponentEditPrivate.h"
-
-#include "MuT/NodeLODPrivate.h"
-#include "MuT/NodeObjectNewPrivate.h"
-#include "MuT/NodeObjectStatePrivate.h"
-#include "MuT/NodeObjectGroupPrivate.h"
-
-#include "MuT/NodeMeshClipMorphPlane.h"
-#include "MuT/NodeMeshClipWithMesh.h"
-
-#include "MuT/NodeModifierPrivate.h"
-#include "MuT/NodeModifierMeshClipMorphPlanePrivate.h"
-#include "MuT/NodeModifierMeshClipWithMeshPrivate.h"
-#include "MuT/NodeComponentNewPrivate.h"
-
-#include "MuT/ASTOpConstantBool.h"
-
-#include "MuR/ModelPrivate.h"
-#include "MuR/Operations.h"
-#include "MuR/SystemPrivate.h"
 
+#include "Misc/AssertionMacros.h"
 #include "MuR/MutableTrace.h"
+#include "MuR/Operations.h"
+#include "MuR/Types.h"
+#include "MuT/ASTOpConstantBool.h"
+#include "MuT/ErrorLog.h"
+#include "MuT/CodeGenerator_FirstPass.h"
 
 using namespace mu;
 
diff --git a/Engine/Plugins/Experimental/Mutable/Source/MutableTools/Private/MuT/CodeGenerator_SecondPass.h b/Engine/Plugins/Experimental/Mutable/Source/MutableTools/Private/MuT/CodeGenerator_SecondPass.h
--- a/Engine/Plugins/Experimental/Mutable/Source/MutableTools/Private/MuT/CodeGenerator_SecondPass.h
+++ b/Engine/Plugins/Experimental/Mutable/Source/MutableTools/Private/MuT/CodeGenerator_SecondPass.h
@@ -2,41 +2,21 @@
 
 #pragma once
 
-#include "MuT/Visitor.h"
-#include "MuT/CodeGenerator_FirstPass.h"
-
+#include "HAL/PlatformCrt.h"
+#include "MuR/MemoryPrivate.h"
+#include "MuR/MutableMemory.h"
+#include "MuR/Ptr.h"
+#include "MuR/RefCounted.h"
 #include "MuT/AST.h"
-#include "MuT/CompilerPrivate.h"
-#include "MuT/NodeLayoutPrivate.h"
-#include "MuT/ErrorLogPrivate.h"
-#include "MuT/NodeObjectNewPrivate.h"
-
-#include "MuT/Table.h"
-
-#include "MuT/NodeBool.h"
-
-#include "MuT/NodeComponent.h"
-#include "MuT/NodeComponentNew.h"
-#include "MuT/NodeComponentEdit.h"
-
-#include "MuT/NodeModifier.h"
-#include "MuT/NodeModifierMeshClipMorphPlane.h"
-#include "MuT/NodeModifierMeshClipWithMesh.h"
-
-#include "MuT/NodeSurface.h"
-#include "MuT/NodeSurfaceNew.h"
-#include "MuT/NodeSurfaceEdit.h"
-
-#include "MuT/NodePatchImage.h"
-#include "MuT/NodeLOD.h"
-#include "MuT/NodeObjectNew.h"
-#include "MuT/NodeObjectState.h"
-#include "MuT/NodeObjectGroup.h"
-#include "MuT/NodePatchMesh.h"
-
-#include "MuR/Operations.h"
-#include "MuR/ModelPrivate.h"
-#include "MuR/ImagePrivate.h"
+#include "MuT/CodeGenerator_FirstPass.h"
+#include "MuT/Compiler.h"
+#include "MuT/Node.h"
+#include "map"
+#include "set"
+
+#include <functional>
+#include <memory>
+#include <utility>
 
 
 namespace mu
diff --git a/Engine/Plugins/Experimental/Mutable/Source/MutableTools/Private/MuT/CodeGenerator_String.cpp b/Engine/Plugins/Experimental/Mutable/Source/MutableTools/Private/MuT/CodeGenerator_String.cpp
--- a/Engine/Plugins/Experimental/Mutable/Source/MutableTools/Private/MuT/CodeGenerator_String.cpp
+++ b/Engine/Plugins/Experimental/Mutable/Source/MutableTools/Private/MuT/CodeGenerator_String.cpp
@@ -1,18 +1,29 @@
 // Copyright Epic Games, Inc. All Rights Reserved.
 
+#include "HAL/PlatformCrt.h"
+#include "HAL/UnrealMemory.h"
+#include "Misc/AssertionMacros.h"
+#include "MuR/MemoryPrivate.h"
+#include "MuR/Operations.h"
+#include "MuR/Parameters.h"
+#include "MuR/ParametersPrivate.h"
+#include "MuR/Ptr.h"
+#include "MuR/RefCounted.h"
+#include "MuR/Types.h"
+#include "MuT/AST.h"
+#include "MuT/ASTOpConstantString.h"
+#include "MuT/ASTOpParameter.h"
 #include "MuT/CodeGenerator.h"
-
-#include "MuR/Platform.h"
-#include "MuT/CompilerPrivate.h"
-
-#include "MuT/NodeStringPrivate.h"
+#include "MuT/NodeRange.h"
+#include "MuT/NodeString.h"
+#include "MuT/NodeStringConstant.h"
 #include "MuT/NodeStringConstantPrivate.h"
+#include "MuT/NodeStringParameter.h"
 #include "MuT/NodeStringParameterPrivate.h"
+#include "map"
 
-#include "MuT/ASTOpConstantString.h"
-
-#include "MuR/ModelPrivate.h"
-#include "MuR/Operations.h"
+#include <memory>
+#include <utility>
 
 
 using namespace mu;
diff --git a/Engine/Plugins/Experimental/Mutable/Source/MutableTools/Private/MuT/CodeOptimiser_GPU.cpp b/Engine/Plugins/Experimental/Mutable/Source/MutableTools/Private/MuT/CodeOptimiser_GPU.cpp
--- a/Engine/Plugins/Experimental/Mutable/Source/MutableTools/Private/MuT/CodeOptimiser_GPU.cpp
+++ b/Engine/Plugins/Experimental/Mutable/Source/MutableTools/Private/MuT/CodeOptimiser_GPU.cpp
@@ -1,1111 +1,6 @@
 // Copyright Epic Games, Inc. All Rights Reserved.
 
 #include "MuT/CodeOptimiser_GPU.h"
-#include "MuT/CodeOptimiser.h"
-#include "MuT/CodeGenerator.h"
-#include "MuT/ModelReportPrivate.h"
-
-#include "MuR/CodeVisitor.h"
-#include "MuR/ModelPrivate.h"
-#include "MuR/SystemPrivate.h"
-#include "MuR/Operations.h"
 
 #define MUTABLE_GPU_CODE_BUFFER_SIZE	512
 
-
-namespace mu
-{
-
-    //---------------------------------------------------------------------------------------------
-    //---------------------------------------------------------------------------------------------
-    //---------------------------------------------------------------------------------------------
-//    class GPUSubtreeGenerator
-//    {
-//    public:
-
-//        virtual ~GPUSubtreeGenerator() {}
-
-//        //! Generate a full program to run internally an expression subtree on the GPU.
-//        virtual GPU_PROGRAM GenerateProgram( OP::ADDRESS at,
-//                              PROGRAM& program,
-//                              int state,
-//                              const GPU_PLATFORM_PROPS& platformOptions,
-//                              const GPU_STATE_PROPS& stateOptions ) = 0;
-
-//        //! Generate a method to run an expression subtree on the GPU externally (integrated on an
-//        //! engine shader).
-//        virtual string GenerateImage( const string& imageName,
-//                                      OP::ADDRESS at,
-//                                      PROGRAM& program,
-//                                      int state,
-//                                      const GPU_PLATFORM_PROPS& platformOptions,
-//                                      const GPU_STATE_PROPS& stateOptions ) = 0;
-
-//        virtual void Clear()
-//        {
-//            m_floatParams.clear();
-//            m_vectorParams.clear();
-//            m_imageParams.clear();
-//            m_currentVar = "";
-//            m_expressionCode = "";
-//        }
-
-//        //! Additional result data
-//        vector<GPU_PROGRAM::PARAM_DATA> m_imageParams;
-//        vector<GPU_PROGRAM::PARAM_DATA> m_vectorParams;
-//        vector<GPU_PROGRAM::PARAM_DATA> m_floatParams;
-//        string m_currentImageName;
-//        string m_currentVar;
-//        string m_expressionCode;
-//    };
-
-
-    //---------------------------------------------------------------------------------------------
-    //---------------------------------------------------------------------------------------------
-    //---------------------------------------------------------------------------------------------
-//    class GPUSubtreeGenerator_GLSL : public GPUSubtreeGenerator, public ConstCodeVisitor
-//    {
-//    public:
-
-//        GPUSubtreeGenerator_GLSL()
-//        {
-//            m_currentLayoutBlock = -1;
-//        }
-
-
-//        //-----------------------------------------------------------------------------------------
-//        // GPUSubtreeGenerator interface
-//        //-----------------------------------------------------------------------------------------
-//        virtual void Clear() override
-//        {
-//            GPUSubtreeGenerator::Clear();
-//            m_generated.clear();
-//            m_usedVars.clear();
-//            m_currentTexCoords = MUTABLE_GPU_CODE_UVS;
-//        }
-
-
-//        //-----------------------------------------------------------------------------------------
-//        string GenerateImage( const string& imageName,
-//                              OP::ADDRESS at,
-//                              PROGRAM& program,
-//                              int state,
-//                              const GPU_PLATFORM_PROPS& platformOptions,
-//                              const GPU_STATE_PROPS& //stateOptions
-//                              ) override
-//        {
-//            check( platformOptions.type == GPU_GL4 );
-//            (void)platformOptions;
-//            m_state = state;
-//            m_currentImageName = imageName;
-
-//            // Clear state
-//            Clear();
-
-//            Traverse( at, program, false );
-
-//            string imageCode = "// begin mutable image["+imageName+"]\n";
-
-//            // Add every source image of the fragment program
-//            size_t sourceImageCount = m_imageParams.size();
-//            for ( size_t i=0; i<sourceImageCount; ++i )
-//            {
-//                char temp[128];
-//                mutable_snprintf( temp, 128, "uniform sampler2D %s;\n",
-//                                  m_imageParams[i].name.c_str() );
-//                imageCode += temp;
-//            }
-
-//            // Add every source vector of the fragment program
-//            size_t sourceVectorCount = m_vectorParams.size();
-//            for ( size_t i=0; i<sourceVectorCount; ++i )
-//            {
-//                char temp[128];
-//                mutable_snprintf( temp, 128, "uniform vec4 %s;\n",
-//                                  m_vectorParams[i].name.c_str() );
-//                imageCode += temp;
-//            }
-
-//            // Add every source float of the fragment program
-//            size_t sourceFloatCount = m_floatParams.size();
-//            for ( size_t i=0; i<sourceFloatCount; ++i )
-//            {
-//                char temp[128];
-//                mutable_snprintf( temp, 128, "uniform float %s;\n",
-//                                  m_floatParams[i].name.c_str() );
-//                imageCode += temp;
-//            }
-
-//            imageCode += "\nvec4 "+imageName+" ( vec2 "+string(MUTABLE_GPU_CODE_UVS)+" )\n";
-//            imageCode += "{\n";
-//            imageCode += m_expressionCode;
-//            imageCode += "\n\treturn "+m_currentVar+";\n";
-//            imageCode += "}\n";
-//            imageCode += "// end of mutable image\n\n";
-
-//            return imageCode;
-//        }
-
-
-//        //-----------------------------------------------------------------------------------------
-//        // ConstCodeVisitor interface
-//        //-----------------------------------------------------------------------------------------
-//        virtual void Visit( OP::ADDRESS at, PROGRAM& program ) override
-//        {
-//            (void)at;
-//            (void)program;
-////            const OP& op = program.m_code[at];
-
-////            if ( !m_generated[at].size() )
-////            {
-////                char temp[MUTABLE_GPU_CODE_BUFFER_SIZE];
-
-////                // See if we can translate the operation to the gpu
-////                bool gpuizable = GetOpDesc( (OP_TYPE)op.type ).gpuizable;
-
-////                // If the op is a state constant, don't translate it.
-////                gpuizable &= !program.m_states[m_state].IsUpdateCache(at);
-
-////                if ( gpuizable )
-////                {
-////                    // Generate more code
-////                    switch( op.type )
-////                    {
-
-////                    case OP_TYPE::IM_MULTIPLY:
-////                    {
-////                        string result = GetNewVar( "multiply" );
-
-////                        Visit( op.args.ImageLayer.base, program );
-////                        string baseVar = m_currentVar;
-
-////                        Visit( op.args.ImageLayer.blended, program );
-////                        string blendedVar = m_currentVar ;
-
-////                        string multiplyVar = GetNewVar("multiply_unmasked");
-////                        mutable_snprintf( temp, MUTABLE_GPU_CODE_BUFFER_SIZE,
-////                                  "\nvec4 %s = %s*%s;\n",
-////                                  multiplyVar.c_str(),
-////                                  baseVar.c_str(),
-////                                  blendedVar.c_str() );
-////                        m_expressionCode += temp;
-
-////                        if ( op.args.ImageLayer.mask )
-////                        {
-////                            Visit( op.args.ImageLayer.mask, program );
-////                            string maskVar = m_currentVar;
-
-////                            mutable_snprintf( temp, MUTABLE_GPU_CODE_BUFFER_SIZE,
-////                                      "vec4 %s = %s*(1.0-%s.rrrr) + %s*%s.rrrr;\n",
-////                                      result.c_str(),
-////                                      baseVar.c_str(),
-////                                      maskVar.c_str(),
-////                                      multiplyVar.c_str(),
-////                                      maskVar.c_str() );
-
-////                            m_expressionCode += temp;
-////                        }
-////                        else
-////                        {
-////                            mutable_snprintf( temp, MUTABLE_GPU_CODE_BUFFER_SIZE,
-////                                      "vec4 %s = %s;\n",
-////                                      result.c_str(),
-////                                      multiplyVar.c_str() );
-////                            m_expressionCode += temp;
-////                        }
-
-////                        m_currentVar = result;
-
-////                        break;
-////                    }
-
-////                    case OP_TYPE::IM_BLEND:
-////                    {
-////                        string result = GetNewVar( "blend" );
-
-////                        Visit( op.args.ImageLayer.base, program );
-////                        string baseVar = m_currentVar;
-
-////                        Visit( op.args.ImageLayer.blended, program );
-////                        string blendedVar = m_currentVar ;
-
-////                        if ( op.args.ImageLayer.mask )
-////                        {
-////                            Visit( op.args.ImageLayer.mask, program );
-////                            string maskVar = m_currentVar;
-
-////                            mutable_snprintf( temp, MUTABLE_GPU_CODE_BUFFER_SIZE,
-////                                      "vec4 %s = %s*(1.0-%s.rrrr) + %s*%s.rrrr;\n",
-////                                      result.c_str(),
-////                                      baseVar.c_str(),
-////                                      maskVar.c_str(),
-////                                      blendedVar.c_str(),
-////                                      maskVar.c_str() );
-
-////                            m_expressionCode += temp;
-////                        }
-////                        else
-////                        {
-////                            mutable_snprintf( temp, MUTABLE_GPU_CODE_BUFFER_SIZE,
-////                                      "vec4 %s = %s;\n",
-////                                      result.c_str(),
-////                                      blendedVar.c_str() );
-////                            m_expressionCode += temp;
-////                        }
-
-////                        m_currentVar = result;
-
-////                        break;
-////                    }
-
-////                    case OP_TYPE::IM_BLENDCOLOUR:
-////                    {
-////                        string result = GetNewVar( "blendcol" );
-
-////                        Visit( op.args.ImageLayerColour.base, program );
-////                        string baseVar = m_currentVar;
-
-////                        Visit( op.args.ImageLayerColour.colour, program );
-////                        string blendedVar = m_currentVar ;
-
-////                        if ( op.args.ImageLayerColour.mask )
-////                        {
-////                            Visit( op.args.ImageLayerColour.mask, program );
-////                            string maskVar = m_currentVar;
-
-////                            mutable_snprintf( temp, MUTABLE_GPU_CODE_BUFFER_SIZE,
-////                                      "vec4 %s = %s*(1.0-%s.rrrr) + %s*%s.rrrr;\n",
-////                                      result.c_str(),
-////                                      baseVar.c_str(),
-////                                      maskVar.c_str(),
-////                                      blendedVar.c_str(),
-////                                      maskVar.c_str() );
-
-////                            m_expressionCode += temp;
-////                        }
-////                        else
-////                        {
-////                            mutable_snprintf( temp, MUTABLE_GPU_CODE_BUFFER_SIZE,
-////                                      "vec4 %s = %s;\n",
-////                                      result.c_str(),
-////                                      blendedVar.c_str() );
-////                            m_expressionCode += temp;
-////                        }
-
-////                        m_currentVar = result;
-
-////                        break;
-////                    }
-
-////                    case OP_TYPE::IM_SOFTLIGHTCOLOUR:
-////                    {
-////                        string result = GetNewVar( "softlight" );
-
-////                        Visit( op.args.ImageLayerColour.base, program );
-////                        string baseVar = m_currentVar;
-
-////                        Visit( op.args.ImageLayerColour.colour, program );
-////                        string colourVar = m_currentVar;
-
-////                        string softVar = GetNewVar("softlight_unmasked");
-////                        mutable_snprintf( temp, MUTABLE_GPU_CODE_BUFFER_SIZE,
-////                                  "\nvec4 %s = (1.0-%s)*%s*%s + %s*(1.0-(1.0-%s)*(1.0-%s));\n",
-////                                  softVar.c_str(),
-////                                  baseVar.c_str(),
-////                                  baseVar.c_str(),
-////                                  colourVar.c_str(),
-////                                  baseVar.c_str(),
-////                                  baseVar.c_str(),
-////                                  colourVar.c_str() );
-////                        m_expressionCode += temp;
-
-////                        if ( op.args.ImageLayerColour.mask )
-////                        {
-////                            Visit( op.args.ImageLayerColour.mask, program );
-////                            string maskVar = m_currentVar;
-
-////                            mutable_snprintf( temp, MUTABLE_GPU_CODE_BUFFER_SIZE,
-////                                      "vec4 %s = %s*(1.0-%s.rrrr) + %s*%s.rrrr;\n",
-////                                      result.c_str(),
-////                                      baseVar.c_str(),
-////                                      maskVar.c_str(),
-////                                      softVar.c_str(),
-////                                      maskVar.c_str() );
-////                            m_expressionCode += temp;
-////                        }
-////                        else
-////                        {
-////                            mutable_snprintf( temp, MUTABLE_GPU_CODE_BUFFER_SIZE,
-////                                      "vec4 %s = %s;\n",
-////                                      result.c_str(),
-////                                      softVar.c_str() );
-////                            m_expressionCode += temp;
-////                        }
-
-////                        m_currentVar = result;
-////                        break;
-////                    }
-
-////                    case OP_TYPE::IM_INTERPOLATE3:
-////                    {
-////                        string result = GetNewVar( "interpolate3" );
-
-////                        Visit( op.args.ImageInterpolate3.factor1, program );
-////                        string factor1Var = m_currentVar;
-
-////                        Visit( op.args.ImageInterpolate3.factor2, program );
-////                        string factor2Var = m_currentVar ;
-
-////                        Visit( op.args.ImageInterpolate3.target0, program );
-////                        string target0Var = m_currentVar;
-
-////                        Visit( op.args.ImageInterpolate3.target1, program );
-////                        string target1Var = m_currentVar;
-
-////                        Visit( op.args.ImageInterpolate3.target2, program );
-////                        string target2Var = m_currentVar;
-
-////                        mutable_snprintf( temp, MUTABLE_GPU_CODE_BUFFER_SIZE,
-////                                          "vec4 %s = %s*(1.0-%s-%s) + %s*%s + %s*%s;\n",
-////                                          result.c_str(),
-////                                          target0Var.c_str(),
-////                                          factor1Var.c_str(),
-////                                          factor2Var.c_str(),
-////                                          target1Var.c_str(),
-////                                          factor1Var.c_str(),
-////                                          target2Var.c_str(),
-////                                          factor2Var.c_str() );
-
-////                        m_expressionCode += temp;
-
-////                        m_currentVar = result;
-
-////                        break;
-////                    }
-
-////                    case OP_TYPE::IM_COMPOSE:
-////                    {
-////                        string result = GetNewVar( "compose" );
-
-////                        // Base image
-////                        Visit( op.args.ImageCompose.base, program );
-////                        string baseVar = m_currentVar;
-
-////                        // Generate the block transform vector
-////                        int oldLayoutBlock = m_currentLayoutBlock;
-////                        m_currentLayoutBlock = op.args.ImageCompose.blockIndex;
-////                        Visit( op.args.ImageCompose.layout, program );
-////                        m_currentLayoutBlock = oldLayoutBlock;
-////                        string transformVar = m_currentVar;
-
-////                        // GLSL like this:
-////                        // vec2 blockCoords = (uvs.xy-transform.xy)/transform.zw;
-////                        // vec4 outsideTest = vec4( blockCoords.x<0, blockCoords.y<0, blockCoords.x>=1, blockCoords.y>=1 );
-////                        // if ( dot(outsideTest,outsideTest) )
-////                        // {
-////                        //     block code;
-////                        //     masking code;
-////                        // }
-////                        // else
-////                        // {
-////                        //     just set the base result;
-////                        // }
-////                        string blockCoordsVar = GetNewVar( "blockCoords" );
-////                        mutable_snprintf( temp, MUTABLE_GPU_CODE_BUFFER_SIZE,
-////                                  "vec2 %s = (%s.xy-%s.xy)/%s.zw;\n",
-////                                  blockCoordsVar.c_str(),
-////                                  m_currentTexCoords.c_str(),
-////                                  transformVar.c_str(),
-////                                  transformVar.c_str() );
-////                        m_expressionCode += temp;
-
-////                        string testVar = GetNewVar( "outsideTest" );
-////                        mutable_snprintf( temp, MUTABLE_GPU_CODE_BUFFER_SIZE,
-////                                          "vec4 %s = vec4( %s.x<0, %s.y<0, %s.x>=1, %s.y>=1 );\n",
-////                                          testVar.c_str(),
-////                                          blockCoordsVar.c_str(),
-////                                          blockCoordsVar.c_str(),
-////                                          blockCoordsVar.c_str(),
-////                                          blockCoordsVar.c_str()
-////                                          );
-////                        m_expressionCode += temp;
-
-////                        // Declare the result
-////                        mutable_snprintf( temp, MUTABLE_GPU_CODE_BUFFER_SIZE,
-////                                          "vec4 %s;\n",
-////                                          result.c_str()
-////                                          );
-////                        m_expressionCode += temp;
-
-////                        // "Inside" conditional
-////                        // TODO: variables declared inside the conditional will generate wrong
-////                        // code if reused outside by the internal map, as they are not accessible
-////                        // outside the "if" scope.
-////                        mutable_snprintf( temp, MUTABLE_GPU_CODE_BUFFER_SIZE,
-////                                          "if ( dot(%s,%s)==0 )\n{\n",
-////                                          testVar.c_str(),
-////                                          testVar.c_str()
-////                                          );
-////                        m_expressionCode += temp;
-
-////                        // Block image
-////                        string oldTexCoords = m_currentTexCoords;
-////                        m_currentTexCoords = blockCoordsVar;
-////                        Visit( op.args.ImageCompose.blockImage, program );
-////                        string blockVar = m_currentVar;
-
-////                        // Block mask
-////                        if ( op.args.ImageCompose.mask )
-////                        {
-////                            Visit( op.args.ImageCompose.mask, program );
-////                            string maskVar = m_currentVar;
-
-////                            mutable_snprintf( temp, MUTABLE_GPU_CODE_BUFFER_SIZE,
-////                                      "%s = %s*(1.0-%s.rrrr) + %s*%s.rrrr;\n",
-////                                      result.c_str(),
-////                                      baseVar.c_str(),
-////                                      maskVar.c_str(),
-////                                      blockVar.c_str(),
-////                                      maskVar.c_str() );
-////                            m_expressionCode += temp;
-////                        }
-////                        else
-////                        {
-////                            mutable_snprintf( temp, MUTABLE_GPU_CODE_BUFFER_SIZE,
-////                                      "%s = %s;\n",
-////                                      result.c_str(),
-////                                      blockVar.c_str() );
-////                            m_expressionCode += temp;
-////                        }
-
-////                        m_currentTexCoords = oldTexCoords;
-
-////                        // If outside
-////                        // TODO: if there is no block mask, skip the base when possible
-////                        m_expressionCode += "}\nelse\n{\n";
-////                        mutable_snprintf( temp, MUTABLE_GPU_CODE_BUFFER_SIZE,
-////                                  "%s = %s;\n",
-////                                  result.c_str(),
-////                                  baseVar.c_str() );
-////                        m_expressionCode += temp;
-////                        m_expressionCode += "}\n";
-
-////                        m_currentVar = result;
-////                        break;
-////                    }
-
-////                    default:
-////                        check( false );
-////                    }
-////                }
-
-////                else
-////                {
-////                    // We cannot generate gpu code below this, make a parameter for what is left
-////                    switch ( GetOpType( op.type ) )
-////                    {
-////                    case DT_IMAGE:
-////                    {
-////                        mutable_snprintf( temp, MUTABLE_GPU_CODE_BUFFER_SIZE, "%s_%s_%d",
-////                                          m_currentImageName.c_str(),
-////                                          MUTABLE_GPU_IMAGE_PARAM, m_imageParams.size() );
-////                        string varName = temp;
-
-////                        m_currentVar = GetNewVar( "imgParam" );
-////                        mutable_snprintf( temp, MUTABLE_GPU_CODE_BUFFER_SIZE,
-////                                          "vec4 %s = texture( %s, %s );\n",
-////                                          m_currentVar.c_str(),
-////                                          varName.c_str(),
-////                                          m_currentTexCoords.c_str() );
-////                        m_expressionCode += temp;
-
-////                        // Make sure the image is in a gpu-usable format
-////                        FImageDesc desc = GetImageDesc( program, at );
-////                        if ( desc.m_format==IF_L_UBYTE_RLE )
-////                        {
-////                            OP conv;
-////                            conv.type = OP_TYPE::IM_PIXELFORMAT;
-////                            conv.args.ImagePixelFormat.format = IF_L_UBYTE;
-////                            conv.args.ImagePixelFormat.source = at;
-////                            at = program.AddOp( conv );
-////                        }
-////                        else if ( desc.m_format==IF_RGB_UBYTE_RLE )
-////                        {
-////                            OP conv;
-////                            conv.type = OP_TYPE::IM_PIXELFORMAT;
-////                            conv.args.ImagePixelFormat.format = IF_RGB_UBYTE;
-////                            conv.args.ImagePixelFormat.source = at;
-////                            at = program.AddOp( conv );
-////                        }
-////                        else if ( desc.m_format==IF_RGBA_UBYTE_RLE )
-////                        {
-////                            OP conv;
-////                            conv.type = OP_TYPE::IM_PIXELFORMAT;
-////                            conv.args.ImagePixelFormat.format = IF_RGBA_UBYTE;
-////                            conv.args.ImagePixelFormat.source = at;
-////                            at = program.AddOp( conv );
-////                        }
-
-////                        m_imageParams.push_back( GPU_PROGRAM::PARAM_DATA( at, varName ) );
-////                        break;
-////                    }
-
-////                    case DT_COLOUR:
-////                        mutable_snprintf( temp, MUTABLE_GPU_CODE_BUFFER_SIZE, "%s_%s_%d",
-////                                          m_currentImageName.c_str(),
-////                                          MUTABLE_GPU_VECTOR_PARAM, m_vectorParams.size() );
-////                        m_currentVar = temp;
-////                        m_vectorParams.push_back( GPU_PROGRAM::PARAM_DATA( at, temp ) );
-////                        break;
-
-////                    case DT_SCALAR:
-////                        mutable_snprintf( temp, MUTABLE_GPU_CODE_BUFFER_SIZE, "%s_%s_%d",
-////                                          m_currentImageName.c_str(),
-////                                          MUTABLE_GPU_FLOAT_PARAM, m_floatParams.size() );
-////                        m_currentVar = temp;
-////                        m_floatParams.push_back( GPU_PROGRAM::PARAM_DATA( at, temp ) );
-////                        break;
-
-////                    case DT_LAYOUT:
-////                    {
-////                        OP trans;
-////                        trans.type = OP_TYPE::CO_LAYOUTBLOCKTRANSFORM;
-////                        trans.args.ColourLayoutBlockTransform.layout = at;
-////                        trans.args.ColourLayoutBlockTransform.block = (uint16_t)m_currentLayoutBlock;
-////                        at = program.AddOp( trans );
-
-////                        mutable_snprintf( temp, 128, "%s_%s_%d", m_currentImageName.c_str(),
-////                                          MUTABLE_GPU_VECTOR_PARAM, m_vectorParams.size() );
-////                        m_currentVar = temp;
-////                        m_vectorParams.push_back( GPU_PROGRAM::PARAM_DATA( at, temp ) );
-////                        break;
-////                    }
-
-////                    default:
-////                        check( false );
-////                    }
-////                }
-
-////                m_generated[at] = m_currentVar;
-////            }
-
-////            else
-////            {
-////                // Reuse previously generated code
-////                m_currentVar = m_generated[at];
-////            }
-
-//        }
-
-//    protected:
-
-//        int m_state;
-
-//        //! Instructions that have been already generated and the variable they reside in.
-//        map<OP::ADDRESS,string> m_generated;
-
-//        vector<string> m_usedVars;
-
-//        //! State when generating image expressions
-//        string m_currentTexCoords;
-
-//        //! State when generating layout expressions
-//        int m_currentLayoutBlock;
-
-//        //! Generate an unused variable name based on the provided one.
-//        string GetNewVar( const string& name )
-//        {
-//            int suffix = 0;
-//            string candidate = name;
-
-//            while( std::find( m_usedVars.begin(), m_usedVars.end(), candidate )
-//                   !=
-//                   m_usedVars.end() )
-//            {
-//                ++suffix;
-
-//                char temp[128];
-//                mutable_snprintf( temp, 128, "%s_%d", name.c_str(), suffix );
-//                candidate = temp;
-//            }
-
-//            m_usedVars.push_back( candidate );
-//            return candidate;
-//        }
-//    };
-
-
-//    //---------------------------------------------------------------------------------------------
-//    //---------------------------------------------------------------------------------------------
-//    //---------------------------------------------------------------------------------------------
-//    class GPUSubtreeGenerator_GL4 : public GPUSubtreeGenerator_GLSL
-//    {
-//    public:
-
-//        GPU_PROGRAM GenerateProgram( OP::ADDRESS at,
-//                              PROGRAM& program,
-//                              int state,
-//                              const GPU_PLATFORM_PROPS& platformOptions,
-//                              const GPU_STATE_PROPS& //stateOptions
-//                              )
-//        {
-//            check( platformOptions.type == GPU_GL4 );
-//            (void)platformOptions;
-//            m_state = state;
-//            m_currentImageName = "";
-
-//            // Clear state
-//            Clear();
-
-//            Traverse( at, program, false );
-
-//            GPU_PROGRAM gpuProgram;
-//            gpuProgram.m_floatParams = m_floatParams;
-//            gpuProgram.m_vectorParams = m_vectorParams;
-//            gpuProgram.m_imageParams = m_imageParams;
-
-////            // Find out what pixel format we need to generate
-////            IMAGE_FORMAT format = GetImageDesc( program, at ).m_format;
-
-////            // Convert to a usable format
-////            switch ( format )
-////            {
-////            case IF_RGBA_UBYTE: gpuProgram.m_format = format; break;
-////            case IF_RGB_UBYTE: gpuProgram.m_format = IF_RGBA_UBYTE; break;
-////            default:
-////                break;
-////            }
-
-
-////            // Generate the size instruction
-////            OP sizeOp;
-////            sizeOp.type = OP_TYPE::CO_IMAGESIZE;
-////            sizeOp.args.ColourImageSize.image = at;
-////            gpuProgram.m_size = program.AddOp( sizeOp );
-
-
-////            // Complete the code
-////            string params =
-////                    "\n"
-////                    "#version 330\n"
-////                    "\n";
-
-////            char temp[MUTABLE_GPU_CODE_BUFFER_SIZE];
-
-////            // Add every source image of the fragment program
-////            for ( size_t i=0; i<m_imageParams.size(); ++i )
-////            {
-////                mutable_snprintf( temp, MUTABLE_GPU_CODE_BUFFER_SIZE, "uniform sampler2D %s;\n",
-////                                  m_imageParams[i].name.c_str() );
-////                params += temp;
-////            }
-
-////            // Add every source vector of the fragment program
-////            for ( size_t i=0; i<m_vectorParams.size(); ++i )
-////            {
-////                mutable_snprintf( temp, MUTABLE_GPU_CODE_BUFFER_SIZE, "uniform vec4 %s;\n",
-////                                  m_vectorParams[i].name.c_str() );
-////                params += temp;
-////            }
-
-////            // Add every source float of the fragment program
-////            for ( size_t i=0; i<m_floatParams.size(); ++i )
-////            {
-////                mutable_snprintf( temp, MUTABLE_GPU_CODE_BUFFER_SIZE, "uniform float %s;\n",
-////                                  m_floatParams[i].name.c_str() );
-////                params += temp;
-////            }
-
-////            string header =
-////                    "\n"
-////                    "in vec2 "+string(MUTABLE_GPU_CODE_UVS)+";\n"
-////                    "out vec4 result;\n"
-////                    "\n"
-////                    "void main(void)\n"
-////                    "{\n";
-
-////            string footer =
-////                    "\n"
-////                    "result = "+m_currentVar+";\n"
-////                    "}\n";
-
-////            gpuProgram.m_fragmentCode = params + header + m_expressionCode + footer;
-
-////            gpuProgram.m_vertexCode =
-////                    "#version 330\n"
-////                    "\n"
-////                    "in vec2 position;"
-////                    "out vec2 "+string(MUTABLE_GPU_CODE_UVS)+";"
-////                    ""
-////                    "void main()"
-////                    "{"
-////                    "	gl_Position = vec4( position.xy, 0.0, 1.0 );"
-////                    "	"+string(MUTABLE_GPU_CODE_UVS)+" = position.xy*0.5+0.5;"
-////                    "}";
-
-////            gpuProgram.m_platform = SS_GL4;
-
-//            return gpuProgram;
-//        }
-
-//    };
-
-
-//    //---------------------------------------------------------------------------------------------
-//    //---------------------------------------------------------------------------------------------
-//    //---------------------------------------------------------------------------------------------
-//    class GPUSubtreeGenerator_GLES2 : public GPUSubtreeGenerator_GLSL
-//    {
-//    public:
-
-
-//        GPU_PROGRAM GenerateProgram( OP::ADDRESS at,
-//                                     PROGRAM& program,
-//                                     int state,
-//                                     const GPU_PLATFORM_PROPS& platformOptions,
-//                                     const GPU_STATE_PROPS& //stateOptions
-//                                     ) override
-//        {
-//            check( platformOptions.type == GPU_GLES2 );
-//            (void)platformOptions;
-
-//            m_state = state;
-//            m_currentImageName = "";
-
-//            // Clear state
-//            Clear();
-
-//            Traverse( at, program, false );
-
-//            GPU_PROGRAM gpuProgram;
-//            gpuProgram.m_floatParams = m_floatParams;
-//            gpuProgram.m_vectorParams = m_vectorParams;
-//            gpuProgram.m_imageParams = m_imageParams;
-
-////            // Find out what pixel format we need to generate
-////            IMAGE_FORMAT format = GetImageDesc( program, at ).m_format;
-
-////            // Convert to a usable format
-////            switch ( format )
-////            {
-////            case IF_RGBA_UBYTE: gpuProgram.m_format = format; break;
-////            case IF_RGB_UBYTE: gpuProgram.m_format = IF_RGBA_UBYTE; break;
-////            default:
-////                break;
-////            }
-
-
-////            // Generate the size instruction
-////            OP sizeOp;
-////            sizeOp.type = OP_TYPE::CO_IMAGESIZE;
-////            sizeOp.args.ColourImageSize.image = at;
-////            gpuProgram.m_size = program.AddOp( sizeOp );
-
-
-////            // Complete the code
-////            string params =
-////                    "\n"
-////                    "precision mediump float;\n"
-////                    "\n";
-
-////            char temp[MUTABLE_GPU_CODE_BUFFER_SIZE];
-
-////            // Add every source image of the fragment program
-////            for ( size_t i=0; i<m_imageParams.size(); ++i )
-////            {
-////                mutable_snprintf( temp, MUTABLE_GPU_CODE_BUFFER_SIZE, "uniform sampler2D %s;\n",
-////                                  m_imageParams[i].name.c_str() );
-////                params += temp;
-////            }
-
-////            // Add every source vector of the fragment program
-////            for ( size_t i=0; i<m_vectorParams.size(); ++i )
-////            {
-////                mutable_snprintf( temp, MUTABLE_GPU_CODE_BUFFER_SIZE, "uniform vec4 %s;\n",
-////                                  m_vectorParams[i].name.c_str() );
-////                params += temp;
-////            }
-
-////            // Add every source float of the fragment program
-////            for ( size_t i=0; i<m_floatParams.size(); ++i )
-////            {
-////                mutable_snprintf( temp, MUTABLE_GPU_CODE_BUFFER_SIZE, "uniform float %s;\n",
-////                                  m_floatParams[i].name.c_str() );
-////                params += temp;
-////            }
-
-////            string header =
-////                    "\n"
-////                    "varying vec2 texCoords;\n"
-////                    "\n"
-////                    "void main(void)\n"
-////                    "{\n";
-
-////            string footer =
-////                    "\n"
-////                    "gl_FragData[0] = "+m_currentVar+";\n"
-////                    "}\n";
-
-////            gpuProgram.m_fragmentCode = params + header + m_expressionCode + footer;
-
-////            gpuProgram.m_vertexCode =
-////                    "\n"
-////                    "attribute vec2 position;"
-////                    "varying vec2 texCoords;"
-////                    ""
-////                    "void main()"
-////                    "{"
-////                    "	gl_Position = vec4( position.xy, 0.0, 1.0 );"
-////                    "	texCoords = position.xy*0.5+0.5;"
-////                    "}";
-
-////            gpuProgram.m_platform = SS_GLES2;
-
-//            return gpuProgram;
-//        }
-
-//    };
-
-
-//    //---------------------------------------------------------------------------------------------
-//    //---------------------------------------------------------------------------------------------
-//    //---------------------------------------------------------------------------------------------
-//    GPUTranslator::GPUTranslator( PROGRAM& program, int s,
-//                                  const GPU_PLATFORM_PROPS& platformOptions,
-//                                  const GPU_STATE_PROPS& , //stateOptions,
-//                                  ModelReportPtr pReport )
-//        //: m_platformOptions( platformOptions )
-//        //, m_stateOptions( stateOptions )
-//        : m_pReport( pReport )
-//    {
-//        m_state = s;
-//        m_visited.resize( program.m_opAddress.size(), 0 );
-
-//        // Create the platfomr specific GPU code generator.
-//        m_pGPUSubtreeGen = 0;
-//        switch ( platformOptions.type )
-//        {
-//        case GPU_GL4:
-//            m_pGPUSubtreeGen = new GPUSubtreeGenerator_GL4;
-//            break;
-
-//        case GPU_GLES2:
-//            m_pGPUSubtreeGen = new GPUSubtreeGenerator_GLES2;
-//            break;
-
-//        default:
-//            check( false );
-//            break;
-//        }
-
-//        // Clone the code before optimising it for the gpu: this way we keep the original
-//        // code to be used by other states.
-//        //CloneVisitor( program, program.m_states[s].m_root, s );
-
-//        Traverse( program.m_states[s].m_root, program );
-//    }
-
-
-//    //---------------------------------------------------------------------------------------------
-//    GPUTranslator::~GPUTranslator()
-//    {
-//        delete m_pGPUSubtreeGen;
-//        m_pGPUSubtreeGen = 0;
-//    }
-
-//    //---------------------------------------------------------------------------------------------
-//    OP::ADDRESS GPUTranslator::Visit( OP::ADDRESS at, PROGRAM& program )
-//    {
-//        // Visit only once
-//        OP::ADDRESS sourceAt = at;
-//        m_visited.resize( program.m_opAddress.size(), 0 );
-//        if ( m_visited[at] )
-//        {
-//            return m_visited[at];
-//        }
-
-//        OP op = program.m_code[at];
-
-//        bool gpuizable = GetOpDesc( (OP_TYPE)op.type ).gpuizable;
-//        gpuizable &= !program.m_states[m_state].IsUpdateCache(at);
-
-////        if ( gpuizable && m_stateOptions.m_internal )
-////        {
-////            // Find out what pixel format we need to generate
-////            IMAGE_FORMAT format = GetImageDesc( program, at ).m_format;
-
-////            // Create the fragment program
-////            GPU_PROGRAM gpuProgram =
-////                    m_pGPUSubtreeGen->GenerateProgram( at, program,
-////                                                       m_state, m_platformOptions, m_stateOptions );
-
-////            // Create a GPU image instruction
-////            OP gpuOp;
-////            gpuOp.type = OP_TYPE::IM_GPU;
-////            gpuOp.args.ImageGPU.program = (OP::ADDRESS)program.m_gpuPrograms.size();
-////            program.m_gpuPrograms.push_back( gpuProgram );
-////            at = program.AddOp( gpuOp );
-
-////            // Check that we have the expected pixel format, as gpu instructions cannot generate
-////            // some of them.
-////            if ( format != gpuProgram.m_format )
-////            {
-////                OP formatOp;
-////                formatOp.type = OP_TYPE::IM_PIXELFORMAT;
-////                formatOp.args.ImagePixelFormat.source = at;
-////                formatOp.args.ImagePixelFormat.format = format;
-////                at = program.AddOp( formatOp );
-////            }
-
-////            // Recurse on the program parameters
-////            //TODO
-////        }
-////        else if ( m_stateOptions.m_external )
-////        {
-////            if ( op.type == OP_TYPE::IN_ADDLOD )
-////            {
-////                for (int t=0;t<MUTABLE_OP_MAX_ADD_COUNT;++t)
-////                {
-////                    if ( op.args.InstanceAddLOD.lod[t] )
-////                    {
-////                        m_currentLOD = REPORT_STATE::LOD();
-////                        op.args.InstanceAddLOD.lod[t] = Visit( op.args.InstanceAddLOD.lod[t], program );
-////                        m_pReport->GetPrivate()->m_states[m_state].m_lods.push_back(m_currentLOD);
-////                    }
-////                }
-
-////                at = program.AddOp(op);
-////            }
-////            else if ( op.type == OP_TYPE::IN_ADDCOMPONENT )
-////            {
-////                op.args.InstanceAdd.instance = Visit( op.args.InstanceAdd.instance, program );
-
-////                if ( op.args.InstanceAdd.value )
-////                {
-////                    m_currentComponent = REPORT_STATE::COMPONENT();
-
-////                    op.args.InstanceAdd.value = Visit( op.args.InstanceAdd.value, program );
-
-////                    m_currentComponent.m_name =
-////                            program.m_constantStrings[ op.args.InstanceAdd.name ];
-////                    m_currentLOD.m_components.push_back(m_currentComponent);
-////                }
-
-////                at = program.AddOp(op);
-////            }
-////            else if ( op.type == OP_TYPE::IN_ADDIMAGE )
-////            {
-////                OP::ADDRESS lastInstance = op.args.InstanceAdd.instance;
-
-////                // Visit the instance
-////                lastInstance = Visit( lastInstance, program );
-
-////                // GPU-ize or visit every image
-////                // We actually unfold the images too
-////                if ( op.args.InstanceAdd.value )
-////                {
-////                    if ( m_stateOptions.m_external )
-////                    {
-////                        REPORT_STATE::IMAGE imageData;
-////                        imageData.m_name = program.m_constantStrings[ op.args.InstanceAdd.name ];
-
-////                        // Create the fragment program
-////                        string imageCode =
-////                                m_pGPUSubtreeGen->GenerateImage( imageData.m_name,
-////                                                                 op.args.InstanceAdd.value,
-////                                                                 program,
-////                                                                 m_state,
-////                                                                 m_platformOptions,
-////                                                                 m_stateOptions );
-
-////                        // Add every source image of the fragment program
-////                        size_t sourceImageCount = m_pGPUSubtreeGen->m_imageParams.size();
-////                        for ( size_t i=0; i<sourceImageCount; ++i )
-////                        {
-////                            OP nop;
-////                            nop.type = OP_TYPE::IN_ADDIMAGE;
-////                            nop.args.InstanceAdd.instance = lastInstance;
-////                            nop.args.InstanceAdd.value =
-////                                    Visit( m_pGPUSubtreeGen->m_imageParams[i].address, program );
-
-////                            nop.args.InstanceAdd.name = program.AddConstant
-////                                    ( m_pGPUSubtreeGen->m_imageParams[i].name.c_str() );
-////                            lastInstance = program.AddOp( nop );
-
-////                            imageData.m_sourceImages.push_back
-////                                    ( m_pGPUSubtreeGen->m_imageParams[i].name );
-////                        }
-
-////                        // Add every source vector of the fragment program
-////                        size_t sourceVectorCount = m_pGPUSubtreeGen->m_vectorParams.size();
-////                        for ( size_t i=0; i<sourceVectorCount; ++i )
-////                        {
-////                            OP nop;
-////                            nop.type = OP_TYPE::IN_ADDVECTOR;
-////                            nop.args.InstanceAdd.instance = lastInstance;
-////                            nop.args.InstanceAdd.value =
-////                                    Visit( m_pGPUSubtreeGen->m_vectorParams[i].address, program );
-
-////                            nop.args.InstanceAdd.name = program.AddConstant
-////                                    ( m_pGPUSubtreeGen->m_vectorParams[i].name.c_str() );
-////                            lastInstance = program.AddOp( nop );
-
-////                            imageData.m_sourceVectors.push_back
-////                                    ( m_pGPUSubtreeGen->m_vectorParams[i].name );
-////                        }
-
-////                        // Add every source float of the fragment program
-////                        size_t sourceFloatCount = m_pGPUSubtreeGen->m_floatParams.size();
-////                        for ( size_t i=0; i<sourceFloatCount; ++i )
-////                        {
-////                            OP nop;
-////                            nop.type = OP_TYPE::IN_ADDSCALAR;
-////                            nop.args.InstanceAdd.instance = lastInstance;
-////                            nop.args.InstanceAdd.value =
-////                                    Visit( m_pGPUSubtreeGen->m_floatParams[i].address, program );
-
-////                            nop.args.InstanceAdd.name = program.AddConstant
-////                                    ( m_pGPUSubtreeGen->m_floatParams[i].name.c_str() );
-////                            lastInstance = program.AddOp( nop );
-
-////                            imageData.m_sourceScalars.push_back
-////                                    ( m_pGPUSubtreeGen->m_floatParams[i].name );
-////                        }
-
-////                        imageData.m_fragmentCode = imageCode;
-////                        m_currentComponent.m_images.push_back(imageData);
-////                    }
-////                    else
-////                    {
-////                        // Recurse and add the image
-////                        OP nop;
-////                        nop.type = OP_TYPE::IN_ADDIMAGE;
-////                        nop.args.InstanceAdd.instance = lastInstance;
-////                        nop.args.InstanceAdd.value =
-////                                Visit( op.args.InstanceAdd.value, program );
-////                        nop.args.InstanceAdd.name = op.args.InstanceAdd.name;
-////                        lastInstance = program.AddOp( nop );
-////                    }
-////                }
-
-////                at = lastInstance;
-////            }
-////            else
-////            {
-////                // We cannot make a gpu program here. Let's see below.
-////                at = Recurse( at, program );
-////            }
-////        }
-////        else
-////        {
-////            // We cannot make a gpu program here. Let's see below.
-////            at = Recurse( at, program );
-////        }
-
-//        m_visited[sourceAt] = at;
-
-//        return at;
-//    }
-
-}
-
diff --git a/Engine/Plugins/Experimental/Mutable/Source/MutableTools/Private/MuT/CodeOptimiser_GPU.h b/Engine/Plugins/Experimental/Mutable/Source/MutableTools/Private/MuT/CodeOptimiser_GPU.h
--- a/Engine/Plugins/Experimental/Mutable/Source/MutableTools/Private/MuT/CodeOptimiser_GPU.h
+++ b/Engine/Plugins/Experimental/Mutable/Source/MutableTools/Private/MuT/CodeOptimiser_GPU.h
@@ -2,57 +2,4 @@
 
 #pragma once
 
-#include "MuR/Operations.h"
-#include "MuR/ModelPrivate.h"
-#include "MuR/ImagePrivate.h"
-#include "MuR/MeshPrivate.h"
-#include "MuR/CodeVisitor.h"
 
-#include "MuT/CompilerPrivate.h"
-#include "MuT/ModelReportPrivate.h"
-
-
-namespace mu
-{
-    //!
-    class GPUSubtreeGenerator;
-
-
-    //---------------------------------------------------------------------------------------------
-    //!
-    //---------------------------------------------------------------------------------------------
-//    class GPUTranslator : public CodeVisitor
-//    {
-//    public:
-
-//        GPUTranslator( PROGRAM& program,
-//                       int state,
-//                       const GPU_PLATFORM_PROPS& platformOptions,
-//                       const GPU_STATE_PROPS& stateOptions,
-//                       ModelReportPtr pReport );
-
-//        virtual ~GPUTranslator();
-
-//        virtual OP::ADDRESS Visit( OP::ADDRESS at, PROGRAM& program );
-
-//    private:
-
-//        vector<OP::ADDRESS> m_visited;
-
-//        int m_state;
-//        //const GPU_PLATFORM_PROPS& m_platformOptions;
-//        //const GPU_STATE_PROPS& m_stateOptions;
-//        REPORT_STATE::LOD m_currentLOD;
-//        REPORT_STATE::COMPONENT m_currentComponent;
-
-//        //!
-//        ModelReportPtr m_pReport;
-
-//        //!
-//        GPUSubtreeGenerator* m_pGPUSubtreeGen;
-
-//    };
-
-
-
-}
diff --git a/Engine/Plugins/Experimental/Mutable/Source/MutableTools/Private/MuT/CodeOptimiser_Logic.cpp b/Engine/Plugins/Experimental/Mutable/Source/MutableTools/Private/MuT/CodeOptimiser_Logic.cpp
--- a/Engine/Plugins/Experimental/Mutable/Source/MutableTools/Private/MuT/CodeOptimiser_Logic.cpp
+++ b/Engine/Plugins/Experimental/Mutable/Source/MutableTools/Private/MuT/CodeOptimiser_Logic.cpp
@@ -1,15 +1,20 @@
 // Copyright Epic Games, Inc. All Rights Reserved.
 
-#include "MuT/CodeOptimiser.h"
-
-#include "MuR/ModelPrivate.h"
+#include "Misc/AssertionMacros.h"
+#include "MuR/MemoryPrivate.h"
 #include "MuR/MutableTrace.h"
-#include "MuR/SystemPrivate.h"
 #include "MuR/Operations.h"
-
+#include "MuR/Ptr.h"
+#include "MuR/RefCounted.h"
+#include "MuT/AST.h"
 #include "MuT/ASTOpConditional.h"
 #include "MuT/ASTOpInstanceAdd.h"
 #include "MuT/ASTOpSwitch.h"
+#include "MuT/CodeOptimiser.h"
+
+#include <functional>
+#include <memory>
+#include <utility>
 
 
 
diff --git a/Engine/Plugins/Experimental/Mutable/Source/MutableTools/Private/MuT/CodeOptimiser_Parameters.cpp b/Engine/Plugins/Experimental/Mutable/Source/MutableTools/Private/MuT/CodeOptimiser_Parameters.cpp
--- a/Engine/Plugins/Experimental/Mutable/Source/MutableTools/Private/MuT/CodeOptimiser_Parameters.cpp
+++ b/Engine/Plugins/Experimental/Mutable/Source/MutableTools/Private/MuT/CodeOptimiser_Parameters.cpp
@@ -1,33 +1,42 @@
 // Copyright Epic Games, Inc. All Rights Reserved.
 
-#include "MuT/CodeOptimiser.h"
-#include "MuT/CodeOptimiser_GPU.h"
-#include "MuT/NodeObjectNewPrivate.h"
-#include "MuT/CodeGenerator.h"
-#include "MuT/DataPacker.h"
-#include "MuT/ModelReportPrivate.h"
-#include "MuT/CompilerPrivate.h"
-
-#include "MuR/CodeVisitor.h"
-#include "MuR/ModelPrivate.h"
-#include "MuR/SystemPrivate.h"
+#include "HAL/PlatformCrt.h"
+#include "HAL/PlatformMath.h"
+#include "Logging/LogCategory.h"
+#include "Logging/LogMacros.h"
+#include "Misc/AssertionMacros.h"
+#include "MuR/Image.h"
+#include "MuR/ImagePrivate.h"
+#include "MuR/MemoryPrivate.h"
+#include "MuR/MutableTrace.h"
 #include "MuR/Operations.h"
-
-#include "MuT/ASTOpImagePixelFormat.h"
-#include "MuT/ASTOpImageMipmap.h"
-#include "MuT/ASTOpImageCompose.h"
+#include "MuR/ParametersPrivate.h"
+#include "MuR/Platform.h"
+#include "MuR/Ptr.h"
+#include "MuR/RefCounted.h"
+#include "MuT/AST.h"
 #include "MuT/ASTOpAddLOD.h"
-#include "MuT/ASTOpSwitch.h"
 #include "MuT/ASTOpConditional.h"
+#include "MuT/ASTOpImageCompose.h"
+#include "MuT/ASTOpImageMipmap.h"
+#include "MuT/ASTOpImagePixelFormat.h"
 #include "MuT/ASTOpInstanceAdd.h"
+#include "MuT/ASTOpParameter.h"
+#include "MuT/CodeOptimiser.h"
+#include "MuT/Compiler.h"
+#include "MuT/CompilerPrivate.h"
+#include "MuT/DataPacker.h"
+#include "Trace/Detail/Channel.h"
 
 #include <array>
-
-#include "MuR/MutableTrace.h"
+#include <functional>
+#include <memory>
+#include <utility>
 
 
 namespace mu
 {
+class TaskManager;
 
 
     //---------------------------------------------------------------------------------------------
diff --git a/Engine/Plugins/Experimental/Mutable/Source/MutableTools/Private/MuT/CodeOptimiser_Semantic.cpp b/Engine/Plugins/Experimental/Mutable/Source/MutableTools/Private/MuT/CodeOptimiser_Semantic.cpp
--- a/Engine/Plugins/Experimental/Mutable/Source/MutableTools/Private/MuT/CodeOptimiser_Semantic.cpp
+++ b/Engine/Plugins/Experimental/Mutable/Source/MutableTools/Private/MuT/CodeOptimiser_Semantic.cpp
@@ -1,24 +1,31 @@
 // Copyright Epic Games, Inc. All Rights Reserved.
 
+#include "HAL/PlatformMath.h"
+#include "Misc/AssertionMacros.h"
+#include "MuR/Image.h"
+#include "MuR/ImagePrivate.h"
+#include "MuR/ImageRLE.h"
+#include "MuR/MemoryPrivate.h"
+#include "MuR/MutableMath.h"
 #include "MuR/MutableTrace.h"
-#include "MuT/CodeOptimiser.h"
-#include "MuT/CodeGenerator.h"
+#include "MuR/Operations.h"
+#include "MuR/Ptr.h"
+#include "MuR/RefCounted.h"
 #include "MuT/AST.h"
-#include "MuT/ASTOpMeshFormat.h"
-#include "MuT/ASTOpImagePixelFormat.h"
-#include "MuT/ASTOpImageMipmap.h"
-#include "MuT/ASTOpImagePatch.h"
-#include "MuT/ASTOpImageCompose.h"
-#include "MuT/ASTOpImageMultiLayer.h"
 #include "MuT/ASTOpConditional.h"
 #include "MuT/ASTOpConstantBool.h"
+#include "MuT/ASTOpImageCompose.h"
+#include "MuT/ASTOpImageMipmap.h"
+#include "MuT/ASTOpImageMultiLayer.h"
+#include "MuT/ASTOpImagePatch.h"
+#include "MuT/ASTOpImagePixelFormat.h"
+#include "MuT/ASTOpSwitch.h"
+#include "MuT/CodeOptimiser.h"
+#include "MuT/Table.h"
 
-#include "MuR/ModelPrivate.h"
-#include "MuR/SystemPrivate.h"
-#include "MuR/Operations.h"
-#include "MuR/ImagePrivate.h"
-#include "MuR/ImageRLE.h"
-#include "MuR/OpMeshRemoveChart.h"
+#include <functional>
+#include <memory>
+#include <utility>
 
 namespace mu
 {
diff --git a/Engine/Plugins/Experimental/Mutable/Source/MutableTools/Private/MuT/Compiler.cpp b/Engine/Plugins/Experimental/Mutable/Source/MutableTools/Private/MuT/Compiler.cpp
--- a/Engine/Plugins/Experimental/Mutable/Source/MutableTools/Private/MuT/Compiler.cpp
+++ b/Engine/Plugins/Experimental/Mutable/Source/MutableTools/Private/MuT/Compiler.cpp
@@ -3,27 +3,41 @@
 
 #include "MuT/Compiler.h"
 
-#include "MuT/CompilerPrivate.h"
+#include "Containers/Array.h"
+#include "Containers/Map.h"
+#include "HAL/LowLevelMemTracker.h"
+#include "HAL/PlatformCrt.h"
+#include "Hash/CityHash.h"
+#include "Logging/LogCategory.h"
+#include "Logging/LogMacros.h"
+#include "Misc/AssertionMacros.h"
+#include "MuR/Image.h"
+#include "MuR/MemoryPrivate.h"
+#include "MuR/Mesh.h"
+#include "MuR/Model.h"
+#include "MuR/ModelPrivate.h"
+#include "MuR/MutableTrace.h"
+#include "MuR/Operations.h"
+#include "MuR/ParametersPrivate.h"
+#include "MuR/Platform.h"
+#include "MuR/Serialisation.h"
+#include "MuR/System.h"
+#include "MuT/AST.h"
+#include "MuT/ASTOpParameter.h"
 #include "MuT/CodeGenerator.h"
 #include "MuT/CodeOptimiser.h"
-#include "MuT/DataPacker.h"
+#include "MuT/CompilerPrivate.h"
+#include "MuT/ErrorLog.h"
 #include "MuT/ErrorLogPrivate.h"
-#include "Hash/CityHash.h"
-
 #include "MuT/Node.h"
-#include "MuT/NodeImage.h"
-#include "MuT/NodeMesh.h"
-#include "MuT/NodeObject.h"
-
-#include "MuR/Operations.h"
-#include "MuR/ModelPrivate.h"
-#include "MuR/SystemPrivate.h"
-
-#include "MuR/Model.h"
-#include "MuR/Instance.h"
-
-#include "MuR/Config.h"
-#include "MuR/MutableTrace.h"
+#include "MuT/NodePrivate.h"
+#include "MuT/Table.h"
+#include "MuT/TaskManager.h"
+#include "Trace/Detail/Channel.h"
+
+#include <algorithm>
+#include <memory>
+#include <utility>
 
 
 namespace mu
diff --git a/Engine/Plugins/Experimental/Mutable/Source/MutableTools/Private/MuT/DataPacker.cpp b/Engine/Plugins/Experimental/Mutable/Source/MutableTools/Private/MuT/DataPacker.cpp
--- a/Engine/Plugins/Experimental/Mutable/Source/MutableTools/Private/MuT/DataPacker.cpp
+++ b/Engine/Plugins/Experimental/Mutable/Source/MutableTools/Private/MuT/DataPacker.cpp
@@ -1,26 +1,32 @@
 // Copyright Epic Games, Inc. All Rights Reserved.
 
 #include "MuT/DataPacker.h"
-#include "MuT/ErrorLogPrivate.h"
-#include "MuT/Streams.h"
 
+#include "Containers/Array.h"
+#include "HAL/PlatformCrt.h"
+#include "HAL/UnrealMemory.h"
 #include "MuR/CodeVisitor.h"
+#include "MuR/Image.h"
+#include "MuR/Layout.h"
+#include "MuR/Mesh.h"
+#include "MuR/MeshBufferSet.h"
 #include "MuR/ModelPrivate.h"
-#include "MuR/SystemPrivate.h"
-#include "MuR/Operations.h"
-#include "MuR/OpImagePixelFormat.h"
-
 #include "MuR/MutableTrace.h"
-#include "MuT/ASTOpImageCompose.h"
+#include "MuR/OpImagePixelFormat.h"
+#include "MuR/Operations.h"
+#include "MuR/Ptr.h"
+#include "MuR/RefCounted.h"
 #include "MuT/ASTOpConstantResource.h"
+#include "MuT/ASTOpImageCompose.h"
 #include "MuT/ASTOpImageMultiLayer.h"
-#include "MuT/ASTOpMeshRemoveMask.h"
 #include "MuT/ASTOpInstanceAdd.h"
 #include "MuT/ASTOpMeshExtractLayoutBlocks.h"
+#include "MuT/ASTOpMeshRemoveMask.h"
 
-#include <algorithm>
 #include <cstdint>
-#include <random>
+#include <functional>
+#include <memory>
+#include <utility>
 
 
 namespace mu
diff --git a/Engine/Plugins/Experimental/Mutable/Source/MutableTools/Private/MuT/DataPacker.h b/Engine/Plugins/Experimental/Mutable/Source/MutableTools/Private/MuT/DataPacker.h
--- a/Engine/Plugins/Experimental/Mutable/Source/MutableTools/Private/MuT/DataPacker.h
+++ b/Engine/Plugins/Experimental/Mutable/Source/MutableTools/Private/MuT/DataPacker.h
@@ -2,12 +2,13 @@
 
 #pragma once
 
-#include "MuT/Compiler.h"
-#include "MuT/AST.h"
 #include "MuR/MemoryPrivate.h"
-#include "MuR/ModelPrivate.h"
-#include "MuR/CodeVisitor.h"
-#include "MuR/RefCounted.h"
+#include "MuR/Operations.h"
+#include "MuT/AST.h"
+
+#include <stdint.h>
+
+namespace mu { struct PROGRAM; }
 
 
 #define MUTABLE_HASH_SEED					((uint32_t)0xcadababa)
diff --git a/Engine/Plugins/Experimental/Mutable/Source/MutableTools/Private/MuT/ErrorLog.cpp b/Engine/Plugins/Experimental/Mutable/Source/MutableTools/Private/MuT/ErrorLog.cpp
--- a/Engine/Plugins/Experimental/Mutable/Source/MutableTools/Private/MuT/ErrorLog.cpp
+++ b/Engine/Plugins/Experimental/Mutable/Source/MutableTools/Private/MuT/ErrorLog.cpp
@@ -1,11 +1,33 @@
 // Copyright Epic Games, Inc. All Rights Reserved.
 
-
 #include "MuT/ErrorLog.h"
 
-#include "MuT/ErrorLogPrivate.h"
-#include "MuR/ModelPrivate.h"
+#include "Containers/Array.h"
+#include "Containers/Map.h"
+#include "Containers/UnrealString.h"
+#include "HAL/PlatformMath.h"
+#include "HAL/UnrealMemory.h"
+#include "Logging/LogCategory.h"
+#include "Logging/LogMacros.h"
+#include "Misc/AssertionMacros.h"
+#include "MuR/Layout.h"
 #include "MuR/MemoryPrivate.h"
+#include "MuR/Mesh.h"
+#include "MuR/Model.h"
+#include "MuR/ModelPrivate.h"
+#include "MuR/MutableMath.h"
+#include "MuR/Operations.h"
+#include "MuR/ParametersPrivate.h"
+#include "MuR/Platform.h"
+#include "MuR/Skeleton.h"
+#include "MuR/System.h"
+#include "MuT/ErrorLogPrivate.h"
+#include "Templates/Tuple.h"
+#include "Trace/Detail/Channel.h"
+
+#include <functional>
+#include <memory>
+#include <utility>
 
 namespace mu
 {
diff --git a/Engine/Plugins/Experimental/Mutable/Source/MutableTools/Private/MuT/ImageDescGenerator.cpp b/Engine/Plugins/Experimental/Mutable/Source/MutableTools/Private/MuT/ImageDescGenerator.cpp
--- a/Engine/Plugins/Experimental/Mutable/Source/MutableTools/Private/MuT/ImageDescGenerator.cpp
+++ b/Engine/Plugins/Experimental/Mutable/Source/MutableTools/Private/MuT/ImageDescGenerator.cpp
@@ -3,40 +3,46 @@
 
 #include "MuT/ImageDescGenerator.h"
 
-#include "MuT/NodeImagePrivate.h"
+#include "Math/IntVector.h"
+#include "Math/UnrealMathSSE.h"
+#include "Misc/AssertionMacros.h"
+#include "MuR/ImagePrivate.h"
+#include "MuR/MutableMath.h"
+#include "MuR/RefCounted.h"
+#include "MuR/Serialisation.h"
+#include "MuT/AST.h"
+#include "MuT/NodeImage.h"
+#include "MuT/NodeImageBinarisePrivate.h"
+#include "MuT/NodeImageColourMapPrivate.h"
+#include "MuT/NodeImageConditionalPrivate.h"
 #include "MuT/NodeImageConstantPrivate.h"
 #include "MuT/NodeImageDifferencePrivate.h"
-#include "MuT/NodeImageInterpolatePrivate.h"
-#include "MuT/NodeImageInterpolate3Private.h"
-#include "MuT/NodeImageSaturatePrivate.h"
-#include "MuT/NodeImageLuminancePrivate.h"
-#include "MuT/NodeImageTablePrivate.h"
-#include "MuT/NodeImageSwizzlePrivate.h"
 #include "MuT/NodeImageFormatPrivate.h"
-#include "MuT/NodeImageSelectColourPrivate.h"
-#include "MuT/NodeImageColourMapPrivate.h"
 #include "MuT/NodeImageGradientPrivate.h"
-#include "MuT/NodeImageBinarisePrivate.h"
-#include "MuT/NodeImageLayerPrivate.h"
+#include "MuT/NodeImageInterpolate3Private.h"
+#include "MuT/NodeImageInterpolatePrivate.h"
+#include "MuT/NodeImageInvertPrivate.h"
 #include "MuT/NodeImageLayerColourPrivate.h"
+#include "MuT/NodeImageLayerPrivate.h"
+#include "MuT/NodeImageLuminancePrivate.h"
+#include "MuT/NodeImageMipmapPrivate.h"
 #include "MuT/NodeImageMultiLayerPrivate.h"
-#include "MuT/NodeImageResizePrivate.h"
+#include "MuT/NodeImageNormalCompositePrivate.h"
 #include "MuT/NodeImagePlainColourPrivate.h"
-#include "MuT/NodeImageSwitchPrivate.h"
-#include "MuT/NodeImageVariationPrivate.h"
 #include "MuT/NodeImageProjectPrivate.h"
-#include "MuT/NodeImageMipmapPrivate.h"
-#include "MuT/NodeImageConditionalPrivate.h"
-#include "MuT/NodeImageParameterPrivate.h"
-#include "MuT/NodeImageInvertPrivate.h"
-#include "MuT/NodeImageNormalCompositePrivate.h"
+#include "MuT/NodeImageResizePrivate.h"
+#include "MuT/NodeImageSaturatePrivate.h"
+#include "MuT/NodeImageSelectColourPrivate.h"
+#include "MuT/NodeImageSwitchPrivate.h"
+#include "MuT/NodeImageSwizzlePrivate.h"
+#include "MuT/NodeImageTablePrivate.h"
 #include "MuT/NodeImageTransformPrivate.h"
-
+#include "MuT/NodeImageVariationPrivate.h"
+#include "MuT/NodePrivate.h"
+#include "MuT/Table.h"
 #include "MuT/TablePrivate.h"
 
-#include "MuR/MeshPrivate.h"
-#include "MuR/ImagePrivate.h"
-#include "MuR/Platform.h"
+#include <utility>
 
 
 
diff --git a/Engine/Plugins/Experimental/Mutable/Source/MutableTools/Private/MuT/ImageDescGenerator.h b/Engine/Plugins/Experimental/Mutable/Source/MutableTools/Private/MuT/ImageDescGenerator.h
--- a/Engine/Plugins/Experimental/Mutable/Source/MutableTools/Private/MuT/ImageDescGenerator.h
+++ b/Engine/Plugins/Experimental/Mutable/Source/MutableTools/Private/MuT/ImageDescGenerator.h
@@ -2,42 +2,41 @@
 
 #pragma once
 
-#include "MuT/Visitor.h"
-
-#include "MuT/AST.h"
-#include "MuT/NodeImage.h"
+#include "MuR/Image.h"
+#include "MuR/MemoryPrivate.h"
+#include "MuR/MutableMemory.h"
+#include "MuR/Ptr.h"
+#include "MuT/Node.h"
+#include "MuT/NodeImageBinarise.h"
+#include "MuT/NodeImageColourMap.h"
+#include "MuT/NodeImageConditional.h"
 #include "MuT/NodeImageConstant.h"
 #include "MuT/NodeImageDifference.h"
+#include "MuT/NodeImageFormat.h"
+#include "MuT/NodeImageGradient.h"
 #include "MuT/NodeImageInterpolate.h"
 #include "MuT/NodeImageInterpolate3.h"
-#include "MuT/NodeImageSaturate.h"
-#include "MuT/NodeImageLuminance.h"
+#include "MuT/NodeImageInvert.h"
 #include "MuT/NodeImageLayer.h"
 #include "MuT/NodeImageLayerColour.h"
+#include "MuT/NodeImageLuminance.h"
+#include "MuT/NodeImageMipmap.h"
 #include "MuT/NodeImageMultiLayer.h"
-#include "MuT/NodeImageTable.h"
-#include "MuT/NodeImageSwizzle.h"
-#include "MuT/NodeImageFormat.h"
-#include "MuT/NodeImageSelectColour.h"
-#include "MuT/NodeImageColourMap.h"
-#include "MuT/NodeImageGradient.h"
-#include "MuT/NodeImageBinarise.h"
-#include "MuT/NodeImageResize.h"
+#include "MuT/NodeImageNormalComposite.h"
+#include "MuT/NodeImageParameter.h"
 #include "MuT/NodeImagePlainColour.h"
-#include "MuT/NodeImageSwitch.h"
-#include "MuT/NodeImageVariation.h"
 #include "MuT/NodeImageProject.h"
-#include "MuT/NodeImageMipmap.h"
-#include "MuT/NodeImageConditional.h"
-#include "MuT/NodeImageParameter.h"
-#include "MuT/NodeImageInvert.h"
-#include "MuT/NodeImageNormalComposite.h"
+#include "MuT/NodeImageResize.h"
+#include "MuT/NodeImageSaturate.h"
+#include "MuT/NodeImageSelectColour.h"
+#include "MuT/NodeImageSwitch.h"
+#include "MuT/NodeImageSwizzle.h"
+#include "MuT/NodeImageTable.h"
 #include "MuT/NodeImageTransform.h"
+#include "MuT/NodeImageVariation.h"
+#include "MuT/Visitor.h"
 
-#include "MuT/ErrorLogPrivate.h"
-
-#include "MuR/Operations.h"
-#include "MuR/ImagePrivate.h"
+namespace mu { class ASTOp; }
 
 
 #define MUTABLE_MISSING_IMAGE_DESC 	FImageDesc( FImageSize( 16, 16 ), EImageFormat::IF_RGB_UBYTE, 1 )
diff --git a/Engine/Plugins/Experimental/Mutable/Source/MutableTools/Private/MuT/ModelReport.cpp b/Engine/Plugins/Experimental/Mutable/Source/MutableTools/Private/MuT/ModelReport.cpp
--- a/Engine/Plugins/Experimental/Mutable/Source/MutableTools/Private/MuT/ModelReport.cpp
+++ b/Engine/Plugins/Experimental/Mutable/Source/MutableTools/Private/MuT/ModelReport.cpp
@@ -2,9 +2,10 @@
 
 
 #include "MuT/ModelReport.h"
-#include "MuT/ModelReportPrivate.h"
-
 
+#include "Misc/AssertionMacros.h"
+#include "MuR/MemoryPrivate.h"
+#include "MuT/ModelReportPrivate.h"
 
 namespace mu
 {
diff --git a/Engine/Plugins/Experimental/Mutable/Source/MutableTools/Private/MuT/MutableToolsModule.cpp b/Engine/Plugins/Experimental/Mutable/Source/MutableTools/Private/MuT/MutableToolsModule.cpp
--- a/Engine/Plugins/Experimental/Mutable/Source/MutableTools/Private/MuT/MutableToolsModule.cpp
+++ b/Engine/Plugins/Experimental/Mutable/Source/MutableTools/Private/MuT/MutableToolsModule.cpp
@@ -1,7 +1,7 @@
 // Copyright Epic Games, Inc. All Rights Reserved.
 
 #include "MuT/MutableToolsModule.h"
-#include "CoreMinimal.h"
+
 #include "Modules/ModuleManager.h"
 
 IMPLEMENT_MODULE(FMutableToolsModule, MutableTools);
diff --git a/Engine/Plugins/Experimental/Mutable/Source/MutableTools/Private/MuT/Node.cpp b/Engine/Plugins/Experimental/Mutable/Source/MutableTools/Private/MuT/Node.cpp
--- a/Engine/Plugins/Experimental/Mutable/Source/MutableTools/Private/MuT/Node.cpp
+++ b/Engine/Plugins/Experimental/Mutable/Source/MutableTools/Private/MuT/Node.cpp
@@ -2,24 +2,32 @@
 
 
 #include "MuT/Node.h"
-#include "MuT/NodePrivate.h"
 
+#include "Misc/AssertionMacros.h"
+#include "MuR/MemoryPrivate.h"
+#include "MuR/Serialisation.h"
+#include "MuR/SerialisationPrivate.h"
+#include "MuT/NodeBool.h"
 #include "MuT/NodeColour.h"
 #include "MuT/NodeComponent.h"
 #include "MuT/NodeImage.h"
-#include "MuT/NodeLayout.h"
 #include "MuT/NodeLOD.h"
+#include "MuT/NodeLayout.h"
 #include "MuT/NodeMesh.h"
+#include "MuT/NodeModifier.h"
 #include "MuT/NodeObject.h"
 #include "MuT/NodePatchImage.h"
 #include "MuT/NodePatchMesh.h"
+#include "MuT/NodePrivate.h"
+#include "MuT/NodeProjector.h"
+#include "MuT/NodeRange.h"
 #include "MuT/NodeScalar.h"
 #include "MuT/NodeString.h"
-#include "MuT/NodeProjector.h"
 #include "MuT/NodeSurface.h"
-#include "MuT/NodeModifier.h"
-#include "MuT/NodeRange.h"
-#include "MuT/NodeBool.h"
+#include "map"
+
+#include <stdint.h>
+#include <utility>
 
 
 namespace mu
diff --git a/Engine/Plugins/Experimental/Mutable/Source/MutableTools/Private/MuT/NodeBool.cpp b/Engine/Plugins/Experimental/Mutable/Source/MutableTools/Private/MuT/NodeBool.cpp
--- a/Engine/Plugins/Experimental/Mutable/Source/MutableTools/Private/MuT/NodeBool.cpp
+++ b/Engine/Plugins/Experimental/Mutable/Source/MutableTools/Private/MuT/NodeBool.cpp
@@ -2,7 +2,12 @@
 
 
 #include "MuT/NodeBool.h"
-#include "MuT/NodePrivate.h"
+
+#include "Misc/AssertionMacros.h"
+#include "MuR/Serialisation.h"
+#include "MuR/SerialisationPrivate.h"
+
+#include <stdint.h>
 
 
 namespace mu
diff --git a/Engine/Plugins/Experimental/Mutable/Source/MutableTools/Private/MuT/NodeBoolAnd.cpp b/Engine/Plugins/Experimental/Mutable/Source/MutableTools/Private/MuT/NodeBoolAnd.cpp
--- a/Engine/Plugins/Experimental/Mutable/Source/MutableTools/Private/MuT/NodeBoolAnd.cpp
+++ b/Engine/Plugins/Experimental/Mutable/Source/MutableTools/Private/MuT/NodeBoolAnd.cpp
@@ -1,10 +1,11 @@
 // Copyright Epic Games, Inc. All Rights Reserved.
 
 
+#include "Misc/AssertionMacros.h"
+#include "MuT/Node.h"
+#include "MuT/NodeBool.h"
 #include "MuT/NodeBoolPrivate.h"
-
-#include "MuR/MutableMath.h"
-#include "MuR/MemoryPrivate.h"
+#include "MuT/NodePrivate.h"
 
 
 namespace mu
diff --git a/Engine/Plugins/Experimental/Mutable/Source/MutableTools/Private/MuT/NodeBoolConstant.cpp b/Engine/Plugins/Experimental/Mutable/Source/MutableTools/Private/MuT/NodeBoolConstant.cpp
--- a/Engine/Plugins/Experimental/Mutable/Source/MutableTools/Private/MuT/NodeBoolConstant.cpp
+++ b/Engine/Plugins/Experimental/Mutable/Source/MutableTools/Private/MuT/NodeBoolConstant.cpp
@@ -1,9 +1,11 @@
 // Copyright Epic Games, Inc. All Rights Reserved.
 
 
+#include "Misc/AssertionMacros.h"
+#include "MuT/Node.h"
+#include "MuT/NodeBool.h"
 #include "MuT/NodeBoolPrivate.h"
-
-#include "MuR/MutableMath.h"
+#include "MuT/NodePrivate.h"
 
 
 #define NODE_INPUT_COUNT 	0
diff --git a/Engine/Plugins/Experimental/Mutable/Source/MutableTools/Private/MuT/NodeBoolIsNull.cpp b/Engine/Plugins/Experimental/Mutable/Source/MutableTools/Private/MuT/NodeBoolIsNull.cpp
--- a/Engine/Plugins/Experimental/Mutable/Source/MutableTools/Private/MuT/NodeBoolIsNull.cpp
+++ b/Engine/Plugins/Experimental/Mutable/Source/MutableTools/Private/MuT/NodeBoolIsNull.cpp
@@ -1,10 +1,10 @@
 // Copyright Epic Games, Inc. All Rights Reserved.
 
 
+#include "MuT/Node.h"
+#include "MuT/NodeBool.h"
 #include "MuT/NodeBoolPrivate.h"
-
-#include "MuR/MutableMath.h"
-#include "MuR/MemoryPrivate.h"
+#include "MuT/NodePrivate.h"
 
 
 namespace mu
diff --git a/Engine/Plugins/Experimental/Mutable/Source/MutableTools/Private/MuT/NodeBoolNot.cpp b/Engine/Plugins/Experimental/Mutable/Source/MutableTools/Private/MuT/NodeBoolNot.cpp
--- a/Engine/Plugins/Experimental/Mutable/Source/MutableTools/Private/MuT/NodeBoolNot.cpp
+++ b/Engine/Plugins/Experimental/Mutable/Source/MutableTools/Private/MuT/NodeBoolNot.cpp
@@ -1,10 +1,11 @@
 // Copyright Epic Games, Inc. All Rights Reserved.
 
 
+#include "Misc/AssertionMacros.h"
+#include "MuT/Node.h"
+#include "MuT/NodeBool.h"
 #include "MuT/NodeBoolPrivate.h"
-
-#include "MuR/MutableMath.h"
-#include "MuR/MemoryPrivate.h"
+#include "MuT/NodePrivate.h"
 
 
 namespace mu
diff --git a/Engine/Plugins/Experimental/Mutable/Source/MutableTools/Private/MuT/NodeBoolParameter.cpp b/Engine/Plugins/Experimental/Mutable/Source/MutableTools/Private/MuT/NodeBoolParameter.cpp
--- a/Engine/Plugins/Experimental/Mutable/Source/MutableTools/Private/MuT/NodeBoolParameter.cpp
+++ b/Engine/Plugins/Experimental/Mutable/Source/MutableTools/Private/MuT/NodeBoolParameter.cpp
@@ -1,10 +1,18 @@
 // Copyright Epic Games, Inc. All Rights Reserved.
 
 
+#include "Misc/AssertionMacros.h"
+#include "MuR/MemoryPrivate.h"
+#include "MuR/Ptr.h"
+#include "MuR/RefCounted.h"
+#include "MuT/Node.h"
+#include "MuT/NodeBool.h"
 #include "MuT/NodeBoolPrivate.h"
+#include "MuT/NodePrivate.h"
+#include "MuT/NodeRange.h"
 
-#include "MuR/MutableMath.h"
-#include "MuR/MemoryPrivate.h"
+#include <memory>
+#include <utility>
 
 
 namespace mu
diff --git a/Engine/Plugins/Experimental/Mutable/Source/MutableTools/Private/MuT/NodeColour.cpp b/Engine/Plugins/Experimental/Mutable/Source/MutableTools/Private/MuT/NodeColour.cpp
--- a/Engine/Plugins/Experimental/Mutable/Source/MutableTools/Private/MuT/NodeColour.cpp
+++ b/Engine/Plugins/Experimental/Mutable/Source/MutableTools/Private/MuT/NodeColour.cpp
@@ -2,17 +2,21 @@
 
 
 #include "MuT/NodeColour.h"
-#include "MuT/NodePrivate.h"
 
+#include "Misc/AssertionMacros.h"
+#include "MuR/Serialisation.h"
+#include "MuR/SerialisationPrivate.h"
+#include "MuT/NodeColourArithmeticOperation.h"
 #include "MuT/NodeColourConstant.h"
+#include "MuT/NodeColourFromScalars.h"
 #include "MuT/NodeColourParameter.h"
 #include "MuT/NodeColourSampleImage.h"
-#include "MuT/NodeColourTable.h"
-#include "MuT/NodeColourFromScalars.h"
-#include "MuT/NodeColourArithmeticOperation.h"
 #include "MuT/NodeColourSwitch.h"
+#include "MuT/NodeColourTable.h"
 #include "MuT/NodeColourVariation.h"
 
+#include <stdint.h>
+
 
 namespace mu
 {
diff --git a/Engine/Plugins/Experimental/Mutable/Source/MutableTools/Private/MuT/NodeColourArithmeticOperation.cpp b/Engine/Plugins/Experimental/Mutable/Source/MutableTools/Private/MuT/NodeColourArithmeticOperation.cpp
--- a/Engine/Plugins/Experimental/Mutable/Source/MutableTools/Private/MuT/NodeColourArithmeticOperation.cpp
+++ b/Engine/Plugins/Experimental/Mutable/Source/MutableTools/Private/MuT/NodeColourArithmeticOperation.cpp
@@ -2,11 +2,10 @@
 
 
 #include "MuT/NodeColourArithmeticOperation.h"
-#include "MuT/NodeColourArithmeticOperationPrivate.h"
-
-#include "MuT/NodeScalar.h"
 
-#include "MuR/ImagePrivate.h"
+#include "Misc/AssertionMacros.h"
+#include "MuT/NodeColourArithmeticOperationPrivate.h"
+#include "MuT/NodePrivate.h"
 
 
 #define NODE_INPUT_COUNT 	2
diff --git a/Engine/Plugins/Experimental/Mutable/Source/MutableTools/Private/MuT/NodeColourConstant.cpp b/Engine/Plugins/Experimental/Mutable/Source/MutableTools/Private/MuT/NodeColourConstant.cpp
--- a/Engine/Plugins/Experimental/Mutable/Source/MutableTools/Private/MuT/NodeColourConstant.cpp
+++ b/Engine/Plugins/Experimental/Mutable/Source/MutableTools/Private/MuT/NodeColourConstant.cpp
@@ -2,9 +2,11 @@
 
 
 #include "MuT/NodeColourConstant.h"
-#include "MuT/NodeColourConstantPrivate.h"
 
+#include "Misc/AssertionMacros.h"
 #include "MuR/MutableMath.h"
+#include "MuT/NodeColourConstantPrivate.h"
+#include "MuT/NodePrivate.h"
 
 
 #define NODE_INPUT_COUNT 	0
diff --git a/Engine/Plugins/Experimental/Mutable/Source/MutableTools/Private/MuT/NodeColourFromScalars.cpp b/Engine/Plugins/Experimental/Mutable/Source/MutableTools/Private/MuT/NodeColourFromScalars.cpp
--- a/Engine/Plugins/Experimental/Mutable/Source/MutableTools/Private/MuT/NodeColourFromScalars.cpp
+++ b/Engine/Plugins/Experimental/Mutable/Source/MutableTools/Private/MuT/NodeColourFromScalars.cpp
@@ -2,12 +2,12 @@
 
 
 #include "MuT/NodeColourFromScalars.h"
-#include "MuT/NodeColourFromScalarsPrivate.h"
 
+#include "Misc/AssertionMacros.h"
+#include "MuT/NodeColourFromScalarsPrivate.h"
+#include "MuT/NodePrivate.h"
 #include "MuT/NodeScalar.h"
 
-#include "MuR/ImagePrivate.h"
-
 
 #define NODE_INPUT_COUNT 	4
 
diff --git a/Engine/Plugins/Experimental/Mutable/Source/MutableTools/Private/MuT/NodeColourParameter.cpp b/Engine/Plugins/Experimental/Mutable/Source/MutableTools/Private/MuT/NodeColourParameter.cpp
--- a/Engine/Plugins/Experimental/Mutable/Source/MutableTools/Private/MuT/NodeColourParameter.cpp
+++ b/Engine/Plugins/Experimental/Mutable/Source/MutableTools/Private/MuT/NodeColourParameter.cpp
@@ -2,10 +2,16 @@
 
 
 #include "MuT/NodeColourParameter.h"
-#include "MuT/NodeColourParameterPrivate.h"
 
-#include "MuR/MutableMath.h"
+#include "Misc/AssertionMacros.h"
 #include "MuR/MemoryPrivate.h"
+#include "MuR/MutableMath.h"
+#include "MuT/NodeColourParameterPrivate.h"
+#include "MuT/NodePrivate.h"
+#include "MuT/NodeRange.h"
+
+#include <memory>
+#include <utility>
 
 
 namespace mu
diff --git a/Engine/Plugins/Experimental/Mutable/Source/MutableTools/Private/MuT/NodeColourSampleImage.cpp b/Engine/Plugins/Experimental/Mutable/Source/MutableTools/Private/MuT/NodeColourSampleImage.cpp
--- a/Engine/Plugins/Experimental/Mutable/Source/MutableTools/Private/MuT/NodeColourSampleImage.cpp
+++ b/Engine/Plugins/Experimental/Mutable/Source/MutableTools/Private/MuT/NodeColourSampleImage.cpp
@@ -2,12 +2,13 @@
 
 
 #include "MuT/NodeColourSampleImage.h"
-#include "MuT/NodeColourSampleImagePrivate.h"
 
+#include "Misc/AssertionMacros.h"
+#include "MuT/NodeColourSampleImagePrivate.h"
+#include "MuT/NodeImage.h"
+#include "MuT/NodePrivate.h"
 #include "MuT/NodeScalar.h"
 
-#include "MuR/ImagePrivate.h"
-
 
 #define NODE_INPUT_COUNT 	3
 
diff --git a/Engine/Plugins/Experimental/Mutable/Source/MutableTools/Private/MuT/NodeColourSwitch.cpp b/Engine/Plugins/Experimental/Mutable/Source/MutableTools/Private/MuT/NodeColourSwitch.cpp
--- a/Engine/Plugins/Experimental/Mutable/Source/MutableTools/Private/MuT/NodeColourSwitch.cpp
+++ b/Engine/Plugins/Experimental/Mutable/Source/MutableTools/Private/MuT/NodeColourSwitch.cpp
@@ -2,9 +2,16 @@
 
 
 #include "MuT/NodeColourSwitch.h"
+
+#include "Misc/AssertionMacros.h"
+#include "MuR/MemoryPrivate.h"
 #include "MuT/NodeColourSwitchPrivate.h"
+#include "MuT/NodePrivate.h"
 #include "MuT/NodeScalar.h"
 
+#include <memory>
+#include <utility>
+
 namespace mu
 {
 
diff --git a/Engine/Plugins/Experimental/Mutable/Source/MutableTools/Private/MuT/NodeColourTable.cpp b/Engine/Plugins/Experimental/Mutable/Source/MutableTools/Private/MuT/NodeColourTable.cpp
--- a/Engine/Plugins/Experimental/Mutable/Source/MutableTools/Private/MuT/NodeColourTable.cpp
+++ b/Engine/Plugins/Experimental/Mutable/Source/MutableTools/Private/MuT/NodeColourTable.cpp
@@ -2,7 +2,12 @@
 
 
 #include "MuT/NodeColourTable.h"
+
+#include "Misc/AssertionMacros.h"
+#include "MuR/MemoryPrivate.h"
 #include "MuT/NodeColourTablePrivate.h"
+#include "MuT/NodePrivate.h"
+#include "MuT/Table.h"
 
 
 
diff --git a/Engine/Plugins/Experimental/Mutable/Source/MutableTools/Private/MuT/NodeColourVariation.cpp b/Engine/Plugins/Experimental/Mutable/Source/MutableTools/Private/MuT/NodeColourVariation.cpp
--- a/Engine/Plugins/Experimental/Mutable/Source/MutableTools/Private/MuT/NodeColourVariation.cpp
+++ b/Engine/Plugins/Experimental/Mutable/Source/MutableTools/Private/MuT/NodeColourVariation.cpp
@@ -2,7 +2,14 @@
 
 
 #include "MuT/NodeColourVariation.h"
+
+#include "Misc/AssertionMacros.h"
+#include "MuR/MemoryPrivate.h"
 #include "MuT/NodeColourVariationPrivate.h"
+#include "MuT/NodePrivate.h"
+
+#include <memory>
+#include <utility>
 
 
 namespace mu
diff --git a/Engine/Plugins/Experimental/Mutable/Source/MutableTools/Private/MuT/NodeComponent.cpp b/Engine/Plugins/Experimental/Mutable/Source/MutableTools/Private/MuT/NodeComponent.cpp
--- a/Engine/Plugins/Experimental/Mutable/Source/MutableTools/Private/MuT/NodeComponent.cpp
+++ b/Engine/Plugins/Experimental/Mutable/Source/MutableTools/Private/MuT/NodeComponent.cpp
@@ -2,13 +2,14 @@
 
 
 #include "MuT/NodeComponent.h"
-#include "MuT/NodeComponentPrivate.h"
 
-#include "MuT/NodeMesh.h"
-#include "MuT/NodeImage.h"
-
-#include "MuT/NodeComponentNew.h"
+#include "Misc/AssertionMacros.h"
+#include "MuR/Serialisation.h"
+#include "MuR/SerialisationPrivate.h"
 #include "MuT/NodeComponentEdit.h"
+#include "MuT/NodeComponentNew.h"
+
+#include <stdint.h>
 
 
 namespace mu
diff --git a/Engine/Plugins/Experimental/Mutable/Source/MutableTools/Private/MuT/NodeComponentEdit.cpp b/Engine/Plugins/Experimental/Mutable/Source/MutableTools/Private/MuT/NodeComponentEdit.cpp
--- a/Engine/Plugins/Experimental/Mutable/Source/MutableTools/Private/MuT/NodeComponentEdit.cpp
+++ b/Engine/Plugins/Experimental/Mutable/Source/MutableTools/Private/MuT/NodeComponentEdit.cpp
@@ -2,10 +2,17 @@
 
 
 #include "MuT/NodeComponentEdit.h"
-#include "MuT/NodeComponentEditPrivate.h"
 
+#include "Misc/AssertionMacros.h"
+#include "MuR/MemoryPrivate.h"
+#include "MuT/NodeComponentEditPrivate.h"
+#include "MuT/NodeComponentNew.h"
+#include "MuT/NodePrivate.h"
 #include "MuT/NodeSurface.h"
 
+#include <memory>
+#include <utility>
+
 
 namespace mu
 {
diff --git a/Engine/Plugins/Experimental/Mutable/Source/MutableTools/Private/MuT/NodeComponentNew.cpp b/Engine/Plugins/Experimental/Mutable/Source/MutableTools/Private/MuT/NodeComponentNew.cpp
--- a/Engine/Plugins/Experimental/Mutable/Source/MutableTools/Private/MuT/NodeComponentNew.cpp
+++ b/Engine/Plugins/Experimental/Mutable/Source/MutableTools/Private/MuT/NodeComponentNew.cpp
@@ -2,10 +2,16 @@
 
 
 #include "MuT/NodeComponentNew.h"
-#include "MuT/NodeComponentNewPrivate.h"
 
+#include "Misc/AssertionMacros.h"
+#include "MuR/MemoryPrivate.h"
+#include "MuT/NodeComponentNewPrivate.h"
+#include "MuT/NodePrivate.h"
 #include "MuT/NodeSurface.h"
 
+#include <memory>
+#include <utility>
+
 
 namespace mu
 {
diff --git a/Engine/Plugins/Experimental/Mutable/Source/MutableTools/Private/MuT/NodeImage.cpp b/Engine/Plugins/Experimental/Mutable/Source/MutableTools/Private/MuT/NodeImage.cpp
--- a/Engine/Plugins/Experimental/Mutable/Source/MutableTools/Private/MuT/NodeImage.cpp
+++ b/Engine/Plugins/Experimental/Mutable/Source/MutableTools/Private/MuT/NodeImage.cpp
@@ -2,35 +2,39 @@
 
 
 #include "MuT/NodeImage.h"
-#include "MuT/NodePrivate.h"
 
+#include "Misc/AssertionMacros.h"
+#include "MuR/Serialisation.h"
+#include "MuR/SerialisationPrivate.h"
+#include "MuT/NodeImageBinarise.h"
+#include "MuT/NodeImageColourMap.h"
+#include "MuT/NodeImageConditional.h"
 #include "MuT/NodeImageConstant.h"
 #include "MuT/NodeImageDifference.h"
-#include "MuT/NodeImageInterpolate.h"
-#include "MuT/NodeImageInterpolate3.h"
-#include "MuT/NodeImageSaturate.h"
-#include "MuT/NodeImageLuminance.h"
-#include "MuT/NodeImageTable.h"
-#include "MuT/NodeImageSwizzle.h"
 #include "MuT/NodeImageFormat.h"
-#include "MuT/NodeImageSelectColour.h"
-#include "MuT/NodeImageColourMap.h"
 #include "MuT/NodeImageGradient.h"
-#include "MuT/NodeImageBinarise.h"
+#include "MuT/NodeImageInterpolate.h"
+#include "MuT/NodeImageInterpolate3.h"
+#include "MuT/NodeImageInvert.h"
 #include "MuT/NodeImageLayer.h"
 #include "MuT/NodeImageLayerColour.h"
+#include "MuT/NodeImageLuminance.h"
+#include "MuT/NodeImageMipmap.h"
 #include "MuT/NodeImageMultiLayer.h"
-#include "MuT/NodeImageResize.h"
+#include "MuT/NodeImageNormalComposite.h"
+#include "MuT/NodeImageParameter.h"
 #include "MuT/NodeImagePlainColour.h"
 #include "MuT/NodeImageProject.h"
-#include "MuT/NodeImageMipmap.h"
+#include "MuT/NodeImageResize.h"
+#include "MuT/NodeImageSaturate.h"
+#include "MuT/NodeImageSelectColour.h"
 #include "MuT/NodeImageSwitch.h"
-#include "MuT/NodeImageConditional.h"
-#include "MuT/NodeImageParameter.h"
-#include "MuT/NodeImageInvert.h"
-#include "MuT/NodeImageVariation.h"
-#include "MuT/NodeImageNormalComposite.h"
+#include "MuT/NodeImageSwizzle.h"
+#include "MuT/NodeImageTable.h"
 #include "MuT/NodeImageTransform.h"
+#include "MuT/NodeImageVariation.h"
+
+#include <stdint.h>
 
 
 namespace mu
diff --git a/Engine/Plugins/Experimental/Mutable/Source/MutableTools/Private/MuT/NodeImageBinarise.cpp b/Engine/Plugins/Experimental/Mutable/Source/MutableTools/Private/MuT/NodeImageBinarise.cpp
--- a/Engine/Plugins/Experimental/Mutable/Source/MutableTools/Private/MuT/NodeImageBinarise.cpp
+++ b/Engine/Plugins/Experimental/Mutable/Source/MutableTools/Private/MuT/NodeImageBinarise.cpp
@@ -2,12 +2,12 @@
 
 
 #include "MuT/NodeImageBinarise.h"
-#include "MuT/NodeImageBinarisePrivate.h"
 
+#include "Misc/AssertionMacros.h"
+#include "MuT/NodeImageBinarisePrivate.h"
+#include "MuT/NodePrivate.h"
 #include "MuT/NodeScalar.h"
 
-#include "MuR/ImagePrivate.h"
-
 
 #define NODE_INPUT_COUNT 	2
 
diff --git a/Engine/Plugins/Experimental/Mutable/Source/MutableTools/Private/MuT/NodeImageColourMap.cpp b/Engine/Plugins/Experimental/Mutable/Source/MutableTools/Private/MuT/NodeImageColourMap.cpp
--- a/Engine/Plugins/Experimental/Mutable/Source/MutableTools/Private/MuT/NodeImageColourMap.cpp
+++ b/Engine/Plugins/Experimental/Mutable/Source/MutableTools/Private/MuT/NodeImageColourMap.cpp
@@ -2,11 +2,10 @@
 
 
 #include "MuT/NodeImageColourMap.h"
-#include "MuT/NodeImageColourMapPrivate.h"
-
-#include "MuT/NodeColour.h"
 
-#include "MuR/ImagePrivate.h"
+#include "Misc/AssertionMacros.h"
+#include "MuT/NodeImageColourMapPrivate.h"
+#include "MuT/NodePrivate.h"
 
 
 namespace mu
