Ran IWYU on Mutable plugin

diff --git a/Engine/Plugins/Experimental/Mutable/Source/MutableTools/Private/MuT/NodeImageConditional.cpp b/Engine/Plugins/Experimental/Mutable/Source/MutableTools/Private/MuT/NodeImageConditional.cpp
--- a/Engine/Plugins/Experimental/Mutable/Source/MutableTools/Private/MuT/NodeImageConditional.cpp
+++ b/Engine/Plugins/Experimental/Mutable/Source/MutableTools/Private/MuT/NodeImageConditional.cpp
@@ -2,11 +2,11 @@
 
 
 #include "MuT/NodeImageConditional.h"
-#include "MuT/NodeImageConditionalPrivate.h"
 
+#include "Misc/AssertionMacros.h"
 #include "MuT/NodeBool.h"
-
-#include "MuR/ImagePrivate.h"
+#include "MuT/NodeImageConditionalPrivate.h"
+#include "MuT/NodePrivate.h"
 
 
 namespace mu
diff --git a/Engine/Plugins/Experimental/Mutable/Source/MutableTools/Private/MuT/NodeImageConstant.cpp b/Engine/Plugins/Experimental/Mutable/Source/MutableTools/Private/MuT/NodeImageConstant.cpp
--- a/Engine/Plugins/Experimental/Mutable/Source/MutableTools/Private/MuT/NodeImageConstant.cpp
+++ b/Engine/Plugins/Experimental/Mutable/Source/MutableTools/Private/MuT/NodeImageConstant.cpp
@@ -2,8 +2,12 @@
 
 
 #include "MuT/NodeImageConstant.h"
+
+#include "Misc/AssertionMacros.h"
+#include "MuR/Image.h"
+#include "MuR/Serialisation.h"
 #include "MuT/NodeImageConstantPrivate.h"
-#include "MuR/ImagePrivate.h"
+#include "MuT/NodePrivate.h"
 
 
 namespace mu
diff --git a/Engine/Plugins/Experimental/Mutable/Source/MutableTools/Private/MuT/NodeImageDifference.cpp b/Engine/Plugins/Experimental/Mutable/Source/MutableTools/Private/MuT/NodeImageDifference.cpp
--- a/Engine/Plugins/Experimental/Mutable/Source/MutableTools/Private/MuT/NodeImageDifference.cpp
+++ b/Engine/Plugins/Experimental/Mutable/Source/MutableTools/Private/MuT/NodeImageDifference.cpp
@@ -2,11 +2,10 @@
 
 
 #include "MuT/NodeImageDifference.h"
-#include "MuT/NodeImageDifferencePrivate.h"
-
-#include "MuT/NodeColour.h"
 
-#include "MuR/ImagePrivate.h"
+#include "Misc/AssertionMacros.h"
+#include "MuT/NodeImageDifferencePrivate.h"
+#include "MuT/NodePrivate.h"
 
 
 #define NODE_INPUT_COUNT 	2
diff --git a/Engine/Plugins/Experimental/Mutable/Source/MutableTools/Private/MuT/NodeImageFormat.cpp b/Engine/Plugins/Experimental/Mutable/Source/MutableTools/Private/MuT/NodeImageFormat.cpp
--- a/Engine/Plugins/Experimental/Mutable/Source/MutableTools/Private/MuT/NodeImageFormat.cpp
+++ b/Engine/Plugins/Experimental/Mutable/Source/MutableTools/Private/MuT/NodeImageFormat.cpp
@@ -2,7 +2,10 @@
 
 
 #include "MuT/NodeImageFormat.h"
+
+#include "Misc/AssertionMacros.h"
 #include "MuT/NodeImageFormatPrivate.h"
+#include "MuT/NodePrivate.h"
 
 
 #define NODE_INPUT_COUNT 	4
diff --git a/Engine/Plugins/Experimental/Mutable/Source/MutableTools/Private/MuT/NodeImageGradient.cpp b/Engine/Plugins/Experimental/Mutable/Source/MutableTools/Private/MuT/NodeImageGradient.cpp
--- a/Engine/Plugins/Experimental/Mutable/Source/MutableTools/Private/MuT/NodeImageGradient.cpp
+++ b/Engine/Plugins/Experimental/Mutable/Source/MutableTools/Private/MuT/NodeImageGradient.cpp
@@ -2,11 +2,12 @@
 
 
 #include "MuT/NodeImageGradient.h"
-#include "MuT/NodeImageGradientPrivate.h"
 
+#include "Misc/AssertionMacros.h"
+#include "MuR/MutableMath.h"
 #include "MuT/NodeColour.h"
-
-#include "MuR/ImagePrivate.h"
+#include "MuT/NodeImageGradientPrivate.h"
+#include "MuT/NodePrivate.h"
 
 
 #define NODE_INPUT_COUNT 	2
diff --git a/Engine/Plugins/Experimental/Mutable/Source/MutableTools/Private/MuT/NodeImageInterpolate.cpp b/Engine/Plugins/Experimental/Mutable/Source/MutableTools/Private/MuT/NodeImageInterpolate.cpp
--- a/Engine/Plugins/Experimental/Mutable/Source/MutableTools/Private/MuT/NodeImageInterpolate.cpp
+++ b/Engine/Plugins/Experimental/Mutable/Source/MutableTools/Private/MuT/NodeImageInterpolate.cpp
@@ -2,11 +2,15 @@
 
 
 #include "MuT/NodeImageInterpolate.h"
-#include "MuT/NodeImageInterpolatePrivate.h"
 
+#include "Misc/AssertionMacros.h"
+#include "MuR/MemoryPrivate.h"
+#include "MuT/NodeImageInterpolatePrivate.h"
+#include "MuT/NodePrivate.h"
 #include "MuT/NodeScalar.h"
 
-#include "MuR/ImagePrivate.h"
+#include <memory>
+#include <utility>
 
 
 namespace mu
diff --git a/Engine/Plugins/Experimental/Mutable/Source/MutableTools/Private/MuT/NodeImageInterpolate3.cpp b/Engine/Plugins/Experimental/Mutable/Source/MutableTools/Private/MuT/NodeImageInterpolate3.cpp
--- a/Engine/Plugins/Experimental/Mutable/Source/MutableTools/Private/MuT/NodeImageInterpolate3.cpp
+++ b/Engine/Plugins/Experimental/Mutable/Source/MutableTools/Private/MuT/NodeImageInterpolate3.cpp
@@ -2,12 +2,12 @@
 
 
 #include "MuT/NodeImageInterpolate3.h"
-#include "MuT/NodeImageInterpolate3Private.h"
 
+#include "Misc/AssertionMacros.h"
+#include "MuT/NodeImageInterpolate3Private.h"
+#include "MuT/NodePrivate.h"
 #include "MuT/NodeScalar.h"
 
-#include "MuR/ImagePrivate.h"
-
 
 namespace mu
 {
diff --git a/Engine/Plugins/Experimental/Mutable/Source/MutableTools/Private/MuT/NodeImageInvert.cpp b/Engine/Plugins/Experimental/Mutable/Source/MutableTools/Private/MuT/NodeImageInvert.cpp
--- a/Engine/Plugins/Experimental/Mutable/Source/MutableTools/Private/MuT/NodeImageInvert.cpp
+++ b/Engine/Plugins/Experimental/Mutable/Source/MutableTools/Private/MuT/NodeImageInvert.cpp
@@ -2,9 +2,10 @@
 
 
 #include "MuT/NodeImageInvert.h"
-#include "MuT/NodeImageInvertPrivate.h"
 
-#include "MuR/ImagePrivate.h"
+#include "Misc/AssertionMacros.h"
+#include "MuT/NodeImageInvertPrivate.h"
+#include "MuT/NodePrivate.h"
 
 
 #define NODE_INPUT_COUNT 1
diff --git a/Engine/Plugins/Experimental/Mutable/Source/MutableTools/Private/MuT/NodeImageLayer.cpp b/Engine/Plugins/Experimental/Mutable/Source/MutableTools/Private/MuT/NodeImageLayer.cpp
--- a/Engine/Plugins/Experimental/Mutable/Source/MutableTools/Private/MuT/NodeImageLayer.cpp
+++ b/Engine/Plugins/Experimental/Mutable/Source/MutableTools/Private/MuT/NodeImageLayer.cpp
@@ -2,11 +2,10 @@
 
 
 #include "MuT/NodeImageLayer.h"
-#include "MuT/NodeImageLayerPrivate.h"
-
-#include "MuT/NodeColour.h"
 
-#include "MuR/ImagePrivate.h"
+#include "Misc/AssertionMacros.h"
+#include "MuT/NodeImageLayerPrivate.h"
+#include "MuT/NodePrivate.h"
 
 
 #define NODE_INPUT_COUNT 	3
diff --git a/Engine/Plugins/Experimental/Mutable/Source/MutableTools/Private/MuT/NodeImageLayerColour.cpp b/Engine/Plugins/Experimental/Mutable/Source/MutableTools/Private/MuT/NodeImageLayerColour.cpp
--- a/Engine/Plugins/Experimental/Mutable/Source/MutableTools/Private/MuT/NodeImageLayerColour.cpp
+++ b/Engine/Plugins/Experimental/Mutable/Source/MutableTools/Private/MuT/NodeImageLayerColour.cpp
@@ -2,11 +2,11 @@
 
 
 #include "MuT/NodeImageLayerColour.h"
-#include "MuT/NodeImageLayerColourPrivate.h"
 
+#include "Misc/AssertionMacros.h"
 #include "MuT/NodeColour.h"
-
-#include "MuR/ImagePrivate.h"
+#include "MuT/NodeImageLayerColourPrivate.h"
+#include "MuT/NodePrivate.h"
 
 
 #define NODE_INPUT_COUNT 	3
diff --git a/Engine/Plugins/Experimental/Mutable/Source/MutableTools/Private/MuT/NodeImageLuminance.cpp b/Engine/Plugins/Experimental/Mutable/Source/MutableTools/Private/MuT/NodeImageLuminance.cpp
--- a/Engine/Plugins/Experimental/Mutable/Source/MutableTools/Private/MuT/NodeImageLuminance.cpp
+++ b/Engine/Plugins/Experimental/Mutable/Source/MutableTools/Private/MuT/NodeImageLuminance.cpp
@@ -2,11 +2,10 @@
 
 
 #include "MuT/NodeImageLuminance.h"
-#include "MuT/NodeImageLuminancePrivate.h"
-
-#include "MuT/NodeScalar.h"
 
-#include "MuR/ImagePrivate.h"
+#include "Misc/AssertionMacros.h"
+#include "MuT/NodeImageLuminancePrivate.h"
+#include "MuT/NodePrivate.h"
 
 
 #define NODE_INPUT_COUNT 	1
diff --git a/Engine/Plugins/Experimental/Mutable/Source/MutableTools/Private/MuT/NodeImageMipmap.cpp b/Engine/Plugins/Experimental/Mutable/Source/MutableTools/Private/MuT/NodeImageMipmap.cpp
--- a/Engine/Plugins/Experimental/Mutable/Source/MutableTools/Private/MuT/NodeImageMipmap.cpp
+++ b/Engine/Plugins/Experimental/Mutable/Source/MutableTools/Private/MuT/NodeImageMipmap.cpp
@@ -2,11 +2,11 @@
 
 
 #include "MuT/NodeImageMipmap.h"
-#include "MuT/NodeImageMipmapPrivate.h"
-
-#include "MuT/NodeScalar.h"
 
+#include "Misc/AssertionMacros.h"
 #include "MuR/ImagePrivate.h"
+#include "MuT/NodeImageMipmapPrivate.h"
+#include "MuT/NodePrivate.h"
 
 
 #define NODE_INPUT_COUNT 	1
diff --git a/Engine/Plugins/Experimental/Mutable/Source/MutableTools/Private/MuT/NodeImageMultiLayer.cpp b/Engine/Plugins/Experimental/Mutable/Source/MutableTools/Private/MuT/NodeImageMultiLayer.cpp
--- a/Engine/Plugins/Experimental/Mutable/Source/MutableTools/Private/MuT/NodeImageMultiLayer.cpp
+++ b/Engine/Plugins/Experimental/Mutable/Source/MutableTools/Private/MuT/NodeImageMultiLayer.cpp
@@ -2,11 +2,10 @@
 
 
 #include "MuT/NodeImageMultiLayer.h"
-#include "MuT/NodeImageMultiLayerPrivate.h"
-
-#include "MuT/NodeColour.h"
 
-#include "MuR/ImagePrivate.h"
+#include "Misc/AssertionMacros.h"
+#include "MuT/NodeImageMultiLayerPrivate.h"
+#include "MuT/NodePrivate.h"
 
 
 #define NODE_INPUT_COUNT 	4
diff --git a/Engine/Plugins/Experimental/Mutable/Source/MutableTools/Private/MuT/NodeImageNormalComposite.cpp b/Engine/Plugins/Experimental/Mutable/Source/MutableTools/Private/MuT/NodeImageNormalComposite.cpp
--- a/Engine/Plugins/Experimental/Mutable/Source/MutableTools/Private/MuT/NodeImageNormalComposite.cpp
+++ b/Engine/Plugins/Experimental/Mutable/Source/MutableTools/Private/MuT/NodeImageNormalComposite.cpp
@@ -2,9 +2,10 @@
 
 
 #include "MuT/NodeImageNormalComposite.h"
-#include "MuT/NodeImageNormalCompositePrivate.h"
 
-#include "MuR/ImagePrivate.h"
+#include "Misc/AssertionMacros.h"
+#include "MuT/NodeImageNormalCompositePrivate.h"
+#include "MuT/NodePrivate.h"
 
 
 namespace mu
diff --git a/Engine/Plugins/Experimental/Mutable/Source/MutableTools/Private/MuT/NodeImageParameter.cpp b/Engine/Plugins/Experimental/Mutable/Source/MutableTools/Private/MuT/NodeImageParameter.cpp
--- a/Engine/Plugins/Experimental/Mutable/Source/MutableTools/Private/MuT/NodeImageParameter.cpp
+++ b/Engine/Plugins/Experimental/Mutable/Source/MutableTools/Private/MuT/NodeImageParameter.cpp
@@ -2,10 +2,11 @@
 
 
 #include "MuT/NodeImageParameter.h"
-#include "MuT/NodeImageParameterPrivate.h"
 
-#include "MuR/MutableMath.h"
+#include "Misc/AssertionMacros.h"
 #include "MuR/MemoryPrivate.h"
+#include "MuT/NodeImageParameterPrivate.h"
+#include "MuT/NodePrivate.h"
 
 
 namespace mu
diff --git a/Engine/Plugins/Experimental/Mutable/Source/MutableTools/Private/MuT/NodeImagePlainColour.cpp b/Engine/Plugins/Experimental/Mutable/Source/MutableTools/Private/MuT/NodeImagePlainColour.cpp
--- a/Engine/Plugins/Experimental/Mutable/Source/MutableTools/Private/MuT/NodeImagePlainColour.cpp
+++ b/Engine/Plugins/Experimental/Mutable/Source/MutableTools/Private/MuT/NodeImagePlainColour.cpp
@@ -2,11 +2,11 @@
 
 
 #include "MuT/NodeImagePlainColour.h"
-#include "MuT/NodeImagePlainColourPrivate.h"
 
+#include "Misc/AssertionMacros.h"
 #include "MuT/NodeColour.h"
-
-#include "MuR/ImagePrivate.h"
+#include "MuT/NodeImagePlainColourPrivate.h"
+#include "MuT/NodePrivate.h"
 
 
 #define NODE_INPUT_COUNT 	1
diff --git a/Engine/Plugins/Experimental/Mutable/Source/MutableTools/Private/MuT/NodeImageProject.cpp b/Engine/Plugins/Experimental/Mutable/Source/MutableTools/Private/MuT/NodeImageProject.cpp
--- a/Engine/Plugins/Experimental/Mutable/Source/MutableTools/Private/MuT/NodeImageProject.cpp
+++ b/Engine/Plugins/Experimental/Mutable/Source/MutableTools/Private/MuT/NodeImageProject.cpp
@@ -2,12 +2,15 @@
 
 
 #include "MuT/NodeImageProject.h"
-#include "MuT/NodeImageProjectPrivate.h"
 
+#include "Math/IntVector.h"
+#include "Misc/AssertionMacros.h"
+#include "MuT/NodeImageProjectPrivate.h"
+#include "MuT/NodeMesh.h"
+#include "MuT/NodePrivate.h"
+#include "MuT/NodeProjector.h"
 #include "MuT/NodeScalar.h"
 
-#include "MuR/ImagePrivate.h"
-
 
 
 namespace mu
diff --git a/Engine/Plugins/Experimental/Mutable/Source/MutableTools/Private/MuT/NodeImageResize.cpp b/Engine/Plugins/Experimental/Mutable/Source/MutableTools/Private/MuT/NodeImageResize.cpp
--- a/Engine/Plugins/Experimental/Mutable/Source/MutableTools/Private/MuT/NodeImageResize.cpp
+++ b/Engine/Plugins/Experimental/Mutable/Source/MutableTools/Private/MuT/NodeImageResize.cpp
@@ -2,11 +2,10 @@
 
 
 #include "MuT/NodeImageResize.h"
-#include "MuT/NodeImageResizePrivate.h"
-
-#include "MuT/NodeScalar.h"
 
-#include "MuR/ImagePrivate.h"
+#include "Misc/AssertionMacros.h"
+#include "MuT/NodeImageResizePrivate.h"
+#include "MuT/NodePrivate.h"
 
 
 #define NODE_INPUT_COUNT 	1
diff --git a/Engine/Plugins/Experimental/Mutable/Source/MutableTools/Private/MuT/NodeImageSaturate.cpp b/Engine/Plugins/Experimental/Mutable/Source/MutableTools/Private/MuT/NodeImageSaturate.cpp
--- a/Engine/Plugins/Experimental/Mutable/Source/MutableTools/Private/MuT/NodeImageSaturate.cpp
+++ b/Engine/Plugins/Experimental/Mutable/Source/MutableTools/Private/MuT/NodeImageSaturate.cpp
@@ -2,12 +2,12 @@
 
 
 #include "MuT/NodeImageSaturate.h"
-#include "MuT/NodeImageSaturatePrivate.h"
 
+#include "Misc/AssertionMacros.h"
+#include "MuT/NodeImageSaturatePrivate.h"
+#include "MuT/NodePrivate.h"
 #include "MuT/NodeScalar.h"
 
-#include "MuR/ImagePrivate.h"
-
 
 #define NODE_INPUT_COUNT 	2
 
diff --git a/Engine/Plugins/Experimental/Mutable/Source/MutableTools/Private/MuT/NodeImageSelectColour.cpp b/Engine/Plugins/Experimental/Mutable/Source/MutableTools/Private/MuT/NodeImageSelectColour.cpp
--- a/Engine/Plugins/Experimental/Mutable/Source/MutableTools/Private/MuT/NodeImageSelectColour.cpp
+++ b/Engine/Plugins/Experimental/Mutable/Source/MutableTools/Private/MuT/NodeImageSelectColour.cpp
@@ -2,11 +2,11 @@
 
 
 #include "MuT/NodeImageSelectColour.h"
-#include "MuT/NodeImageSelectColourPrivate.h"
 
+#include "Misc/AssertionMacros.h"
 #include "MuT/NodeColour.h"
-
-#include "MuR/ImagePrivate.h"
+#include "MuT/NodeImageSelectColourPrivate.h"
+#include "MuT/NodePrivate.h"
 
 
 #define NODE_INPUT_COUNT 	2
diff --git a/Engine/Plugins/Experimental/Mutable/Source/MutableTools/Private/MuT/NodeImageSwitch.cpp b/Engine/Plugins/Experimental/Mutable/Source/MutableTools/Private/MuT/NodeImageSwitch.cpp
--- a/Engine/Plugins/Experimental/Mutable/Source/MutableTools/Private/MuT/NodeImageSwitch.cpp
+++ b/Engine/Plugins/Experimental/Mutable/Source/MutableTools/Private/MuT/NodeImageSwitch.cpp
@@ -2,11 +2,15 @@
 
 
 #include "MuT/NodeImageSwitch.h"
-#include "MuT/NodeImageSwitchPrivate.h"
 
+#include "Misc/AssertionMacros.h"
+#include "MuR/MemoryPrivate.h"
+#include "MuT/NodeImageSwitchPrivate.h"
+#include "MuT/NodePrivate.h"
 #include "MuT/NodeScalar.h"
 
-#include "MuR/ImagePrivate.h"
+#include <memory>
+#include <utility>
 
 
 namespace mu
diff --git a/Engine/Plugins/Experimental/Mutable/Source/MutableTools/Private/MuT/NodeImageSwizzle.cpp b/Engine/Plugins/Experimental/Mutable/Source/MutableTools/Private/MuT/NodeImageSwizzle.cpp
--- a/Engine/Plugins/Experimental/Mutable/Source/MutableTools/Private/MuT/NodeImageSwizzle.cpp
+++ b/Engine/Plugins/Experimental/Mutable/Source/MutableTools/Private/MuT/NodeImageSwizzle.cpp
@@ -2,7 +2,15 @@
 
 
 #include "MuT/NodeImageSwizzle.h"
+
+#include "Misc/AssertionMacros.h"
+#include "MuR/ImagePrivate.h"
+#include "MuR/MemoryPrivate.h"
 #include "MuT/NodeImageSwizzlePrivate.h"
+#include "MuT/NodePrivate.h"
+
+#include <memory>
+#include <utility>
 
 
 #define NODE_INPUT_COUNT 	4
diff --git a/Engine/Plugins/Experimental/Mutable/Source/MutableTools/Private/MuT/NodeImageTable.cpp b/Engine/Plugins/Experimental/Mutable/Source/MutableTools/Private/MuT/NodeImageTable.cpp
--- a/Engine/Plugins/Experimental/Mutable/Source/MutableTools/Private/MuT/NodeImageTable.cpp
+++ b/Engine/Plugins/Experimental/Mutable/Source/MutableTools/Private/MuT/NodeImageTable.cpp
@@ -2,7 +2,12 @@
 
 
 #include "MuT/NodeImageTable.h"
+
+#include "Misc/AssertionMacros.h"
+#include "MuR/MemoryPrivate.h"
 #include "MuT/NodeImageTablePrivate.h"
+#include "MuT/NodePrivate.h"
+#include "MuT/Table.h"
 
 
 namespace mu
diff --git a/Engine/Plugins/Experimental/Mutable/Source/MutableTools/Private/MuT/NodeImageTransform.cpp b/Engine/Plugins/Experimental/Mutable/Source/MutableTools/Private/MuT/NodeImageTransform.cpp
--- a/Engine/Plugins/Experimental/Mutable/Source/MutableTools/Private/MuT/NodeImageTransform.cpp
+++ b/Engine/Plugins/Experimental/Mutable/Source/MutableTools/Private/MuT/NodeImageTransform.cpp
@@ -2,12 +2,12 @@
 
 
 #include "MuT/NodeImageTransform.h"
-#include "MuT/NodeImageTransformPrivate.h"
 
+#include "Misc/AssertionMacros.h"
+#include "MuT/NodeImageTransformPrivate.h"
+#include "MuT/NodePrivate.h"
 #include "MuT/NodeScalar.h"
 
-#include "MuR/ImagePrivate.h"
-
 
 #define NODE_INPUT_COUNT 	6
 
diff --git a/Engine/Plugins/Experimental/Mutable/Source/MutableTools/Private/MuT/NodeImageVariation.cpp b/Engine/Plugins/Experimental/Mutable/Source/MutableTools/Private/MuT/NodeImageVariation.cpp
--- a/Engine/Plugins/Experimental/Mutable/Source/MutableTools/Private/MuT/NodeImageVariation.cpp
+++ b/Engine/Plugins/Experimental/Mutable/Source/MutableTools/Private/MuT/NodeImageVariation.cpp
@@ -2,7 +2,14 @@
 
 
 #include "MuT/NodeImageVariation.h"
+
+#include "Misc/AssertionMacros.h"
+#include "MuR/MemoryPrivate.h"
 #include "MuT/NodeImageVariationPrivate.h"
+#include "MuT/NodePrivate.h"
+
+#include <memory>
+#include <utility>
 
 
 namespace mu
diff --git a/Engine/Plugins/Experimental/Mutable/Source/MutableTools/Private/MuT/NodeLOD.cpp b/Engine/Plugins/Experimental/Mutable/Source/MutableTools/Private/MuT/NodeLOD.cpp
--- a/Engine/Plugins/Experimental/Mutable/Source/MutableTools/Private/MuT/NodeLOD.cpp
+++ b/Engine/Plugins/Experimental/Mutable/Source/MutableTools/Private/MuT/NodeLOD.cpp
@@ -2,7 +2,17 @@
 
 
 #include "MuT/NodeLOD.h"
+
+#include "Misc/AssertionMacros.h"
+#include "MuR/MemoryPrivate.h"
+#include "MuT/NodeComponent.h"
+#include "MuT/NodeComponentNew.h"
 #include "MuT/NodeLODPrivate.h"
+#include "MuT/NodeModifier.h"
+#include "MuT/NodePrivate.h"
+
+#include <memory>
+#include <utility>
 
 
 namespace mu
diff --git a/Engine/Plugins/Experimental/Mutable/Source/MutableTools/Private/MuT/NodeLayout.cpp b/Engine/Plugins/Experimental/Mutable/Source/MutableTools/Private/MuT/NodeLayout.cpp
--- a/Engine/Plugins/Experimental/Mutable/Source/MutableTools/Private/MuT/NodeLayout.cpp
+++ b/Engine/Plugins/Experimental/Mutable/Source/MutableTools/Private/MuT/NodeLayout.cpp
@@ -2,9 +2,22 @@
 
 
 #include "MuT/NodeLayout.h"
+
+#include "Math/IntPoint.h"
+#include "Math/UnrealMathSSE.h"
+#include "Misc/AssertionMacros.h"
+#include "MuR/ConvertData.h"
+#include "MuR/MemoryPrivate.h"
+#include "MuR/MeshBufferSet.h"
+#include "MuR/MeshPrivate.h"
+#include "MuR/MutableMath.h"
+#include "MuR/Serialisation.h"
+#include "MuR/SerialisationPrivate.h"
 #include "MuT/NodeLayoutPrivate.h"
+#include "MuT/NodePrivate.h"
 
-#include "MuT/NodeImage.h"
+#include <memory>
+#include <utility>
 
 
 namespace mu
diff --git a/Engine/Plugins/Experimental/Mutable/Source/MutableTools/Private/MuT/NodeMesh.cpp b/Engine/Plugins/Experimental/Mutable/Source/MutableTools/Private/MuT/NodeMesh.cpp
--- a/Engine/Plugins/Experimental/Mutable/Source/MutableTools/Private/MuT/NodeMesh.cpp
+++ b/Engine/Plugins/Experimental/Mutable/Source/MutableTools/Private/MuT/NodeMesh.cpp
@@ -1,27 +1,30 @@
 // Copyright Epic Games, Inc. All Rights Reserved.
 
 #include "MuT/NodeMesh.h"
-#include "MuT/NodeMeshPrivate.h"
 
-#include "MuT/NodeMesh.h"
+#include "Misc/AssertionMacros.h"
+#include "MuR/Serialisation.h"
+#include "MuR/SerialisationPrivate.h"
+#include "MuT/NodeMeshApplyPose.h"
+#include "MuT/NodeMeshClipDeform.h"
+#include "MuT/NodeMeshClipMorphPlane.h"
+#include "MuT/NodeMeshClipWithMesh.h"
 #include "MuT/NodeMeshConstant.h"
 #include "MuT/NodeMeshFormat.h"
+#include "MuT/NodeMeshFragment.h"
+#include "MuT/NodeMeshGeometryOperation.h"
 #include "MuT/NodeMeshInterpolate.h"
-#include "MuT/NodeMeshTable.h"
-#include "MuT/NodeMeshSubtract.h"
-#include "MuT/NodeMeshTangents.h"
-#include "MuT/NodeMeshMorph.h"
 #include "MuT/NodeMeshMakeMorph.h"
+#include "MuT/NodeMeshMorph.h"
+#include "MuT/NodeMeshReshape.h"
+#include "MuT/NodeMeshSubtract.h"
 #include "MuT/NodeMeshSwitch.h"
-#include "MuT/NodeMeshFragment.h"
+#include "MuT/NodeMeshTable.h"
+#include "MuT/NodeMeshTangents.h"
 #include "MuT/NodeMeshTransform.h"
-#include "MuT/NodeMeshClipMorphPlane.h"
-#include "MuT/NodeMeshClipWithMesh.h"
-#include "MuT/NodeMeshApplyPose.h"
 #include "MuT/NodeMeshVariation.h"
-#include "MuT/NodeMeshGeometryOperation.h"
-#include "MuT/NodeMeshReshape.h"
-#include "MuT/NodeMeshClipDeform.h"
+
+#include <stdint.h>
 
 namespace mu
 {
diff --git a/Engine/Plugins/Experimental/Mutable/Source/MutableTools/Private/MuT/NodeMeshApplyPose.cpp b/Engine/Plugins/Experimental/Mutable/Source/MutableTools/Private/MuT/NodeMeshApplyPose.cpp
--- a/Engine/Plugins/Experimental/Mutable/Source/MutableTools/Private/MuT/NodeMeshApplyPose.cpp
+++ b/Engine/Plugins/Experimental/Mutable/Source/MutableTools/Private/MuT/NodeMeshApplyPose.cpp
@@ -2,11 +2,12 @@
 
 
 #include "MuT/NodeMeshApplyPose.h"
-#include "MuT/NodeMeshApplyPosePrivate.h"
-
-#include "MuT/NodeScalar.h"
 
-#include "MuR/MeshPrivate.h"
+#include "Misc/AssertionMacros.h"
+#include "MuT/NodeLayout.h"
+#include "MuT/NodeMeshApplyPosePrivate.h"
+#include "MuT/NodeMeshPrivate.h"
+#include "MuT/NodePrivate.h"
 
 
 
diff --git a/Engine/Plugins/Experimental/Mutable/Source/MutableTools/Private/MuT/NodeMeshClipDeform.cpp b/Engine/Plugins/Experimental/Mutable/Source/MutableTools/Private/MuT/NodeMeshClipDeform.cpp
--- a/Engine/Plugins/Experimental/Mutable/Source/MutableTools/Private/MuT/NodeMeshClipDeform.cpp
+++ b/Engine/Plugins/Experimental/Mutable/Source/MutableTools/Private/MuT/NodeMeshClipDeform.cpp
@@ -1,12 +1,13 @@
 // Copyright Epic Games, Inc. All Rights Reserved.
 
 #include "MuT/NodeMeshClipDeform.h"
-#include "MuT/NodeMeshClipDeformPrivate.h"
-
-#include "MuT/NodeScalar.h"
 
-#include "MuR/MeshPrivate.h"
-#include "MuR/ImagePrivate.h"
+#include "Misc/AssertionMacros.h"
+#include "MuR/RefCounted.h"
+#include "MuT/NodeLayout.h"
+#include "MuT/NodeMeshClipDeformPrivate.h"
+#include "MuT/NodeMeshPrivate.h"
+#include "MuT/NodePrivate.h"
 
 
 namespace mu
diff --git a/Engine/Plugins/Experimental/Mutable/Source/MutableTools/Private/MuT/NodeMeshClipMorphPlane.cpp b/Engine/Plugins/Experimental/Mutable/Source/MutableTools/Private/MuT/NodeMeshClipMorphPlane.cpp
--- a/Engine/Plugins/Experimental/Mutable/Source/MutableTools/Private/MuT/NodeMeshClipMorphPlane.cpp
+++ b/Engine/Plugins/Experimental/Mutable/Source/MutableTools/Private/MuT/NodeMeshClipMorphPlane.cpp
@@ -2,8 +2,17 @@
 
 
 #include "MuT/NodeMeshClipMorphPlane.h"
+
+#include "Misc/AssertionMacros.h"
+#include "MuR/MemoryPrivate.h"
+#include "MuR/MutableMath.h"
+#include "MuT/NodeLayout.h"
 #include "MuT/NodeMeshClipMorphPlanePrivate.h"
-#include "MuR/MeshPrivate.h"
+#include "MuT/NodeMeshPrivate.h"
+#include "MuT/NodePrivate.h"
+
+#include <memory>
+#include <utility>
 
 
 #define NODE_INPUT_COUNT 	1
diff --git a/Engine/Plugins/Experimental/Mutable/Source/MutableTools/Private/MuT/NodeMeshClipWithMesh.cpp b/Engine/Plugins/Experimental/Mutable/Source/MutableTools/Private/MuT/NodeMeshClipWithMesh.cpp
--- a/Engine/Plugins/Experimental/Mutable/Source/MutableTools/Private/MuT/NodeMeshClipWithMesh.cpp
+++ b/Engine/Plugins/Experimental/Mutable/Source/MutableTools/Private/MuT/NodeMeshClipWithMesh.cpp
@@ -2,8 +2,16 @@
 
 
 #include "MuT/NodeMeshClipWithMesh.h"
+
+#include "Misc/AssertionMacros.h"
+#include "MuR/MemoryPrivate.h"
+#include "MuT/NodeLayout.h"
 #include "MuT/NodeMeshClipWithMeshPrivate.h"
-#include "MuR/MeshPrivate.h"
+#include "MuT/NodeMeshPrivate.h"
+#include "MuT/NodePrivate.h"
+
+#include <memory>
+#include <utility>
 
 
 #define NODE_INPUT_COUNT 	1
diff --git a/Engine/Plugins/Experimental/Mutable/Source/MutableTools/Private/MuT/NodeMeshConstant.cpp b/Engine/Plugins/Experimental/Mutable/Source/MutableTools/Private/MuT/NodeMeshConstant.cpp
--- a/Engine/Plugins/Experimental/Mutable/Source/MutableTools/Private/MuT/NodeMeshConstant.cpp
+++ b/Engine/Plugins/Experimental/Mutable/Source/MutableTools/Private/MuT/NodeMeshConstant.cpp
@@ -2,13 +2,15 @@
 
 
 #include "MuT/NodeMeshConstant.h"
-#include "MuT/NodeMeshConstantPrivate.h"
 
+#include "Misc/AssertionMacros.h"
+#include "MuR/MemoryPrivate.h"
 #include "MuT/NodeLayout.h"
-#include "MuT/NodeLayoutPrivate.h"
+#include "MuT/NodeMeshConstantPrivate.h"
+#include "MuT/NodePrivate.h"
 
-#include "MuR/MeshPrivate.h"
-#include "MuR/MemoryPrivate.h"
+#include <memory>
+#include <utility>
 
 
 namespace mu
diff --git a/Engine/Plugins/Experimental/Mutable/Source/MutableTools/Private/MuT/NodeMeshFormat.cpp b/Engine/Plugins/Experimental/Mutable/Source/MutableTools/Private/MuT/NodeMeshFormat.cpp
--- a/Engine/Plugins/Experimental/Mutable/Source/MutableTools/Private/MuT/NodeMeshFormat.cpp
+++ b/Engine/Plugins/Experimental/Mutable/Source/MutableTools/Private/MuT/NodeMeshFormat.cpp
@@ -2,12 +2,17 @@
 
 
 #include "MuT/NodeMeshFormat.h"
+
+#include "Misc/AssertionMacros.h"
+#include "MuT/NodeLayout.h"
 #include "MuT/NodeMeshFormatPrivate.h"
-#include "MuR/MeshPrivate.h"
+#include "MuT/NodeMeshPrivate.h"
+#include "MuT/NodePrivate.h"
 
 
 namespace mu
 {
+class FMeshBufferSet;
 
 	//---------------------------------------------------------------------------------------------
 	// Static initialisation
diff --git a/Engine/Plugins/Experimental/Mutable/Source/MutableTools/Private/MuT/NodeMeshFragment.cpp b/Engine/Plugins/Experimental/Mutable/Source/MutableTools/Private/MuT/NodeMeshFragment.cpp
--- a/Engine/Plugins/Experimental/Mutable/Source/MutableTools/Private/MuT/NodeMeshFragment.cpp
+++ b/Engine/Plugins/Experimental/Mutable/Source/MutableTools/Private/MuT/NodeMeshFragment.cpp
@@ -2,11 +2,16 @@
 
 
 #include "MuT/NodeMeshFragment.h"
-#include "MuT/NodeMeshFragmentPrivate.h"
 
-#include "MuT/NodeScalar.h"
+#include "Misc/AssertionMacros.h"
+#include "MuR/MemoryPrivate.h"
+#include "MuT/NodeLayout.h"
+#include "MuT/NodeMeshFragmentPrivate.h"
+#include "MuT/NodeMeshPrivate.h"
+#include "MuT/NodePrivate.h"
 
-#include "MuR/MeshPrivate.h"
+#include <memory>
+#include <utility>
 
 
 namespace mu
diff --git a/Engine/Plugins/Experimental/Mutable/Source/MutableTools/Private/MuT/NodeMeshGeometryOperation.cpp b/Engine/Plugins/Experimental/Mutable/Source/MutableTools/Private/MuT/NodeMeshGeometryOperation.cpp
--- a/Engine/Plugins/Experimental/Mutable/Source/MutableTools/Private/MuT/NodeMeshGeometryOperation.cpp
+++ b/Engine/Plugins/Experimental/Mutable/Source/MutableTools/Private/MuT/NodeMeshGeometryOperation.cpp
@@ -2,12 +2,14 @@
 
 
 #include "MuT/NodeMeshGeometryOperation.h"
-#include "MuT/NodeMeshGeometryOperationPrivate.h"
 
+#include "Misc/AssertionMacros.h"
+#include "MuT/NodeLayout.h"
+#include "MuT/NodeMeshGeometryOperationPrivate.h"
+#include "MuT/NodeMeshPrivate.h"
+#include "MuT/NodePrivate.h"
 #include "MuT/NodeScalar.h"
 
-#include "MuR/MeshPrivate.h"
-
 
 
 namespace mu
diff --git a/Engine/Plugins/Experimental/Mutable/Source/MutableTools/Private/MuT/NodeMeshInterpolate.cpp b/Engine/Plugins/Experimental/Mutable/Source/MutableTools/Private/MuT/NodeMeshInterpolate.cpp
--- a/Engine/Plugins/Experimental/Mutable/Source/MutableTools/Private/MuT/NodeMeshInterpolate.cpp
+++ b/Engine/Plugins/Experimental/Mutable/Source/MutableTools/Private/MuT/NodeMeshInterpolate.cpp
@@ -2,11 +2,17 @@
 
 
 #include "MuT/NodeMeshInterpolate.h"
-#include "MuT/NodeMeshInterpolatePrivate.h"
 
+#include "Misc/AssertionMacros.h"
+#include "MuR/MemoryPrivate.h"
+#include "MuT/NodeLayout.h"
+#include "MuT/NodeMeshInterpolatePrivate.h"
+#include "MuT/NodeMeshPrivate.h"
+#include "MuT/NodePrivate.h"
 #include "MuT/NodeScalar.h"
 
-#include "MuR/MeshPrivate.h"
+#include <memory>
+#include <utility>
 
 
 #define NODE_INPUT_COUNT 	4
diff --git a/Engine/Plugins/Experimental/Mutable/Source/MutableTools/Private/MuT/NodeMeshMakeMorph.cpp b/Engine/Plugins/Experimental/Mutable/Source/MutableTools/Private/MuT/NodeMeshMakeMorph.cpp
--- a/Engine/Plugins/Experimental/Mutable/Source/MutableTools/Private/MuT/NodeMeshMakeMorph.cpp
+++ b/Engine/Plugins/Experimental/Mutable/Source/MutableTools/Private/MuT/NodeMeshMakeMorph.cpp
@@ -2,8 +2,12 @@
 
 
 #include "MuT/NodeMeshMakeMorph.h"
+
+#include "Misc/AssertionMacros.h"
+#include "MuT/NodeLayout.h"
 #include "MuT/NodeMeshMakeMorphPrivate.h"
-#include "MuR/MeshPrivate.h"
+#include "MuT/NodeMeshPrivate.h"
+#include "MuT/NodePrivate.h"
 
 
 #define NODE_INPUT_COUNT 	2
diff --git a/Engine/Plugins/Experimental/Mutable/Source/MutableTools/Private/MuT/NodeMeshMorph.cpp b/Engine/Plugins/Experimental/Mutable/Source/MutableTools/Private/MuT/NodeMeshMorph.cpp
--- a/Engine/Plugins/Experimental/Mutable/Source/MutableTools/Private/MuT/NodeMeshMorph.cpp
+++ b/Engine/Plugins/Experimental/Mutable/Source/MutableTools/Private/MuT/NodeMeshMorph.cpp
@@ -2,11 +2,19 @@
 
 
 #include "MuT/NodeMeshMorph.h"
-#include "MuT/NodeMeshMorphPrivate.h"
 
+#include "Containers/Array.h"
+#include "HAL/PlatformCrt.h"
+#include "Misc/AssertionMacros.h"
+#include "MuR/MemoryPrivate.h"
+#include "MuT/NodeLayout.h"
+#include "MuT/NodeMeshMorphPrivate.h"
+#include "MuT/NodeMeshPrivate.h"
+#include "MuT/NodePrivate.h"
 #include "MuT/NodeScalar.h"
 
-#include "MuR/MeshPrivate.h"
+#include <memory>
+#include <utility>
 
 
 
diff --git a/Engine/Plugins/Experimental/Mutable/Source/MutableTools/Private/MuT/NodeMeshReshape.cpp b/Engine/Plugins/Experimental/Mutable/Source/MutableTools/Private/MuT/NodeMeshReshape.cpp
--- a/Engine/Plugins/Experimental/Mutable/Source/MutableTools/Private/MuT/NodeMeshReshape.cpp
+++ b/Engine/Plugins/Experimental/Mutable/Source/MutableTools/Private/MuT/NodeMeshReshape.cpp
@@ -1,11 +1,18 @@
 // Copyright Epic Games, Inc. All Rights Reserved.
 
 #include "MuT/NodeMeshReshape.h"
-#include "MuT/NodeMeshReshapePrivate.h"
 
-#include "MuT/NodeScalar.h"
+#include "Containers/Array.h"
+#include "HAL/PlatformCrt.h"
+#include "Misc/AssertionMacros.h"
+#include "MuR/MemoryPrivate.h"
+#include "MuR/RefCounted.h"
+#include "MuT/NodeLayout.h"
+#include "MuT/NodeMeshPrivate.h"
+#include "MuT/NodeMeshReshapePrivate.h"
+#include "MuT/NodePrivate.h"
 
-#include "MuR/MeshPrivate.h"
+#include <memory>
 
 
 
diff --git a/Engine/Plugins/Experimental/Mutable/Source/MutableTools/Private/MuT/NodeMeshSubtract.cpp b/Engine/Plugins/Experimental/Mutable/Source/MutableTools/Private/MuT/NodeMeshSubtract.cpp
--- a/Engine/Plugins/Experimental/Mutable/Source/MutableTools/Private/MuT/NodeMeshSubtract.cpp
+++ b/Engine/Plugins/Experimental/Mutable/Source/MutableTools/Private/MuT/NodeMeshSubtract.cpp
@@ -2,8 +2,12 @@
 
 
 #include "MuT/NodeMeshSubtract.h"
+
+#include "Misc/AssertionMacros.h"
+#include "MuT/NodeLayout.h"
+#include "MuT/NodeMeshPrivate.h"
 #include "MuT/NodeMeshSubtractPrivate.h"
-#include "MuR/MeshPrivate.h"
+#include "MuT/NodePrivate.h"
 
 
 #define NODE_INPUT_COUNT 	2
diff --git a/Engine/Plugins/Experimental/Mutable/Source/MutableTools/Private/MuT/NodeMeshSwitch.cpp b/Engine/Plugins/Experimental/Mutable/Source/MutableTools/Private/MuT/NodeMeshSwitch.cpp
--- a/Engine/Plugins/Experimental/Mutable/Source/MutableTools/Private/MuT/NodeMeshSwitch.cpp
+++ b/Engine/Plugins/Experimental/Mutable/Source/MutableTools/Private/MuT/NodeMeshSwitch.cpp
@@ -2,11 +2,17 @@
 
 
 #include "MuT/NodeMeshSwitch.h"
-#include "MuT/NodeMeshSwitchPrivate.h"
 
+#include "Misc/AssertionMacros.h"
+#include "MuR/MemoryPrivate.h"
+#include "MuT/NodeLayout.h"
+#include "MuT/NodeMeshPrivate.h"
+#include "MuT/NodeMeshSwitchPrivate.h"
+#include "MuT/NodePrivate.h"
 #include "MuT/NodeScalar.h"
 
-#include "MuR/MeshPrivate.h"
+#include <memory>
+#include <utility>
 
 
 namespace mu
diff --git a/Engine/Plugins/Experimental/Mutable/Source/MutableTools/Private/MuT/NodeMeshTable.cpp b/Engine/Plugins/Experimental/Mutable/Source/MutableTools/Private/MuT/NodeMeshTable.cpp
--- a/Engine/Plugins/Experimental/Mutable/Source/MutableTools/Private/MuT/NodeMeshTable.cpp
+++ b/Engine/Plugins/Experimental/Mutable/Source/MutableTools/Private/MuT/NodeMeshTable.cpp
@@ -2,7 +2,16 @@
 
 
 #include "MuT/NodeMeshTable.h"
+
+#include "Misc/AssertionMacros.h"
+#include "MuR/MemoryPrivate.h"
+#include "MuT/NodeLayout.h"
 #include "MuT/NodeMeshTablePrivate.h"
+#include "MuT/NodePrivate.h"
+#include "MuT/Table.h"
+
+#include <memory>
+#include <utility>
 
 
 namespace mu
diff --git a/Engine/Plugins/Experimental/Mutable/Source/MutableTools/Private/MuT/NodeMeshTangents.cpp b/Engine/Plugins/Experimental/Mutable/Source/MutableTools/Private/MuT/NodeMeshTangents.cpp
--- a/Engine/Plugins/Experimental/Mutable/Source/MutableTools/Private/MuT/NodeMeshTangents.cpp
+++ b/Engine/Plugins/Experimental/Mutable/Source/MutableTools/Private/MuT/NodeMeshTangents.cpp
@@ -2,8 +2,12 @@
 
 
 #include "MuT/NodeMeshTangents.h"
+
+#include "Misc/AssertionMacros.h"
+#include "MuT/NodeLayout.h"
+#include "MuT/NodeMeshPrivate.h"
 #include "MuT/NodeMeshTangentsPrivate.h"
-#include "MuR/MeshPrivate.h"
+#include "MuT/NodePrivate.h"
 
 
 #define NODE_INPUT_COUNT 	1
diff --git a/Engine/Plugins/Experimental/Mutable/Source/MutableTools/Private/MuT/NodeMeshTransform.cpp b/Engine/Plugins/Experimental/Mutable/Source/MutableTools/Private/MuT/NodeMeshTransform.cpp
--- a/Engine/Plugins/Experimental/Mutable/Source/MutableTools/Private/MuT/NodeMeshTransform.cpp
+++ b/Engine/Plugins/Experimental/Mutable/Source/MutableTools/Private/MuT/NodeMeshTransform.cpp
@@ -2,8 +2,14 @@
 
 
 #include "MuT/NodeMeshTransform.h"
+
+#include "HAL/PlatformString.h"
+#include "Misc/AssertionMacros.h"
+#include "MuR/MutableMath.h"
+#include "MuT/NodeLayout.h"
+#include "MuT/NodeMeshPrivate.h"
 #include "MuT/NodeMeshTransformPrivate.h"
-#include "MuR/MeshPrivate.h"
+#include "MuT/NodePrivate.h"
 
 
 #define NODE_INPUT_COUNT 	1
diff --git a/Engine/Plugins/Experimental/Mutable/Source/MutableTools/Private/MuT/NodeModifier.cpp b/Engine/Plugins/Experimental/Mutable/Source/MutableTools/Private/MuT/NodeModifier.cpp
--- a/Engine/Plugins/Experimental/Mutable/Source/MutableTools/Private/MuT/NodeModifier.cpp
+++ b/Engine/Plugins/Experimental/Mutable/Source/MutableTools/Private/MuT/NodeModifier.cpp
@@ -2,11 +2,20 @@
 
 
 #include "MuT/NodeModifier.h"
-#include "MuT/NodeModifierPrivate.h"
 
+#include "Misc/AssertionMacros.h"
+#include "MuR/MemoryPrivate.h"
+#include "MuR/Serialisation.h"
+#include "MuR/SerialisationPrivate.h"
+#include "MuT/NodeModifierMeshClipDeform.h"
 #include "MuT/NodeModifierMeshClipMorphPlane.h"
 #include "MuT/NodeModifierMeshClipWithMesh.h"
-#include "MuT/NodeModifierMeshClipDeform.h"
+#include "MuT/NodeModifierPrivate.h"
+#include "MuT/NodePrivate.h"
+
+#include <stdint.h>
+#include <memory>
+#include <utility>
 
 
 namespace mu
diff --git a/Engine/Plugins/Experimental/Mutable/Source/MutableTools/Private/MuT/NodeModifierMeshClipDeform.cpp b/Engine/Plugins/Experimental/Mutable/Source/MutableTools/Private/MuT/NodeModifierMeshClipDeform.cpp
--- a/Engine/Plugins/Experimental/Mutable/Source/MutableTools/Private/MuT/NodeModifierMeshClipDeform.cpp
+++ b/Engine/Plugins/Experimental/Mutable/Source/MutableTools/Private/MuT/NodeModifierMeshClipDeform.cpp
@@ -2,8 +2,12 @@
 
 
 #include "MuT/NodeModifierMeshClipDeform.h"
+
+#include "Misc/AssertionMacros.h"
+#include "MuR/Mesh.h"
+#include "MuT/NodeMesh.h"
 #include "MuT/NodeModifierMeshClipDeformPrivate.h"
-#include "MuR/MeshPrivate.h"
+#include "MuT/NodePrivate.h"
 
 
 namespace mu
diff --git a/Engine/Plugins/Experimental/Mutable/Source/MutableTools/Private/MuT/NodeModifierMeshClipMorphPlane.cpp b/Engine/Plugins/Experimental/Mutable/Source/MutableTools/Private/MuT/NodeModifierMeshClipMorphPlane.cpp
--- a/Engine/Plugins/Experimental/Mutable/Source/MutableTools/Private/MuT/NodeModifierMeshClipMorphPlane.cpp
+++ b/Engine/Plugins/Experimental/Mutable/Source/MutableTools/Private/MuT/NodeModifierMeshClipMorphPlane.cpp
@@ -2,8 +2,12 @@
 
 
 #include "MuT/NodeModifierMeshClipMorphPlane.h"
+
+#include "Misc/AssertionMacros.h"
+#include "MuR/MemoryPrivate.h"
+#include "MuR/MutableMath.h"
 #include "MuT/NodeModifierMeshClipMorphPlanePrivate.h"
-#include "MuR/MeshPrivate.h"
+#include "MuT/NodePrivate.h"
 
 
 #define NODE_INPUT_COUNT 	0
diff --git a/Engine/Plugins/Experimental/Mutable/Source/MutableTools/Private/MuT/NodeModifierMeshClipWithMesh.cpp b/Engine/Plugins/Experimental/Mutable/Source/MutableTools/Private/MuT/NodeModifierMeshClipWithMesh.cpp
--- a/Engine/Plugins/Experimental/Mutable/Source/MutableTools/Private/MuT/NodeModifierMeshClipWithMesh.cpp
+++ b/Engine/Plugins/Experimental/Mutable/Source/MutableTools/Private/MuT/NodeModifierMeshClipWithMesh.cpp
@@ -2,8 +2,11 @@
 
 
 #include "MuT/NodeModifierMeshClipWithMesh.h"
+
+#include "Misc/AssertionMacros.h"
+#include "MuT/NodeMesh.h"
 #include "MuT/NodeModifierMeshClipWithMeshPrivate.h"
-#include "MuR/MeshPrivate.h"
+#include "MuT/NodePrivate.h"
 
 
 namespace mu
diff --git a/Engine/Plugins/Experimental/Mutable/Source/MutableTools/Private/MuT/NodeObject.cpp b/Engine/Plugins/Experimental/Mutable/Source/MutableTools/Private/MuT/NodeObject.cpp
--- a/Engine/Plugins/Experimental/Mutable/Source/MutableTools/Private/MuT/NodeObject.cpp
+++ b/Engine/Plugins/Experimental/Mutable/Source/MutableTools/Private/MuT/NodeObject.cpp
@@ -2,12 +2,16 @@
 
 
 #include "MuT/NodeObject.h"
-#include "MuT/NodeObjectPrivate.h"
 
-#include "MuT/NodeObjectNew.h"
+#include "Misc/AssertionMacros.h"
+#include "MuR/Serialisation.h"
+#include "MuR/SerialisationPrivate.h"
 #include "MuT/NodeObjectGroup.h"
+#include "MuT/NodeObjectNew.h"
 #include "MuT/NodeObjectState.h"
 
+#include <stdint.h>
+
 
 namespace mu
 {
diff --git a/Engine/Plugins/Experimental/Mutable/Source/MutableTools/Private/MuT/NodeObjectGroup.cpp b/Engine/Plugins/Experimental/Mutable/Source/MutableTools/Private/MuT/NodeObjectGroup.cpp
--- a/Engine/Plugins/Experimental/Mutable/Source/MutableTools/Private/MuT/NodeObjectGroup.cpp
+++ b/Engine/Plugins/Experimental/Mutable/Source/MutableTools/Private/MuT/NodeObjectGroup.cpp
@@ -2,9 +2,17 @@
 
 
 #include "MuT/NodeObjectGroup.h"
-#include "MuT/NodeObjectGroupPrivate.h"
 
+#include "HAL/PlatformCrt.h"
+#include "Misc/AssertionMacros.h"
+#include "MuR/MemoryPrivate.h"
 #include "MuT/NodeLayout.h"
+#include "MuT/NodeObjectGroupPrivate.h"
+#include "MuT/NodeObjectPrivate.h"
+#include "MuT/NodePrivate.h"
+
+#include <memory>
+#include <utility>
 
 
 namespace mu
diff --git a/Engine/Plugins/Experimental/Mutable/Source/MutableTools/Private/MuT/NodeObjectNew.cpp b/Engine/Plugins/Experimental/Mutable/Source/MutableTools/Private/MuT/NodeObjectNew.cpp
--- a/Engine/Plugins/Experimental/Mutable/Source/MutableTools/Private/MuT/NodeObjectNew.cpp
+++ b/Engine/Plugins/Experimental/Mutable/Source/MutableTools/Private/MuT/NodeObjectNew.cpp
@@ -2,14 +2,26 @@
 
 
 #include "MuT/NodeObjectNew.h"
-#include "MuT/NodeObjectNewPrivate.h"
 
-#include "MuT/NodeMeshPrivate.h"
-#include "MuT/NodeLayout.h"
-#include "MuT/NodeMeshConstant.h"
+#include "HAL/PlatformString.h"
+#include "Misc/AssertionMacros.h"
+#include "MuR/MemoryPrivate.h"
+#include "MuT/CompilerPrivate.h"
+#include "MuT/NodeComponent.h"
 #include "MuT/NodeComponentNew.h"
+#include "MuT/NodeLOD.h"
+#include "MuT/NodeLODPrivate.h"
+#include "MuT/NodeMesh.h"
+#include "MuT/NodeMeshConstant.h"
+#include "MuT/NodeMeshPrivate.h"
+#include "MuT/NodeObjectNewPrivate.h"
+#include "MuT/NodePrivate.h"
+#include "MuT/NodeSurface.h"
 #include "MuT/NodeSurfaceNew.h"
 
+#include <memory>
+#include <utility>
+
 
 namespace mu
 {
diff --git a/Engine/Plugins/Experimental/Mutable/Source/MutableTools/Private/MuT/NodeObjectState.cpp b/Engine/Plugins/Experimental/Mutable/Source/MutableTools/Private/MuT/NodeObjectState.cpp
--- a/Engine/Plugins/Experimental/Mutable/Source/MutableTools/Private/MuT/NodeObjectState.cpp
+++ b/Engine/Plugins/Experimental/Mutable/Source/MutableTools/Private/MuT/NodeObjectState.cpp
@@ -2,7 +2,18 @@
 
 
 #include "MuT/NodeObjectState.h"
+
+#include "HAL/PlatformString.h"
+#include "Misc/AssertionMacros.h"
+#include "MuR/MemoryPrivate.h"
+#include "MuT/CompilerPrivate.h"
+#include "MuT/NodeLayout.h"
+#include "MuT/NodeObjectPrivate.h"
 #include "MuT/NodeObjectStatePrivate.h"
+#include "MuT/NodePrivate.h"
+
+#include <memory>
+#include <utility>
 
 
 namespace mu
diff --git a/Engine/Plugins/Experimental/Mutable/Source/MutableTools/Private/MuT/NodePatchImage.cpp b/Engine/Plugins/Experimental/Mutable/Source/MutableTools/Private/MuT/NodePatchImage.cpp
--- a/Engine/Plugins/Experimental/Mutable/Source/MutableTools/Private/MuT/NodePatchImage.cpp
+++ b/Engine/Plugins/Experimental/Mutable/Source/MutableTools/Private/MuT/NodePatchImage.cpp
@@ -2,9 +2,15 @@
 
 
 #include "MuT/NodePatchImage.h"
-#include "MuT/NodePatchImagePrivate.h"
 
+#include "Misc/AssertionMacros.h"
+#include "MuR/MemoryPrivate.h"
 #include "MuT/NodeImage.h"
+#include "MuT/NodePatchImagePrivate.h"
+#include "MuT/NodePrivate.h"
+
+#include <memory>
+#include <utility>
 
 
 namespace mu
diff --git a/Engine/Plugins/Experimental/Mutable/Source/MutableTools/Private/MuT/NodePatchMesh.cpp b/Engine/Plugins/Experimental/Mutable/Source/MutableTools/Private/MuT/NodePatchMesh.cpp
--- a/Engine/Plugins/Experimental/Mutable/Source/MutableTools/Private/MuT/NodePatchMesh.cpp
+++ b/Engine/Plugins/Experimental/Mutable/Source/MutableTools/Private/MuT/NodePatchMesh.cpp
@@ -2,8 +2,11 @@
 
 
 #include "MuT/NodePatchMesh.h"
+
+#include "Misc/AssertionMacros.h"
+#include "MuT/NodeMesh.h"
 #include "MuT/NodePatchMeshPrivate.h"
-#include "MuR/MeshPrivate.h"
+#include "MuT/NodePrivate.h"
 
 
 namespace mu
diff --git a/Engine/Plugins/Experimental/Mutable/Source/MutableTools/Private/MuT/NodeProjector.cpp b/Engine/Plugins/Experimental/Mutable/Source/MutableTools/Private/MuT/NodeProjector.cpp
--- a/Engine/Plugins/Experimental/Mutable/Source/MutableTools/Private/MuT/NodeProjector.cpp
+++ b/Engine/Plugins/Experimental/Mutable/Source/MutableTools/Private/MuT/NodeProjector.cpp
@@ -3,8 +3,12 @@
 
 #include "MuT/NodeProjector.h"
 
+#include "Misc/AssertionMacros.h"
+#include "MuR/Serialisation.h"
+#include "MuR/SerialisationPrivate.h"
 #include "MuT/Node.h"
-#include "MuT/NodePrivate.h"
+
+#include <stdint.h>
 
 
 namespace mu
diff --git a/Engine/Plugins/Experimental/Mutable/Source/MutableTools/Private/MuT/NodeProjectorConstant.cpp b/Engine/Plugins/Experimental/Mutable/Source/MutableTools/Private/MuT/NodeProjectorConstant.cpp
--- a/Engine/Plugins/Experimental/Mutable/Source/MutableTools/Private/MuT/NodeProjectorConstant.cpp
+++ b/Engine/Plugins/Experimental/Mutable/Source/MutableTools/Private/MuT/NodeProjectorConstant.cpp
@@ -1,9 +1,13 @@
 // Copyright Epic Games, Inc. All Rights Reserved.
 
 
-#include "MuT/NodeProjectorPrivate.h"
-
+#include "Misc/AssertionMacros.h"
 #include "MuR/MutableMath.h"
+#include "MuR/Parameters.h"
+#include "MuT/Node.h"
+#include "MuT/NodePrivate.h"
+#include "MuT/NodeProjector.h"
+#include "MuT/NodeProjectorPrivate.h"
 
 
 namespace mu
diff --git a/Engine/Plugins/Experimental/Mutable/Source/MutableTools/Private/MuT/NodeProjectorParameter.cpp b/Engine/Plugins/Experimental/Mutable/Source/MutableTools/Private/MuT/NodeProjectorParameter.cpp
--- a/Engine/Plugins/Experimental/Mutable/Source/MutableTools/Private/MuT/NodeProjectorParameter.cpp
+++ b/Engine/Plugins/Experimental/Mutable/Source/MutableTools/Private/MuT/NodeProjectorParameter.cpp
@@ -1,10 +1,20 @@
 // Copyright Epic Games, Inc. All Rights Reserved.
 
 
+#include "Misc/AssertionMacros.h"
+#include "MuR/MemoryPrivate.h"
+#include "MuR/MutableMath.h"
+#include "MuR/Parameters.h"
+#include "MuR/Ptr.h"
+#include "MuR/RefCounted.h"
+#include "MuT/Node.h"
+#include "MuT/NodePrivate.h"
+#include "MuT/NodeProjector.h"
 #include "MuT/NodeProjectorPrivate.h"
+#include "MuT/NodeRange.h"
 
-#include "MuR/MutableMath.h"
-#include "MuR/MemoryPrivate.h"
+#include <memory>
+#include <utility>
 
 
 namespace mu
diff --git a/Engine/Plugins/Experimental/Mutable/Source/MutableTools/Private/MuT/NodeRange.cpp b/Engine/Plugins/Experimental/Mutable/Source/MutableTools/Private/MuT/NodeRange.cpp
--- a/Engine/Plugins/Experimental/Mutable/Source/MutableTools/Private/MuT/NodeRange.cpp
+++ b/Engine/Plugins/Experimental/Mutable/Source/MutableTools/Private/MuT/NodeRange.cpp
@@ -2,10 +2,15 @@
 
 
 #include "MuT/NodeRange.h"
+
+#include "Misc/AssertionMacros.h"
+#include "MuR/Serialisation.h"
+#include "MuR/SerialisationPrivate.h"
 #include "MuT/Node.h"
-#include "MuT/NodePrivate.h"
 #include "MuT/NodeRangeFromScalar.h"
 
+#include <stdint.h>
+
 
 namespace mu
 {
diff --git a/Engine/Plugins/Experimental/Mutable/Source/MutableTools/Private/MuT/NodeRangeFromScalar.cpp b/Engine/Plugins/Experimental/Mutable/Source/MutableTools/Private/MuT/NodeRangeFromScalar.cpp
--- a/Engine/Plugins/Experimental/Mutable/Source/MutableTools/Private/MuT/NodeRangeFromScalar.cpp
+++ b/Engine/Plugins/Experimental/Mutable/Source/MutableTools/Private/MuT/NodeRangeFromScalar.cpp
@@ -2,12 +2,14 @@
 
 
 #include "MuT/NodeRangeFromScalar.h"
-#include "MuT/NodeRangeFromScalarPrivate.h"
 
+#include "Misc/AssertionMacros.h"
+#include "MuR/MemoryPrivate.h"
+#include "MuR/RefCounted.h"
+#include "MuT/NodePrivate.h"
+#include "MuT/NodeRangeFromScalarPrivate.h"
 #include "MuT/NodeScalar.h"
 
-#include "MuR/ImagePrivate.h"
-
 
 #define NODE_INPUT_COUNT 	1
 
diff --git a/Engine/Plugins/Experimental/Mutable/Source/MutableTools/Private/MuT/NodeScalar.cpp b/Engine/Plugins/Experimental/Mutable/Source/MutableTools/Private/MuT/NodeScalar.cpp
--- a/Engine/Plugins/Experimental/Mutable/Source/MutableTools/Private/MuT/NodeScalar.cpp
+++ b/Engine/Plugins/Experimental/Mutable/Source/MutableTools/Private/MuT/NodeScalar.cpp
@@ -3,17 +3,20 @@
 
 #include "MuT/NodeScalar.h"
 
+#include "Misc/AssertionMacros.h"
+#include "MuR/Serialisation.h"
+#include "MuR/SerialisationPrivate.h"
 #include "MuT/Node.h"
-#include "MuT/NodePrivate.h"
-
+#include "MuT/NodeScalarArithmeticOperation.h"
 #include "MuT/NodeScalarConstant.h"
-#include "MuT/NodeScalarParameter.h"
-#include "MuT/NodeScalarEnumParameter.h"
 #include "MuT/NodeScalarCurve.h"
+#include "MuT/NodeScalarEnumParameter.h"
+#include "MuT/NodeScalarParameter.h"
 #include "MuT/NodeScalarSwitch.h"
-#include "MuT/NodeScalarArithmeticOperation.h"
-#include "MuT/NodeScalarVariation.h"
 #include "MuT/NodeScalarTable.h"
+#include "MuT/NodeScalarVariation.h"
+
+#include <stdint.h>
 
 
 namespace mu
diff --git a/Engine/Plugins/Experimental/Mutable/Source/MutableTools/Private/MuT/NodeScalarArithmeticOperation.cpp b/Engine/Plugins/Experimental/Mutable/Source/MutableTools/Private/MuT/NodeScalarArithmeticOperation.cpp
--- a/Engine/Plugins/Experimental/Mutable/Source/MutableTools/Private/MuT/NodeScalarArithmeticOperation.cpp
+++ b/Engine/Plugins/Experimental/Mutable/Source/MutableTools/Private/MuT/NodeScalarArithmeticOperation.cpp
@@ -2,11 +2,11 @@
 
 
 #include "MuT/NodeScalarArithmeticOperation.h"
-#include "MuT/NodeScalarArithmeticOperationPrivate.h"
 
+#include "Misc/AssertionMacros.h"
+#include "MuT/NodePrivate.h"
 #include "MuT/NodeScalar.h"
-
-#include "MuR/ImagePrivate.h"
+#include "MuT/NodeScalarArithmeticOperationPrivate.h"
 
 
 #define NODE_INPUT_COUNT 	2
diff --git a/Engine/Plugins/Experimental/Mutable/Source/MutableTools/Private/MuT/NodeScalarConstant.cpp b/Engine/Plugins/Experimental/Mutable/Source/MutableTools/Private/MuT/NodeScalarConstant.cpp
--- a/Engine/Plugins/Experimental/Mutable/Source/MutableTools/Private/MuT/NodeScalarConstant.cpp
+++ b/Engine/Plugins/Experimental/Mutable/Source/MutableTools/Private/MuT/NodeScalarConstant.cpp
@@ -2,9 +2,10 @@
 
 
 #include "MuT/NodeScalarConstant.h"
-#include "MuT/NodeScalarConstantPrivate.h"
 
-#include "MuR/MutableMath.h"
+#include "Misc/AssertionMacros.h"
+#include "MuT/NodePrivate.h"
+#include "MuT/NodeScalarConstantPrivate.h"
 
 
 #define NODE_INPUT_COUNT 	0
diff --git a/Engine/Plugins/Experimental/Mutable/Source/MutableTools/Private/MuT/NodeScalarCurve.cpp b/Engine/Plugins/Experimental/Mutable/Source/MutableTools/Private/MuT/NodeScalarCurve.cpp
--- a/Engine/Plugins/Experimental/Mutable/Source/MutableTools/Private/MuT/NodeScalarCurve.cpp
+++ b/Engine/Plugins/Experimental/Mutable/Source/MutableTools/Private/MuT/NodeScalarCurve.cpp
@@ -2,9 +2,16 @@
 
 
 #include "MuT/NodeScalarCurve.h"
-#include "MuT/NodeScalarCurvePrivate.h"
 
+#include "Misc/AssertionMacros.h"
+#include "MuR/MemoryPrivate.h"
+#include "MuR/ParametersPrivate.h"
+#include "MuT/NodePrivate.h"
 #include "MuT/NodeScalar.h"
+#include "MuT/NodeScalarCurvePrivate.h"
+
+#include <memory>
+#include <utility>
 
 
 #define NODE_INPUT_COUNT 	1
diff --git a/Engine/Plugins/Experimental/Mutable/Source/MutableTools/Private/MuT/NodeScalarEnumParameter.cpp b/Engine/Plugins/Experimental/Mutable/Source/MutableTools/Private/MuT/NodeScalarEnumParameter.cpp
--- a/Engine/Plugins/Experimental/Mutable/Source/MutableTools/Private/MuT/NodeScalarEnumParameter.cpp
+++ b/Engine/Plugins/Experimental/Mutable/Source/MutableTools/Private/MuT/NodeScalarEnumParameter.cpp
@@ -2,10 +2,14 @@
 
 
 #include "MuT/NodeScalarEnumParameter.h"
-#include "MuT/NodeScalarEnumParameterPrivate.h"
 
-#include "MuR/MutableMath.h"
+#include "Misc/AssertionMacros.h"
 #include "MuR/MemoryPrivate.h"
+#include "MuT/NodePrivate.h"
+#include "MuT/NodeScalarEnumParameterPrivate.h"
+
+#include <memory>
+#include <utility>
 
 
 namespace mu
diff --git a/Engine/Plugins/Experimental/Mutable/Source/MutableTools/Private/MuT/NodeScalarParameter.cpp b/Engine/Plugins/Experimental/Mutable/Source/MutableTools/Private/MuT/NodeScalarParameter.cpp
--- a/Engine/Plugins/Experimental/Mutable/Source/MutableTools/Private/MuT/NodeScalarParameter.cpp
+++ b/Engine/Plugins/Experimental/Mutable/Source/MutableTools/Private/MuT/NodeScalarParameter.cpp
@@ -2,10 +2,16 @@
 
 
 #include "MuT/NodeScalarParameter.h"
-#include "MuT/NodeScalarParameterPrivate.h"
 
-#include "MuR/MutableMath.h"
+#include "Misc/AssertionMacros.h"
 #include "MuR/MemoryPrivate.h"
+#include "MuT/NodeImage.h"
+#include "MuT/NodePrivate.h"
+#include "MuT/NodeRange.h"
+#include "MuT/NodeScalarParameterPrivate.h"
+
+#include <memory>
+#include <utility>
 
 
 namespace mu
diff --git a/Engine/Plugins/Experimental/Mutable/Source/MutableTools/Private/MuT/NodeScalarSwitch.cpp b/Engine/Plugins/Experimental/Mutable/Source/MutableTools/Private/MuT/NodeScalarSwitch.cpp
--- a/Engine/Plugins/Experimental/Mutable/Source/MutableTools/Private/MuT/NodeScalarSwitch.cpp
+++ b/Engine/Plugins/Experimental/Mutable/Source/MutableTools/Private/MuT/NodeScalarSwitch.cpp
@@ -2,9 +2,15 @@
 
 
 #include "MuT/NodeScalarSwitch.h"
-#include "MuT/NodeScalarSwitchPrivate.h"
 
+#include "Misc/AssertionMacros.h"
+#include "MuR/MemoryPrivate.h"
+#include "MuT/NodePrivate.h"
 #include "MuT/NodeScalar.h"
+#include "MuT/NodeScalarSwitchPrivate.h"
+
+#include <memory>
+#include <utility>
 
 
 namespace mu
diff --git a/Engine/Plugins/Experimental/Mutable/Source/MutableTools/Private/MuT/NodeScalarVariation.cpp b/Engine/Plugins/Experimental/Mutable/Source/MutableTools/Private/MuT/NodeScalarVariation.cpp
--- a/Engine/Plugins/Experimental/Mutable/Source/MutableTools/Private/MuT/NodeScalarVariation.cpp
+++ b/Engine/Plugins/Experimental/Mutable/Source/MutableTools/Private/MuT/NodeScalarVariation.cpp
@@ -2,8 +2,15 @@
 
 
 #include "MuT/NodeScalarVariation.h"
+
+#include "Misc/AssertionMacros.h"
+#include "MuR/MemoryPrivate.h"
+#include "MuT/NodePrivate.h"
 #include "MuT/NodeScalarVariationPrivate.h"
 
+#include <memory>
+#include <utility>
+
 
 namespace mu
 {
diff --git a/Engine/Plugins/Experimental/Mutable/Source/MutableTools/Private/MuT/NodeString.cpp b/Engine/Plugins/Experimental/Mutable/Source/MutableTools/Private/MuT/NodeString.cpp
--- a/Engine/Plugins/Experimental/Mutable/Source/MutableTools/Private/MuT/NodeString.cpp
+++ b/Engine/Plugins/Experimental/Mutable/Source/MutableTools/Private/MuT/NodeString.cpp
@@ -3,12 +3,15 @@
 
 #include "MuT/NodeString.h"
 
+#include "Misc/AssertionMacros.h"
+#include "MuR/Serialisation.h"
+#include "MuR/SerialisationPrivate.h"
 #include "MuT/Node.h"
-#include "MuT/NodePrivate.h"
-
 #include "MuT/NodeStringConstant.h"
 #include "MuT/NodeStringParameter.h"
 
+#include <stdint.h>
+
 
 namespace mu
 {
diff --git a/Engine/Plugins/Experimental/Mutable/Source/MutableTools/Private/MuT/NodeStringConstant.cpp b/Engine/Plugins/Experimental/Mutable/Source/MutableTools/Private/MuT/NodeStringConstant.cpp
--- a/Engine/Plugins/Experimental/Mutable/Source/MutableTools/Private/MuT/NodeStringConstant.cpp
+++ b/Engine/Plugins/Experimental/Mutable/Source/MutableTools/Private/MuT/NodeStringConstant.cpp
@@ -2,9 +2,11 @@
 
 
 #include "MuT/NodeStringConstant.h"
-#include "MuT/NodeStringConstantPrivate.h"
 
-#include "MuR/MutableMath.h"
+#include "Misc/AssertionMacros.h"
+#include "MuR/MemoryPrivate.h"
+#include "MuT/NodePrivate.h"
+#include "MuT/NodeStringConstantPrivate.h"
 
 
 #define NODE_INPUT_COUNT 	0
diff --git a/Engine/Plugins/Experimental/Mutable/Source/MutableTools/Private/MuT/NodeStringParameter.cpp b/Engine/Plugins/Experimental/Mutable/Source/MutableTools/Private/MuT/NodeStringParameter.cpp
--- a/Engine/Plugins/Experimental/Mutable/Source/MutableTools/Private/MuT/NodeStringParameter.cpp
+++ b/Engine/Plugins/Experimental/Mutable/Source/MutableTools/Private/MuT/NodeStringParameter.cpp
@@ -2,10 +2,16 @@
 
 
 #include "MuT/NodeStringParameter.h"
-#include "MuT/NodeStringParameterPrivate.h"
 
-#include "MuR/MutableMath.h"
+#include "Misc/AssertionMacros.h"
 #include "MuR/MemoryPrivate.h"
+#include "MuT/NodeImage.h"
+#include "MuT/NodePrivate.h"
+#include "MuT/NodeRange.h"
+#include "MuT/NodeStringParameterPrivate.h"
+
+#include <memory>
+#include <utility>
 
 
 namespace mu
diff --git a/Engine/Plugins/Experimental/Mutable/Source/MutableTools/Private/MuT/NodeSurface.cpp b/Engine/Plugins/Experimental/Mutable/Source/MutableTools/Private/MuT/NodeSurface.cpp
--- a/Engine/Plugins/Experimental/Mutable/Source/MutableTools/Private/MuT/NodeSurface.cpp
+++ b/Engine/Plugins/Experimental/Mutable/Source/MutableTools/Private/MuT/NodeSurface.cpp
@@ -3,14 +3,14 @@
 
 #include "MuT/NodeSurface.h"
 
-#include "MuT/NodeMesh.h"
-#include "MuT/NodeImage.h"
-
-#include "MuT/NodeSurfaceNew.h"
+#include "Misc/AssertionMacros.h"
+#include "MuR/Serialisation.h"
+#include "MuR/SerialisationPrivate.h"
 #include "MuT/NodeSurfaceEdit.h"
+#include "MuT/NodeSurfaceNew.h"
 #include "MuT/NodeSurfaceVariation.h"
 
-#include "MuR/SerialisationPrivate.h"
+#include <stdint.h>
 
 
 namespace mu
diff --git a/Engine/Plugins/Experimental/Mutable/Source/MutableTools/Private/MuT/NodeSurfaceEdit.cpp b/Engine/Plugins/Experimental/Mutable/Source/MutableTools/Private/MuT/NodeSurfaceEdit.cpp
--- a/Engine/Plugins/Experimental/Mutable/Source/MutableTools/Private/MuT/NodeSurfaceEdit.cpp
+++ b/Engine/Plugins/Experimental/Mutable/Source/MutableTools/Private/MuT/NodeSurfaceEdit.cpp
@@ -2,11 +2,19 @@
 
 
 #include "MuT/NodeSurfaceEdit.h"
-#include "MuT/NodeSurfaceEditPrivate.h"
 
-#include "MuT/NodeMesh.h"
+#include "Misc/AssertionMacros.h"
+#include "MuR/MemoryPrivate.h"
 #include "MuT/NodeImage.h"
+#include "MuT/NodeMesh.h"
+#include "MuT/NodePatchImage.h"
+#include "MuT/NodePatchMesh.h"
+#include "MuT/NodePrivate.h"
 #include "MuT/NodeScalar.h"
+#include "MuT/NodeSurfaceEditPrivate.h"
+
+#include <memory>
+#include <utility>
 
 
 namespace mu
diff --git a/Engine/Plugins/Experimental/Mutable/Source/MutableTools/Private/MuT/NodeSurfaceNew.cpp b/Engine/Plugins/Experimental/Mutable/Source/MutableTools/Private/MuT/NodeSurfaceNew.cpp
--- a/Engine/Plugins/Experimental/Mutable/Source/MutableTools/Private/MuT/NodeSurfaceNew.cpp
+++ b/Engine/Plugins/Experimental/Mutable/Source/MutableTools/Private/MuT/NodeSurfaceNew.cpp
@@ -2,13 +2,21 @@
 
 
 #include "MuT/NodeSurfaceNew.h"
-#include "MuT/NodeSurfaceNewPrivate.h"
 
-#include "MuT/NodeMesh.h"
+#include "HAL/PlatformMath.h"
+#include "Misc/AssertionMacros.h"
+#include "MuR/MemoryPrivate.h"
+#include "MuT/NodeColour.h"
 #include "MuT/NodeImage.h"
+#include "MuT/NodeMesh.h"
+#include "MuT/NodePatchImage.h"
+#include "MuT/NodePrivate.h"
 #include "MuT/NodeScalar.h"
 #include "MuT/NodeString.h"
-#include "MuT/NodeColour.h"
+#include "MuT/NodeSurfaceNewPrivate.h"
+
+#include <memory>
+#include <utility>
 
 namespace mu
 {
diff --git a/Engine/Plugins/Experimental/Mutable/Source/MutableTools/Private/MuT/NodeSurfaceVariation.cpp b/Engine/Plugins/Experimental/Mutable/Source/MutableTools/Private/MuT/NodeSurfaceVariation.cpp
--- a/Engine/Plugins/Experimental/Mutable/Source/MutableTools/Private/MuT/NodeSurfaceVariation.cpp
+++ b/Engine/Plugins/Experimental/Mutable/Source/MutableTools/Private/MuT/NodeSurfaceVariation.cpp
@@ -2,10 +2,16 @@
 
 
 #include "MuT/NodeSurfaceVariation.h"
+
+#include "HAL/PlatformCrt.h"
+#include "Misc/AssertionMacros.h"
+#include "MuR/MemoryPrivate.h"
+#include "MuT/NodeModifier.h"
+#include "MuT/NodePrivate.h"
 #include "MuT/NodeSurfaceVariationPrivate.h"
 
-#include "MuT/NodeMesh.h"
-#include "MuT/NodeImage.h"
+#include <memory>
+#include <utility>
 
 
 namespace mu
diff --git a/Engine/Plugins/Experimental/Mutable/Source/MutableTools/Private/MuT/Streams.cpp b/Engine/Plugins/Experimental/Mutable/Source/MutableTools/Private/MuT/Streams.cpp
--- a/Engine/Plugins/Experimental/Mutable/Source/MutableTools/Private/MuT/Streams.cpp
+++ b/Engine/Plugins/Experimental/Mutable/Source/MutableTools/Private/MuT/Streams.cpp
@@ -2,11 +2,18 @@
 
 
 #include "MuT/Streams.h"
-#include "MuT/StreamsPrivate.h"
-#include "MuT/Platform.h"
-#include "MuT/NodeImageConstant.h"
+
+#include "GenericPlatform/GenericPlatformFile.h"
+#include "HAL/PlatformFileManager.h"
+#include "HAL/PlatformMath.h"
+#include "HAL/UnrealMemory.h"
+#include "Misc/AssertionMacros.h"
 #include "MuR/Image.h"
-#include "MuR/MutableTrace.h"
+#include "MuR/Model.h"
+#include "MuR/MutableMemory.h"
+#include "MuR/RefCounted.h"
+#include "MuT/Platform.h"
+#include "MuT/StreamsPrivate.h"
 
 #define MUTABLE_STREAM_BUFFER_SIZE 0x10000
 
diff --git a/Engine/Plugins/Experimental/Mutable/Source/MutableTools/Private/MuT/Table.cpp b/Engine/Plugins/Experimental/Mutable/Source/MutableTools/Private/MuT/Table.cpp
--- a/Engine/Plugins/Experimental/Mutable/Source/MutableTools/Private/MuT/Table.cpp
+++ b/Engine/Plugins/Experimental/Mutable/Source/MutableTools/Private/MuT/Table.cpp
@@ -2,13 +2,20 @@
 
 
 #include "MuT/Table.h"
-#include "MuT/TablePrivate.h"
-
-#include "MuR/Platform.h"
-#include "MuT/ErrorLogPrivate.h"
 
+#include "HAL/PlatformMath.h"
+#include "Misc/AssertionMacros.h"
 #include "MuR/Image.h"
+#include "MuR/MemoryPrivate.h"
 #include "MuR/Mesh.h"
+#include "MuR/MutableMath.h"
+#include "MuR/Serialisation.h"
+#include "MuR/SerialisationPrivate.h"
+#include "MuT/TablePrivate.h"
+
+#include <memory>
+#include <string>
+#include <utility>
 
 
 namespace mu
diff --git a/Engine/Plugins/Experimental/Mutable/Source/MutableTools/Private/MuT/TaskManager.h b/Engine/Plugins/Experimental/Mutable/Source/MutableTools/Private/MuT/TaskManager.h
--- a/Engine/Plugins/Experimental/Mutable/Source/MutableTools/Private/MuT/TaskManager.h
+++ b/Engine/Plugins/Experimental/Mutable/Source/MutableTools/Private/MuT/TaskManager.h
@@ -3,11 +3,10 @@
 
 #pragma once
 
-#include "MuR/System.h"
-#include "MuR/MemoryPrivate.h"
-
-#include "Async/Async.h"
 #include "Async/TaskGraphInterfaces.h"
+#include "Containers/Array.h"
+#include "Containers/Map.h"
+#include "MuR/MutableMemory.h"
 
 
 namespace mu
diff --git a/Engine/Plugins/Experimental/Mutable/Source/MutableTools/Private/MuT/TypeInfo.cpp b/Engine/Plugins/Experimental/Mutable/Source/MutableTools/Private/MuT/TypeInfo.cpp
--- a/Engine/Plugins/Experimental/Mutable/Source/MutableTools/Private/MuT/TypeInfo.cpp
+++ b/Engine/Plugins/Experimental/Mutable/Source/MutableTools/Private/MuT/TypeInfo.cpp
@@ -2,6 +2,9 @@
 
 #include "MuT/TypeInfo.h"
 
+#include "MuR/ParametersPrivate.h"
+#include "MuR/Types.h"
+
 using namespace mu;
 
 
diff --git a/Engine/Plugins/Experimental/Mutable/Source/MutableTools/Public/MuT/Compiler.h b/Engine/Plugins/Experimental/Mutable/Source/MutableTools/Public/MuT/Compiler.h
--- a/Engine/Plugins/Experimental/Mutable/Source/MutableTools/Public/MuT/Compiler.h
+++ b/Engine/Plugins/Experimental/Mutable/Source/MutableTools/Public/MuT/Compiler.h
@@ -2,9 +2,9 @@
 
 #pragma once
 
-#include "MuR/RefCounted.h"
+#include "HAL/PlatformMath.h"
 #include "MuR/Ptr.h"
-#include "MuR/Types.h"
+#include "MuR/RefCounted.h"
 
 
 namespace mu
diff --git a/Engine/Plugins/Experimental/Mutable/Source/MutableTools/Public/MuT/ErrorLog.h b/Engine/Plugins/Experimental/Mutable/Source/MutableTools/Public/MuT/ErrorLog.h
--- a/Engine/Plugins/Experimental/Mutable/Source/MutableTools/Public/MuT/ErrorLog.h
+++ b/Engine/Plugins/Experimental/Mutable/Source/MutableTools/Public/MuT/ErrorLog.h
@@ -2,8 +2,9 @@
 
 #pragma once
 
-#include "MuR/RefCounted.h"
+#include "HAL/PlatformCrt.h"
 #include "MuR/Ptr.h"
+#include "MuR/RefCounted.h"
 
 
 namespace mu
diff --git a/Engine/Plugins/Experimental/Mutable/Source/MutableTools/Public/MuT/ModelReport.h b/Engine/Plugins/Experimental/Mutable/Source/MutableTools/Public/MuT/ModelReport.h
--- a/Engine/Plugins/Experimental/Mutable/Source/MutableTools/Public/MuT/ModelReport.h
+++ b/Engine/Plugins/Experimental/Mutable/Source/MutableTools/Public/MuT/ModelReport.h
@@ -2,9 +2,8 @@
 
 #pragma once
 
-#include "MuR/RefCounted.h"
 #include "MuR/Ptr.h"
-#include "MuR/Types.h"
+#include "MuR/RefCounted.h"
 
 
 namespace mu
diff --git a/Engine/Plugins/Experimental/Mutable/Source/MutableTools/Public/MuT/Node.h b/Engine/Plugins/Experimental/Mutable/Source/MutableTools/Public/MuT/Node.h
--- a/Engine/Plugins/Experimental/Mutable/Source/MutableTools/Public/MuT/Node.h
+++ b/Engine/Plugins/Experimental/Mutable/Source/MutableTools/Public/MuT/Node.h
@@ -2,10 +2,12 @@
 
 #pragma once
 
-#include "MuR/RefCounted.h"
+#include "HAL/Platform.h"
 #include "MuR/Ptr.h"
-#include "MuR/Types.h"
-#include "MuR/Serialisation.h"
+#include "MuR/RefCounted.h"
+
+namespace mu { class InputArchive; }
+namespace mu { class OutputArchive; }
 
 
 //! This tag is used to identify files containing serialised Node hierarchies. The tag is not added
diff --git a/Engine/Plugins/Experimental/Mutable/Source/MutableTools/Public/MuT/NodeBool.h b/Engine/Plugins/Experimental/Mutable/Source/MutableTools/Public/MuT/NodeBool.h
--- a/Engine/Plugins/Experimental/Mutable/Source/MutableTools/Public/MuT/NodeBool.h
+++ b/Engine/Plugins/Experimental/Mutable/Source/MutableTools/Public/MuT/NodeBool.h
@@ -2,6 +2,9 @@
 
 #pragma once
 
+#include "HAL/Platform.h"
+#include "MuR/Ptr.h"
+#include "MuR/RefCounted.h"
 #include "MuT/Node.h"
 
 
diff --git a/Engine/Plugins/Experimental/Mutable/Source/MutableTools/Public/MuT/NodeColour.h b/Engine/Plugins/Experimental/Mutable/Source/MutableTools/Public/MuT/NodeColour.h
--- a/Engine/Plugins/Experimental/Mutable/Source/MutableTools/Public/MuT/NodeColour.h
+++ b/Engine/Plugins/Experimental/Mutable/Source/MutableTools/Public/MuT/NodeColour.h
@@ -2,6 +2,9 @@
 
 #pragma once
 
+#include "HAL/Platform.h"
+#include "MuR/Ptr.h"
+#include "MuR/RefCounted.h"
 #include "MuT/Node.h"
 
 
diff --git a/Engine/Plugins/Experimental/Mutable/Source/MutableTools/Public/MuT/NodeColourArithmeticOperation.h b/Engine/Plugins/Experimental/Mutable/Source/MutableTools/Public/MuT/NodeColourArithmeticOperation.h
--- a/Engine/Plugins/Experimental/Mutable/Source/MutableTools/Public/MuT/NodeColourArithmeticOperation.h
+++ b/Engine/Plugins/Experimental/Mutable/Source/MutableTools/Public/MuT/NodeColourArithmeticOperation.h
@@ -2,6 +2,9 @@
 
 #pragma once
 
+#include "MuR/Ptr.h"
+#include "MuR/RefCounted.h"
+#include "MuT/Node.h"
 #include "MuT/NodeColour.h"
 
 
diff --git a/Engine/Plugins/Experimental/Mutable/Source/MutableTools/Public/MuT/NodeColourFromScalars.h b/Engine/Plugins/Experimental/Mutable/Source/MutableTools/Public/MuT/NodeColourFromScalars.h
--- a/Engine/Plugins/Experimental/Mutable/Source/MutableTools/Public/MuT/NodeColourFromScalars.h
+++ b/Engine/Plugins/Experimental/Mutable/Source/MutableTools/Public/MuT/NodeColourFromScalars.h
@@ -2,10 +2,11 @@
 
 #pragma once
 
+#include "MuR/Ptr.h"
+#include "MuR/RefCounted.h"
+#include "MuT/Node.h"
 #include "MuT/NodeColour.h"
 
-#include "MuR/Image.h"
-
 
 namespace mu
 {
diff --git a/Engine/Plugins/Experimental/Mutable/Source/MutableTools/Public/MuT/NodeColourSampleImage.h b/Engine/Plugins/Experimental/Mutable/Source/MutableTools/Public/MuT/NodeColourSampleImage.h
--- a/Engine/Plugins/Experimental/Mutable/Source/MutableTools/Public/MuT/NodeColourSampleImage.h
+++ b/Engine/Plugins/Experimental/Mutable/Source/MutableTools/Public/MuT/NodeColourSampleImage.h
@@ -2,10 +2,11 @@
 
 #pragma once
 
+#include "MuR/Ptr.h"
+#include "MuR/RefCounted.h"
+#include "MuT/Node.h"
 #include "MuT/NodeColour.h"
 
-#include "MuR/Image.h"
-
 
 namespace mu
 {
diff --git a/Engine/Plugins/Experimental/Mutable/Source/MutableTools/Public/MuT/NodeComponentNew.h b/Engine/Plugins/Experimental/Mutable/Source/MutableTools/Public/MuT/NodeComponentNew.h
--- a/Engine/Plugins/Experimental/Mutable/Source/MutableTools/Public/MuT/NodeComponentNew.h
+++ b/Engine/Plugins/Experimental/Mutable/Source/MutableTools/Public/MuT/NodeComponentNew.h
@@ -2,6 +2,10 @@
 
 #pragma once
 
+#include "HAL/Platform.h"
+#include "MuR/Ptr.h"
+#include "MuR/RefCounted.h"
+#include "MuT/Node.h"
 #include "MuT/NodeComponent.h"
 
 
diff --git a/Engine/Plugins/Experimental/Mutable/Source/MutableTools/Public/MuT/NodeImageColourMap.h b/Engine/Plugins/Experimental/Mutable/Source/MutableTools/Public/MuT/NodeImageColourMap.h
--- a/Engine/Plugins/Experimental/Mutable/Source/MutableTools/Public/MuT/NodeImageColourMap.h
+++ b/Engine/Plugins/Experimental/Mutable/Source/MutableTools/Public/MuT/NodeImageColourMap.h
@@ -2,6 +2,9 @@
 
 #pragma once
 
+#include "MuR/Ptr.h"
+#include "MuR/RefCounted.h"
+#include "MuT/Node.h"
 #include "MuT/NodeImage.h"
 
 
@@ -13,6 +16,9 @@ namespace mu
 	typedef Ptr<NodeImageColourMap> NodeImageColourMapPtr;
 	typedef Ptr<const NodeImageColourMap> NodeImageColourMapPtrConst;
 
+	class InputArchive;
+	class OutputArchive;
+
 
 	//! This node changes the colours of a selectd part of the image, applying a colour map from
 	//! conteined in another image.
@@ -66,6 +72,7 @@ namespace mu
 		// Interface pattern
 		//-----------------------------------------------------------------------------------------
 		class Private;
+
 		Private* GetPrivate() const;
         Node::Private* GetBasePrivate() const override;
 
diff --git a/Engine/Plugins/Experimental/Mutable/Source/MutableTools/Public/MuT/NodeImageConstant.h b/Engine/Plugins/Experimental/Mutable/Source/MutableTools/Public/MuT/NodeImageConstant.h
--- a/Engine/Plugins/Experimental/Mutable/Source/MutableTools/Public/MuT/NodeImageConstant.h
+++ b/Engine/Plugins/Experimental/Mutable/Source/MutableTools/Public/MuT/NodeImageConstant.h
@@ -2,9 +2,11 @@
 
 #pragma once
 
-#include "MuT/NodeImage.h"
-
 #include "MuR/Image.h"
+#include "MuR/Ptr.h"
+#include "MuR/RefCounted.h"
+#include "MuT/Node.h"
+#include "MuT/NodeImage.h"
 
 
 namespace mu
@@ -15,6 +17,9 @@ namespace mu
     typedef Ptr<NodeImageConstant> NodeImageConstantPtr;
     typedef Ptr<const NodeImageConstant> NodeImageConstantPtrConst;
 
+	class InputArchive;
+	class OutputArchive;
+
 
 	//! Node that outputs a constant image.
 	//! \ingroup model
diff --git a/Engine/Plugins/Experimental/Mutable/Source/MutableTools/Public/MuT/NodeImageDifference.h b/Engine/Plugins/Experimental/Mutable/Source/MutableTools/Public/MuT/NodeImageDifference.h
--- a/Engine/Plugins/Experimental/Mutable/Source/MutableTools/Public/MuT/NodeImageDifference.h
+++ b/Engine/Plugins/Experimental/Mutable/Source/MutableTools/Public/MuT/NodeImageDifference.h
@@ -2,6 +2,9 @@
 
 #pragma once
 
+#include "MuR/Ptr.h"
+#include "MuR/RefCounted.h"
+#include "MuT/Node.h"
 #include "MuT/NodeImage.h"
 
 
@@ -13,6 +16,9 @@ namespace mu
 	typedef Ptr<NodeImageDifference> NodeImageDifferencePtr;
 	typedef Ptr<const NodeImageDifference> NodeImageDifferencePtrConst;
 
+	class InputArchive;
+	class OutputArchive;
+
 
 	//! Node that compares two images and returns a black-and-white image with pixels that are
 	//! "significally different", meaning around 0.03 different. This node is for testing
diff --git a/Engine/Plugins/Experimental/Mutable/Source/MutableTools/Public/MuT/NodeImageFormat.h b/Engine/Plugins/Experimental/Mutable/Source/MutableTools/Public/MuT/NodeImageFormat.h
--- a/Engine/Plugins/Experimental/Mutable/Source/MutableTools/Public/MuT/NodeImageFormat.h
+++ b/Engine/Plugins/Experimental/Mutable/Source/MutableTools/Public/MuT/NodeImageFormat.h
@@ -2,9 +2,11 @@
 
 #pragma once
 
-#include "MuT/NodeImage.h"
-
 #include "MuR/Image.h"
+#include "MuR/Ptr.h"
+#include "MuR/RefCounted.h"
+#include "MuT/Node.h"
+#include "MuT/NodeImage.h"
 
 
 namespace mu
@@ -15,6 +17,9 @@ namespace mu
 	typedef Ptr<NodeImageFormat> NodeImageFormatPtr;
 	typedef Ptr<const NodeImageFormat> NodeImageFormatPtrConst;
 
+	class InputArchive;
+	class OutputArchive;
+
 
 	//! Node that composes a new image by gathering pixel data from channels in other images.
 	//! \ingroup model
diff --git a/Engine/Plugins/Experimental/Mutable/Source/MutableTools/Public/MuT/NodeImageInterpolate3.h b/Engine/Plugins/Experimental/Mutable/Source/MutableTools/Public/MuT/NodeImageInterpolate3.h
--- a/Engine/Plugins/Experimental/Mutable/Source/MutableTools/Public/MuT/NodeImageInterpolate3.h
+++ b/Engine/Plugins/Experimental/Mutable/Source/MutableTools/Public/MuT/NodeImageInterpolate3.h
@@ -2,9 +2,11 @@
 
 #pragma once
 
-#include "MuT/NodeImage.h"
-
 #include "MuR/Image.h"
+#include "MuR/Ptr.h"
+#include "MuR/RefCounted.h"
+#include "MuT/Node.h"
+#include "MuT/NodeImage.h"
 
 
 namespace mu
diff --git a/Engine/Plugins/Experimental/Mutable/Source/MutableTools/Public/MuT/NodeImageLayerColour.h b/Engine/Plugins/Experimental/Mutable/Source/MutableTools/Public/MuT/NodeImageLayerColour.h
--- a/Engine/Plugins/Experimental/Mutable/Source/MutableTools/Public/MuT/NodeImageLayerColour.h
+++ b/Engine/Plugins/Experimental/Mutable/Source/MutableTools/Public/MuT/NodeImageLayerColour.h
@@ -2,8 +2,11 @@
 
 #pragma once
 
-#include "MuT/NodeImageLayer.h"
 #include "MuR/Image.h"
+#include "MuR/Ptr.h"
+#include "MuR/RefCounted.h"
+#include "MuT/Node.h"
+#include "MuT/NodeImage.h"
 
 namespace mu
 {
@@ -17,6 +20,9 @@ namespace mu
 	typedef Ptr<NodeImageLayerColour> NodeImageLayerColourPtr;
 	typedef Ptr<const NodeImageLayerColour> NodeImageLayerColourPtrConst;
 
+	class InputArchive;
+	class OutputArchive;
+
 
 	//! This node applies a layer blending effect on a base image using a mask and a colour.
 	//! \ingroup model
diff --git a/Engine/Plugins/Experimental/Mutable/Source/MutableTools/Public/MuT/NodeImageLuminance.h b/Engine/Plugins/Experimental/Mutable/Source/MutableTools/Public/MuT/NodeImageLuminance.h
--- a/Engine/Plugins/Experimental/Mutable/Source/MutableTools/Public/MuT/NodeImageLuminance.h
+++ b/Engine/Plugins/Experimental/Mutable/Source/MutableTools/Public/MuT/NodeImageLuminance.h
@@ -2,10 +2,11 @@
 
 #pragma once
 
+#include "MuR/Ptr.h"
+#include "MuR/RefCounted.h"
+#include "MuT/Node.h"
 #include "MuT/NodeImage.h"
 
-#include "MuR/Image.h"
-
 
 namespace mu
 {
@@ -15,6 +16,9 @@ namespace mu
 	typedef Ptr<NodeImageLuminance> NodeImageLuminancePtr;
 	typedef Ptr<const NodeImageLuminance> NodeImageLuminancePtrConst;
 
+	class InputArchive;
+	class OutputArchive;
+
 
 	//! Calculate the luminance of an image into a new single-channel image..
 	//! \ingroup model
diff --git a/Engine/Plugins/Experimental/Mutable/Source/MutableTools/Public/MuT/NodeImageMipmap.h b/Engine/Plugins/Experimental/Mutable/Source/MutableTools/Public/MuT/NodeImageMipmap.h
--- a/Engine/Plugins/Experimental/Mutable/Source/MutableTools/Public/MuT/NodeImageMipmap.h
+++ b/Engine/Plugins/Experimental/Mutable/Source/MutableTools/Public/MuT/NodeImageMipmap.h
@@ -2,9 +2,11 @@
 
 #pragma once
 
-#include "MuT/NodeImage.h"
-
 #include "MuR/Image.h"
+#include "MuR/Ptr.h"
+#include "MuR/RefCounted.h"
+#include "MuT/Node.h"
+#include "MuT/NodeImage.h"
 
 
 namespace mu
diff --git a/Engine/Plugins/Experimental/Mutable/Source/MutableTools/Public/MuT/NodeImageMultiLayer.h b/Engine/Plugins/Experimental/Mutable/Source/MutableTools/Public/MuT/NodeImageMultiLayer.h
--- a/Engine/Plugins/Experimental/Mutable/Source/MutableTools/Public/MuT/NodeImageMultiLayer.h
+++ b/Engine/Plugins/Experimental/Mutable/Source/MutableTools/Public/MuT/NodeImageMultiLayer.h
@@ -2,10 +2,12 @@
 
 #pragma once
 
+#include "MuR/Image.h"
+#include "MuR/Ptr.h"
+#include "MuR/RefCounted.h"
+#include "MuT/Node.h"
 #include "MuT/NodeImage.h"
-#include "MuT/NodeImageLayer.h"
 #include "MuT/NodeRange.h"
-#include "MuR/Image.h"
 
 
 namespace mu
@@ -16,6 +18,9 @@ namespace mu
     typedef Ptr<NodeImageMultiLayer> NodeImageMultiLayerPtr;
     typedef Ptr<const NodeImageMultiLayer> NodeImageMultiLayerPtrConst;
 
+	class InputArchive;
+	class OutputArchive;
+
 
     //! This node applies any numbre of layer blending effect on a base image using a mask and a
     //! blended image. The number of layers depends on a scalar input of this node.
diff --git a/Engine/Plugins/Experimental/Mutable/Source/MutableTools/Public/MuT/NodeImageNormalComposite.h b/Engine/Plugins/Experimental/Mutable/Source/MutableTools/Public/MuT/NodeImageNormalComposite.h
--- a/Engine/Plugins/Experimental/Mutable/Source/MutableTools/Public/MuT/NodeImageNormalComposite.h
+++ b/Engine/Plugins/Experimental/Mutable/Source/MutableTools/Public/MuT/NodeImageNormalComposite.h
@@ -2,8 +2,11 @@
 
 #pragma once
 
-#include "MuT/NodeImage.h"
 #include "MuR/Image.h"
+#include "MuR/Ptr.h"
+#include "MuR/RefCounted.h"
+#include "MuT/Node.h"
+#include "MuT/NodeImage.h"
 
 namespace mu
 {
diff --git a/Engine/Plugins/Experimental/Mutable/Source/MutableTools/Public/MuT/NodeImageParameter.h b/Engine/Plugins/Experimental/Mutable/Source/MutableTools/Public/MuT/NodeImageParameter.h
--- a/Engine/Plugins/Experimental/Mutable/Source/MutableTools/Public/MuT/NodeImageParameter.h
+++ b/Engine/Plugins/Experimental/Mutable/Source/MutableTools/Public/MuT/NodeImageParameter.h
@@ -2,8 +2,10 @@
 
 #pragma once
 
+#include "MuR/Ptr.h"
+#include "MuR/RefCounted.h"
+#include "MuT/Node.h"
 #include "MuT/NodeImage.h"
-#include "MuR/Image.h"
 
 
 namespace mu
@@ -13,6 +15,9 @@ namespace mu
     typedef Ptr<NodeImageParameter> NodeImageParameterPtr;
     typedef Ptr<const NodeImageParameter> NodeImageParameterPtrConst;
 
+	class InputArchive;
+	class OutputArchive;
+
 
     //! Node that defines a Image model parameter.
 	//! \ingroup model
diff --git a/Engine/Plugins/Experimental/Mutable/Source/MutableTools/Public/MuT/NodeImagePlainColour.h b/Engine/Plugins/Experimental/Mutable/Source/MutableTools/Public/MuT/NodeImagePlainColour.h
--- a/Engine/Plugins/Experimental/Mutable/Source/MutableTools/Public/MuT/NodeImagePlainColour.h
+++ b/Engine/Plugins/Experimental/Mutable/Source/MutableTools/Public/MuT/NodeImagePlainColour.h
@@ -2,6 +2,9 @@
 
 #pragma once
 
+#include "MuR/Ptr.h"
+#include "MuR/RefCounted.h"
+#include "MuT/Node.h"
 #include "MuT/NodeImage.h"
 
 
@@ -17,6 +20,9 @@ namespace mu
 	typedef Ptr<NodeImagePlainColour> NodeImagePlainColourPtr;
 	typedef Ptr<const NodeImagePlainColour> NodeImagePlainColourPtrConst;
 
+	class InputArchive;
+	class OutputArchive;
+
 
 	//! Node that multiplies the colors of an image, channel by channel.
 	//! \ingroup model
diff --git a/Engine/Plugins/Experimental/Mutable/Source/MutableTools/Public/MuT/NodeImageProject.h b/Engine/Plugins/Experimental/Mutable/Source/MutableTools/Public/MuT/NodeImageProject.h
--- a/Engine/Plugins/Experimental/Mutable/Source/MutableTools/Public/MuT/NodeImageProject.h
+++ b/Engine/Plugins/Experimental/Mutable/Source/MutableTools/Public/MuT/NodeImageProject.h
@@ -3,9 +3,11 @@
 
 #pragma once
 
+#include "Math/UnrealMathSSE.h"
+#include "MuR/Ptr.h"
+#include "MuR/RefCounted.h"
+#include "MuT/Node.h"
 #include "MuT/NodeImage.h"
-#include "MuR/Image.h"
-#include "Math/IntVector.h"
 
 namespace mu
 {
@@ -31,6 +33,9 @@ namespace mu
 	typedef Ptr<NodeImageProject> NodeImageProjectPtr;
 	typedef Ptr<const NodeImageProject> NodeImageProjectPtrConst;
 
+	class InputArchive;
+	class OutputArchive;
+
 
 	//!
 	//! \ingroup model
@@ -103,6 +108,7 @@ namespace mu
 		// Interface pattern
 		//-----------------------------------------------------------------------------------------
 		class Private;
+
 		Private* GetPrivate() const;
         Node::Private* GetBasePrivate() const override;
 
diff --git a/Engine/Plugins/Experimental/Mutable/Source/MutableTools/Public/MuT/NodeImageSaturate.h b/Engine/Plugins/Experimental/Mutable/Source/MutableTools/Public/MuT/NodeImageSaturate.h
--- a/Engine/Plugins/Experimental/Mutable/Source/MutableTools/Public/MuT/NodeImageSaturate.h
+++ b/Engine/Plugins/Experimental/Mutable/Source/MutableTools/Public/MuT/NodeImageSaturate.h
@@ -2,10 +2,11 @@
 
 #pragma once
 
+#include "MuR/Ptr.h"
+#include "MuR/RefCounted.h"
+#include "MuT/Node.h"
 #include "MuT/NodeImage.h"
 
-#include "MuR/Image.h"
-
 
 namespace mu
 {
@@ -19,6 +20,9 @@ namespace mu
 	typedef Ptr<NodeImageSaturate> NodeImageSaturatePtr;
 	typedef Ptr<const NodeImageSaturate> NodeImageSaturatePtrConst;
 
+	class InputArchive;
+	class OutputArchive;
+
 
 	//! Change the saturation of an image. This node can be used to increase the saturation with a
 	//! factor bigger than 1, or to decrease it or desaturate it completely with a factor smaller
@@ -70,6 +74,7 @@ namespace mu
 		// Interface pattern
 		//-----------------------------------------------------------------------------------------
 		class Private;
+
 		Private* GetPrivate() const;
         Node::Private* GetBasePrivate() const override;
 
diff --git a/Engine/Plugins/Experimental/Mutable/Source/MutableTools/Public/MuT/NodeImageSelectColour.h b/Engine/Plugins/Experimental/Mutable/Source/MutableTools/Public/MuT/NodeImageSelectColour.h
--- a/Engine/Plugins/Experimental/Mutable/Source/MutableTools/Public/MuT/NodeImageSelectColour.h
+++ b/Engine/Plugins/Experimental/Mutable/Source/MutableTools/Public/MuT/NodeImageSelectColour.h
@@ -2,10 +2,11 @@
 
 #pragma once
 
+#include "MuR/Ptr.h"
+#include "MuR/RefCounted.h"
+#include "MuT/Node.h"
 #include "MuT/NodeImage.h"
 
-#include "MuR/Image.h"
-
 
 namespace mu
 {
@@ -19,6 +20,9 @@ namespace mu
 	typedef Ptr<NodeImageSelectColour> NodeImageSelectColourPtr;
 	typedef Ptr<const NodeImageSelectColour> NodeImageSelectColourPtrConst;
 
+	class InputArchive;
+	class OutputArchive;
+
 
 	//! Create a black and white mask image by selecting the pixels from a source image that match
 	//! a colour.
diff --git a/Engine/Plugins/Experimental/Mutable/Source/MutableTools/Public/MuT/NodeImageSwizzle.h b/Engine/Plugins/Experimental/Mutable/Source/MutableTools/Public/MuT/NodeImageSwizzle.h
--- a/Engine/Plugins/Experimental/Mutable/Source/MutableTools/Public/MuT/NodeImageSwizzle.h
+++ b/Engine/Plugins/Experimental/Mutable/Source/MutableTools/Public/MuT/NodeImageSwizzle.h
@@ -2,9 +2,11 @@
 
 #pragma once
 
-#include "MuT/NodeImage.h"
-
 #include "MuR/Image.h"
+#include "MuR/Ptr.h"
+#include "MuR/RefCounted.h"
+#include "MuT/Node.h"
+#include "MuT/NodeImage.h"
 
 
 namespace mu
@@ -15,6 +17,9 @@ namespace mu
 	typedef Ptr<NodeImageSwizzle> NodeImageSwizzlePtr;
 	typedef Ptr<const NodeImageSwizzle> NodeImageSwizzlePtrConst;
 
+	class InputArchive;
+	class OutputArchive;
+
 
 	//! Node that composes a new image by gathering pixel data from channels in other images.
 	//! \ingroup model
@@ -66,6 +71,7 @@ namespace mu
 		// Interface pattern
 		//-----------------------------------------------------------------------------------------
 		class Private;
+
 		Private* GetPrivate() const;
         Node::Private* GetBasePrivate() const override;
 
diff --git a/Engine/Plugins/Experimental/Mutable/Source/MutableTools/Public/MuT/NodeImageTransform.h b/Engine/Plugins/Experimental/Mutable/Source/MutableTools/Public/MuT/NodeImageTransform.h
--- a/Engine/Plugins/Experimental/Mutable/Source/MutableTools/Public/MuT/NodeImageTransform.h
+++ b/Engine/Plugins/Experimental/Mutable/Source/MutableTools/Public/MuT/NodeImageTransform.h
@@ -2,6 +2,9 @@
 
 #pragma once
 
+#include "MuR/Ptr.h"
+#include "MuR/RefCounted.h"
+#include "MuT/Node.h"
 #include "MuT/NodeImage.h"
 
 
@@ -17,6 +20,9 @@ namespace mu
 	typedef Ptr<NodeImageTransform> NodeImageTransformPtr;
 	typedef Ptr<const NodeImageTransform> NodeImageTransformPtrConst;
 
+	class InputArchive;
+	class OutputArchive;
+
 
 	//! Node that multiplies the colors of an image, channel by channel.
 	//! \ingroup model
@@ -70,6 +76,7 @@ namespace mu
 		// Interface pattern
 		//-----------------------------------------------------------------------------------------
 		class Private;
+
 		Private* GetPrivate() const;
         Node::Private* GetBasePrivate() const override;
 
diff --git a/Engine/Plugins/Experimental/Mutable/Source/MutableTools/Public/MuT/NodeLayout.h b/Engine/Plugins/Experimental/Mutable/Source/MutableTools/Public/MuT/NodeLayout.h
--- a/Engine/Plugins/Experimental/Mutable/Source/MutableTools/Public/MuT/NodeLayout.h
+++ b/Engine/Plugins/Experimental/Mutable/Source/MutableTools/Public/MuT/NodeLayout.h
@@ -2,9 +2,12 @@
 
 #pragma once
 
-#include "MuT/Node.h"
-#include "MuR/Mesh.h"
+#include "HAL/Platform.h"
 #include "MuR/Layout.h"
+#include "MuR/Mesh.h"
+#include "MuR/Ptr.h"
+#include "MuR/RefCounted.h"
+#include "MuT/Node.h"
 
 
 namespace mu
@@ -20,6 +23,8 @@ namespace mu
 	typedef Ptr<NodeLayoutBlocks> NodeLayoutBlocksPtr;
 	typedef Ptr<const NodeLayoutBlocks> NodeLayoutBlocksPtrConst;
 
+	class InputArchive;
+	class OutputArchive;
 
 	//! This node is used to define the texture layout for a texture coordinates channel of a
 	//! constant mesh.
diff --git a/Engine/Plugins/Experimental/Mutable/Source/MutableTools/Public/MuT/NodeMeshApplyPose.h b/Engine/Plugins/Experimental/Mutable/Source/MutableTools/Public/MuT/NodeMeshApplyPose.h
--- a/Engine/Plugins/Experimental/Mutable/Source/MutableTools/Public/MuT/NodeMeshApplyPose.h
+++ b/Engine/Plugins/Experimental/Mutable/Source/MutableTools/Public/MuT/NodeMeshApplyPose.h
@@ -2,10 +2,11 @@
 
 #pragma once
 
+#include "MuR/Ptr.h"
+#include "MuR/RefCounted.h"
+#include "MuT/Node.h"
 #include "MuT/NodeMesh.h"
 
-#include "MuR/Mesh.h"
-
 
 namespace mu
 {
diff --git a/Engine/Plugins/Experimental/Mutable/Source/MutableTools/Public/MuT/NodeMeshClipDeform.h b/Engine/Plugins/Experimental/Mutable/Source/MutableTools/Public/MuT/NodeMeshClipDeform.h
--- a/Engine/Plugins/Experimental/Mutable/Source/MutableTools/Public/MuT/NodeMeshClipDeform.h
+++ b/Engine/Plugins/Experimental/Mutable/Source/MutableTools/Public/MuT/NodeMeshClipDeform.h
@@ -2,11 +2,10 @@
 
 #pragma once
 
-#include "MuT/NodeMesh.h"
+#include "MuR/Ptr.h"
+#include "MuT/Node.h"
 #include "MuT/NodeImage.h"
-
-#include "MuR/Mesh.h"
-#include "MuR/Image.h"
+#include "MuT/NodeMesh.h"
 
 
 namespace mu
diff --git a/Engine/Plugins/Experimental/Mutable/Source/MutableTools/Public/MuT/NodeMeshClipMorphPlane.h b/Engine/Plugins/Experimental/Mutable/Source/MutableTools/Public/MuT/NodeMeshClipMorphPlane.h
--- a/Engine/Plugins/Experimental/Mutable/Source/MutableTools/Public/MuT/NodeMeshClipMorphPlane.h
+++ b/Engine/Plugins/Experimental/Mutable/Source/MutableTools/Public/MuT/NodeMeshClipMorphPlane.h
@@ -2,10 +2,11 @@
 
 #pragma once
 
+#include "MuR/Ptr.h"
+#include "MuR/RefCounted.h"
+#include "MuT/Node.h"
 #include "MuT/NodeMesh.h"
 
-#include "MuR/Mesh.h"
-
 
 namespace mu
 {
diff --git a/Engine/Plugins/Experimental/Mutable/Source/MutableTools/Public/MuT/NodeMeshConstant.h b/Engine/Plugins/Experimental/Mutable/Source/MutableTools/Public/MuT/NodeMeshConstant.h
--- a/Engine/Plugins/Experimental/Mutable/Source/MutableTools/Public/MuT/NodeMeshConstant.h
+++ b/Engine/Plugins/Experimental/Mutable/Source/MutableTools/Public/MuT/NodeMeshConstant.h
@@ -2,10 +2,11 @@
 
 #pragma once
 
-#include "MuT/NodeMesh.h"
-
 #include "MuR/Mesh.h"
-#include "MuR/Types.h"
+#include "MuR/Ptr.h"
+#include "MuR/RefCounted.h"
+#include "MuT/Node.h"
+#include "MuT/NodeMesh.h"
 
 
 namespace mu
diff --git a/Engine/Plugins/Experimental/Mutable/Source/MutableTools/Public/MuT/NodeMeshFormat.h b/Engine/Plugins/Experimental/Mutable/Source/MutableTools/Public/MuT/NodeMeshFormat.h
--- a/Engine/Plugins/Experimental/Mutable/Source/MutableTools/Public/MuT/NodeMeshFormat.h
+++ b/Engine/Plugins/Experimental/Mutable/Source/MutableTools/Public/MuT/NodeMeshFormat.h
@@ -2,11 +2,11 @@
 
 #pragma once
 
+#include "MuR/Ptr.h"
+#include "MuR/RefCounted.h"
+#include "MuT/Node.h"
 #include "MuT/NodeMesh.h"
 
-#include "MuR/Mesh.h"
-
-
 namespace mu
 {
 
@@ -15,6 +15,9 @@ namespace mu
 	typedef Ptr<NodeMeshFormat> NodeMeshFormatPtr;
 	typedef Ptr<const NodeMeshFormat> NodeMeshFormatPtrConst;
 
+	class InputArchive;
+	class FMeshBufferSet;
+	class OutputArchive;
 
 	//! This node can change the buffer formats of a mesh vertices, indices and faces.
 	//! \ingroup model
diff --git a/Engine/Plugins/Experimental/Mutable/Source/MutableTools/Public/MuT/NodeMeshGeometryOperation.h b/Engine/Plugins/Experimental/Mutable/Source/MutableTools/Public/MuT/NodeMeshGeometryOperation.h
--- a/Engine/Plugins/Experimental/Mutable/Source/MutableTools/Public/MuT/NodeMeshGeometryOperation.h
+++ b/Engine/Plugins/Experimental/Mutable/Source/MutableTools/Public/MuT/NodeMeshGeometryOperation.h
@@ -2,8 +2,10 @@
 
 #pragma once
 
+#include "MuR/Ptr.h"
+#include "MuR/RefCounted.h"
+#include "MuT/Node.h"
 #include "MuT/NodeMesh.h"
-#include "MuR/Mesh.h"
 
 
 namespace mu
diff --git a/Engine/Plugins/Experimental/Mutable/Source/MutableTools/Public/MuT/NodeMeshInterpolate.h b/Engine/Plugins/Experimental/Mutable/Source/MutableTools/Public/MuT/NodeMeshInterpolate.h
--- a/Engine/Plugins/Experimental/Mutable/Source/MutableTools/Public/MuT/NodeMeshInterpolate.h
+++ b/Engine/Plugins/Experimental/Mutable/Source/MutableTools/Public/MuT/NodeMeshInterpolate.h
@@ -2,10 +2,12 @@
 
 #pragma once
 
+#include "MuR/MeshBufferSet.h"
+#include "MuR/Ptr.h"
+#include "MuR/RefCounted.h"
+#include "MuT/Node.h"
 #include "MuT/NodeMesh.h"
 
-#include "MuR/Mesh.h"
-
 
 namespace mu
 {
diff --git a/Engine/Plugins/Experimental/Mutable/Source/MutableTools/Public/MuT/NodeMeshReshape.h b/Engine/Plugins/Experimental/Mutable/Source/MutableTools/Public/MuT/NodeMeshReshape.h
--- a/Engine/Plugins/Experimental/Mutable/Source/MutableTools/Public/MuT/NodeMeshReshape.h
+++ b/Engine/Plugins/Experimental/Mutable/Source/MutableTools/Public/MuT/NodeMeshReshape.h
@@ -2,8 +2,9 @@
 
 #pragma once
 
+#include "MuR/Ptr.h"
+#include "MuT/Node.h"
 #include "MuT/NodeMesh.h"
-#include "MuR/Mesh.h"
 
 
 namespace mu
diff --git a/Engine/Plugins/Experimental/Mutable/Source/MutableTools/Public/MuT/NodePatchImage.h b/Engine/Plugins/Experimental/Mutable/Source/MutableTools/Public/MuT/NodePatchImage.h
--- a/Engine/Plugins/Experimental/Mutable/Source/MutableTools/Public/MuT/NodePatchImage.h
+++ b/Engine/Plugins/Experimental/Mutable/Source/MutableTools/Public/MuT/NodePatchImage.h
@@ -2,9 +2,11 @@
 
 #pragma once
 
-#include "MuT/Node.h"
-#include "MuT/NodeImageLayer.h"
 #include "MuR/Image.h"
+#include "MuR/Ptr.h"
+#include "MuR/RefCounted.h"
+#include "MuT/Node.h"
+#include "MuT/NodeImage.h"
 
 
 namespace mu
diff --git a/Engine/Plugins/Experimental/Mutable/Source/MutableTools/Public/MuT/NodePatchMesh.h b/Engine/Plugins/Experimental/Mutable/Source/MutableTools/Public/MuT/NodePatchMesh.h
--- a/Engine/Plugins/Experimental/Mutable/Source/MutableTools/Public/MuT/NodePatchMesh.h
+++ b/Engine/Plugins/Experimental/Mutable/Source/MutableTools/Public/MuT/NodePatchMesh.h
@@ -2,8 +2,10 @@
 
 #pragma once
 
-#include "MuT/NodeMesh.h"
 
+#include "MuR/Ptr.h"
+#include "MuR/RefCounted.h"
+#include "MuT/Node.h"
 
 namespace mu
 {
@@ -13,6 +15,7 @@ namespace mu
 	typedef Ptr<NodePatchMesh> NodePatchMeshPtr;
 	typedef Ptr<const NodePatchMesh> NodePatchMeshPtrConst;
 
+	class NodeMesh;
 
 	//!
 	//! \ingroup model
diff --git a/Engine/Plugins/Experimental/Mutable/Source/MutableTools/Public/MuT/NodeRangeFromScalar.h b/Engine/Plugins/Experimental/Mutable/Source/MutableTools/Public/MuT/NodeRangeFromScalar.h
--- a/Engine/Plugins/Experimental/Mutable/Source/MutableTools/Public/MuT/NodeRangeFromScalar.h
+++ b/Engine/Plugins/Experimental/Mutable/Source/MutableTools/Public/MuT/NodeRangeFromScalar.h
@@ -2,6 +2,8 @@
 
 #pragma once
 
+#include "MuR/Ptr.h"
+#include "MuT/Node.h"
 #include "MuT/NodeRange.h"
 
 namespace mu
diff --git a/Engine/Plugins/Experimental/Mutable/Source/MutableTools/Public/MuT/NodeScalar.h b/Engine/Plugins/Experimental/Mutable/Source/MutableTools/Public/MuT/NodeScalar.h
--- a/Engine/Plugins/Experimental/Mutable/Source/MutableTools/Public/MuT/NodeScalar.h
+++ b/Engine/Plugins/Experimental/Mutable/Source/MutableTools/Public/MuT/NodeScalar.h
@@ -2,8 +2,10 @@
 
 #pragma once
 
+#include "HAL/Platform.h"
+#include "MuR/Ptr.h"
+#include "MuR/RefCounted.h"
 #include "MuT/Node.h"
-#include "MuR/Parameters.h"
 
 
 namespace mu
diff --git a/Engine/Plugins/Experimental/Mutable/Source/MutableTools/Public/MuT/NodeScalarConstant.h b/Engine/Plugins/Experimental/Mutable/Source/MutableTools/Public/MuT/NodeScalarConstant.h
--- a/Engine/Plugins/Experimental/Mutable/Source/MutableTools/Public/MuT/NodeScalarConstant.h
+++ b/Engine/Plugins/Experimental/Mutable/Source/MutableTools/Public/MuT/NodeScalarConstant.h
@@ -2,8 +2,10 @@
 
 #pragma once
 
+#include "MuR/Ptr.h"
+#include "MuR/RefCounted.h"
+#include "MuT/Node.h"
 #include "MuT/NodeScalar.h"
-#include "MuR/Parameters.h"
 
 
 namespace mu
diff --git a/Engine/Plugins/Experimental/Mutable/Source/MutableTools/Public/MuT/NodeScalarEnumParameter.h b/Engine/Plugins/Experimental/Mutable/Source/MutableTools/Public/MuT/NodeScalarEnumParameter.h
--- a/Engine/Plugins/Experimental/Mutable/Source/MutableTools/Public/MuT/NodeScalarEnumParameter.h
+++ b/Engine/Plugins/Experimental/Mutable/Source/MutableTools/Public/MuT/NodeScalarEnumParameter.h
@@ -2,9 +2,11 @@
 
 #pragma once
 
-#include "MuT/NodeScalar.h"
-#include "MuR/Parameters.h"
+#include "MuR/Ptr.h"
+#include "MuR/RefCounted.h"
+#include "MuT/Node.h"
 #include "MuT/NodeRange.h"
+#include "MuT/NodeScalar.h"
 
 
 namespace mu
diff --git a/Engine/Plugins/Experimental/Mutable/Source/MutableTools/Public/MuT/NodeScalarParameter.h b/Engine/Plugins/Experimental/Mutable/Source/MutableTools/Public/MuT/NodeScalarParameter.h
--- a/Engine/Plugins/Experimental/Mutable/Source/MutableTools/Public/MuT/NodeScalarParameter.h
+++ b/Engine/Plugins/Experimental/Mutable/Source/MutableTools/Public/MuT/NodeScalarParameter.h
@@ -2,8 +2,11 @@
 
 #pragma once
 
-#include "MuT/NodeScalar.h"
 #include "MuR/Parameters.h"
+#include "MuR/Ptr.h"
+#include "MuR/RefCounted.h"
+#include "MuT/Node.h"
+#include "MuT/NodeScalar.h"
 
 
 namespace mu
diff --git a/Engine/Plugins/Experimental/Mutable/Source/MutableTools/Public/MuT/NodeString.h b/Engine/Plugins/Experimental/Mutable/Source/MutableTools/Public/MuT/NodeString.h
--- a/Engine/Plugins/Experimental/Mutable/Source/MutableTools/Public/MuT/NodeString.h
+++ b/Engine/Plugins/Experimental/Mutable/Source/MutableTools/Public/MuT/NodeString.h
@@ -2,8 +2,10 @@
 
 #pragma once
 
+#include "HAL/Platform.h"
+#include "MuR/Ptr.h"
+#include "MuR/RefCounted.h"
 #include "MuT/Node.h"
-#include "MuR/Parameters.h"
 
 
 namespace mu
diff --git a/Engine/Plugins/Experimental/Mutable/Source/MutableTools/Public/MuT/NodeStringConstant.h b/Engine/Plugins/Experimental/Mutable/Source/MutableTools/Public/MuT/NodeStringConstant.h
--- a/Engine/Plugins/Experimental/Mutable/Source/MutableTools/Public/MuT/NodeStringConstant.h
+++ b/Engine/Plugins/Experimental/Mutable/Source/MutableTools/Public/MuT/NodeStringConstant.h
@@ -2,8 +2,10 @@
 
 #pragma once
 
+#include "MuR/Ptr.h"
+#include "MuR/RefCounted.h"
+#include "MuT/Node.h"
 #include "MuT/NodeString.h"
-#include "MuR/Parameters.h"
 
 
 namespace mu
diff --git a/Engine/Plugins/Experimental/Mutable/Source/MutableTools/Public/MuT/NodeStringParameter.h b/Engine/Plugins/Experimental/Mutable/Source/MutableTools/Public/MuT/NodeStringParameter.h
--- a/Engine/Plugins/Experimental/Mutable/Source/MutableTools/Public/MuT/NodeStringParameter.h
+++ b/Engine/Plugins/Experimental/Mutable/Source/MutableTools/Public/MuT/NodeStringParameter.h
@@ -2,8 +2,11 @@
 
 #pragma once
 
-#include "MuT/NodeString.h"
 #include "MuR/Parameters.h"
+#include "MuR/Ptr.h"
+#include "MuR/RefCounted.h"
+#include "MuT/Node.h"
+#include "MuT/NodeString.h"
 
 
 namespace mu
diff --git a/Engine/Plugins/Experimental/Mutable/Source/MutableTools/Public/MuT/NodeSurfaceVariation.h b/Engine/Plugins/Experimental/Mutable/Source/MutableTools/Public/MuT/NodeSurfaceVariation.h
--- a/Engine/Plugins/Experimental/Mutable/Source/MutableTools/Public/MuT/NodeSurfaceVariation.h
+++ b/Engine/Plugins/Experimental/Mutable/Source/MutableTools/Public/MuT/NodeSurfaceVariation.h
@@ -2,9 +2,11 @@
 
 #pragma once
 
+#include "MuR/Ptr.h"
+#include "MuR/RefCounted.h"
+#include "MuT/Node.h"
 #include "MuT/NodeSurface.h"
 
-
 namespace mu
 {
 
diff --git a/Engine/Plugins/Experimental/Mutable/Source/MutableTools/Public/MuT/Streams.h b/Engine/Plugins/Experimental/Mutable/Source/MutableTools/Public/MuT/Streams.h
--- a/Engine/Plugins/Experimental/Mutable/Source/MutableTools/Public/MuT/Streams.h
+++ b/Engine/Plugins/Experimental/Mutable/Source/MutableTools/Public/MuT/Streams.h
@@ -2,12 +2,13 @@
 
 #pragma once
 
-#include "MuR/RefCounted.h"
+#include "Containers/UnrealString.h"
+#include "HAL/Platform.h"
+#include "MuR/Ptr.h"
 #include "MuR/Serialisation.h"
-#include "MuR/System.h"
-#include "MuR/Model.h"
 
-#include "Containers/UnrealString.h"
+namespace mu { class Image; }
+namespace mu { class Model; }
 
 
 //! This tag is used to identify files containing serialised Model objects. The tag is not added
diff --git a/Engine/Plugins/Experimental/Mutable/Source/MutableTools/Public/MuT/Table.h b/Engine/Plugins/Experimental/Mutable/Source/MutableTools/Public/MuT/Table.h
--- a/Engine/Plugins/Experimental/Mutable/Source/MutableTools/Public/MuT/Table.h
+++ b/Engine/Plugins/Experimental/Mutable/Source/MutableTools/Public/MuT/Table.h
@@ -2,10 +2,10 @@
 
 #pragma once
 
-#include "MuR/RefCounted.h"
 #include "MuR/Ptr.h"
-#include "MuR/Types.h"
-#include "MuR/Serialisation.h"
+#include "MuR/RefCounted.h"
+
+#include <stdint.h>
 
 
 
@@ -24,6 +24,8 @@ namespace mu
 	typedef Ptr<Image> ImagePtr;
 	typedef Ptr<const Image> ImagePtrConst;
 
+	class InputArchive;
+	class OutputArchive;
 
 	//! Types of the values for the table cells.
 	//! \ingroup runtime
diff --git a/Engine/Plugins/Experimental/Mutable/Source/MutableTools/Public/MuT/TypeInfo.h b/Engine/Plugins/Experimental/Mutable/Source/MutableTools/Public/MuT/TypeInfo.h
--- a/Engine/Plugins/Experimental/Mutable/Source/MutableTools/Public/MuT/TypeInfo.h
+++ b/Engine/Plugins/Experimental/Mutable/Source/MutableTools/Public/MuT/TypeInfo.h
@@ -2,11 +2,11 @@
 
 #pragma once
 
-#include "MuR/Types.h"
+#include "../Private/MuR/ParametersPrivate.h"
+#include "HAL/PlatformMath.h"
 #include "MuR/Image.h"
-#include "MuR/Mesh.h"
+#include "MuR/MeshBufferSet.h"
 #include "MuR/Parameters.h"
-#include "../Private/MuR/ParametersPrivate.h"
 
 // This header file provides additional information about data types defined in the runtime, that
 // can be useful only for the tools. For instance, it provides strings for some enumeration values.
