[mutable] Fix incorrect use of forward declaration.

diff --git a/Engine/Plugins/Experimental/Mutable/Source/CustomizableObject/Classes/MuCO/CustomizableObjectSystem.h b/Engine/Plugins/Experimental/Mutable/Source/CustomizableObject/Classes/MuCO/CustomizableObjectSystem.h
--- a/Engine/Plugins/Experimental/Mutable/Source/CustomizableObject/Classes/MuCO/CustomizableObjectSystem.h
+++ b/Engine/Plugins/Experimental/Mutable/Source/CustomizableObject/Classes/MuCO/CustomizableObjectSystem.h
@@ -16,6 +16,7 @@
 #include "UObject/Object.h"
 #include "UObject/ObjectPtr.h"
 #include "UObject/UObjectGlobals.h"
+#include "AssetRegistry/AssetData.h"
 
 #if WITH_EDITOR
 #include "Framework/Notifications/NotificationManager.h"
@@ -29,7 +30,6 @@ class SNotificationItem;
 class UCustomizableObject;
 class USkeletalMesh;
 class UTexture2D;
-struct FAssetData;
 struct FFrame;
 struct FGuid;
 
