#rnx

diff --git a/Engine/Plugins/Experimental/Mutable/Source/CustomizableObject/Private/MuCO/CustomizableSkeletalComponent.cpp b/Engine/Plugins/Experimental/Mutable/Source/CustomizableObject/Private/MuCO/CustomizableSkeletalComponent.cpp
--- a/Engine/Plugins/Experimental/Mutable/Source/CustomizableObject/Private/MuCO/CustomizableSkeletalComponent.cpp
+++ b/Engine/Plugins/Experimental/Mutable/Source/CustomizableObject/Private/MuCO/CustomizableSkeletalComponent.cpp
@@ -315,7 +315,7 @@ void UCustomizableSkeletalComponent::TickComponent(float DeltaTime, ELevelTick T
 	{
 		USkeletalMeshComponent* Parent = Cast<USkeletalMeshComponent>(GetAttachParent());
 
-		if (Parent)
+		if (Parent && CustomizableObjectInstance)
 		{
 			UCustomizableObject* CustomizableObject = CustomizableObjectInstance->GetCustomizableObject(); 
 
