[mutable] Fix for progressive texturing matching of image surfaces.

diff --git a/Engine/Plugins/Experimental/Mutable/Source/CustomizableObject/Classes/MuCO/CustomizableObjectSystem.h b/Engine/Plugins/Experimental/Mutable/Source/CustomizableObject/Classes/MuCO/CustomizableObjectSystem.h
--- a/Engine/Plugins/Experimental/Mutable/Source/CustomizableObject/Classes/MuCO/CustomizableObjectSystem.h
+++ b/Engine/Plugins/Experimental/Mutable/Source/CustomizableObject/Classes/MuCO/CustomizableObjectSystem.h
@@ -135,10 +135,12 @@ struct FMutableImageReference
 	*/
 	uint32 ImageID = 0;
 
+	/** */
+	uint32 SurfaceId = 0;
+
 	/** */
 	uint8 LOD = 0;
 	uint8 Component = 0;
-	uint8 Surface = 0;
 	uint8 Image = 0;
 };
 
diff --git a/Engine/Plugins/Experimental/Mutable/Source/CustomizableObject/Private/MuCO/CustomizableObjectInstance.cpp b/Engine/Plugins/Experimental/Mutable/Source/CustomizableObject/Private/MuCO/CustomizableObjectInstance.cpp
--- a/Engine/Plugins/Experimental/Mutable/Source/CustomizableObject/Private/MuCO/CustomizableObjectInstance.cpp
+++ b/Engine/Plugins/Experimental/Mutable/Source/CustomizableObject/Private/MuCO/CustomizableObjectInstance.cpp
@@ -5009,7 +5009,7 @@ void UCustomizableInstancePrivateData::BuildMaterials(const TSharedPtr<FMutableO
 							UTexture2D* Texture = nullptr;
 
 							// \TODO: Change this key to a struct.
-							FString TextureReuseCacheRef = bReuseTextures ? FString::Printf(TEXT("%d-%d-%d-%d"), LODIndex, ComponentIndex, SurfaceIndex, ImageIndex) : FString();
+							FString TextureReuseCacheRef = bReuseTextures ? FString::Printf(TEXT("%d-%d-%d-%d"), LODIndex, ComponentIndex, Surface.SurfaceId, ImageIndex) : FString();
 
 							// If the mutable image is null, it must be in the cache
 							FMutableImageCacheKey ImageCacheKey = { Image.ImageID };//, OperationData->MipsToSkip }; // Todo: Reenable extended FMutableImageCacheKey
@@ -5120,8 +5120,12 @@ void UCustomizableInstancePrivateData::BuildMaterials(const TSharedPtr<FMutableO
 														if (MutableMipDataProviderFactory)
 														{
 															MutableMipDataProviderFactory->CustomizableObjectInstance = Public;
-															check(LODIndex<256 && ComponentIndex<256 && SurfaceIndex<256 && ImageIndex<256 );
-															MutableMipDataProviderFactory->ImageRef = { Image.ImageID, uint8(LODIndex), uint8(ComponentIndex), uint8(SurfaceIndex), uint8(ImageIndex) };
+															check(LODIndex < 256 && ComponentIndex < 256 && ImageIndex < 256);
+															MutableMipDataProviderFactory->ImageRef.ImageID = Image.ImageID;
+															MutableMipDataProviderFactory->ImageRef.SurfaceId = Surface.SurfaceId;
+															MutableMipDataProviderFactory->ImageRef.LOD = uint8(LODIndex);
+															MutableMipDataProviderFactory->ImageRef.Component = uint8(ComponentIndex);
+															MutableMipDataProviderFactory->ImageRef.Image = uint8(ImageIndex);
 															MutableMipDataProviderFactory->UpdateContext = UpdateContext;
 															Texture->AddAssetUserData(MutableMipDataProviderFactory);
 														}
diff --git a/Engine/Plugins/Experimental/Mutable/Source/CustomizableObject/Private/MuCO/CustomizableObjectMipDataProvider.cpp b/Engine/Plugins/Experimental/Mutable/Source/CustomizableObject/Private/MuCO/CustomizableObjectMipDataProvider.cpp
--- a/Engine/Plugins/Experimental/Mutable/Source/CustomizableObject/Private/MuCO/CustomizableObjectMipDataProvider.cpp
+++ b/Engine/Plugins/Experimental/Mutable/Source/CustomizableObject/Private/MuCO/CustomizableObjectMipDataProvider.cpp
@@ -98,9 +98,12 @@ namespace impl
 
 			const FMutableImageReference& ImageRef = OperationData->RequestedImage;
 
+			int32 SurfaceIndex = Instance->FindSurfaceById(ImageRef.LOD, ImageRef.Component, ImageRef.SurfaceId);
+			check(SurfaceIndex>=0);
+
 			// This ID may be different than the ID obtained the first time the image was generated, because the mutable
 			// runtime cannot remember all the resources it has built, and only remembers a fixed amount.
-			mu::RESOURCE_ID MipImageID = Instance->GetImageId(ImageRef.LOD, ImageRef.Component, ImageRef.Surface, ImageRef.Image);
+			mu::RESOURCE_ID MipImageID = Instance->GetImageId(ImageRef.LOD, ImageRef.Component, SurfaceIndex, ImageRef.Image);
 
 			// TODO: Why do we need to do this again? The full size should be stored in the initial image creation.
 			mu::FImageDesc ImageDesc;
diff --git a/Engine/Plugins/Experimental/Mutable/Source/CustomizableObject/Private/MuCO/CustomizableObjectSystem.cpp b/Engine/Plugins/Experimental/Mutable/Source/CustomizableObject/Private/MuCO/CustomizableObjectSystem.cpp
--- a/Engine/Plugins/Experimental/Mutable/Source/CustomizableObject/Private/MuCO/CustomizableObjectSystem.cpp
+++ b/Engine/Plugins/Experimental/Mutable/Source/CustomizableObject/Private/MuCO/CustomizableObjectSystem.cpp
@@ -1091,7 +1091,8 @@ namespace impl
 				// Materials and images
 				for (int32 MeshSurfaceIndex = 0; MeshSurfaceIndex < Component.Mesh->GetSurfaceCount(); ++MeshSurfaceIndex)
 				{
-					int32 InstanceSurfaceIndex = Instance->FindSurfaceById(MutableLODIndex, ComponentIndex, Component.Mesh->GetSurfaceId(MeshSurfaceIndex));
+					uint32 SurfaceId = Component.Mesh->GetSurfaceId(MeshSurfaceIndex);
+					int32 InstanceSurfaceIndex = Instance->FindSurfaceById(MutableLODIndex, ComponentIndex, SurfaceId);
 					check(Component.Mesh->GetVertexCount() == 0 || InstanceSurfaceIndex >= 0);
 
 					if (InstanceSurfaceIndex >= 0)
@@ -1101,6 +1102,7 @@ namespace impl
 						++Component.SurfaceCount;
 
 						Surface.MaterialIndex = Instance->GetSurfaceCustomId(MutableLODIndex, ComponentIndex, InstanceSurfaceIndex);
+						Surface.SurfaceId = SurfaceId;
 
 						// Images
 						Surface.FirstImage = OperationData->InstanceUpdateData.Images.Num();
diff --git a/Engine/Plugins/Experimental/Mutable/Source/CustomizableObject/Private/MuCO/CustomizableObjectSystemPrivate.h b/Engine/Plugins/Experimental/Mutable/Source/CustomizableObject/Private/MuCO/CustomizableObjectSystemPrivate.h
--- a/Engine/Plugins/Experimental/Mutable/Source/CustomizableObject/Private/MuCO/CustomizableObjectSystemPrivate.h
+++ b/Engine/Plugins/Experimental/Mutable/Source/CustomizableObject/Private/MuCO/CustomizableObjectSystemPrivate.h
@@ -285,6 +285,9 @@ struct FInstanceUpdateData
 		uint16 ScalarCount = 0;
 
 		uint32 MaterialIndex = 0;
+
+		/** Id of the surface in the mutable core instance. */
+		uint32 SurfaceId = 0;
 	};
 
 	struct FComponent
