#rnx

diff --git a/Engine/Plugins/Experimental/Mutable/Source/CustomizableObject/Classes/MuCO/CustomizableObjectSystem.h b/Engine/Plugins/Experimental/Mutable/Source/CustomizableObject/Classes/MuCO/CustomizableObjectSystem.h
--- a/Engine/Plugins/Experimental/Mutable/Source/CustomizableObject/Classes/MuCO/CustomizableObjectSystem.h
+++ b/Engine/Plugins/Experimental/Mutable/Source/CustomizableObject/Classes/MuCO/CustomizableObjectSystem.h
@@ -171,8 +171,11 @@ public:
 		// Data will be provided with size and pointer
 		Raw,
 
-		// Data will be provided from an unreal texture
-		Unreal
+		// Data will be provided from an unreal texture, loaded in the game thread and kept in memory
+		Unreal,
+
+		// Data will be provided from an unreal texture, and will only be loaded when actually needed in the Mutable thread
+		Unreal_Deferred
 	};
 
 	// Query that Mutable will run to find out if a texture will be provided as an Unreal UTexture2D,
@@ -229,22 +232,9 @@ private:
 };
 
 
-struct FMutableMipUpdateLevel
-{
-	FMutableMipUpdateLevel(int32 InMipLevel, void* InDest, int32 InSizeX, int32 InSizeY, int32 InDataSize, EPixelFormat InFormat) :
-		MipLevel(InMipLevel), Dest(InDest), SizeX(InSizeX), SizeY(InSizeY), DataSize(InDataSize), Format(InFormat) {}
-
-	int32 MipLevel;
-	void* Dest; // Only access from the FMutableTextureMipDataProvider, owned by the FTextureMipInfoArray so don't delete
-	int32 SizeX;
-	int32 SizeY;
-	int32 DataSize;
-	EPixelFormat Format;
-};
-
-
 class UCustomizableInstanceLODManagementBase;
 
+
 UCLASS(Blueprintable, BlueprintType)
 class CUSTOMIZABLEOBJECT_API UCustomizableObjectSystem : public UObject
 {
@@ -326,9 +316,26 @@ public:
 	/** [Texture Parameters] Add a new image provider to the CustomizableObject System. This will be queried for when an external image ID is provided to mutable in an Texture Parameter node. */
 	void RegisterImageProvider(UCustomizableSystemImageProvider* Provider);
 
-	/** [Texture Parameters] Remove a previously regiostered provider. */
+	/** [Texture Parameters] Remove a previously registered provider. */
 	void UnregisterImageProvider(UCustomizableSystemImageProvider* Provider);
 
+	/** [Texture Parameters] Interface to actually cache Images in the Mutable system and make them availabe at run-time.
+		Any cached image has to be registered by an Image provider before caching it.
+		Have in mind that once an image has been cached, it will spend memory according to its size, except in the case
+		of images of type UCustomizableSystemImageProvider::ValueType::Unreal_Deferred, where only a very small amount of
+		memory is used and the real texel data is loaded when needed during an update and then immediately discarded */
+
+	/** [Texture Parameters] Cache an image which has to have been previously registered by an Image provider with the parameter id. */
+	void CacheImage(uint64 ImageId);
+	/** [Texture Parameters] Remove an image from the cache. */
+	void UnCacheImage(uint64 ImageId);
+	/** [Texture Parameters] Cache all images which have been previously registered by all registered Image provider with the parameter id
+		that was used in the provider. If bClearPreviousCacheImages is true, then all previous cache state is cleared */
+	void CacheAllImagesInAllProviders(bool bClearPreviousCacheImages);
+	/** [Texture Parameters] Remove all images from the cache. */
+	void ClearImageCache();
+
+
 #if WITH_EDITOR
 	/** [Texture Parameters] Get the editor-side preview external texture provider. 
 	 * \TODO: Move to the editor? 
diff --git a/Engine/Plugins/Experimental/Mutable/Source/CustomizableObject/Private/MuCO/CustomizableObjectMipDataProvider.cpp b/Engine/Plugins/Experimental/Mutable/Source/CustomizableObject/Private/MuCO/CustomizableObjectMipDataProvider.cpp
--- a/Engine/Plugins/Experimental/Mutable/Source/CustomizableObject/Private/MuCO/CustomizableObjectMipDataProvider.cpp
+++ b/Engine/Plugins/Experimental/Mutable/Source/CustomizableObject/Private/MuCO/CustomizableObjectMipDataProvider.cpp
@@ -55,17 +55,8 @@ namespace impl
 	{
 		MUTABLE_CPUPROFILER_SCOPE(Task_Mutable_UpdateImage);
 
-		// TODO we cannot access all this from the mips steaming thread
-		// Get any external texture that may be requested for this update
-		//MUTABLE_CPUPROFILER_SCOPE(GetExternalImages);
-		//if (System->GetPrivate()->ImageProvider)
-		//{
-		//	for (const FCustomizableObjectTextureParameterValue& Param : CandidateInstance->TextureParameters)
-		//	{
-		//		System->GetPrivate()->ImageProvider->CacheImage(Param.ParameterValue);
-		//	}
-		//}
-
+		// Any external texture that may be needed for this update will be requested from Mutable Core's GetImage
+		// which will safely access the GlobalExternalImages map, and then just get the cached image or issue a disk read
 
 		// This runs in a worker thread.
 		check(OperationData.IsValid());
diff --git a/Engine/Plugins/Experimental/Mutable/Source/CustomizableObject/Private/MuCO/CustomizableObjectSystem.cpp b/Engine/Plugins/Experimental/Mutable/Source/CustomizableObject/Private/MuCO/CustomizableObjectSystem.cpp
--- a/Engine/Plugins/Experimental/Mutable/Source/CustomizableObject/Private/MuCO/CustomizableObjectSystem.cpp
+++ b/Engine/Plugins/Experimental/Mutable/Source/CustomizableObject/Private/MuCO/CustomizableObjectSystem.cpp
@@ -1901,15 +1901,6 @@ namespace impl
 				// Unlock the resource cache for the object used by this instance to avoid
 				// the destruction of resources that we may want to reuse.
 				System->ClearResourceCacheProtected();
-
-				// Free cached external images for this operation
-				{
-					MUTABLE_CPUPROFILER_SCOPE(ClearExternalImages);
-					if (System->GetPrivate()->ImageProvider)
-					{
-						System->GetPrivate()->ImageProvider->ClearCache();
-					}
-				}
 			}
 
 			// Next Task: Callbacks
@@ -1980,18 +1971,8 @@ namespace impl
 			}
 		}
 
-		// Get any external texture that may be requested for this update
-		{
-			MUTABLE_CPUPROFILER_SCOPE(GetExternalImages);
-			if (System->GetPrivate()->ImageProvider)
-			{
-				for (const FCustomizableObjectTextureParameterValue& Param : ObjectInstance->GetTextureParameters())
-				{
-					System->GetPrivate()->ImageProvider->CacheImage(Param.ParameterValue);
-				}
-			}
-		}
-
+		// Any external texture that may be needed for this update will be requested from Mutable Core's GetImage
+		// which will safely access the GlobalExternalImages map, and then just get the cached image or issue a disk read
 
 		// Copy data generated in the mutable thread over to the instance
 		ObjectInstance->GetPrivate()->PrepareForUpdate(OperationData);
@@ -2853,4 +2834,40 @@ uint64 UCustomizableObjectSystem::GetMaxChunkSizeForPlatform(const ITargetPlatfo
 	return MaxChunkSize;
 }
 
+
+void UCustomizableObjectSystem::CacheImage(uint64 ImageId)
+{
+	if (GetPrivate()->ImageProvider)
+	{
+		GetPrivate()->ImageProvider->CacheImage(ImageId);
+	}
+}
+
+
+void UCustomizableObjectSystem::UnCacheImage(uint64 ImageId)
+{
+	if (GetPrivate()->ImageProvider)
+	{
+		GetPrivate()->ImageProvider->UnCacheImage(ImageId);
+	}
+}
+
+
+void UCustomizableObjectSystem::CacheAllImagesInAllProviders(bool bClearPreviousCacheImages)
+{
+	if (GetPrivate()->ImageProvider)
+	{
+		GetPrivate()->ImageProvider->CacheAllImagesInAllProviders(bClearPreviousCacheImages);
+	}
+}
+
+
+void UCustomizableObjectSystem::ClearImageCache()
+{
+	if (GetPrivate()->ImageProvider)
+	{
+		GetPrivate()->ImageProvider->ClearCache();
+	}
+}
+
 #endif // WITH_EDITOR
diff --git a/Engine/Plugins/Experimental/Mutable/Source/CustomizableObject/Private/MuCO/CustomizableObjectSystemPrivate.h b/Engine/Plugins/Experimental/Mutable/Source/CustomizableObject/Private/MuCO/CustomizableObjectSystemPrivate.h
--- a/Engine/Plugins/Experimental/Mutable/Source/CustomizableObject/Private/MuCO/CustomizableObjectSystemPrivate.h
+++ b/Engine/Plugins/Experimental/Mutable/Source/CustomizableObject/Private/MuCO/CustomizableObjectSystemPrivate.h
@@ -454,8 +454,8 @@ public:
 	class FUnrealMutableModelBulkStreamer* Streamer = nullptr;
 
 	// This object is responsible for providing custom images to mutable models (for image parameters)
-	// This object is called from the mutable thread, and it should only access data already build in
-	// the game thread and stored in FUnrealMutableImageProvider::ExternalImagesForCurrentInstance.
+	// This object is called from the mutable thread, and it should only access data already safely submitted from
+	// the game thread and stored in FUnrealMutableImageProvider::GlobalExternalImages.
 	// Non-owned reference
 	class FUnrealMutableImageProvider* ImageProvider = nullptr;
 
diff --git a/Engine/Plugins/Experimental/Mutable/Source/CustomizableObject/Private/MuCO/UnrealMutableImageProvider.cpp b/Engine/Plugins/Experimental/Mutable/Source/CustomizableObject/Private/MuCO/UnrealMutableImageProvider.cpp
--- a/Engine/Plugins/Experimental/Mutable/Source/CustomizableObject/Private/MuCO/UnrealMutableImageProvider.cpp
+++ b/Engine/Plugins/Experimental/Mutable/Source/CustomizableObject/Private/MuCO/UnrealMutableImageProvider.cpp
@@ -11,6 +11,7 @@
 #include "Math/UnrealMathSSE.h"
 #include "Misc/AssertionMacros.h"
 #include "MuCO/CustomizableObjectSystem.h"
+#include "MuCO/CustomizableObject.h"
 #include "MuR/Ptr.h"
 #include "UObject/WeakObjectPtr.h"
 
@@ -33,7 +34,7 @@ namespace
 		switch (Format)
 		{
 		case ETextureSourceFormat::TSF_BGRA8: MutableFormat = mu::EImageFormat::IF_BGRA_UBYTE; break;
-		// This format is deprecated ans using the enum fails to compile in some cases.
+		// This format is deprecated and using the enum fails to compile in some cases.
 		//case ETextureSourceFormat::TSF_RGBA8: MutableFormat = mu::EImageFormat::IF_RGBA_UBYTE; break;
 		case ETextureSourceFormat::TSF_G8: MutableFormat = mu::EImageFormat::IF_L_UBYTE; break;
 		default:
@@ -51,6 +52,37 @@ namespace
 		}
 
 		Texture->Source.UnlockMip(0);
+#else
+		check(Texture->GetPlatformData()->Mips[0].BulkData.IsBulkDataLoaded());
+
+		int32 LODs = 1;
+		int32 SizeX = Texture->GetSizeX();
+		int32 SizeY = Texture->GetSizeY();
+		EPixelFormat Format = Texture->GetPlatformData()->PixelFormat;
+		mu::EImageFormat MutableFormat = mu::EImageFormat::IF_NONE;
+
+		switch (Format)
+		{
+		case EPixelFormat::PF_B8G8R8A8: MutableFormat = mu::EImageFormat::IF_BGRA_UBYTE; break;
+			// This format is deprecated and using the enum fails to compile in some cases.
+			//case ETextureSourceFormat::TSF_RGBA8: MutableFormat = mu::EImageFormat::IF_RGBA_UBYTE; break;
+		case EPixelFormat::PF_G8: MutableFormat = mu::EImageFormat::IF_L_UBYTE; break;
+		default:
+			break;
+		}
+
+		pResult = new mu::Image(SizeX, SizeY, LODs, MutableFormat);
+
+		// If not locked ReadOnly the Texture Source's FGuid can change, invalidating the texture's caching/shaders
+		// making shader compile and cook times increase
+		const void* pSource = Texture->GetPlatformData()->Mips[0].BulkData.LockReadOnly();
+
+		if (pSource)
+		{
+			FMemory::Memcpy(pResult->GetData(), pSource, pResult->GetDataSize());
+		}
+
+		Texture->GetPlatformData()->Mips[0].BulkData.Unlock();
 #endif
 
 		return pResult;
@@ -64,8 +96,108 @@ mu::ImagePtr FUnrealMutableImageProvider::GetImage(mu::EXTERNAL_IMAGE_ID id)
 {
 	// Thread: Mutable worker
 
-	check(ExternalImagesForCurrentInstance.Contains(id));
-	return ExternalImagesForCurrentInstance[id];
+	// Some data that may have to be copied from the GlobalExternalImages while it's locked
+	TUniquePtr<IAsyncReadFileHandle> FileHandle;
+	int32 LODs = 1;
+	int32 SizeX = 0;
+	int32 SizeY = 0;
+	EPixelFormat Format = EPixelFormat::PF_Unknown;
+	int32 BulkDataSize = 0;
+	
+	{
+		FScopeLock Lock(&ExternalImagesLock);
+		// Inside this scope it's safe to access GlobalExternalImages
+
+		if (!GlobalExternalImages.Contains(id))
+		{
+			// Null case, no image was provided
+			return CreateDummy();
+		}
+
+		const FUnrealMutableImageInfo& ImageInfo = GlobalExternalImages[id];
+
+		if (ImageInfo.Image)
+		{
+			// Easy case where the image was directly provided
+			return ImageInfo.Image;
+		}
+		else if (UTexture2D* TextureToLoad = ImageInfo.TextureToLoad)
+		{
+			// It's safe to access TextureToLoad because ExternalImagesLock guarantees that the data in GlobalExternalImages is valid,
+			// not being modified by the game thread at the moment and the texture cannot be GCed because of the AddReferencedObjects
+			// in the FUnrealMutableImageProvider
+
+#if WITH_EDITOR
+			// In the editor the src data can be directly accessed
+			return ConvertTextureUnrealToMutable(TextureToLoad);
+#else
+			// In a packaged game the bulk data has to be loaded
+			// Get the actual file to read the mip 0 data, do not keep any reference to TextureToLoad because once outside of the lock
+			// it may be GCed or changed. Just keep the actual file handle and some sizes instead of the texture
+			FByteBulkData& BulkData = TextureToLoad->GetPlatformData()->Mips[0].BulkData;
+			check(BulkData.GetBulkDataSize() > 0);
+			FileHandle.Reset(BulkData.OpenAsyncReadHandle());
+
+			SizeX = TextureToLoad->GetSizeX();
+			SizeY = TextureToLoad->GetSizeY();
+			BulkDataSize = BulkData.GetBulkDataSize();
+			Format = TextureToLoad->GetPlatformData()->PixelFormat;
+#endif
+		}
+		else
+		{
+			// No UTexture2D was provided, cannot do anything, just provide a dummy texture
+			UE_LOG(LogMutable, Warning, TEXT("No UTexture2D was provided for an application-specific image parameter."));
+			return CreateDummy();
+		}
+	}
+
+	// Case where the bulk data has to be loaded from disk. The ExternalImagesLock has been unlocked because the file access
+	// must not be done while the lock is on as it could stall the game thread for too long
+	if (FileHandle)
+	{
+		mu::ImagePtr Image;
+		IAsyncReadRequest* AsyncRequest = FileHandle->ReadRequest(0, BulkDataSize,
+			AIOP_FLAG_DONTCACHE | AIOP_BelowNormal, nullptr, nullptr);
+
+		if (AsyncRequest)
+		{
+			// It's OK to wait since this must run on the Mutable thread
+			bool bComplete = AsyncRequest->WaitCompletion();
+			uint8* Results = AsyncRequest->GetReadResults();
+
+			if (bComplete && Results)
+			{
+				// TODO: For the moment only this uncompressed format is supported in packaged games
+				check(Format == EPixelFormat::PF_B8G8R8A8);
+				mu::EImageFormat MutableFormat = mu::EImageFormat::IF_BGRA_UBYTE;
+
+				check(BulkDataSize == AsyncRequest->GetSizeResults());
+				check(Image->GetDataSize() == BulkDataSize);
+				Image = new mu::Image(SizeX, SizeY, LODs, MutableFormat);
+				FMemory::Memcpy(Image->GetData(), Results, Image->GetDataSize());
+
+				FMemory::Free(Results);
+				Results = nullptr;
+			}
+		}
+
+		delete AsyncRequest;
+
+		if (!Image)
+		{
+			// Something failed in the loading of the bulk data, just return a dummy
+			UE_LOG(LogMutable, Warning, TEXT("UTexture2D BulkData failed loading for an application-specific image parameter."));
+			CreateDummy();
+		}
+
+		return Image;
+	}
+	else
+	{
+		UE_LOG(LogMutable, Warning, TEXT("Failed to get a FileHandle for a UTexture2D BulkData for an application-specific image parameter."));
+		return CreateDummy();
+	}
 }
 
 
@@ -74,9 +206,10 @@ void FUnrealMutableImageProvider::CacheImage(mu::EXTERNAL_IMAGE_ID id)
 	check( IsInGameThread() );
 
 	mu::ImagePtr pResult;
+	UTexture2D* UnrealDeferredTexture = nullptr;
 
-	// See if any providere provides this id.
-	for (int p = 0; !pResult && p<ImageProviders.Num(); ++p)
+	// See if any provider provides this id.
+	for (int32 p = 0; !pResult && p < ImageProviders.Num(); ++p)
 	{
 		TWeakObjectPtr<UCustomizableSystemImageProvider> Provider = ImageProviders[p];
 
@@ -100,18 +233,67 @@ void FUnrealMutableImageProvider::CacheImage(mu::EXTERNAL_IMAGE_ID id)
 				break;
 			}
 
+			case UCustomizableSystemImageProvider::ValueType::Unreal_Deferred:
+			{
+				UnrealDeferredTexture = Provider->GetTextureParameterValue(id);
+				break;
+			}
+
 			default:
 				break;
 			}
 		}
 	}
 
-	if (!pResult)
+	if (!pResult && !UnrealDeferredTexture)
 	{
 		pResult = CreateDummy();
 	}
 
-	ExternalImagesForCurrentInstance.Add(id,pResult);
+	{
+		FScopeLock Lock(&ExternalImagesLock);
+
+		GlobalExternalImages.Add(id, FUnrealMutableImageInfo(pResult, UnrealDeferredTexture));
+	}
+}
+
+
+void FUnrealMutableImageProvider::UnCacheImage(mu::EXTERNAL_IMAGE_ID id)
+{
+	check(IsInGameThread());
+
+	FScopeLock Lock(&ExternalImagesLock);
+	GlobalExternalImages.Remove(id);
+}
+
+
+void FUnrealMutableImageProvider::CacheAllImagesInAllProviders(bool bClearPreviousCacheImages)
+{
+	check(IsInGameThread());
+
+	FScopeLock Lock(&ExternalImagesLock);
+
+	if (bClearPreviousCacheImages)
+	{
+		ClearCache();
+	}
+
+	// See if any provider provides this id.
+	for (int32 p = 0; p < ImageProviders.Num(); ++p)
+	{
+		TWeakObjectPtr<UCustomizableSystemImageProvider> Provider = ImageProviders[p];
+
+		if (Provider.IsValid())
+		{
+			TArray<FCustomizableObjectExternalTexture> OutValues;
+			Provider->GetTextureParameterValues(OutValues);
+
+			for (const FCustomizableObjectExternalTexture& Value : OutValues)
+			{
+				CacheImage(Value.Value);
+			}
+		}
+	}
 }
 
 
@@ -119,7 +301,8 @@ void FUnrealMutableImageProvider::ClearCache()
 {
 	check(IsInGameThread());
 
-	ExternalImagesForCurrentInstance.Empty();
+	FScopeLock Lock(&ExternalImagesLock);
+	GlobalExternalImages.Empty();
 }
 
 
@@ -151,3 +334,13 @@ mu::ImagePtr FUnrealMutableImageProvider::CreateDummy()
 }
 
 
+void FUnrealMutableImageProvider::AddReferencedObjects(FReferenceCollector& Collector)
+{
+	for (TPair<uint64, FUnrealMutableImageInfo>& Image : GlobalExternalImages)
+	{
+		if (Image.Value.TextureToLoad)
+		{
+			Collector.AddReferencedObject(Image.Value.TextureToLoad);
+		}
+	}
+}
diff --git a/Engine/Plugins/Experimental/Mutable/Source/CustomizableObject/Private/MuCO/UnrealMutableImageProvider.h b/Engine/Plugins/Experimental/Mutable/Source/CustomizableObject/Private/MuCO/UnrealMutableImageProvider.h
--- a/Engine/Plugins/Experimental/Mutable/Source/CustomizableObject/Private/MuCO/UnrealMutableImageProvider.h
+++ b/Engine/Plugins/Experimental/Mutable/Source/CustomizableObject/Private/MuCO/UnrealMutableImageProvider.h
@@ -9,10 +9,11 @@
 #include "MuR/System.h"
 #include "MuR/Types.h"
 #include "UObject/WeakObjectPtrTemplates.h"
+#include "UObject/GCObject.h"
 
 
 /** Implementation of a mutable core provider for image parameters that are application-specific. */
-class FUnrealMutableImageProvider : public mu::ImageParameterGenerator
+class FUnrealMutableImageProvider : public mu::ImageParameterGenerator, public FGCObject
 {
 public:
 
@@ -23,22 +24,45 @@ public:
 	// Own interface
 	// Thread: Game
 	void CacheImage(mu::EXTERNAL_IMAGE_ID id);
+	void UnCacheImage(mu::EXTERNAL_IMAGE_ID id);
+	void CacheAllImagesInAllProviders(bool bClearPreviousCacheImages);
 	void ClearCache();
 
 	/** List of actual image providers that have been registered to the CustomizableObjectSystem. */
 	TArray< TWeakObjectPtr<class UCustomizableSystemImageProvider> > ImageProviders;
 
+	virtual void AddReferencedObjects(FReferenceCollector& Collector) override;
+	virtual FString GetReferencerName() const override
+	{
+		return TEXT("FUnrealMutableImageProvider");
+	}
+
 private:
 
+	struct FUnrealMutableImageInfo
+	{
+		FUnrealMutableImageInfo() {}
+		FUnrealMutableImageInfo(mu::ImagePtr InImage, class UTexture2D* InTextureToLoad) : Image(InImage), TextureToLoad(InTextureToLoad) { }
+
+		mu::ImagePtr Image;
+
+		/** If the above Image has not been loaded in the game thread, the TextureToLoad bulk data will be loaded
+		* from the Mutable thread when it's needed
+		*/
+		UTexture2D* TextureToLoad = nullptr;
+	};
+
 	/** This will be called if an image Id has been requested by Mutable core but it has not been provided by any provider. */
 	mu::ImagePtr CreateDummy();
 
-	/** List of external textures that may be required for the current instance under construction.
-	* This is only written from the game thread before scheduling a new instance update, and it
-	* is read from the mutable thread during the update.
+	/** Map of Ids to external textures that may be required for any instance or Mutable texture mip under construction.
+	* This is only safely written from the game thread protected by the following critical section, and it
+	* is safely read from the mutable thread during the update of the instance or texture mips
 	*/
-	TMap<uint64, mu::ImagePtr> ExternalImagesForCurrentInstance;
+	TMap<uint64, FUnrealMutableImageInfo> GlobalExternalImages;
+	
+	/** Access to GlobalExternalImages must be protected with this because it may be accessed concurrently from the 
+	Game thread to modify it and from the Mutable thread to read it. */
+	FCriticalSection ExternalImagesLock;
 
 };
-
-
diff --git a/Engine/Plugins/Experimental/Mutable/Source/CustomizableObject/Public/MuCO/CustomizableObjectMipDataProvider.h b/Engine/Plugins/Experimental/Mutable/Source/CustomizableObject/Public/MuCO/CustomizableObjectMipDataProvider.h
--- a/Engine/Plugins/Experimental/Mutable/Source/CustomizableObject/Public/MuCO/CustomizableObjectMipDataProvider.h
+++ b/Engine/Plugins/Experimental/Mutable/Source/CustomizableObject/Public/MuCO/CustomizableObjectMipDataProvider.h
@@ -39,6 +39,20 @@ struct FMutableUpdateContext
 };
 
 
+struct FMutableMipUpdateLevel
+{
+	FMutableMipUpdateLevel(int32 InMipLevel, void* InDest, int32 InSizeX, int32 InSizeY, int32 InDataSize, EPixelFormat InFormat) :
+		Dest(InDest), MipLevel(InMipLevel), SizeX(InSizeX), SizeY(InSizeY), DataSize(InDataSize), Format(InFormat) {}
+
+	void* Dest; // Only access from the FMutableTextureMipDataProvider, owned by the FTextureMipInfoArray so don't delete
+	int32 MipLevel;
+	int32 SizeX;
+	int32 SizeY;
+	int32 DataSize;
+	EPixelFormat Format;
+};
+
+
 /** Runtime data used during a mutable image mipmap update */
 struct FMutableImageOperationData
 {
diff --git a/Engine/Plugins/Experimental/Mutable/Source/CustomizableObjectEditor/Private/MuCOE/CustomizableInstanceDetails.cpp b/Engine/Plugins/Experimental/Mutable/Source/CustomizableObjectEditor/Private/MuCOE/CustomizableInstanceDetails.cpp
--- a/Engine/Plugins/Experimental/Mutable/Source/CustomizableObjectEditor/Private/MuCOE/CustomizableInstanceDetails.cpp
+++ b/Engine/Plugins/Experimental/Mutable/Source/CustomizableObjectEditor/Private/MuCOE/CustomizableInstanceDetails.cpp
@@ -80,7 +80,10 @@ void FCustomizableInstanceDetails::Refresh() const
 
 void FCustomizableInstanceDetails::UpdateInstance() const
 {
-		CustomInstance->UpdateSkeletalMeshAsync(true, true);
+	UCustomizableObjectSystem* System = UCustomizableObjectSystem::GetInstance();
+	System->CacheAllImagesInAllProviders(true);
+
+	CustomInstance->UpdateSkeletalMeshAsync(true, true);
 }
 
 
