[Mutable] Add support for 16 bit InfluenceWeights

#rnx

diff --git a/Engine/Plugins/Experimental/Mutable/Source/CustomizableObject/Classes/MuCO/CustomizableObject.h b/Engine/Plugins/Experimental/Mutable/Source/CustomizableObject/Classes/MuCO/CustomizableObject.h
--- a/Engine/Plugins/Experimental/Mutable/Source/CustomizableObject/Classes/MuCO/CustomizableObject.h
+++ b/Engine/Plugins/Experimental/Mutable/Source/CustomizableObject/Classes/MuCO/CustomizableObject.h
@@ -197,6 +197,9 @@ struct FCompilationOptions
 	// Used to enable the use of clothing.
 	bool bClothingEnabled = false;
 
+	// Used to enable 16 bit bone weights
+	bool b16BitBoneWeightsEnabled = false;
+
 	// Used to reduce the number of notifications when compiling objects
 	bool bSilentCompilation = true;
 
@@ -1005,23 +1008,27 @@ public:
 	FCompilationOptions CompileOptions;
 
 	// 
-	UPROPERTY(EditAnywhere, Category = CustomizableObject)
+	UPROPERTY(EditAnywhere, Category = CompileOptions)
 	bool bDisableTextureLayoutManagement = false;
 
 	//
-	UPROPERTY(EditAnywhere, Category = CustomizableObject)
+	UPROPERTY(EditAnywhere, Category = CompileOptions)
 	bool bEnableRealTimeMorphTargets = false;
 
 	//
-	UPROPERTY(EditAnywhere, Category = CustomizableObject)
+	UPROPERTY(EditAnywhere, Category = CompileOptions)
 	bool bEnableClothing = false;
 
+	// TODO: Enable 16 bit weights 
+	UPROPERTY(VisibleAnywhere, Category = CompileOptions)
+	bool bEnable16BitBoneWeights = false;
+
 	// Options when compiling this customizable object (see EMutableCompileMeshType declaration for info)
-	UPROPERTY(EditAnywhere, Category = CustomizableObject)
+	UPROPERTY(EditAnywhere, Category = CompileOptions)
 	EMutableCompileMeshType MeshCompileType = EMutableCompileMeshType::LocalAndChildren;
 
 	// Array of elements to use with compile option CompileType = WorkingSet
-	UPROPERTY(EditAnywhere, Category = CustomizableObject)
+	UPROPERTY(EditAnywhere, Category = CompileOptions)
 	TArray<TSoftObjectPtr<UCustomizableObject>> WorkingSet;
 
 	// Editor graph
@@ -1125,7 +1132,7 @@ private:
 	// This is a manual version number for the binary blobs in this asset.
 	// Increasing it invalidates all the previously compiled models.
 	// Warning: If while merging code both versions have changed, take the highest+1.
-	static const int32 CurrentSupportedVersion = 347;
+	static const int32 CurrentSupportedVersion = 348;
 
 public:
 
diff --git a/Engine/Plugins/Experimental/Mutable/Source/CustomizableObject/Private/MuCO/CustomizableObjectInstance.cpp b/Engine/Plugins/Experimental/Mutable/Source/CustomizableObject/Private/MuCO/CustomizableObjectInstance.cpp
--- a/Engine/Plugins/Experimental/Mutable/Source/CustomizableObject/Private/MuCO/CustomizableObjectInstance.cpp
+++ b/Engine/Plugins/Experimental/Mutable/Source/CustomizableObject/Private/MuCO/CustomizableObjectInstance.cpp
@@ -4325,48 +4325,20 @@ bool UCustomizableInstancePrivateData::BuildSkeletalMeshRenderData(const TShared
 			return false;
 		}
 
-		// Detect if the produced mesh is correct
-		bool bSuccess = true;
-
 		SetLastMeshId(Component.Id, LODIndex, Component.Mesh->GetId());
 
-		FSkeletalMeshLODRenderData& LODModel = Helper_GetLODData(SkeletalMesh)[MeshLODIndex];
-		const mu::FMeshBufferSet& MutableMeshVertexBuffers = Component.Mesh->GetVertexBuffers();
-
-		int32 NumVerticesLODModel = Component.Mesh->GetVertexCount();
-		int32 SurfaceCount = Component.Mesh->GetSurfaceCount();
-
-		mu::MESH_BUFFER_FORMAT boneIndexFormat = mu::MBF_NONE;
-		int boneIndexComponents = 0;
-		int boneIndexOffset = 0;
-		int boneIndexBuffer = -1;
-		int boneIndexChannel = -1;
-		MutableMeshVertexBuffers.FindChannel(mu::MBS_BONEINDICES, 0, &boneIndexBuffer, &boneIndexChannel);
-		if (boneIndexBuffer >= 0 || boneIndexChannel >= 0)
-		{
-			MutableMeshVertexBuffers.GetChannel(boneIndexBuffer, boneIndexChannel,
-				nullptr, nullptr, &boneIndexFormat, &boneIndexComponents, &boneIndexOffset);
-		}
-
-		int NumBoneInfluences = boneIndexComponents;
-
-		TArray<uint16>& BoneMap = Component.BoneMap;
 		UnrealConversionUtils::SetupRenderSections(
-			Component.Mesh,
 			SkeletalMesh,
+			Component.Mesh,
 			MeshLODIndex,
-			NumBoneInfluences,
-			BoneMap);
+			Component.BoneMap);
 
 		UnrealConversionUtils::CopyMutableVertexBuffers(
 			SkeletalMesh,
-			NumVerticesLODModel,
-			NumBoneInfluences,
-			boneIndexBuffer,
-			MutableMeshVertexBuffers,
-			MeshLODIndex,
-			boneIndexFormat);
+			Component.Mesh,
+			MeshLODIndex);
 
+		FSkeletalMeshLODRenderData& LODModel = Helper_GetLODData(SkeletalMesh)[MeshLODIndex];
 
 		if (LODModel.DoesVertexBufferUse16BitBoneIndex() && !UCustomizableObjectSystem::GetInstance()->IsSupport16BitBoneIndexEnabled())
 		{
@@ -4385,17 +4357,15 @@ bool UCustomizableInstancePrivateData::BuildSkeletalMeshRenderData(const TShared
 #endif
 		}
 
-
 		// Update active and required bones
-		{
-			MUTABLE_CPUPROFILER_SCOPE(UpdateSkeletalMesh_SurfaceLoop_UpdateActiveAndRequiredBones);
-			LODModel.ActiveBoneIndices.Append(Component.ActiveBones);
-			LODModel.RequiredBones.Append(Component.ActiveBones);
-		}
+		LODModel.ActiveBoneIndices.Append(Component.ActiveBones);
+		LODModel.RequiredBones.Append(Component.ActiveBones);
 
-		if (bSuccess)
+		// Copy indices.
+		if (!UnrealConversionUtils::CopyMutableIndexBuffers(Component.Mesh, LODModel))
 		{
-			bSuccess = UnrealConversionUtils::CopyMutableIndexBuffers(Component.Mesh,LODModel);
+			// End with failure
+			return false;
 		}
 
 		// Update LOD and streaming data
@@ -4403,12 +4373,6 @@ bool UCustomizableInstancePrivateData::BuildSkeletalMeshRenderData(const TShared
 		LODModel.bIsLODOptional = RefLODRenderData.bIsLODOptional;
 		LODModel.bStreamedDataInlined = RefLODRenderData.bStreamedDataInlined;
 
-		if (!bSuccess)
-		{
-			// End with failure
-			return false;
-		}
-
 	}
 	
 
diff --git a/Engine/Plugins/Experimental/Mutable/Source/CustomizableObject/Private/MuCO/MutableMeshBufferUtils.cpp b/Engine/Plugins/Experimental/Mutable/Source/CustomizableObject/Private/MuCO/MutableMeshBufferUtils.cpp
--- a/Engine/Plugins/Experimental/Mutable/Source/CustomizableObject/Private/MuCO/MutableMeshBufferUtils.cpp
+++ b/Engine/Plugins/Experimental/Mutable/Source/CustomizableObject/Private/MuCO/MutableMeshBufferUtils.cpp
@@ -92,13 +92,14 @@ void MutableMeshBufferUtils::SetupTexCoordinatesBuffer(const int32& InCurrentVer
 }
 
 
-void MutableMeshBufferUtils::SetupSkinBuffer(const int32& InCurrentVertexBuffer, const int32& MaxBoneIndexTypeSizeBytes,
-                                             const int32& MaxNumBonesPerVertex,
-                                             mu::FMeshBufferSet& OutTargetVertexBuffers)
+void MutableMeshBufferUtils::SetupSkinBuffer(const int32& InCurrentVertexBuffer,
+	const int32& MaxBoneIndexTypeSizeBytes,
+	const int32& MaxBoneWeightTypeSizeBytes,
+	const int32& MaxNumBonesPerVertex,
+	mu::FMeshBufferSet& OutTargetVertexBuffers)
 {
 	using namespace mu;
-	const int32 meshWeightTypeSizeBytes = 1;
-	const int32 ElementSize = (meshWeightTypeSizeBytes + MaxBoneIndexTypeSizeBytes) * MaxNumBonesPerVertex;
+	const int32 ElementSize = (MaxBoneWeightTypeSizeBytes + MaxBoneIndexTypeSizeBytes) * MaxNumBonesPerVertex;
 	constexpr int32 ChannelCount = 2;
 	const MESH_BUFFER_SEMANTIC Semantics[ChannelCount] = {MBS_BONEINDICES, MBS_BONEWEIGHTS};
 	const int32 SemanticIndices[ChannelCount] = {0, 0};
@@ -118,7 +119,8 @@ void MutableMeshBufferUtils::SetupSkinBuffer(const int32& InCurrentVertexBuffer,
 		Formats[0] = mu::MBF_NONE;
 		break;
 	}
-	switch (meshWeightTypeSizeBytes)
+
+	switch (MaxBoneWeightTypeSizeBytes)
 	{
 	case 1: Formats[1] = mu::MBF_NUINT8;
 		break;
diff --git a/Engine/Plugins/Experimental/Mutable/Source/CustomizableObject/Private/MuCO/UnrealConversionUtils.cpp b/Engine/Plugins/Experimental/Mutable/Source/CustomizableObject/Private/MuCO/UnrealConversionUtils.cpp
--- a/Engine/Plugins/Experimental/Mutable/Source/CustomizableObject/Private/MuCO/UnrealConversionUtils.cpp
+++ b/Engine/Plugins/Experimental/Mutable/Source/CustomizableObject/Private/MuCO/UnrealConversionUtils.cpp
@@ -51,26 +51,26 @@ namespace UnrealConversionUtils
 		 * @param InMutableTextureData - Mutable texture data buffer
 		 */
 		void FStaticMeshVertexBuffers_InitWithMutableData(
-		FStaticMeshVertexBuffers& OutVertexBuffers,
-		const int32 NumVertices,
-		const int32 NumTexCoords,
-		const bool bUseFullPrecisionUVs,
-		const bool bNeedCPUAccess,
-		const void* InMutablePositionData,
-		const void* InMutableTangentData,
-		const void* InMutableTextureData)
-{
-	// positions
-	OutVertexBuffers.PositionVertexBuffer.Init(NumVertices, bNeedCPUAccess);
-	FMemory::Memcpy(OutVertexBuffers.PositionVertexBuffer.GetVertexData(), InMutablePositionData, NumVertices * OutVertexBuffers.PositionVertexBuffer.GetStride());
-
-	// tangent and texture coords
-	OutVertexBuffers.StaticMeshVertexBuffer.SetUseFullPrecisionUVs(true);
-	OutVertexBuffers.StaticMeshVertexBuffer.SetUseHighPrecisionTangentBasis(false);
-	OutVertexBuffers.StaticMeshVertexBuffer.Init(NumVertices, NumTexCoords, bNeedCPUAccess);
-	FMemory::Memcpy(OutVertexBuffers.StaticMeshVertexBuffer.GetTangentData(), InMutableTangentData, OutVertexBuffers.StaticMeshVertexBuffer.GetTangentSize());
-	FMemory::Memcpy(OutVertexBuffers.StaticMeshVertexBuffer.GetTexCoordData(), InMutableTextureData, OutVertexBuffers.StaticMeshVertexBuffer.GetTexCoordSize());
-}
+			FStaticMeshVertexBuffers& OutVertexBuffers,
+			const int32 NumVertices,
+			const int32 NumTexCoords,
+			const bool bUseFullPrecisionUVs,
+			const bool bNeedCPUAccess,
+			const void* InMutablePositionData,
+			const void* InMutableTangentData,
+			const void* InMutableTextureData)
+		{
+			// positions
+			OutVertexBuffers.PositionVertexBuffer.Init(NumVertices, bNeedCPUAccess);
+			FMemory::Memcpy(OutVertexBuffers.PositionVertexBuffer.GetVertexData(), InMutablePositionData, NumVertices * OutVertexBuffers.PositionVertexBuffer.GetStride());
+
+			// tangent and texture coords
+			OutVertexBuffers.StaticMeshVertexBuffer.SetUseFullPrecisionUVs(true);
+			OutVertexBuffers.StaticMeshVertexBuffer.SetUseHighPrecisionTangentBasis(false);
+			OutVertexBuffers.StaticMeshVertexBuffer.Init(NumVertices, NumTexCoords, bNeedCPUAccess);
+			FMemory::Memcpy(OutVertexBuffers.StaticMeshVertexBuffer.GetTangentData(), InMutableTangentData, OutVertexBuffers.StaticMeshVertexBuffer.GetTangentSize());
+			FMemory::Memcpy(OutVertexBuffers.StaticMeshVertexBuffer.GetTexCoordData(), InMutableTextureData, OutVertexBuffers.StaticMeshVertexBuffer.GetTexCoordSize());
+		}
 
 
 		/**
@@ -80,15 +80,15 @@ namespace UnrealConversionUtils
 		 * @param InMutableColorData - Mutable color data buffer
 		 */
 		void FColorVertexBuffers_InitWithMutableData(
-	FStaticMeshVertexBuffers& OutVertexBuffers,
-	const int32 NumVertices,
-	const void* InMutableColorData
-)
-{
-	// positions
-	OutVertexBuffers.ColorVertexBuffer.Init(NumVertices);
-	FMemory::Memcpy(OutVertexBuffers.ColorVertexBuffer.GetVertexData(), InMutableColorData, NumVertices * OutVertexBuffers.ColorVertexBuffer.GetStride());
-}
+			FStaticMeshVertexBuffers& OutVertexBuffers,
+			const int32 NumVertices,
+			const void* InMutableColorData
+		)
+		{
+			// positions
+			OutVertexBuffers.ColorVertexBuffer.Init(NumVertices);
+			FMemory::Memcpy(OutVertexBuffers.ColorVertexBuffer.GetVertexData(), InMutableColorData, NumVertices * OutVertexBuffers.ColorVertexBuffer.GetStride());
+		}
 
 
 		/**
@@ -101,255 +101,278 @@ namespace UnrealConversionUtils
 		 * @param InMutableData - Mutable data buffer
 		 */
 		void FSkinWeightVertexBuffer_InitWithMutableData(
-	FSkinWeightVertexBuffer& OutVertexWeightBuffer,
-	const int32 NumVertices,
-	const int32 NumBones,
-	const int32 NumBoneInfluences,
-	const bool bNeedCPUAccess,
-	const void* InMutableData)
-{
-	// \todo Ugly cast.
-	FSkinWeightDataVertexBuffer* VertexBuffer = const_cast<FSkinWeightDataVertexBuffer*>(OutVertexWeightBuffer.GetDataVertexBuffer());
-	VertexBuffer->SetMaxBoneInfluences(NumBoneInfluences);
-	VertexBuffer->Init(NumBones, NumVertices);
+			FSkinWeightVertexBuffer& OutVertexWeightBuffer,
+			const int32 NumVertices,
+			const int32 NumBones,
+			const int32 NumBoneInfluences,
+			const bool bNeedCPUAccess,
+			const void* InMutableData)
+		{
+			// \todo Ugly cast.
+			FSkinWeightDataVertexBuffer* VertexBuffer = const_cast<FSkinWeightDataVertexBuffer*>(OutVertexWeightBuffer.GetDataVertexBuffer());
+			VertexBuffer->SetMaxBoneInfluences(NumBoneInfluences);
+			VertexBuffer->Init(NumBones, NumVertices);
 
-	if (NumVertices)
-	{
-		OutVertexWeightBuffer.SetNeedsCPUAccess(bNeedCPUAccess);
+			if (NumVertices)
+			{
+				OutVertexWeightBuffer.SetNeedsCPUAccess(bNeedCPUAccess);
 
-		uint8* Data = VertexBuffer->GetWeightData();
-		FMemory::Memcpy(Data, InMutableData, OutVertexWeightBuffer.GetVertexDataSize());
-	}
-}
+				uint8* Data = VertexBuffer->GetWeightData();
+				FMemory::Memcpy(Data, InMutableData, OutVertexWeightBuffer.GetVertexDataSize());
+			}
+		}
 	}
 
 
-	
+
 	void BuildRefSkeleton(FInstanceUpdateData::FSkeletonData* OutMutSkeletonData
-	,const FReferenceSkeleton& InSourceReferenceSkeleton,const TArray<bool>& InUsedBones,
-	FReferenceSkeleton& InRefSkeleton, const USkeleton* InSkeleton )
-{
-	const int32 SourceBoneCount = InSourceReferenceSkeleton.GetNum();
-	const TArray<FTransform>& SourceRawMeshBonePose = InSourceReferenceSkeleton.GetRawRefBonePose();
-	
-	MUTABLE_CPUPROFILER_SCOPE(BuildSkeletonData_BuildRefSkeleton);
+		, const FReferenceSkeleton& InSourceReferenceSkeleton, const TArray<bool>& InUsedBones,
+		FReferenceSkeleton& InRefSkeleton, const USkeleton* InSkeleton)
+	{
+		const int32 SourceBoneCount = InSourceReferenceSkeleton.GetNum();
+		const TArray<FTransform>& SourceRawMeshBonePose = InSourceReferenceSkeleton.GetRawRefBonePose();
 
-	// Build new RefSkeleton	
-	FReferenceSkeletonModifier RefSkeletonModifier(InRefSkeleton, InSkeleton);
+		MUTABLE_CPUPROFILER_SCOPE(BuildSkeletonData_BuildRefSkeleton);
 
-	const TArray<FMeshBoneInfo>& BoneInfo = InSourceReferenceSkeleton.GetRawRefBoneInfo();
+		// Build new RefSkeleton	
+		FReferenceSkeletonModifier RefSkeletonModifier(InRefSkeleton, InSkeleton);
 
-	TMap<FName, uint16> BoneToFinalBoneIndexMap;
-	BoneToFinalBoneIndexMap.Reserve(SourceBoneCount);
-	
-	uint32 FinalBoneCount = 0;
-	for (int32 BoneIndex = 0; BoneIndex < SourceBoneCount; ++BoneIndex)
-	{
-		if (!InUsedBones[BoneIndex])
+		const TArray<FMeshBoneInfo>& BoneInfo = InSourceReferenceSkeleton.GetRawRefBoneInfo();
+
+		TMap<FName, uint16> BoneToFinalBoneIndexMap;
+		BoneToFinalBoneIndexMap.Reserve(SourceBoneCount);
+
+		uint32 FinalBoneCount = 0;
+		for (int32 BoneIndex = 0; BoneIndex < SourceBoneCount; ++BoneIndex)
 		{
-			continue;
-		}
+			if (!InUsedBones[BoneIndex])
+			{
+				continue;
+			}
 
-		FName BoneName = BoneInfo[BoneIndex].Name;
+			FName BoneName = BoneInfo[BoneIndex].Name;
 
-		// Build a bone to index map so we can remap BoneMaps and ActiveBoneIndices later on
-		BoneToFinalBoneIndexMap.Add(BoneName, FinalBoneCount);
-		FinalBoneCount++;
+			// Build a bone to index map so we can remap BoneMaps and ActiveBoneIndices later on
+			BoneToFinalBoneIndexMap.Add(BoneName, FinalBoneCount);
+			FinalBoneCount++;
 
-		// Find parent index
-		const int32 SourceParentIndex = BoneInfo[BoneIndex].ParentIndex;
-		const int32 ParentIndex = SourceParentIndex != INDEX_NONE ? BoneToFinalBoneIndexMap[BoneInfo[SourceParentIndex].Name] : INDEX_NONE;
+			// Find parent index
+			const int32 SourceParentIndex = BoneInfo[BoneIndex].ParentIndex;
+			const int32 ParentIndex = SourceParentIndex != INDEX_NONE ? BoneToFinalBoneIndexMap[BoneInfo[SourceParentIndex].Name] : INDEX_NONE;
 
-		RefSkeletonModifier.Add(FMeshBoneInfo(BoneName, BoneName.ToString(), ParentIndex), SourceRawMeshBonePose[BoneIndex]);
+			RefSkeletonModifier.Add(FMeshBoneInfo(BoneName, BoneName.ToString(), ParentIndex), SourceRawMeshBonePose[BoneIndex]);
+		}
 	}
-}
 
 
 	void BuildSkeletalMeshElementDataAtLOD(const int32 MeshLODIndex,
-	mu::MeshPtrConst InMutableMesh, USkeletalMesh* OutSkeletalMesh)
-{
-	int32 SurfaceCount = 0;
-	if (InMutableMesh)
+		mu::MeshPtrConst InMutableMesh, USkeletalMesh* OutSkeletalMesh)
 	{
-		SurfaceCount = InMutableMesh->GetSurfaceCount();
-	}
-	
-	Helper_GetLODInfoArray(OutSkeletalMesh)[MeshLODIndex].LODMaterialMap.SetNum(1);
-	Helper_GetLODInfoArray(OutSkeletalMesh)[MeshLODIndex].LODMaterialMap[0] = 0;
+		int32 SurfaceCount = 0;
+		if (InMutableMesh)
+		{
+			SurfaceCount = InMutableMesh->GetSurfaceCount();
+		}
 
-	const int32 MaterialCount = SurfaceCount;
-	int32 NewMaterialIndex = Helper_GetLODRenderSections(OutSkeletalMesh, MeshLODIndex).Num();
-	for (int32 MaterialIndex = 0; MaterialIndex < MaterialCount; MaterialIndex++, NewMaterialIndex++)
-	{
-		new(Helper_GetLODRenderSections(OutSkeletalMesh, MeshLODIndex)) Helper_SkelMeshRenderSection();
-		FSkelMeshRenderSection& Section = Helper_GetLODRenderSections(OutSkeletalMesh, MeshLODIndex)[NewMaterialIndex];
-		Section.MaterialIndex = 0;
-		Section.BaseIndex = 0;
-		Section.NumTriangles = 0;
+		Helper_GetLODInfoArray(OutSkeletalMesh)[MeshLODIndex].LODMaterialMap.SetNum(1);
+		Helper_GetLODInfoArray(OutSkeletalMesh)[MeshLODIndex].LODMaterialMap[0] = 0;
+
+		for (int32 SurfaceIndex = 0; SurfaceIndex < SurfaceCount; ++SurfaceIndex)
+		{
+			new(Helper_GetLODRenderSections(OutSkeletalMesh, MeshLODIndex)) Helper_SkelMeshRenderSection();
+		}
 	}
-}
 
 
 	void SetupRenderSections(
-	const mu::MeshPtrConst InMutableMesh,
-	const USkeletalMesh* OutSkeletalMesh,
-	const int32 MeshLODIndex,
-	const int32 NumBoneInfluences,
-	const TArray<uint16>& InBoneMap)
-{
-	const int32 SurfaceCount = InMutableMesh->GetSurfaceCount();
-	for (int32 Surface = 0; Surface < SurfaceCount; Surface++)
+		const USkeletalMesh* OutSkeletalMesh,
+		const mu::MeshPtrConst InMutableMesh,
+		const int32 MeshLODIndex,
+		const TArray<uint16>& InBoneMap)
 	{
-		MUTABLE_CPUPROFILER_SCOPE(UpdateSkeletalMesh_SurfaceLoop);
-
-		int32 FirstIndex;
-		int32 IndexCount;
-		int32 FirstVertex;
-		int32 VertexCount;
-		InMutableMesh->GetSurface(Surface, &FirstVertex, &VertexCount, &FirstIndex, &IndexCount);
-		FSkelMeshRenderSection& Section = Helper_GetLODRenderSections(OutSkeletalMesh, MeshLODIndex)[Surface];
-
-		Section.DuplicatedVerticesBuffer.Init(1, TMap<int, TArray<int32>>());
-			
-		if (VertexCount == 0 || IndexCount == 0)
+		check(InMutableMesh);
+
+		const mu::FMeshBufferSet& MutableMeshVertexBuffers = InMutableMesh->GetVertexBuffers();
+
+		// Find the number of influences from this mesh
+		int32 NumBoneInfluences = 0;
+		int32 boneIndexBuffer = -1;
+		int32 boneIndexChannel = -1;
+		MutableMeshVertexBuffers.FindChannel(mu::MBS_BONEINDICES, 0, &boneIndexBuffer, &boneIndexChannel);
+		if (boneIndexBuffer >= 0 || boneIndexChannel >= 0)
+		{
+			MutableMeshVertexBuffers.GetChannel(boneIndexBuffer, boneIndexChannel,
+				nullptr, nullptr, nullptr, &NumBoneInfluences, nullptr);
+		}
+
+		const int32 SurfaceCount = InMutableMesh->GetSurfaceCount();
+		for (int32 SurfaceIndex = 0; SurfaceIndex < SurfaceCount; ++SurfaceIndex)
 		{
-			Section.bDisabled = true;
-			continue; // Unreal doesn't like empty meshes
+			MUTABLE_CPUPROFILER_SCOPE(UpdateSkeletalMesh_SurfaceLoop);
+
+			int32 FirstIndex;
+			int32 IndexCount;
+			int32 FirstVertex;
+			int32 VertexCount;
+			InMutableMesh->GetSurface(SurfaceIndex, &FirstVertex, &VertexCount, &FirstIndex, &IndexCount);
+			FSkelMeshRenderSection& Section = Helper_GetLODRenderSections(OutSkeletalMesh, MeshLODIndex)[SurfaceIndex];
+
+			Section.DuplicatedVerticesBuffer.Init(1, TMap<int, TArray<int32>>());
+
+			if (VertexCount == 0 || IndexCount == 0)
+			{
+				Section.bDisabled = true;
+				continue; // Unreal doesn't like empty meshes
+			}
+
+			Section.BaseIndex = FirstIndex;
+			Section.NumTriangles = IndexCount / 3;
+			Section.BaseVertexIndex = FirstVertex;
+			Section.MaxBoneInfluences = NumBoneInfluences;
+			Section.NumVertices = VertexCount;
+			Section.BoneMap.Append(InBoneMap);
 		}
-			
-		Section.BaseIndex = FirstIndex;
-		Section.NumTriangles = IndexCount / 3;
-		Section.BaseVertexIndex = FirstVertex;
-		Section.MaxBoneInfluences = NumBoneInfluences;
-		Section.NumVertices = VertexCount;
-
-		//Section.BoneMap.Append(OperationData->InstanceUpdateData.LODs[LODIndex].Components[ComponentIndex].BoneMap);
-		Section.BoneMap.Append(InBoneMap);
 	}
-}
 
 
 	void CopyMutableVertexBuffers(
-	USkeletalMesh* OutSkeletalMesh,
-	const int32 NumVerticesLODModel, const int32 NumBoneInfluences,
-	const int32 BoneIndexBuffer,
-	const mu::FMeshBufferSet& MutableMeshVertexBuffers,
-	const int32 MeshLODIndex,
-	const mu::MESH_BUFFER_FORMAT& InBoneIndexFormat
-	)
-{
-	FSkeletalMeshLODRenderData& LODModel = Helper_GetLODData(OutSkeletalMesh)[MeshLODIndex];
-	
-	MUTABLE_CPUPROFILER_SCOPE(UpdateSkeletalMesh_SurfaceLoop_MemCpy);
+		USkeletalMesh* SkeletalMesh,
+		const mu::MeshPtrConst MutableMesh,
+		const int32 MeshLODIndex)
 
-	const ESkeletalMeshVertexFlags BuildFlags = OutSkeletalMesh->GetVertexBufferFlags();
-	const bool bUseFullPrecisionUVs = EnumHasAllFlags(BuildFlags, ESkeletalMeshVertexFlags::UseFullPrecisionUVs);
-	const bool bHasVertexColors = EnumHasAllFlags(BuildFlags, ESkeletalMeshVertexFlags::HasVertexColors);
-	const int NumTexCoords = MutableMeshVertexBuffers.GetBufferChannelCount(MUTABLE_VERTEXBUFFER_TEXCOORDS);
+	{
+		FSkeletalMeshLODRenderData& LODModel = Helper_GetLODData(SkeletalMesh)[MeshLODIndex];
 
-	const bool bNeedsCPUAccess = Helper_GetLODInfoArray(OutSkeletalMesh)[MeshLODIndex].bAllowCPUAccess;
-	
-	FStaticMeshVertexBuffers_InitWithMutableData(
-		LODModel.StaticVertexBuffers,
-		NumVerticesLODModel, 
-		NumTexCoords, 
-		bUseFullPrecisionUVs, 
-		bNeedsCPUAccess,
-		MutableMeshVertexBuffers.GetBufferData(MUTABLE_VERTEXBUFFER_POSITION),
-		MutableMeshVertexBuffers.GetBufferData(MUTABLE_VERTEXBUFFER_TANGENT),
-		MutableMeshVertexBuffers.GetBufferData(MUTABLE_VERTEXBUFFER_TEXCOORDS)
+		const mu::FMeshBufferSet& MutableMeshVertexBuffers = MutableMesh->GetVertexBuffers();
+		const int32 NumVertices = MutableMeshVertexBuffers.GetElementCount();
+
+		MUTABLE_CPUPROFILER_SCOPE(UpdateSkeletalMesh_SurfaceLoop_MemCpy);
+
+		const ESkeletalMeshVertexFlags BuildFlags = SkeletalMesh->GetVertexBufferFlags();
+		const bool bUseFullPrecisionUVs = EnumHasAllFlags(BuildFlags, ESkeletalMeshVertexFlags::UseFullPrecisionUVs);
+		const bool bHasVertexColors = EnumHasAllFlags(BuildFlags, ESkeletalMeshVertexFlags::HasVertexColors);
+		const int NumTexCoords = MutableMeshVertexBuffers.GetBufferChannelCount(MUTABLE_VERTEXBUFFER_TEXCOORDS);
+
+		const bool bNeedsCPUAccess = Helper_GetLODInfoArray(SkeletalMesh)[MeshLODIndex].bAllowCPUAccess;
+
+		FStaticMeshVertexBuffers_InitWithMutableData(
+			LODModel.StaticVertexBuffers,
+			NumVertices,
+			NumTexCoords,
+			bUseFullPrecisionUVs,
+			bNeedsCPUAccess,
+			MutableMeshVertexBuffers.GetBufferData(MUTABLE_VERTEXBUFFER_POSITION),
+			MutableMeshVertexBuffers.GetBufferData(MUTABLE_VERTEXBUFFER_TANGENT),
+			MutableMeshVertexBuffers.GetBufferData(MUTABLE_VERTEXBUFFER_TEXCOORDS)
 		);
 
-	// Init skin weight buffer
-	FSkinWeightVertexBuffer_InitWithMutableData(
-		LODModel.SkinWeightVertexBuffer,
-		NumVerticesLODModel,
-		NumBoneInfluences * NumVerticesLODModel,
-		NumBoneInfluences,
-		bNeedsCPUAccess,
-		MutableMeshVertexBuffers.GetBufferData(BoneIndexBuffer)
-	);
-
-	if (InBoneIndexFormat == mu::MBF_UINT16)
-	{
-		LODModel.SkinWeightVertexBuffer.SetUse16BitBoneIndex(true);
-	}
+		mu::MESH_BUFFER_FORMAT BoneIndexFormat = mu::MBF_NONE;
+		int32 NumBoneInfluences = 0;
+		int32 BoneIndexBuffer = -1;
+		int32 BoneIndexChannel = -1;
+		MutableMeshVertexBuffers.FindChannel(mu::MBS_BONEINDICES, 0, &BoneIndexBuffer, &BoneIndexChannel);
+		if (BoneIndexBuffer >= 0 || BoneIndexChannel >= 0)
+		{
+			MutableMeshVertexBuffers.GetChannel(BoneIndexBuffer, BoneIndexChannel,
+				nullptr, nullptr, &BoneIndexFormat, &NumBoneInfluences, nullptr);
+		}
 
-	// Optional buffers
-	for ( int Buffer= MUTABLE_VERTEXBUFFER_TEXCOORDS + 1; Buffer<MutableMeshVertexBuffers.GetBufferCount(); ++Buffer)
-	{
-		if (MutableMeshVertexBuffers.GetBufferChannelCount(Buffer) > 0)
+		mu::MESH_BUFFER_FORMAT BoneWeightFormat = mu::MBF_NONE;
+		int32 BoneWeightBuffer = -1;
+		int32 BoneWeightChannel = -1;
+		MutableMeshVertexBuffers.FindChannel(mu::MBS_BONEWEIGHTS, 0, &BoneWeightBuffer, &BoneWeightChannel);
+		if (BoneWeightBuffer >= 0 || BoneWeightChannel >= 0)
 		{
-			mu::MESH_BUFFER_SEMANTIC Semantic;
-			mu::MESH_BUFFER_FORMAT Format;
-			int32 SemanticIndex;
-			int32 ComponentCount;
-			int32 Offset;
-			MutableMeshVertexBuffers.GetChannel(Buffer, 0, &Semantic, &SemanticIndex, &Format, &ComponentCount, &Offset);
-			
-			if (Semantic == mu::MBS_BONEINDICES)
-			{
-				const int32 BonesPerVertex = ComponentCount;
-				const int32 NumBones = BonesPerVertex * NumVerticesLODModel;
-
-				check(LODModel.SkinWeightVertexBuffer.GetVariableBonesPerVertex() == false);
-				FSkinWeightVertexBuffer_InitWithMutableData(
-					LODModel.SkinWeightVertexBuffer,
-					NumVerticesLODModel,
-					NumBones, 
-					NumBoneInfluences,
-					bNeedsCPUAccess,
-					MutableMeshVertexBuffers.GetBufferData(Buffer)
-				);
-			}
+			MutableMeshVertexBuffers.GetChannel(BoneWeightBuffer, BoneWeightChannel,
+				nullptr, nullptr, &BoneWeightFormat, nullptr, nullptr);
+		}
+
+		// Init skin weight buffer
+		FSkinWeightVertexBuffer_InitWithMutableData(
+			LODModel.SkinWeightVertexBuffer,
+			NumVertices,
+			NumBoneInfluences * NumVertices,
+			NumBoneInfluences,
+			bNeedsCPUAccess,
+			MutableMeshVertexBuffers.GetBufferData(BoneIndexBuffer)
+		);
+
+		if (BoneIndexFormat == mu::MBF_UINT16)
+		{
+			LODModel.SkinWeightVertexBuffer.SetUse16BitBoneIndex(true);
+		}
 
-			// colour buffer?
-			else if (Semantic == mu::MBS_COLOUR)
+		if (BoneWeightFormat == mu::MBF_NUINT16)
+		{
+			// TODO: 
+			unimplemented()
+		}
+
+		// Optional buffers
+		for (int Buffer = MUTABLE_VERTEXBUFFER_TEXCOORDS + 1; Buffer < MutableMeshVertexBuffers.GetBufferCount(); ++Buffer)
+		{
+			if (MutableMeshVertexBuffers.GetBufferChannelCount(Buffer) > 0)
 			{
-				OutSkeletalMesh->SetHasVertexColors(true);
-				const void* DataPtr = MutableMeshVertexBuffers.GetBufferData(Buffer);
-				FColorVertexBuffers_InitWithMutableData(LODModel.StaticVertexBuffers, NumVerticesLODModel, DataPtr);
-				check(LODModel.StaticVertexBuffers.ColorVertexBuffer.GetStride() == MutableMeshVertexBuffers.GetElementSize(Buffer));
+				mu::MESH_BUFFER_SEMANTIC Semantic;
+				mu::MESH_BUFFER_FORMAT Format;
+				int32 SemanticIndex;
+				int32 ComponentCount;
+				int32 Offset;
+				MutableMeshVertexBuffers.GetChannel(Buffer, 0, &Semantic, &SemanticIndex, &Format, &ComponentCount, &Offset);
+
+				if (Semantic == mu::MBS_BONEINDICES)
+				{
+					const int32 BonesPerVertex = ComponentCount;
+					const int32 NumBones = BonesPerVertex * NumVertices;
+
+					check(LODModel.SkinWeightVertexBuffer.GetVariableBonesPerVertex() == false);
+					FSkinWeightVertexBuffer_InitWithMutableData(
+						LODModel.SkinWeightVertexBuffer,
+						NumVertices,
+						NumBones,
+						NumBoneInfluences,
+						bNeedsCPUAccess,
+						MutableMeshVertexBuffers.GetBufferData(Buffer)
+					);
+				}
+
+				// colour buffer?
+				else if (Semantic == mu::MBS_COLOUR)
+				{
+					SkeletalMesh->SetHasVertexColors(true);
+					const void* DataPtr = MutableMeshVertexBuffers.GetBufferData(Buffer);
+					FColorVertexBuffers_InitWithMutableData(LODModel.StaticVertexBuffers, NumVertices, DataPtr);
+					check(LODModel.StaticVertexBuffers.ColorVertexBuffer.GetStride() == MutableMeshVertexBuffers.GetElementSize(Buffer));
+				}
 			}
 		}
 	}
-}
 
 
-	bool CopyMutableIndexBuffers(mu::MeshPtrConst InMutableMesh,FSkeletalMeshLODRenderData& LODModel)
-{
-	MUTABLE_CPUPROFILER_SCOPE(UpdateSkeletalMesh_BuildSkeletalMeshRenderData_IndexLoop);
+	bool CopyMutableIndexBuffers(mu::MeshPtrConst InMutableMesh, FSkeletalMeshLODRenderData& LODModel)
+	{
+		MUTABLE_CPUPROFILER_SCOPE(UpdateSkeletalMesh_BuildSkeletalMeshRenderData_IndexLoop);
 
-	const int32 IndexCount = InMutableMesh->GetIndexBuffers().GetElementCount();
+		const int32 IndexCount = InMutableMesh->GetIndexBuffers().GetElementCount();
 
-	if (IndexCount)
-	{
-		if (InMutableMesh->GetIndexBuffers().GetElementSize(0) == 2)
+		if (IndexCount == 0)
 		{
-			const uint16* IndexPtr = (const uint16*)InMutableMesh->GetIndexBuffers().GetBufferData(0);
-			LODModel.MultiSizeIndexContainer.CreateIndexBuffer(2);
-			LODModel.MultiSizeIndexContainer.GetIndexBuffer()->Insert(0, IndexCount);
-			FMemory::Memcpy(LODModel.MultiSizeIndexContainer.GetIndexBuffer()->GetPointerTo(0), IndexPtr, IndexCount * sizeof(uint16));
-		}
-		else
-		{
-			const uint32* IndexPtr = (const uint32*)InMutableMesh->GetIndexBuffers().GetBufferData(0);
-			LODModel.MultiSizeIndexContainer.CreateIndexBuffer(4);
-			LODModel.MultiSizeIndexContainer.GetIndexBuffer()->Insert(0, IndexCount);
-			FMemory::Memcpy(LODModel.MultiSizeIndexContainer.GetIndexBuffer()->GetPointerTo(0), IndexPtr, IndexCount * sizeof(uint32));
+			// Copy indices from an empty buffer
+			UE_LOG(LogMutable, Error, TEXT("UCustomizableInstancePrivateData::BuildSkeletalMeshRenderData is converting an empty mesh."));
+			return false;
 		}
 
-	}
-	else
-	{
-		UE_LOG(LogMutable, Error, TEXT("UCustomizableInstancePrivateData::BuildSkeletalMeshRenderData is converting an empty mesh."));
-		return false;
+		const uint8* DataPtr = InMutableMesh->GetIndexBuffers().GetBufferData(0);
+		const int32 ElementSize = InMutableMesh->GetIndexBuffers().GetElementSize(0);
+
+		LODModel.MultiSizeIndexContainer.CreateIndexBuffer(ElementSize);
+		LODModel.MultiSizeIndexContainer.GetIndexBuffer()->Insert(0, IndexCount);
+		FMemory::Memcpy(LODModel.MultiSizeIndexContainer.GetIndexBuffer()->GetPointerTo(0), DataPtr, IndexCount * ElementSize);
+
+		return true;
 	}
 
-	return true;
 }
 
-		}
-
 	
 		
diff --git a/Engine/Plugins/Experimental/Mutable/Source/CustomizableObject/Public/MuCO/MutableMeshBufferUtils.h b/Engine/Plugins/Experimental/Mutable/Source/CustomizableObject/Public/MuCO/MutableMeshBufferUtils.h
--- a/Engine/Plugins/Experimental/Mutable/Source/CustomizableObject/Public/MuCO/MutableMeshBufferUtils.h
+++ b/Engine/Plugins/Experimental/Mutable/Source/CustomizableObject/Public/MuCO/MutableMeshBufferUtils.h
@@ -46,7 +46,8 @@ namespace  MutableMeshBufferUtils
 	 * @param MaxNumBonesPerVertex Amount of bone influences a vertex can have.
 	 * @param OutTargetVertexBuffers Mesh buffer to be set up. 
 	 */
-	CUSTOMIZABLEOBJECT_API void SetupSkinBuffer(const int32& InCurrentVertexBuffer,const int32& MaxBoneIndexTypeSizeBytes, const int32& MaxNumBonesPerVertex, mu::FMeshBufferSet& OutTargetVertexBuffers);
+	CUSTOMIZABLEOBJECT_API void SetupSkinBuffer(const int32& InCurrentVertexBuffer,const int32& MaxBoneIndexTypeSizeBytes, const int32& MaxBoneWeightTypeSizeBytes,
+		const int32& MaxNumBonesPerVertex, mu::FMeshBufferSet& OutTargetVertexBuffers);
 	
 	/**
 	 * Sets up all the channels that ought to be part of the Vertex Color Buffer of the Mutable mesh owning the provided MeshBufferSet
diff --git a/Engine/Plugins/Experimental/Mutable/Source/CustomizableObject/Public/MuCO/UnrealConversionUtils.h b/Engine/Plugins/Experimental/Mutable/Source/CustomizableObject/Public/MuCO/UnrealConversionUtils.h
--- a/Engine/Plugins/Experimental/Mutable/Source/CustomizableObject/Public/MuCO/UnrealConversionUtils.h
+++ b/Engine/Plugins/Experimental/Mutable/Source/CustomizableObject/Public/MuCO/UnrealConversionUtils.h
@@ -53,37 +53,27 @@ namespace UnrealConversionUtils
 	
 	/**
 	 * Prepares the render sections found on the InSkeletalMesh and sets them up accordingly what the InMutableMesh requires
-	 * @param InMutableMesh - Mutable mesh to be used as reference for the section data update on the Skeletal Mesh
 	 * @param OutSkeletalMesh - The Skeletal mesh whose sections are ought to be updated
+	 * @param InMutableMesh - Mutable mesh to be used as reference for the section data update on the Skeletal Mesh
 	 * @param MeshLODIndex - The Mesh lod to be working with
-	 * @param NumBoneInfluences - The amount of bones a pixel can be affected by
 	 * @param InBoneMap - The bones to be set as part of the sections.
 	 */
 	CUSTOMIZABLEOBJECT_API void SetupRenderSections(
-		const mu::MeshPtrConst InMutableMesh,
 		const USkeletalMesh* OutSkeletalMesh,
-		const int32 MeshLODIndex,
-		const int32 NumBoneInfluences,
+		const mu::MeshPtrConst InMutableMesh,
+		const int32 InMeshLODIndex,
 		const TArray<uint16>& InBoneMap);
 
 
 	/** Performs a copy of the data found on the vertex buffers on the mutable mesh to the buffers of the skeletal mesh
 	 * @param OutSkeletalMesh - Skeletal Mesh to be updated with new buffer data
-	 * @param NumVerticesLODModel - The amount of vertices on the LODModel.
-	 * @param NumBoneInfluences - The amount of bones that can influence a vertex
-	 * @param BoneIndexBuffer - The buffer containing the indices for the bones.
-	 * @param InMutableMeshVertexBuffers - The mutable buffers to be reading data from
+	 * @param InMutableMesh - Mutable mesh to be used as reference for the section data update on the Skeletal Mesh
 	 * @param MeshLODIndex - The LOD index we are working with.
-	 * @param InBoneIndexFormat - The format of the bone indices on the mutable mesh.
 	 */
 	CUSTOMIZABLEOBJECT_API void CopyMutableVertexBuffers(
 		USkeletalMesh* OutSkeletalMesh,
-		const int32 NumVerticesLODModel,
-		const int32 NumBoneInfluences,
-		const int32 BoneIndexBuffer,
-		const mu::FMeshBufferSet& InMutableMeshVertexBuffers,
-		const int32 MeshLODIndex,
-		const mu::MESH_BUFFER_FORMAT& InBoneIndexFormat);
+		const mu::MeshPtrConst InMutableMesh,
+		const int32 InMeshLODIndex);
 
 	
 	/**
diff --git a/Engine/Plugins/Experimental/Mutable/Source/CustomizableObjectEditor/Private/MuCOE/CustomizableObjectCompiler.cpp b/Engine/Plugins/Experimental/Mutable/Source/CustomizableObjectEditor/Private/MuCOE/CustomizableObjectCompiler.cpp
--- a/Engine/Plugins/Experimental/Mutable/Source/CustomizableObjectEditor/Private/MuCOE/CustomizableObjectCompiler.cpp
+++ b/Engine/Plugins/Experimental/Mutable/Source/CustomizableObjectEditor/Private/MuCOE/CustomizableObjectCompiler.cpp
@@ -948,6 +948,7 @@ void FCustomizableObjectCompiler::CompileInternal(UCustomizableObject* Object, c
 	Options.bExtraBoneInfluencesEnabled = bAreExtraBoneInfluencesEnabled;
 	Options.bRealTimeMorphTargetsEnabled = Object->bEnableRealTimeMorphTargets;
 	Options.bClothingEnabled = Object->bEnableClothing;
+	Options.b16BitBoneWeightsEnabled = Object->bEnable16BitBoneWeights;
 
 	if (Object->IsLocked() || !UCustomizableObjectSystem::GetInstance()->LockObject(Object))
 	{
diff --git a/Engine/Plugins/Experimental/Mutable/Source/CustomizableObjectEditor/Private/MuCOE/GenerateMutableSource/GenerateMutableSourceMesh.cpp b/Engine/Plugins/Experimental/Mutable/Source/CustomizableObjectEditor/Private/MuCOE/GenerateMutableSource/GenerateMutableSourceMesh.cpp
--- a/Engine/Plugins/Experimental/Mutable/Source/CustomizableObjectEditor/Private/MuCOE/GenerateMutableSource/GenerateMutableSourceMesh.cpp
+++ b/Engine/Plugins/Experimental/Mutable/Source/CustomizableObjectEditor/Private/MuCOE/GenerateMutableSource/GenerateMutableSourceMesh.cpp
@@ -494,7 +494,11 @@ mu::MeshPtr ConvertSkeletalMeshToMutable(USkeletalMesh* InSkeletalMesh, int LOD,
 		const int ChannelCount = 11;
 		const MESH_BUFFER_SEMANTIC Semantics[ChannelCount] = { MBS_POSITION, MBS_TANGENT, MBS_BINORMAL, MBS_NORMAL, MBS_TEXCOORDS, MBS_TEXCOORDS, MBS_TEXCOORDS, MBS_TEXCOORDS, MBS_COLOUR, MBS_BONEINDICES, MBS_BONEWEIGHTS };
 		const int SemanticIndices[ChannelCount] = { 0, 0, 0, 0, 0, 1, 2, 3, 0, 0, 0 };
-		const MESH_BUFFER_FORMAT Formats[ChannelCount] = { MBF_FLOAT32, MBF_FLOAT32, MBF_FLOAT32, MBF_FLOAT32, MBF_FLOAT32, MBF_FLOAT32, MBF_FLOAT32, MBF_FLOAT32, MBF_NUINT8, MBF_UINT16, MBF_NUINT8 };
+
+		// TODO: Remove BoneWeightFormat after merge
+		MESH_BUFFER_FORMAT BoneWeightFormat = sizeof(TDecay<decltype(DeclVal<FSoftSkinVertex>().InfluenceWeights[0])>::Type) == 1 ? MBF_NUINT8 : MBF_NUINT16;
+		const MESH_BUFFER_FORMAT Formats[ChannelCount] = { MBF_FLOAT32, MBF_FLOAT32, MBF_FLOAT32, MBF_FLOAT32, MBF_FLOAT32, MBF_FLOAT32, MBF_FLOAT32, MBF_FLOAT32, MBF_NUINT8, MBF_UINT16, BoneWeightFormat };
+
 		int Components[ChannelCount] = { 3, 3, 3, 4, 2, 2, 2, 2, 4, 4, 4 };
 		if (GenerationContext.Options.bExtraBoneInfluencesEnabled && MaxSectionInfluences > 4)
 		{
@@ -502,7 +506,7 @@ mu::MeshPtr ConvertSkeletalMeshToMutable(USkeletalMesh* InSkeletalMesh, int LOD,
 			Components[10] = EXTRA_BONE_INFLUENCES;
 		}
 
-        constexpr size_t SoftSkinVertexUVsElemSize = sizeof(TDecay<decltype(DeclVal<FSoftSkinVertex>().UVs[0])>::Type); 
+		constexpr size_t SoftSkinVertexUVsElemSize = sizeof(TDecay<decltype(DeclVal<FSoftSkinVertex>().UVs[0])>::Type);
 		const int Offsets[ChannelCount] =
 		{
 			STRUCT_OFFSET(FSoftSkinVertex, Position),
@@ -518,6 +522,30 @@ mu::MeshPtr ConvertSkeletalMeshToMutable(USkeletalMesh* InSkeletalMesh, int LOD,
 			STRUCT_OFFSET(FSoftSkinVertex, InfluenceWeights),
 		};
 
+		// Fix bone weights if required (uint8 -> uint16)
+		if (BoneWeightFormat == MBF_NUINT16 && Vertices.IsValidIndex(VertexStart))
+		{
+			FSoftSkinVertex FirstVertex = Vertices[VertexStart];
+
+			uint16 TotalWeight = 0;
+			for (int32 InfluenceIndex = 0; InfluenceIndex < MaxSectionInfluences; ++InfluenceIndex)
+			{
+				TotalWeight += FirstVertex.InfluenceWeights[InfluenceIndex];
+			}
+
+			if (TotalWeight <= 255)
+			{
+				for (int32 VertexIndex = VertexStart; VertexIndex < VertexStart + VertexCount && VertexIndex < Vertices.Num(); ++VertexIndex)
+				{
+					FSoftSkinVertex& Vertex = Vertices[VertexIndex];
+					for (int32 InfluenceIndex = 0; InfluenceIndex < MaxSectionInfluences; ++InfluenceIndex)
+					{
+						Vertex.InfluenceBones[InfluenceIndex] = Vertex.InfluenceBones[InfluenceIndex] * (65535 / 255);
+					}
+				}
+			}
+		}
+
 		if (GenerationContext.Options.bExtraBoneInfluencesEnabled && MaxSectionInfluences < EXTRA_BONE_INFLUENCES)
 		{
 			for (int32 VertexIndex = VertexStart; VertexIndex < VertexStart + VertexCount && VertexIndex < Vertices.Num(); ++VertexIndex)
@@ -532,7 +560,8 @@ mu::MeshPtr ConvertSkeletalMeshToMutable(USkeletalMesh* InSkeletalMesh, int LOD,
 		}
 		else if (!GenerationContext.Options.bExtraBoneInfluencesEnabled)
 		{
-			int32 MaxSectionBoneMapIndex = ImportedModel->LODModels[LOD].Sections[MaterialIndex].BoneMap.Num();
+			const int32 MaxBoneWeightValue = BoneWeightFormat == MBF_NUINT16 ? 65535 : 255;
+			const int32 MaxSectionBoneMapIndex = ImportedModel->LODModels[LOD].Sections[MaterialIndex].BoneMap.Num();
 
 			// Renormalize to 4 weights per vertex
 			for (int32 VertexIndex = VertexStart; VertexIndex < VertexStart + VertexCount && VertexIndex < Vertices.Num(); ++VertexIndex)
@@ -633,7 +662,7 @@ mu::MeshPtr ConvertSkeletalMeshToMutable(USkeletalMesh* InSkeletalMesh, int LOD,
 						if (j < MaxMutableWeights)
 						{
 							float Aux = Vertex.InfluenceWeights[j];
-							int32 Res = FMath::RoundToInt(Aux / TotalWeight * 255);
+							int32 Res = FMath::RoundToInt(Aux / TotalWeight * MaxBoneWeightValue);
 							AssignedWeight += Res;
 							Vertex.InfluenceWeights[j] = Res;
 						}
@@ -643,11 +672,11 @@ mu::MeshPtr ConvertSkeletalMeshToMutable(USkeletalMesh* InSkeletalMesh, int LOD,
 						}
 					}
 
-					Vertex.InfluenceWeights[0] = 255 - AssignedWeight;
+					Vertex.InfluenceWeights[0] = MaxBoneWeightValue - AssignedWeight;
 				}
 				else
 				{
-					Vertex.InfluenceWeights[0] = 255;
+					Vertex.InfluenceWeights[0] = MaxBoneWeightValue;
 
 					for (int32 j = 1; j < MaxMutableWeights; ++j)
 					{
diff --git a/Engine/Plugins/Experimental/Mutable/Source/CustomizableObjectEditor/Private/MuCOE/GenerateMutableSource/GenerateMutableSourceSurface.cpp b/Engine/Plugins/Experimental/Mutable/Source/CustomizableObjectEditor/Private/MuCOE/GenerateMutableSource/GenerateMutableSourceSurface.cpp
--- a/Engine/Plugins/Experimental/Mutable/Source/CustomizableObjectEditor/Private/MuCOE/GenerateMutableSource/GenerateMutableSourceSurface.cpp
+++ b/Engine/Plugins/Experimental/Mutable/Source/CustomizableObjectEditor/Private/MuCOE/GenerateMutableSource/GenerateMutableSourceSurface.cpp
@@ -129,7 +129,7 @@ void LayoutOperationUVNormalizedWarning(FMutableGraphGenerationContext& Generati
 }
 
 void SetSurfaceFormat( mu::FMeshBufferSet& OutVertexBufferFormat, mu::FMeshBufferSet& OutIndexBufferFormat, const FMutableGraphMeshGenerationData& MeshData, 
-					   bool bWithExtraBoneInfluences, bool bWithRealTimeMorphs, bool bWithClothing )
+					   bool bWithExtraBoneInfluences, bool bWithRealTimeMorphs, bool bWithClothing, bool bWith16BitWeights )
 {
 	// Limit skinning weights if necessary
 	// \todo: make it more flexible to support 3 or 5 or 1 weight, since there is support for this in 4.25
@@ -180,7 +180,8 @@ void SetSurfaceFormat( mu::FMeshBufferSet& OutVertexBufferFormat, mu::FMeshBuffe
 	// Skin buffer
 	if (MeshData.MaxNumBonesPerVertex > 0 && MeshData.MaxBoneIndexTypeSizeBytes > 0)
 	{
-		MutableMeshBufferUtils::SetupSkinBuffer(CurrentVertexBuffer, MeshData.MaxBoneIndexTypeSizeBytes, MutableBonesPerVertex, OutVertexBufferFormat);
+		const int32 MaxBoneWeightTypeSizeBytes = bWith16BitWeights ? 2 : 1;
+		MutableMeshBufferUtils::SetupSkinBuffer(CurrentVertexBuffer, MeshData.MaxBoneIndexTypeSizeBytes, MaxBoneWeightTypeSizeBytes, MutableBonesPerVertex, OutVertexBufferFormat);
 		++CurrentVertexBuffer;
 	}
 
@@ -393,7 +394,8 @@ mu::NodeSurfacePtr GenerateMutableSourceSurface(const UEdGraphPin * Pin, FMutabl
 						MeshFormatNode->GetVertexBuffers(), MeshFormatNode->GetIndexBuffers(), MeshData,
 						GenerationContext.Options.bExtraBoneInfluencesEnabled,
 						GenerationContext.Options.bRealTimeMorphTargetsEnabled, 
-						GenerationContext.Options.bClothingEnabled);
+						GenerationContext.Options.bClothingEnabled,
+						GenerationContext.Options.b16BitBoneWeightsEnabled);
 				MeshFormatNode->SetMessageContext(Node);
 
 				SurfNode->SetMesh(0, MeshFormatNode);
@@ -1013,7 +1015,8 @@ mu::NodeSurfacePtr GenerateMutableSourceSurface(const UEdGraphPin * Pin, FMutabl
 						MeshFormat->GetVertexBuffers(), MeshFormat->GetIndexBuffers(), MeshData,
 						GenerationContext.Options.bExtraBoneInfluencesEnabled,
 						GenerationContext.Options.bRealTimeMorphTargetsEnabled, 
-						GenerationContext.Options.bClothingEnabled );
+						GenerationContext.Options.bClothingEnabled,
+						GenerationContext.Options.b16BitBoneWeightsEnabled);
 
 				MeshFormat->SetSource(AddMeshNode.get());
 
diff --git a/Engine/Plugins/Experimental/Mutable/Source/CustomizableObjectEditor/Private/MuCOE/MutableMeshPreviewUtils.cpp b/Engine/Plugins/Experimental/Mutable/Source/CustomizableObjectEditor/Private/MuCOE/MutableMeshPreviewUtils.cpp
--- a/Engine/Plugins/Experimental/Mutable/Source/CustomizableObjectEditor/Private/MuCOE/MutableMeshPreviewUtils.cpp
+++ b/Engine/Plugins/Experimental/Mutable/Source/CustomizableObjectEditor/Private/MuCOE/MutableMeshPreviewUtils.cpp
@@ -72,6 +72,8 @@ namespace MutableMeshPreviewUtils
 				bool bOriginDoesHaveVertexColorData = false;
 				int32 NumbOfTextCoordChannels = 0;
 				int32 MaxNumBonesPerVertex = 0;
+				int32 BoneIndicesSizeBytes = 0;
+				int32 BoneWeightsSizeBytes = 0;
 
 				// Contingency to avoid mismatches between the ideal semantic indices on the provided mutable mesh against the
 				// ones that should be .
@@ -124,6 +126,14 @@ namespace MutableMeshPreviewUtils
 							{
 								// Store the amount of bones a vertex can be skinned to.
 								MaxNumBonesPerVertex = BufferComponentCount;
+								BoneIndicesSizeBytes = BufferFormat == mu::MESH_BUFFER_FORMAT::MBF_INT16 ? 2 : 1;
+								break;
+							}
+
+							case mu::MESH_BUFFER_SEMANTIC::MBS_BONEWEIGHTS:
+							{
+								// Store the amount of bones a vertex can be skinned to.
+								BoneWeightsSizeBytes = BufferFormat == mu::MESH_BUFFER_FORMAT::MBF_NUINT16 ? 2 : 1;
 								break;
 							}
 
@@ -171,21 +181,8 @@ namespace MutableMeshPreviewUtils
 					CurrentVertexBuffer++;
 				
 					// Skin buffer
-					{
-						// Use the index buffer to know the amount of bones
-						const int32 NumberOfBones = InMutableMesh->GetVertexCount() * MaxNumBonesPerVertex;;
-					
-						// Failing this check would mean that the provided mutable mesh is in very bad shape (missing verts or
-						// the BONE_INDICES buffer channel
-						if (!NumberOfBones)
-						{
-							UE_LOG(LogTemp,Warning,TEXT("Unable to extract the bone count data from the provided mutable mesh. Make sure it has at least buffer data  "))
-						}
-
-						const int32 MaxBoneIndexTypeSizeBytes = NumberOfBones > 256 ? 2 : 1;
-						MutableMeshBufferUtils::SetupSkinBuffer(CurrentVertexBuffer,MaxBoneIndexTypeSizeBytes,MaxNumBonesPerVertex,FormattedVertexBuffers);
-						CurrentVertexBuffer++;
-					}
+					MutableMeshBufferUtils::SetupSkinBuffer(CurrentVertexBuffer, BoneIndicesSizeBytes, BoneWeightsSizeBytes, MaxNumBonesPerVertex, FormattedVertexBuffers);
+					CurrentVertexBuffer++;
 
 					// Colour buffer
 					if (bOriginDoesHaveVertexColorData)
@@ -602,93 +599,44 @@ namespace MutableMeshPreviewUtils
 			OutSkeletalMesh->SetHasBeenSimplified(false);
 			OutSkeletalMesh->SetHasVertexColors(false);
 
-			// Constant values since we are working with a simple mutable mesh (represents only one LOD)
-			constexpr int32 MeshLODIndex = 0;
-			constexpr int32 LODIndex = 0;
-			
-			// Detect if the produced mesh is correct
-			bool bSuccess = true;
-
-			FSkeletalMeshLODRenderData& LODModel = Helper_GetLODData(OutSkeletalMesh)[MeshLODIndex];
-			const mu::FMeshBufferSet& MutableMeshVertexBuffers = InMutableMesh->GetVertexBuffers();
-			
 			// Load buffer data found on the mutable model onto the out skeletal mesh
 			// It includes vertex and index buffers
-			{
-				const int32 NumVerticesLODModel = InMutableMesh->GetVertexCount();
-				//const int32 SurfaceCount = InMutableMesh->GetSurfaceCount();
-
-				mu::MESH_BUFFER_FORMAT boneIndexFormat = mu::MBF_NONE;
-				int32 BoneIndexComponents = 0;
-				int32 BoneIndexOffset = 0;
-				int32 BoneIndexBuffer = -1;
-				int32 BoneIndexChannel = -1;
-				MutableMeshVertexBuffers.FindChannel(mu::MBS_BONEINDICES, 0, &BoneIndexBuffer, &BoneIndexChannel);
-				if (BoneIndexBuffer >= 0 || BoneIndexChannel >= 0)
-				{
-					MutableMeshVertexBuffers.GetChannel(BoneIndexBuffer, BoneIndexChannel,
-					nullptr, nullptr, &boneIndexFormat, &BoneIndexComponents, &BoneIndexOffset);
-				}
-
-				const int32 NumBoneInfluences = BoneIndexComponents;
-
-				UnrealConversionUtils::SetupRenderSections(
-					InMutableMesh,
-					OutSkeletalMesh,
-					MeshLODIndex,
-					NumBoneInfluences,
-					InBoneMap);
+			UnrealConversionUtils::SetupRenderSections(
+				OutSkeletalMesh,
+				InMutableMesh,
+				0,
+				InBoneMap);
+
+			UnrealConversionUtils::CopyMutableVertexBuffers(
+				OutSkeletalMesh,
+				InMutableMesh,
+				0); 
 			
-				UnrealConversionUtils::CopyMutableVertexBuffers(
-					OutSkeletalMesh,
-					NumVerticesLODModel,
-					NumBoneInfluences,
-					BoneIndexBuffer,
-					MutableMeshVertexBuffers,
-					MeshLODIndex,
-					boneIndexFormat);
-			}
-			
-			
-			{
-				MUTABLE_CPUPROFILER_SCOPE(UpdateSkeletalMesh_SurfaceLoop_UpdateActiveAndRequiredBones);
+			FSkeletalMeshLODRenderData& LODModel = Helper_GetLODData(OutSkeletalMesh)[0];
 
-				// Generate the active bones array based on the bones found on BoneMap
-				TArray<uint16> ActiveBones;
+			// Generate the active bones array based on the bones found on BoneMap
+			TArray<uint16> ActiveBones;
+			{
+				ActiveBones.Reserve(InBoneMap.Num());
+				for (const uint16& BoneMapIndex : InBoneMap)
 				{
-					ActiveBones.Reserve(InBoneMap.Num());
-					for (const uint16& BoneMapIndex : InBoneMap)
-					{
-						ActiveBones.AddUnique(BoneMapIndex);
-					}
-					ActiveBones.Sort();
+					ActiveBones.AddUnique(BoneMapIndex);
 				}
-				
-				LODModel.ActiveBoneIndices.Append(ActiveBones);
-				LODModel.RequiredBones.Append(ActiveBones);
-			}
-			
-			if (bSuccess)
-			{
-				bSuccess = UnrealConversionUtils::CopyMutableIndexBuffers(InMutableMesh,LODModel);
+				ActiveBones.Sort();
 			}
 
+			LODModel.ActiveBoneIndices.Append(ActiveBones);
+			LODModel.RequiredBones.Append(ActiveBones);
 			
-			// Update LOD and streaming data
-			if (InRefSkeletalMesh && LODIndex >= 0 && LODIndex < Helper_GetLODNum(InRefSkeletalMesh))
-			{
-				const FSkeletalMeshLODRenderData& ReferenceLODModel = Helper_GetLODData(InRefSkeletalMesh)[LODIndex];
-
-				LODModel.bIsLODOptional = ReferenceLODModel.bIsLODOptional;
-				LODModel.bStreamedDataInlined = ReferenceLODModel.bStreamedDataInlined;
-			}
-
-
-			if (!bSuccess)
+			// Copy index buffers or fail to generate the mesh
+			if (!UnrealConversionUtils::CopyMutableIndexBuffers(InMutableMesh, LODModel))
 			{
-				// End with failure
 				return false;
 			}
+			
+			// Update LOD and streaming data
+			LODModel.bIsLODOptional = false;
+			LODModel.bStreamedDataInlined = false;
 
 			if (InRefSkeletalMesh)
 			{
diff --git a/Engine/Plugins/Experimental/Mutable/Source/MutableRuntime/Private/MuR/ConvertData.h b/Engine/Plugins/Experimental/Mutable/Source/MutableRuntime/Private/MuR/ConvertData.h
--- a/Engine/Plugins/Experimental/Mutable/Source/MutableRuntime/Private/MuR/ConvertData.h
+++ b/Engine/Plugins/Experimental/Mutable/Source/MutableRuntime/Private/MuR/ConvertData.h
@@ -955,6 +955,13 @@ namespace mu
 				break;
 			}
 
+			case MBF_NUINT16:
+			{
+				const uint16_t* pTypedSource = reinterpret_cast<const uint16_t*>(pSource);
+				pTypedResult[channel] = pTypedSource[channel] / (65535 / 255);
+				break;
+			}
+
 			case MBF_FLOAT32:
 			{
 				const float* pTypedSource = reinterpret_cast<const float*>(pSource);
@@ -1013,6 +1020,13 @@ namespace mu
 				break;
 			}
 
+			case MBF_NUINT8:
+			{
+				const uint8_t* pTypedSource = reinterpret_cast<const uint8_t*>(pSource);
+				pTypedResult[channel] = pTypedSource[channel] * (65535 / 255);
+				break;
+			}
+
 			case MBF_FLOAT32:
 			{
 				const float* pTypedSource = reinterpret_cast<const float*>( pSource );
diff --git a/Engine/Plugins/Experimental/Mutable/Source/MutableRuntime/Private/MuR/OpMeshFormat.cpp b/Engine/Plugins/Experimental/Mutable/Source/MutableRuntime/Private/MuR/OpMeshFormat.cpp
--- a/Engine/Plugins/Experimental/Mutable/Source/MutableRuntime/Private/MuR/OpMeshFormat.cpp
+++ b/Engine/Plugins/Experimental/Mutable/Source/MutableRuntime/Private/MuR/OpMeshFormat.cpp
@@ -632,7 +632,7 @@ namespace mu
 							{
 								if (resultFormat == MBF_NUINT8)
 								{
-									auto pData = (uint8_t*)pResultBuf;
+									uint8_t* pData = (uint8_t*)pResultBuf;
 									uint8_t accum = 0;
 									for (int i = 0; i < resultComponents; ++i)
 									{
@@ -640,6 +640,17 @@ namespace mu
 									}
 									pData[0] += 255 - accum;
 								}
+
+								else if (resultFormat == MBF_NUINT16)
+								{
+									uint16_t* pData = (uint16_t*)pResultBuf;
+									uint16_t accum = 0;
+									for (int i = 0; i < resultComponents; ++i)
+									{
+										accum += pData[i];
+									}
+									pData[0] += 65535 - accum;
+								}
 							}
 						}
 
