[mutable] Avoid CO validation if performing a cooking operation.
Now, the data validation of the COs will be not performed if we are trying to run it as part of a commandlet. This stops CI from raising warnings and errors related to the COs.
- Reverted changes on CL 22250497 as asked.

diff --git a/Engine/Plugins/Experimental/Mutable/Source/CustomizableObject/Private/MuCO/CustomizableObject.cpp b/Engine/Plugins/Experimental/Mutable/Source/CustomizableObject/Private/MuCO/CustomizableObject.cpp
--- a/Engine/Plugins/Experimental/Mutable/Source/CustomizableObject/Private/MuCO/CustomizableObject.cpp
+++ b/Engine/Plugins/Experimental/Mutable/Source/CustomizableObject/Private/MuCO/CustomizableObject.cpp
@@ -1488,6 +1488,15 @@ EDataValidationResult UCustomizableObject::IsDataValid(TArray<FText>& Validation
 		bIsValidationTriggeredBySave = false;
     	return Result;
     }
+
+	// Skip validation when cooking the assets. The validation of the CO is designed to be used explicitly by the user
+	// and not during automated operations like saving or cooking or any other automated action.
+	if (IsRunningCommandlet())
+	{
+		return Result;
+	}
+
+	UE_LOG(LogMutable,Display,TEXT("Running data validation checks for %s CO"),*this->GetName());
 	
 	// Request a compiler to be able to locate the root and to compile it
 	FCustomizableObjectCompilerBase* Compiler = UCustomizableObjectSystem::GetInstance()->GetNewCompiler();
@@ -1525,10 +1534,8 @@ EDataValidationResult UCustomizableObject::IsDataValid(TArray<FText>& Validation
 			break;
 		
 			// Invalid options for this context:
+			// ECustomizableObjectCompilationState::None: would mean we are working with a locked asset, with should not be possible with a sync compilation
 		case ECustomizableObjectCompilationState::None:
-			{
-				// This would mean we are working with a locked asset, with should not be possible with a sync compilation
-			}
 		case ECustomizableObjectCompilationState::InProgress:
 		default:
 			checkNoEntry();
diff --git a/Engine/Plugins/Experimental/Mutable/Source/CustomizableObjectEditor/Private/MuCOE/GenerateMutableSource/GenerateMutableSourceModifier.cpp b/Engine/Plugins/Experimental/Mutable/Source/CustomizableObjectEditor/Private/MuCOE/GenerateMutableSource/GenerateMutableSourceModifier.cpp
--- a/Engine/Plugins/Experimental/Mutable/Source/CustomizableObjectEditor/Private/MuCOE/GenerateMutableSource/GenerateMutableSourceModifier.cpp
+++ b/Engine/Plugins/Experimental/Mutable/Source/CustomizableObjectEditor/Private/MuCOE/GenerateMutableSource/GenerateMutableSourceModifier.cpp
@@ -153,7 +153,7 @@ mu::NodeModifierPtr GenerateMutableSourceModifier(const UEdGraphPin * Pin, FMuta
 
 				if (!bClosed)
 				{
-					//GenerationContext.Compiler->CompilerLog(LOCTEXT("Clipping mesh", "Clipping mesh not closed (i.e., it does not enclose a volume)."), MeshData.Node, EMessageSeverity::Warning);
+					GenerationContext.Compiler->CompilerLog(LOCTEXT("Clipping mesh", "Clipping mesh not closed (i.e., it does not enclose a volume)."), MeshData.Node, EMessageSeverity::Warning);
 				}
 			}
 
