[mutable]Added a warning message when there is a mesh in a table with less LODs than the default value.

#rnx

diff --git a/Engine/Plugins/Experimental/Mutable/Source/CustomizableObjectEditor/Private/MuCOE/GenerateMutableSource/GenerateMutableSourceTable.cpp b/Engine/Plugins/Experimental/Mutable/Source/CustomizableObjectEditor/Private/MuCOE/GenerateMutableSource/GenerateMutableSourceTable.cpp
--- a/Engine/Plugins/Experimental/Mutable/Source/CustomizableObjectEditor/Private/MuCOE/GenerateMutableSource/GenerateMutableSourceTable.cpp
+++ b/Engine/Plugins/Experimental/Mutable/Source/CustomizableObjectEditor/Private/MuCOE/GenerateMutableSource/GenerateMutableSourceTable.cpp
@@ -181,12 +181,16 @@ void FillTableColumn(const UEdGraphPin* Pin, mu::TablePtr MutableTable, FString
 
 			int NumLODs = Helper_GetLODInfoArray(SkeletalMesh).Num();
 
-			if (NumLODs > ReferenceNumLODs)
+			if (NumLODs != ReferenceNumLODs)
 			{
-				FString msg = FString::Printf(TEXT("Mesh from column [%s] row [%s] has more LODs than the reference mesh. Some will be ignored"), *ColumnName, *RowName);
+				FString Dif_1 = NumLODs > ReferenceNumLODs ? "more" : "less";
+				FString Dif_2 = NumLODs > ReferenceNumLODs ? "Some will be ignored" : "This can cause some Errors.";
+
+				FString msg = FString::Printf(TEXT("Mesh from column [%s] row [%s] has %s LODs than the reference mesh. %s"), *ColumnName, *RowName, *Dif_1, *Dif_2);
 				GenerationContext.Compiler->CompilerLog(FText::FromString(msg), CustomNodeTable);
 			}
 
+			// Parameter used for LOD differences
 			int32 LODs = NumLODs <= ReferenceNumLODs ? NumLODs : ReferenceNumLODs;
 
 			// Generating skeletal mesh columns
@@ -261,9 +265,12 @@ void FillTableColumn(const UEdGraphPin* Pin, mu::TablePtr MutableTable, FString
 			int NumLODs = StaticMesh->GetRenderData()->LODResources.Num();
 			int ReferenceNumLODs = ReferenceStaticMesh->GetRenderData()->LODResources.Num();
 
-			if (NumLODs > ReferenceNumLODs)
+			if (NumLODs != ReferenceNumLODs)
 			{
-				FString msg = FString::Printf(TEXT("Mesh from column [%s] row [%s] has more LODs than the reference mesh. Some will be ignored"), *ColumnName, *RowName);
+				FString Dif_1 = NumLODs > ReferenceNumLODs ? "more" : "less";
+				FString Dif_2 = NumLODs > ReferenceNumLODs ? "Some will be ignored" : "This can cause some Errors.";
+
+				FString msg = FString::Printf(TEXT("Mesh from column [%s] row [%s] has %s LODs than the reference mesh. %s"), *ColumnName, *RowName, *Dif_1, *Dif_2);
 				GenerationContext.Compiler->CompilerLog(FText::FromString(msg), CustomNodeTable);
 			}
 
