[mutable] Add bBuildParameterDecorations to the Instance Descriptor.

diff --git a/Engine/Plugins/Experimental/Mutable/Source/CustomizableObject/Classes/MuCO/CustomizableObjectInstance.h b/Engine/Plugins/Experimental/Mutable/Source/CustomizableObject/Classes/MuCO/CustomizableObjectInstance.h
--- a/Engine/Plugins/Experimental/Mutable/Source/CustomizableObject/Classes/MuCO/CustomizableObjectInstance.h
+++ b/Engine/Plugins/Experimental/Mutable/Source/CustomizableObject/Classes/MuCO/CustomizableObjectInstance.h
@@ -120,11 +120,6 @@ public:
 	UPROPERTY(Transient, VisibleAnywhere, Category = CustomizableSkeletalMesh)
 	TArray< TObjectPtr<USkeletalMesh> > SkeletalMeshes;
 
-	// Flag to control the build of the parameter description images required for customization UI
-	// These descriptions get generated with every instance update, so it should be disabled when not needed.
-	UPROPERTY(Transient, BlueprintReadwrite, Category=CustomizableObjectInstance, EditAnywhere)
-	bool bBuildParameterDecorations = false;
-
 	/** Broadcasts when the Customizable Object Instance is updated. */
 	UPROPERTY(Transient, BlueprintAssignable, Category = CustomizableObjectInstance)
 	FObjectInstanceUpdatedDelegate UpdatedDelegate;
@@ -174,6 +169,12 @@ public:
 	UFUNCTION(BlueprintCallable, Category = CustomizableObjectInstance)
 	UCustomizableObject* GetCustomizableObject() const;
 
+	UFUNCTION(BlueprintCallable, Category = CustomizableObjectInstance)
+	bool GetBuildParameterDecorations() const;
+
+	UFUNCTION(BlueprintCallable, Category = CustomizableObjectInstance)
+	void SetBuildParameterDecorations(bool Value);
+
 	int32 GetState() const;
 	void SetState(int32 InState);
 
@@ -653,6 +654,9 @@ private:
 
 	UPROPERTY()
 	TArray<FCustomizableObjectProjectorParameterValue> ProjectorParameters_DEPRECATED;
+	
+	UPROPERTY()
+	bool bBuildParameterDecorations_DEPRECATED;
 };
 
 #if WITH_EDITOR
diff --git a/Engine/Plugins/Experimental/Mutable/Source/CustomizableObject/Private/MuCO/CustomizableObjectInstance.cpp b/Engine/Plugins/Experimental/Mutable/Source/CustomizableObject/Private/MuCO/CustomizableObjectInstance.cpp
--- a/Engine/Plugins/Experimental/Mutable/Source/CustomizableObject/Private/MuCO/CustomizableObjectInstance.cpp
+++ b/Engine/Plugins/Experimental/Mutable/Source/CustomizableObject/Private/MuCO/CustomizableObjectInstance.cpp
@@ -470,6 +470,11 @@ void UCustomizableObjectInstance::Serialize(FArchive& Ar)
 		Descriptor.VectorParameters = VectorParameters_DEPRECATED;
 		Descriptor.ProjectorParameters = ProjectorParameters_DEPRECATED;
 	}
+	
+	if (CustomizableObjectCustomVersion < FCustomizableObjectCustomVersion::DescriptorBuildParameterDecorations)
+	{
+		Descriptor.bBuildParameterDecorations = bBuildParameterDecorations_DEPRECATED;
+	}
 }
 
 
@@ -480,10 +485,6 @@ void UCustomizableObjectInstance::PostLoad()
 
 	Super::PostLoad();
 
-	// This property is transient and should be false unless explicitely set to true. In case old
-	// data still has it, reset it here
-	bBuildParameterDecorations = false;
-
 	Descriptor.CreateParametersLookupTable();
 }
 
@@ -1146,6 +1147,18 @@ UCustomizableObject* UCustomizableObjectInstance::GetCustomizableObject() const
 }
 
 
+bool UCustomizableObjectInstance::GetBuildParameterDecorations() const
+{
+	return Descriptor.GetBuildParameterDecorations();
+}
+
+
+void UCustomizableObjectInstance::SetBuildParameterDecorations(const bool Value)
+{
+	Descriptor.SetBuildParameterDecorations(Value);
+}
+
+
 int32 UCustomizableObjectInstance::GetState() const
 {
 	return Descriptor.GetState();
diff --git a/Engine/Plugins/Experimental/Mutable/Source/CustomizableObject/Private/MuCO/CustomizableObjectInstanceDescriptor.cpp b/Engine/Plugins/Experimental/Mutable/Source/CustomizableObject/Private/MuCO/CustomizableObjectInstanceDescriptor.cpp
--- a/Engine/Plugins/Experimental/Mutable/Source/CustomizableObject/Private/MuCO/CustomizableObjectInstanceDescriptor.cpp
+++ b/Engine/Plugins/Experimental/Mutable/Source/CustomizableObject/Private/MuCO/CustomizableObjectInstanceDescriptor.cpp
@@ -49,20 +49,6 @@ FCustomizableObjectInstanceDescriptor::FCustomizableObjectInstanceDescriptor(UCu
 }
 
 
-FCustomizableObjectInstanceDescriptor::FCustomizableObjectInstanceDescriptor(const FCustomizableObjectInstanceDescriptor& Other)
-	: FCustomizableObjectInstanceDescriptor(*Other.CustomizableObject)
-{
-	BoolParameters = Other.BoolParameters;
-	IntParameters =	Other.IntParameters;
-	IntParametersLookupTable = Other.IntParametersLookupTable;
-	FloatParameters = Other.FloatParameters;
-	TextureParameters =	Other.TextureParameters;
-	VectorParameters = Other.VectorParameters;
-	ProjectorParameters = Other.ProjectorParameters;
-	State = Other.State;
-}
-
-
 void FCustomizableObjectInstanceDescriptor::SaveDescriptor(FArchive& Ar)
 {
 	check(CustomizableObject);
@@ -449,6 +435,18 @@ UCustomizableObject* FCustomizableObjectInstanceDescriptor::GetCustomizableObjec
 }
 
 
+bool FCustomizableObjectInstanceDescriptor::GetBuildParameterDecorations() const
+{
+	return bBuildParameterDecorations;
+}
+
+
+void FCustomizableObjectInstanceDescriptor::SetBuildParameterDecorations(const bool Value)
+{
+	bBuildParameterDecorations = Value;
+}
+
+
 TArray<FCustomizableObjectBoolParameterValue>& FCustomizableObjectInstanceDescriptor::GetBoolParameters()
 {
 	return BoolParameters;
@@ -1450,5 +1448,7 @@ uint32 GetTypeHash(const FCustomizableObjectInstanceDescriptor& Key)
 	
 	Hash = HashCombine(Hash, GetTypeHash(Key.State));
 
+	Hash = HashCombine(Hash, GetTypeHash(Key.bBuildParameterDecorations));
+
 	return Hash;
 }
diff --git a/Engine/Plugins/Experimental/Mutable/Source/CustomizableObject/Private/MuCO/CustomizableObjectSystem.cpp b/Engine/Plugins/Experimental/Mutable/Source/CustomizableObject/Private/MuCO/CustomizableObjectSystem.cpp
--- a/Engine/Plugins/Experimental/Mutable/Source/CustomizableObject/Private/MuCO/CustomizableObjectSystem.cpp
+++ b/Engine/Plugins/Experimental/Mutable/Source/CustomizableObject/Private/MuCO/CustomizableObjectSystem.cpp
@@ -2423,14 +2423,9 @@ FMutableOperation FMutableOperation::CreateInstanceUpdate(UCustomizableObjectIns
 	Op.CustomizableObjectInstance = InCustomizableObjectInstance;
 	Op.InstanceDescriptorHash = GetTypeHash(InCustomizableObjectInstance->GetDescriptor());
 	Op.bStarted = false;
-	Op.bBuildParameterDecorations = InCustomizableObjectInstance->bBuildParameterDecorations;
+	Op.bBuildParameterDecorations = InCustomizableObjectInstance->GetBuildParameterDecorations();
 	Op.Parameters = InCustomizableObjectInstance->GetPrivate()->ReloadParametersFromObject(InCustomizableObjectInstance, false);
 
-	// TODO: Remove once MTBL-973 is done
-#if WITH_EDITOR
-	Op.bBuildParameterDecorations = true;
-#endif
-
 	InCustomizableObjectInstance->GetCustomizableObject()->ApplyStateForcedValuesToParameters(InCustomizableObjectInstance->GetState(), Op.Parameters.get());
 
 	if (!Op.Parameters)
diff --git a/Engine/Plugins/Experimental/Mutable/Source/CustomizableObject/Private/MuCO/LogInformationUtil.cpp b/Engine/Plugins/Experimental/Mutable/Source/CustomizableObject/Private/MuCO/LogInformationUtil.cpp
--- a/Engine/Plugins/Experimental/Mutable/Source/CustomizableObject/Private/MuCO/LogInformationUtil.cpp
+++ b/Engine/Plugins/Experimental/Mutable/Source/CustomizableObject/Private/MuCO/LogInformationUtil.cpp
@@ -160,7 +160,7 @@ void LogInformationUtil::LogShowInstanceDataFull(const UCustomizableObjectInstan
 	FString MessageChunk;
 
 	MessageChunk += "\n\t";
-	MessageChunk += FString::Printf(TEXT("        bBuildParameterDecorations = %d\n"), CustomizableObjectInstance->bBuildParameterDecorations);
+	MessageChunk += FString::Printf(TEXT("        bBuildParameterDecorations = %d\n"), CustomizableObjectInstance->GetBuildParameterDecorations());
 	MessageChunk += FString::Printf(TEXT("        bShowOnlyRuntimeParameters = %d\n"), CustomizableObjectInstance->bShowOnlyRuntimeParameters);
 	MessageChunk += FString::Printf(TEXT("        bShowOnlyRelevantParameters = %d\n"), CustomizableObjectInstance->bShowOnlyRelevantParameters);
 	MessageChunk += FString::Printf(TEXT("        MinSquareDistFromComponentToPlayer = %.2f\n"), CustomizableObjectInstance->GetPrivate()->MinSquareDistFromComponentToPlayer);
diff --git a/Engine/Plugins/Experimental/Mutable/Source/CustomizableObject/Public/MuCO/CustomizableObjectCustomVersion.h b/Engine/Plugins/Experimental/Mutable/Source/CustomizableObject/Public/MuCO/CustomizableObjectCustomVersion.h
--- a/Engine/Plugins/Experimental/Mutable/Source/CustomizableObject/Public/MuCO/CustomizableObjectCustomVersion.h
+++ b/Engine/Plugins/Experimental/Mutable/Source/CustomizableObject/Public/MuCO/CustomizableObjectCustomVersion.h
@@ -56,6 +56,8 @@ struct CUSTOMIZABLEOBJECT_API FCustomizableObjectCustomVersion
 
 		CustomizableObjectInstanceDescriptor,
 		
+		DescriptorBuildParameterDecorations,
+		
 		// -----<new versions can be added above this line>-------------------------------------------------
 		VersionPlusOne,
 		LatestVersion = VersionPlusOne - 1
diff --git a/Engine/Plugins/Experimental/Mutable/Source/CustomizableObject/Public/MuCO/CustomizableObjectInstanceDescriptor.h b/Engine/Plugins/Experimental/Mutable/Source/CustomizableObject/Public/MuCO/CustomizableObjectInstanceDescriptor.h
--- a/Engine/Plugins/Experimental/Mutable/Source/CustomizableObject/Public/MuCO/CustomizableObjectInstanceDescriptor.h
+++ b/Engine/Plugins/Experimental/Mutable/Source/CustomizableObject/Public/MuCO/CustomizableObjectInstanceDescriptor.h
@@ -31,7 +31,7 @@ struct FCustomizableObjectInstanceDescriptor
 	
 	FCustomizableObjectInstanceDescriptor(UCustomizableObject& Object);
 
-	FCustomizableObjectInstanceDescriptor(const FCustomizableObjectInstanceDescriptor& Other);
+	FCustomizableObjectInstanceDescriptor(const FCustomizableObjectInstanceDescriptor& Other) = default;
 
 	/** Serialize this object. */
 	void SaveDescriptor(FArchive &Ar);
@@ -40,6 +40,10 @@ struct FCustomizableObjectInstanceDescriptor
 	void LoadDescriptor(FArchive &Ar);
 
 	UCustomizableObject* GetCustomizableObject() const;
+
+	bool GetBuildParameterDecorations() const;
+	
+	void SetBuildParameterDecorations(bool Value);
 	
 	// ------------------------------------------------------------
 	// Parameters
@@ -246,6 +250,10 @@ private:
 	TArray<FCustomizableObjectProjectorParameterValue> ProjectorParameters;
 
 	int32 State = 0;
+	
+	/** Flag to control the build of the parameter description images required for customization UI.
+     * These descriptions get generated with every instance update, so it should be disabled when not needed. */
+	bool bBuildParameterDecorations = false;
 
 	TMap<FString, int32> IntParametersLookupTable;
 	
diff --git a/Engine/Plugins/Experimental/Mutable/Source/CustomizableObjectEditor/Private/MuCOE/CustomizableObjectEditor.cpp b/Engine/Plugins/Experimental/Mutable/Source/CustomizableObjectEditor/Private/MuCOE/CustomizableObjectEditor.cpp
--- a/Engine/Plugins/Experimental/Mutable/Source/CustomizableObjectEditor/Private/MuCOE/CustomizableObjectEditor.cpp
+++ b/Engine/Plugins/Experimental/Mutable/Source/CustomizableObjectEditor/Private/MuCOE/CustomizableObjectEditor.cpp
@@ -539,7 +539,7 @@ void FCustomizableObjectEditor::CreatePreviewInstance()
 				// Asset loading works in the Main Thread, there's a risk the sync loading could put to sleep the thread while
 				// waiting for the asset registry to load the content (that never gets executed)
 				Viewport->SetAssetRegistryLoaded(true);
-				PreviewInstance->bBuildParameterDecorations = true;
+				PreviewInstance->SetBuildParameterDecorations(true);
 				PreviewInstance->UpdateSkeletalMeshAsync(true, true);
 			}
 			else
diff --git a/Engine/Plugins/Experimental/Mutable/Source/CustomizableObjectEditor/Private/MuCOE/SCustomizableInstanceProperties.cpp b/Engine/Plugins/Experimental/Mutable/Source/CustomizableObjectEditor/Private/MuCOE/SCustomizableInstanceProperties.cpp
--- a/Engine/Plugins/Experimental/Mutable/Source/CustomizableObjectEditor/Private/MuCOE/SCustomizableInstanceProperties.cpp
+++ b/Engine/Plugins/Experimental/Mutable/Source/CustomizableObjectEditor/Private/MuCOE/SCustomizableInstanceProperties.cpp
@@ -1314,9 +1314,9 @@ void SCustomizableInstanceProperties::Tick(const FGeometry& AllottedGeometry, co
 	SCompoundWidget::Tick(AllottedGeometry, InCurrentTime, InDeltaTime);
 
 	// We set this every frame, in case other widgets are also modifying this flag
-	if (!CustomInstance->bBuildParameterDecorations)
+	if (!CustomInstance->GetBuildParameterDecorations())
 	{
-		CustomInstance->bBuildParameterDecorations = true;
+		CustomInstance->SetBuildParameterDecorations(true);
 		CustomInstance->UpdateSkeletalMeshAsync(true);
 	}
 
diff --git a/Engine/Plugins/Experimental/Mutable/Source/CustomizableObjectPopulationEditor/Private/MuCOPE/CustomizableObjectPopulationClassEditor.cpp b/Engine/Plugins/Experimental/Mutable/Source/CustomizableObjectPopulationEditor/Private/MuCOPE/CustomizableObjectPopulationClassEditor.cpp
--- a/Engine/Plugins/Experimental/Mutable/Source/CustomizableObjectPopulationEditor/Private/MuCOPE/CustomizableObjectPopulationClassEditor.cpp
+++ b/Engine/Plugins/Experimental/Mutable/Source/CustomizableObjectPopulationEditor/Private/MuCOPE/CustomizableObjectPopulationClassEditor.cpp
@@ -173,7 +173,7 @@ void FCustomizableObjectPopulationClassEditor::InitCustomizableObjectPopulationC
 		if (PopulationClass->CustomizableObject->IsCompiled() && !PopulationClass->CustomizableObjectInstance)
 		{
 			PopulationClass->CustomizableObjectInstance = PopulationClass->CustomizableObject->CreateInstance();
-			PopulationClass->CustomizableObjectInstance->bBuildParameterDecorations = true;
+			PopulationClass->CustomizableObjectInstance->SetBuildParameterDecorations(true);
 			PopulationClass->CustomizableObjectInstance->UpdateSkeletalMeshAsync(true, true);
 
 			bRefreshDetailsView = true;
@@ -310,7 +310,7 @@ void FCustomizableObjectPopulationClassEditor::Tick(float InDeltaTime)
 		if (PopulationClass->CustomizableObject->IsCompiled() && !PopulationClass->CustomizableObjectInstance)
 		{
 			PopulationClass->CustomizableObjectInstance = PopulationClass->CustomizableObject->CreateInstance();
-			PopulationClass->CustomizableObjectInstance->bBuildParameterDecorations = true;
+			PopulationClass->CustomizableObjectInstance->SetBuildParameterDecorations(true);
 			PopulationClass->CustomizableObjectInstance->UpdateSkeletalMeshAsync(true, true);
 
 			bRefreshDetailsView = true;
@@ -834,7 +834,7 @@ void FCustomizableObjectPopulationClassEditor::TestPopulationClass()
 						// Attaching the Customizable Skeletal Component to the Skeletal Mesh Component
 						if (PreviewCustomizableSkeletalComponent && SkeletalMeshComponent && ColliderComponent)
 						{
-							ViewportInstances[i]->bBuildParameterDecorations = true;
+							ViewportInstances[i]->SetBuildParameterDecorations(true);
 							ViewportInstances[i]->UpdateSkeletalMeshAsync(true, true);
 
 							PreviewCustomizableSkeletalComponent->CustomizableObjectInstance = ViewportInstances[i];
diff --git a/Engine/Plugins/Experimental/Mutable/Source/CustomizableObjectPopulationEditor/Private/MuCOPE/CustomizableObjectPopulationEditor.cpp b/Engine/Plugins/Experimental/Mutable/Source/CustomizableObjectPopulationEditor/Private/MuCOPE/CustomizableObjectPopulationEditor.cpp
--- a/Engine/Plugins/Experimental/Mutable/Source/CustomizableObjectPopulationEditor/Private/MuCOPE/CustomizableObjectPopulationEditor.cpp
+++ b/Engine/Plugins/Experimental/Mutable/Source/CustomizableObjectPopulationEditor/Private/MuCOPE/CustomizableObjectPopulationEditor.cpp
@@ -512,7 +512,7 @@ void FCustomizableObjectPopulationEditor::TestPopulation()
 				// Attaching the Customizable Skeletal Component to the Skeletal Mesh Component
 				if (PreviewCustomizableSkeletalComponent && PreviewSkeletalMeshComponent && ColliderComponent)
 				{
-					ViewportInstances[i]->bBuildParameterDecorations = true;
+					ViewportInstances[i]->SetBuildParameterDecorations(true);
 					ViewportInstances[i]->UpdateSkeletalMeshAsync(true, true);
 
 					PreviewCustomizableSkeletalComponent->CustomizableObjectInstance = ViewportInstances[i];
