#rnx

diff --git a/Engine/Plugins/Experimental/Mutable/Source/CustomizableObjectEditor/Private/MuCOE/Nodes/CustomizableObjectNodeMaterial.cpp b/Engine/Plugins/Experimental/Mutable/Source/CustomizableObjectEditor/Private/MuCOE/Nodes/CustomizableObjectNodeMaterial.cpp
--- a/Engine/Plugins/Experimental/Mutable/Source/CustomizableObjectEditor/Private/MuCOE/Nodes/CustomizableObjectNodeMaterial.cpp
+++ b/Engine/Plugins/Experimental/Mutable/Source/CustomizableObjectEditor/Private/MuCOE/Nodes/CustomizableObjectNodeMaterial.cpp
@@ -869,24 +869,27 @@ FText UCustomizableObjectNodeMaterial::GetParameterLayerName(const EMaterialPara
 
 bool UCustomizableObjectNodeMaterial::HasParameter(const FGuid& ParameterId) const
 {
-	for (const EMaterialParameterType Type : ParameterTypes)
+	if (Material)
 	{
-		const FMaterialCachedExpressionData& Data = Material->GetCachedExpressionData();
-		const FMaterialCachedParameterEntry& Entry = Data.GetParameterTypeEntry(Type);
-
-		if (!Data.EditorOnlyData || Data.EditorOnlyData->EditorEntries[(int32)Type].EditorInfo.IsEmpty())
-		{
-			continue;
-		}
-			
-		for (TSet<FMaterialParameterInfo>::TConstIterator It(Entry.ParameterInfoSet); It; ++It)
+		for (const EMaterialParameterType Type : ParameterTypes)
 		{
-			const int32 IteratorIndex = It.GetId().AsInteger();
-			const FGuid& ParamGUid = Data.EditorOnlyData->EditorEntries[(int32)Type].EditorInfo[IteratorIndex].ExpressionGuid;
-			
-			if (ParamGUid == ParameterId)
+			const FMaterialCachedExpressionData& Data = Material->GetCachedExpressionData();
+			const FMaterialCachedParameterEntry& Entry = Data.GetParameterTypeEntry(Type);
+
+			if (!Data.EditorOnlyData || Data.EditorOnlyData->EditorEntries[(int32)Type].EditorInfo.IsEmpty())
 			{
-				return true;
+				continue;
+			}
+
+			for (TSet<FMaterialParameterInfo>::TConstIterator It(Entry.ParameterInfoSet); It; ++It)
+			{
+				const int32 IteratorIndex = It.GetId().AsInteger();
+				const FGuid& ParamGUid = Data.EditorOnlyData->EditorEntries[(int32)Type].EditorInfo[IteratorIndex].ExpressionGuid;
+
+				if (ParamGUid == ParameterId)
+				{
+					return true;
+				}
 			}
 		}
 	}
