[mutable] Fix bug with "reuse mesh" optimization vs "override layouts".
Additional cleanup towards UE coding standards.
#rnx

diff --git a/Engine/Plugins/Experimental/Mutable/Source/MutableRuntime/Private/MuR/CodeVisitor.cpp b/Engine/Plugins/Experimental/Mutable/Source/MutableRuntime/Private/MuR/CodeVisitor.cpp
--- a/Engine/Plugins/Experimental/Mutable/Source/MutableRuntime/Private/MuR/CodeVisitor.cpp
+++ b/Engine/Plugins/Experimental/Mutable/Source/MutableRuntime/Private/MuR/CodeVisitor.cpp
@@ -5,7 +5,6 @@
 
 #include "HAL/PlatformCrt.h"
 #include "MuR/MutableTrace.h"
-#include "map"
 
 #include <functional>
 #include <memory>
diff --git a/Engine/Plugins/Experimental/Mutable/Source/MutableTools/Private/MuT/CodeGenerator.cpp b/Engine/Plugins/Experimental/Mutable/Source/MutableTools/Private/MuT/CodeGenerator.cpp
--- a/Engine/Plugins/Experimental/Mutable/Source/MutableTools/Private/MuT/CodeGenerator.cpp
+++ b/Engine/Plugins/Experimental/Mutable/Source/MutableTools/Private/MuT/CodeGenerator.cpp
@@ -267,7 +267,7 @@ namespace mu
         // Free caches
         m_compiled.Reset();
         m_constantMeshes.clear();
-        m_addedLayouts.clear();
+        m_addedLayouts.Empty();
         m_nodeVariables.clear();
         m_generatedMeshes.Reset();
         m_generatedProjectors.Reset();
@@ -512,24 +512,31 @@ namespace mu
 
 
     //---------------------------------------------------------------------------------------------
-	Ptr<const Layout> CodeGenerator::AddLayout(Ptr<const Layout> pLayout )
+	Ptr<const Layout> CodeGenerator::AddLayout(Ptr<const Layout> SourceLayout )
     {
-        auto it = m_addedLayouts.find( pLayout.get() );
+		// The layout we are adding must be a source layout, without block ids yet.
+		check(SourceLayout->m_blocks.IsEmpty() || SourceLayout->m_blocks[0].m_id == -1);
+		
+		Ptr<const Layout>* it = m_addedLayouts.Find(SourceLayout.get() );
 
 		Ptr<const Layout> pResult;
-        if ( it != m_addedLayouts.end() )
+        if ( it )
         {
-            pResult = it->second;
+            pResult = *it;
         }
         else
         {
-			Ptr<Layout> pCloned = pLayout->Clone();
-            for (int32 b=0;b< pCloned->m_blocks.Num();++b)
+			Ptr<Layout> ClonedLayout = SourceLayout->Clone();
+            for (int32 b=0;b< ClonedLayout->m_blocks.Num();++b)
             {
-				pCloned->m_blocks[b].m_id = m_absoluteLayoutIndex++;
+				// This is a hard limit due to layout block index data being stored in 16 bit.
+				check(m_absoluteLayoutIndex<65536);
+				ClonedLayout->m_blocks[b].m_id = m_absoluteLayoutIndex++;
             }
-            m_addedLayouts[ pLayout.get() ] = pCloned;
-			pResult = pCloned;
+			check(SourceLayout->m_blocks.Num() == ClonedLayout->m_blocks.Num());
+			check(ClonedLayout->m_blocks.IsEmpty() || ClonedLayout->m_blocks[0].m_id != -1);
+			m_addedLayouts.Add(SourceLayout.get(), ClonedLayout);
+			pResult = ClonedLayout;
         }
 
         return pResult;
@@ -834,7 +841,7 @@ namespace mu
         const auto& node = *surfaceNode->GetPrivate();
 
         // Clear the surface generation state
-//        m_addedLayouts.clear();
+//        m_addedLayouts.Empty();
 //        m_absoluteLayoutIndex = 0;
 //        m_currentLayoutMesh = nullptr;
 //        m_currentLayoutChannel = 0;
diff --git a/Engine/Plugins/Experimental/Mutable/Source/MutableTools/Private/MuT/CodeGenerator.h b/Engine/Plugins/Experimental/Mutable/Source/MutableTools/Private/MuT/CodeGenerator.h
--- a/Engine/Plugins/Experimental/Mutable/Source/MutableTools/Private/MuT/CodeGenerator.h
+++ b/Engine/Plugins/Experimental/Mutable/Source/MutableTools/Private/MuT/CodeGenerator.h
@@ -50,70 +50,70 @@
 
 namespace mu
 {
-class ASTOpParameter;
-class Layout;
-class NodeColourArithmeticOperation;
-class NodeColourConstant;
-class NodeColourFromScalars;
-class NodeColourParameter;
-class NodeColourSwitch;
-class NodeColourTable;
-class NodeColourVariation;
-class NodeImageBinarise;
-class NodeImageColourMap;
-class NodeImageConditional;
-class NodeImageConstant;
-class NodeImageDifference;
-class NodeImageFormat;
-class NodeImageGradient;
-class NodeImageInterpolate3;
-class NodeImageInterpolate;
-class NodeImageInvert;
-class NodeImageLayer;
-class NodeImageLayerColour;
-class NodeImageLuminance;
-class NodeImageMipmap;
-class NodeImageMultiLayer;
-class NodeImageNormalComposite;
-class NodeImageParameter;
-class NodeImagePlainColour;
-class NodeImageResize;
-class NodeImageSaturate;
-class NodeImageSelectColour;
-class NodeImageSwitch;
-class NodeImageSwizzle;
-class NodeImageTable;
-class NodeImageTransform;
-class NodeImageVariation;
-class NodeMeshApplyPose;
-class NodeMeshClipDeform;
-class NodeMeshClipMorphPlane;
-class NodeMeshClipWithMesh;
-class NodeMeshConstant;
-class NodeMeshFormat;
-class NodeMeshFragment;
-class NodeMeshInterpolate;
-class NodeMeshMakeMorph;
-class NodeMeshMorph;
-class NodeMeshReshape;
-class NodeMeshSubtract;
-class NodeMeshSwitch;
-class NodeMeshTransform;
-class NodeMeshVariation;
-class NodeRange;
-class NodeScalarArithmeticOperation;
-class NodeScalarConstant;
-class NodeScalarCurve;
-class NodeScalarEnumParameter;
-class NodeScalarParameter;
-class NodeScalarSwitch;
-class NodeScalarTable;
-class NodeScalarVariation;
-class NodeStringConstant;
-class NodeStringParameter;
-class TaskManager;
-struct OBJECT_STATE;
-struct PROGRAM;
+	class ASTOpParameter;
+	class Layout;
+	class NodeColourArithmeticOperation;
+	class NodeColourConstant;
+	class NodeColourFromScalars;
+	class NodeColourParameter;
+	class NodeColourSwitch;
+	class NodeColourTable;
+	class NodeColourVariation;
+	class NodeImageBinarise;
+	class NodeImageColourMap;
+	class NodeImageConditional;
+	class NodeImageConstant;
+	class NodeImageDifference;
+	class NodeImageFormat;
+	class NodeImageGradient;
+	class NodeImageInterpolate3;
+	class NodeImageInterpolate;
+	class NodeImageInvert;
+	class NodeImageLayer;
+	class NodeImageLayerColour;
+	class NodeImageLuminance;
+	class NodeImageMipmap;
+	class NodeImageMultiLayer;
+	class NodeImageNormalComposite;
+	class NodeImageParameter;
+	class NodeImagePlainColour;
+	class NodeImageResize;
+	class NodeImageSaturate;
+	class NodeImageSelectColour;
+	class NodeImageSwitch;
+	class NodeImageSwizzle;
+	class NodeImageTable;
+	class NodeImageTransform;
+	class NodeImageVariation;
+	class NodeMeshApplyPose;
+	class NodeMeshClipDeform;
+	class NodeMeshClipMorphPlane;
+	class NodeMeshClipWithMesh;
+	class NodeMeshConstant;
+	class NodeMeshFormat;
+	class NodeMeshFragment;
+	class NodeMeshInterpolate;
+	class NodeMeshMakeMorph;
+	class NodeMeshMorph;
+	class NodeMeshReshape;
+	class NodeMeshSubtract;
+	class NodeMeshSwitch;
+	class NodeMeshTransform;
+	class NodeMeshVariation;
+	class NodeRange;
+	class NodeScalarArithmeticOperation;
+	class NodeScalarConstant;
+	class NodeScalarCurve;
+	class NodeScalarEnumParameter;
+	class NodeScalarParameter;
+	class NodeScalarSwitch;
+	class NodeScalarTable;
+	class NodeScalarVariation;
+	class NodeStringConstant;
+	class NodeStringParameter;
+	class TaskManager;
+	struct OBJECT_STATE;
+	struct PROGRAM;
 
     //---------------------------------------------------------------------------------------------
     //! Code generator
@@ -241,7 +241,7 @@ struct PROGRAM;
 
         //! Map of layouts found in the code already generated. The map is from the source layout
         //! pointer to the cloned layout. The cloned layout will have absolute block ids assigned.
-        map<Ptr<const Layout>,Ptr<const Layout>> m_addedLayouts;
+        TMap<Ptr<const Layout>,Ptr<const Layout>> m_addedLayouts;
 
         //! First free index for a layout block
         int32_t m_absoluteLayoutIndex = 0;
@@ -249,9 +249,6 @@ struct PROGRAM;
         //! First free index to be used to identify mesh vertices.
         uint32_t m_freeVertexIndex = 0;
 
-        //! First free index to be used to identify mesh masks.
-        uint32_t m_freeMaskIndex = 0;
-
         //! When generating images, here we have the entire source image size and the rect of the
         //! image that we are generating.
         struct IMAGE_STATE
@@ -325,11 +322,6 @@ struct PROGRAM;
         //! Variables added for every node
         map< Ptr<const Node>, Ptr<ASTOpParameter> > m_nodeVariables;
 
-        //! This is a list of nodes that are set to replace the identity nodes found while
-        //! generating code. This is useful when applying object transforms.
-        typedef map< NodePtr, Ptr<ASTOp> > IdentityMap;
-        vector<IdentityMap> m_identityReplacement;
-
 		//-----------------------------------------------------------------------------------------
 
 		// Get the modifiers that have to be applied to elements with a specific tag.
@@ -479,10 +471,8 @@ struct PROGRAM;
 		void GenerateMesh_Reshape(MESH_GENERATION_RESULT&, const NodeMeshReshape*);
 		void GenerateMesh_ClipDeform(MESH_GENERATION_RESULT&, const NodeMeshClipDeform*);
 
-		void GenerateLayout(MESH_GENERATION_RESULT& result, const NodeLayoutBlocksPtrConst& node, uint32 currentLayoutChannel, const MeshPtr currentLayoutMesh);
-
 		//!
-		Ptr<const Layout> AddLayout(Ptr<const Layout> pLayout);
+		Ptr<const Layout> AddLayout(Ptr<const Layout> SourceLayout);
 
         //-----------------------------------------------------------------------------------------
         // Projectors
diff --git a/Engine/Plugins/Experimental/Mutable/Source/MutableTools/Private/MuT/CodeGenerator_FirstPass.h b/Engine/Plugins/Experimental/Mutable/Source/MutableTools/Private/MuT/CodeGenerator_FirstPass.h
--- a/Engine/Plugins/Experimental/Mutable/Source/MutableTools/Private/MuT/CodeGenerator_FirstPass.h
+++ b/Engine/Plugins/Experimental/Mutable/Source/MutableTools/Private/MuT/CodeGenerator_FirstPass.h
@@ -44,13 +44,15 @@ struct OBJECT_STATE;
         //! Original base mesh before removes, morphs, etc.
         Ptr<ASTOp> baseMeshOp;
 
-        vector<Ptr<const Layout>> layouts;
+		/** Source node layouts to use with these extra mesh. They don't have block ids. */
+		vector<Ptr<const Layout>> layouts;
 
         vector<Ptr<ASTOp>> layoutOps;
 
 
         struct EXTRA_LAYOUTS
         {
+			/** Source node layouts to use with these extra mesh. They don't have block ids. */
             vector<Ptr<const Layout>> layouts;
             Ptr<ASTOp> condition;
             Ptr<ASTOp> meshFragment;
diff --git a/Engine/Plugins/Experimental/Mutable/Source/MutableTools/Private/MuT/CodeGenerator_Mesh.cpp b/Engine/Plugins/Experimental/Mutable/Source/MutableTools/Private/MuT/CodeGenerator_Mesh.cpp
--- a/Engine/Plugins/Experimental/Mutable/Source/MutableTools/Private/MuT/CodeGenerator_Mesh.cpp
+++ b/Engine/Plugins/Experimental/Mutable/Source/MutableTools/Private/MuT/CodeGenerator_Mesh.cpp
@@ -95,7 +95,7 @@ class Node;
             return;
         }
 
-        // TODO: Done twice? Check new component code.
+        // 
 		Ptr<const Layout> pLayout = AddLayout( pSourceLayout );
         currentLayoutMesh->AddLayout( pLayout );
 
@@ -281,28 +281,6 @@ class Node;
     }
 
 
-    //---------------------------------------------------------------------------------------------
-    void CodeGenerator::GenerateLayout( 
-		MESH_GENERATION_RESULT& result, 
-		const NodeLayoutBlocksPtrConst& node,
-		uint32 currentLayoutChannel, const MeshPtr currentLayoutMesh )
-    {
-        if (!currentLayoutMesh)
-        {
-            m_pErrorLog->GetPrivate()->Add( "Generating a layout node without a parent mesh.",
-                                            ELMT_ERROR, node->GetPrivate()->m_errorContext );
-            return;
-        }
-
-        LayoutPtr pSourceLayout = node->GetPrivate()->m_pLayout;
-        result.layouts.push_back( pSourceLayout );
-
-        PrepareForLayout( pSourceLayout,
-                          currentLayoutMesh, currentLayoutChannel,
-                          node->GetPrivate()->m_errorContext  );
-    }
-
-
     //---------------------------------------------------------------------------------------------
     void CodeGenerator::GenerateMesh( MESH_GENERATION_RESULT& result, const NodeMeshPtrConst& Untyped )
     {
@@ -918,13 +896,15 @@ class Node;
 		result.baseMeshOp = op;
 		result.meshOp = op;
 
+		bool bIsOverridingLayouts = !m_overrideLayoutsStack.empty();
+
         MeshPtr pMesh = node.m_pValue.get();
         if (pMesh)
         {
             // Clone the mesh
             MeshPtr pCloned = pMesh->Clone();
             pCloned->EnsureSurfaceData();
-            if ( m_overrideLayoutsStack.empty() )
+            if ( !bIsOverridingLayouts )
             {
                 // This means that we are processing a base mesh
 
@@ -932,24 +912,37 @@ class Node;
 				result.layouts.clear();
 
                 // Apply whatever transform is necessary for every layout
-                for ( std::size_t l=0; l<node.m_layouts.size();++l )
+                for ( std::size_t LayoutIndex=0; LayoutIndex<node.m_layouts.size(); ++LayoutIndex)
                 {
-					NodeLayoutPtr pLayoutNode = node.m_layouts[l];
+					NodeLayoutPtr pLayoutNode = node.m_layouts[LayoutIndex];
 					// TODO: In a cleanup of the design of the layouts, we should remove this cast.
 					const NodeLayoutBlocks* TypedNode = dynamic_cast<NodeLayoutBlocks*>(pLayoutNode.get());
 					if (TypedNode)
 					{
-						GenerateLayout(result, TypedNode, l, pCloned);
+						if (!pCloned)
+						{
+							m_pErrorLog->GetPrivate()->Add("Generating a layout node without a parent mesh.",
+								ELMT_ERROR, TypedNode->GetPrivate()->m_errorContext);
+							return;
+						}
+
+						LayoutPtr pSourceLayout = TypedNode->GetPrivate()->m_pLayout;
+						result.layouts.push_back(pSourceLayout);
+
+						PrepareForLayout(pSourceLayout,
+							pCloned, LayoutIndex,
+							TypedNode->GetPrivate()->m_errorContext);
+
 					}
                 }
             }
             else
             {
                 // We need to apply the transform of the source layouts
-                for ( std::size_t l=0; l<m_overrideLayoutsStack.back().size(); ++l )
+                for ( std::size_t LayoutIndex=0; LayoutIndex<m_overrideLayoutsStack.back().size(); ++LayoutIndex)
                 {
-                    PrepareForLayout( m_overrideLayoutsStack.back()[l],
-                                      pCloned, l,
+                    PrepareForLayout( m_overrideLayoutsStack.back()[LayoutIndex],
+                                      pCloned, LayoutIndex,
                                       node.m_errorContext  );
 
                 }
@@ -967,41 +960,53 @@ class Node;
 
             // See if we already have a mesh identical to this one, except for the internal data
             // like vertex indices.
-            bool isDuplicated = false;
-            MeshPtrConst pCandidate;
-            for ( size_t i=0; !isDuplicated && i<m_constantMeshes.size(); ++i)
+			// For now we cannot handle the case were we are overriding layouts and reusing meshes at the same time.
+            bool bIsDuplicated = false;
+            MeshPtrConst Candidate;
+            for ( size_t i=0; !bIsOverridingLayouts && !bIsDuplicated && i<m_constantMeshes.size(); ++i)
             {
-                pCandidate = m_constantMeshes[i];
+				Candidate = m_constantMeshes[i];
 
-                if ( pCandidate->IsSimilar(*pCloned) )
+                if (Candidate->IsSimilar(*pCloned) )
                 {
-                    isDuplicated = true;
-
-                    // Remap layouts from the source mesh to the the ones created for the mesh we will use instead.
-                    for (int l=0; l<pCandidate->GetLayoutCount(); ++l)
-                    {
-                        check(pCandidate->GetLayoutCount()==pCloned->GetLayoutCount());
-
-                        const Layout* pSourceLayoutValue = pCandidate->GetLayout(l);
-                        const Layout* pDestLayoutValue = pCloned->GetLayout(l);
+					bIsDuplicated = true;
 
-                        Ptr<const Layout> pDestLayoutKey = nullptr;
-                        for( const auto& e : m_addedLayouts )
-                        {
-                            if (e.second==pDestLayoutValue)
-                            {
-                                pDestLayoutKey = e.first;
-                            }
-                        }
-                        check(pDestLayoutKey);
-                        m_addedLayouts[pDestLayoutKey] = pSourceLayoutValue;
-                    }
+					// We want to reuse the layouts from the similar mesh.					
+					// Remap layouts that we have created in the previosu step from the source mesh to the ones created for the mesh we will reuse instead.
+					for (int l = 0; l < Candidate->GetLayoutCount(); ++l)
+					{
+						check(Candidate->GetLayoutCount() == pCloned->GetLayoutCount());
+
+						const Layout* ClonedMeshLayout = pCloned->GetLayout(l);
+
+						Ptr<const Layout> DestLayoutKey = nullptr;
+						for (const auto& e : m_addedLayouts)
+						{
+							if (e.Value == ClonedMeshLayout)
+							{
+								DestLayoutKey = e.Key;
+							}
+						}
+						check(DestLayoutKey);
+
+						const Layout* GeneratedMeshLayout = Candidate->GetLayout(l);
+						if (bIsOverridingLayouts)
+						{
+							const Layout* OverridingLayout = m_overrideLayoutsStack.back()[l].get();
+							check(OverridingLayout->m_blocks.IsEmpty() || OverridingLayout->m_blocks[0].m_id == -1);
+							GeneratedMeshLayout = m_addedLayouts[OverridingLayout].get();
+						}
+
+						check(DestLayoutKey->m_blocks.IsEmpty() || DestLayoutKey->m_blocks[0].m_id == -1);
+						check(GeneratedMeshLayout->m_blocks.IsEmpty() || GeneratedMeshLayout->m_blocks[0].m_id != -1);
+						m_addedLayouts[DestLayoutKey] = GeneratedMeshLayout;
+					}
                 }
             }
 
-            if (isDuplicated)
+            if (bIsDuplicated)
             {
-                op->SetValue( pCandidate, m_compilerOptions->m_optimisationOptions.m_useDiskCache );
+                op->SetValue(Candidate, m_compilerOptions->m_optimisationOptions.m_useDiskCache );
             }
             else
             {
