#rnx

diff --git a/Engine/Plugins/Experimental/Mutable/Source/CustomizableObjectEditor/Private/MuCOE/CustomizableObjectEditorViewportClient.cpp b/Engine/Plugins/Experimental/Mutable/Source/CustomizableObjectEditor/Private/MuCOE/CustomizableObjectEditorViewportClient.cpp
--- a/Engine/Plugins/Experimental/Mutable/Source/CustomizableObjectEditor/Private/MuCOE/CustomizableObjectEditorViewportClient.cpp
+++ b/Engine/Plugins/Experimental/Mutable/Source/CustomizableObjectEditor/Private/MuCOE/CustomizableObjectEditorViewportClient.cpp
@@ -2396,17 +2396,19 @@ void FCustomizableObjectEditorViewportClient::BakeInstance()
 					TArray<TMap<int, UTexture*>> TextureReplacementMaps;
 
 					// Duplicate Mutable generated textures
-					for (int m = 0; m < Mesh->GetMaterials().Num(); ++m)
+					for (int32 m = 0; m < Mesh->GetMaterials().Num(); ++m)
 					{
+						UMaterialInterface* Interface = Mesh->GetMaterials()[m].MaterialInterface;
+						Material = Interface->GetMaterial();
+						MaterialName = Material ? Material->GetName() : "Material";
 						Inst = Cast<UMaterialInstance>(Mesh->GetMaterials()[m].MaterialInterface);
-						MaterialName = (Inst->Parent != nullptr) ? Inst->Parent->GetName() : "Material";
-						Material = Inst->Parent->GetMaterial();
 
-						if (Material != nullptr)
-						{
-							TMap<int, UTexture*> ReplacementTextures;
-							TextureReplacementMaps.Add(ReplacementTextures);
+						TMap<int, UTexture*> ReplacementTextures;
+						TextureReplacementMaps.Add(ReplacementTextures);
 
+						// The material will only have Mutable generated textures if it's actually a UMaterialInstance
+						if (Material != nullptr && Inst != nullptr)
+						{
 							TArray<FName> ParameterNames = GetTextureParameterNames(Material);
 
 							for (int32 i = 0; i < ParameterNames.Num(); i++)
@@ -2471,11 +2473,11 @@ void FCustomizableObjectEditorViewportClient::BakeInstance()
 					}
 
 					// Duplicate non-Mutable material textures
-					for (int m = 0; m < Mesh->GetMaterials().Num(); ++m)
+					for (int32 m = 0; m < Mesh->GetMaterials().Num(); ++m)
 					{
-						Inst = Cast<UMaterialInstance>(Mesh->GetMaterials()[m].MaterialInterface);
-						MaterialName = (Inst->Parent != nullptr) ? Inst->Parent->GetName() : "Material";
-						Material = Inst->Parent->GetMaterial();
+						UMaterialInterface* Interface = Mesh->GetMaterials()[m].MaterialInterface;
+						Material = Interface->GetMaterial();
+						MaterialName = Material ? Material->GetName() : "Material";
 
 						if (Material != nullptr)
 						{
@@ -2483,7 +2485,11 @@ void FCustomizableObjectEditorViewportClient::BakeInstance()
 
 							for (int32 i = 0; i < ParameterNames.Num(); i++)
 							{
-								if (Inst->GetTextureParameterValue(ParameterNames[i], Texture))
+								TArray<FMaterialParameterInfo> InfoArray;
+								TArray<FGuid> GuidArray;
+								Material->GetAllTextureParameterInfo(InfoArray, GuidArray);
+								
+								if (Material->GetTextureParameterValue(InfoArray[i], Texture))
 								{
 									FString ParameterSanitized = ParameterNames[i].GetPlainNameString();
 									RemoveRestrictedChars(ParameterSanitized);
@@ -2512,97 +2518,98 @@ void FCustomizableObjectEditorViewportClient::BakeInstance()
 					}
 
 
-					// Duplicate the materials used by each material instance so replacement map has proper information when duplicating the material instances
-					for (int m = 0; m < Mesh->GetMaterials().Num(); ++m)
+					// Duplicate the materials used by each material instance so that the replacement map has proper information 
+					// when duplicating the material instances
+					for (int32 m = 0; m < Mesh->GetMaterials().Num(); ++m)
 					{
-						Inst = Cast<UMaterialInstance>(Mesh->GetMaterials()[m].MaterialInterface);
+						UMaterialInterface* Interface = Mesh->GetMaterials()[m].MaterialInterface;
+						Material = Interface->GetMaterial();
 
-						if ((Inst != nullptr) && (Inst->GetMaterial() != nullptr))
+						if (Interface && Material)
 						{
-							Material = Inst->GetMaterial();
+							ResourceName = ObjectName + "_Material_" + Material->GetName();
 
-							if (Material != nullptr)
+							if (!GetUniqueResourceName(Material, ResourceName, ArrayCachedObject, ArrayCachedElement))
 							{
-								ResourceName = ObjectName + "_Material_" + Material->GetName();
-
-								if (!GetUniqueResourceName(Material, ResourceName, ArrayCachedObject, ArrayCachedElement))
-								{
-									continue;
-								}
+								continue;
+							}
 
-								if (!ManageBakingAction(AssetPath, ResourceName))
-								{
-									return;
-								}
+							if (!ManageBakingAction(AssetPath, ResourceName))
+							{
+								return;
+							}
 
-								PackageName = FolderDlg->GetAssetPath() + FString("/") + ResourceName;
-								TMap<UObject*, UObject*> FakeReplacementMap;
-								DuplicatedObject = BakeHelper_DuplicateAsset(Material, ResourceName, PackageName, false, FakeReplacementMap, BakingOverwritePermission);
-								ArrayCachedElement.Add(ResourceName);
-								ArrayCachedObject.Add(DuplicatedObject);
-								ReplacementMap.Add(Inst, DuplicatedObject);
-								PackagesToSave.Add(DuplicatedObject->GetPackage());
+							PackageName = FolderDlg->GetAssetPath() + FString("/") + ResourceName;
+							TMap<UObject*, UObject*> FakeReplacementMap;
+							DuplicatedObject = BakeHelper_DuplicateAsset(Material, ResourceName, PackageName, false, FakeReplacementMap, BakingOverwritePermission);
+							ArrayCachedElement.Add(ResourceName);
+							ArrayCachedObject.Add(DuplicatedObject);
+							ReplacementMap.Add(Interface, DuplicatedObject);
+							PackagesToSave.Add(DuplicatedObject->GetPackage());
 
-								if (UMaterial* DupMaterial = Cast<UMaterial>(DuplicatedObject))
+							if (UMaterial* DupMaterial = Cast<UMaterial>(DuplicatedObject))
+							{
+								TArray<FMaterialParameterInfo> parametersInfo;
+								TArray<FGuid> parametersGuids;
+
+								// copy scalar parameters
+								TArray<FMaterialParameterInfo> ScalarParameterInfoArray;
+								TArray<FGuid> GuidArray;
+								Interface->GetAllScalarParameterInfo(ScalarParameterInfoArray, GuidArray);
+								for (const FMaterialParameterInfo& Param : ScalarParameterInfoArray)
 								{
-									TArray<FMaterialParameterInfo> parametersInfo;
-									TArray<FGuid> parametersGuids;
-
-									// copy scalar parameters
-									Inst->GetAllScalarParameterInfo(parametersInfo, parametersGuids);
-									for (const FMaterialParameterInfo& param : parametersInfo)
+									float Value = 0.f;
+									if (Interface->GetScalarParameterValue(Param, Value))
 									{
-										float value = 0.f;
-										if (Inst->GetScalarParameterValue(param, value))
-										{
-											DupMaterial->SetScalarParameterValueEditorOnly(param.Name, value);
-										}
+										DupMaterial->SetScalarParameterValueEditorOnly(Param.Name, Value);
 									}
+								}
 
-									// copy vector parameters
-									Inst->GetAllVectorParameterInfo(parametersInfo, parametersGuids);
-									for (const FMaterialParameterInfo& param : parametersInfo)
+								// copy vector parameters
+								TArray<FMaterialParameterInfo> VectorParameterInfoArray;
+								Interface->GetAllVectorParameterInfo(VectorParameterInfoArray, GuidArray);
+								for (const FMaterialParameterInfo& Param : VectorParameterInfoArray)
+								{
+									FLinearColor Value;
+									if (Interface->GetVectorParameterValue(Param, Value))
 									{
-										FLinearColor value;
-										if (Inst->GetVectorParameterValue(param, value))
-										{
-											DupMaterial->SetVectorParameterValueEditorOnly(param.Name, value);
-										}
+										DupMaterial->SetVectorParameterValueEditorOnly(Param.Name, Value);
 									}
+								}
 
-									// copy switch parameters
-									Inst->GetAllStaticSwitchParameterInfo(parametersInfo, parametersGuids);
-									for (int i = 0; i < parametersInfo.Num(); ++i)
+								// copy switch parameters								
+								TArray<FMaterialParameterInfo> StaticSwitchParameterInfoArray;
+								Interface->GetAllStaticSwitchParameterInfo(StaticSwitchParameterInfoArray, GuidArray);
+								for (int i = 0; i < StaticSwitchParameterInfoArray.Num(); ++i)
+								{
+									bool Value = false;
+									if (Interface->GetStaticSwitchParameterValue(StaticSwitchParameterInfoArray[i].Name, Value, GuidArray[i]))
 									{
-										bool value = false;
-										if (Inst->GetStaticSwitchParameterValue(parametersInfo[i].Name, value, parametersGuids[i]))
-										{
-											DupMaterial->SetStaticSwitchParameterValueEditorOnly(parametersInfo[i].Name, value, parametersGuids[i]);
-										}
+										DupMaterial->SetStaticSwitchParameterValueEditorOnly(StaticSwitchParameterInfoArray[i].Name, Value, GuidArray[i]);
 									}
+								}
 
-									// Replace Textures
-									TArray<FName> ParameterNames = GetTextureParameterNames(Material);
-									for (const TPair<int, UTexture*>& it : TextureReplacementMaps[m])
+								// Replace Textures
+								TArray<FName> ParameterNames = GetTextureParameterNames(Material);
+								for (const TPair<int, UTexture*>& it : TextureReplacementMaps[m])
+								{
+									if (ParameterNames.IsValidIndex(it.Key))
 									{
-										if (ParameterNames.IsValidIndex(it.Key))
-										{
-											DupMaterial->SetTextureParameterValueEditorOnly(ParameterNames[it.Key], it.Value);
-										}
+										DupMaterial->SetTextureParameterValueEditorOnly(ParameterNames[it.Key], it.Value);
 									}
+								}
 
-									// Fix potential errors compiling materials due to Sampler Types
-									for (TObjectPtr<UMaterialExpression> Expression : DupMaterial->GetExpressions())
+								// Fix potential errors compiling materials due to Sampler Types
+								for (TObjectPtr<UMaterialExpression> Expression : DupMaterial->GetExpressions())
+								{
+									if (UMaterialExpressionTextureBase* MatExpressionTexBase = Cast<UMaterialExpressionTextureBase>(Expression))
 									{
-										if (UMaterialExpressionTextureBase* MatExpressionTexBase = Cast<UMaterialExpressionTextureBase>(Expression))
-										{
-											MatExpressionTexBase->AutoSetSampleType();
-										}
+										MatExpressionTexBase->AutoSetSampleType();
 									}
-
-									DuplicatedObject->PreEditChange(NULL);
-									DuplicatedObject->PostEditChange();
 								}
+
+								DuplicatedObject->PreEditChange(NULL);
+								DuplicatedObject->PostEditChange();
 							}
 						}
 					}
@@ -2612,13 +2619,14 @@ void FCustomizableObjectEditorViewportClient::BakeInstance()
 					// Duplicate the material instances
 					for (int m = 0; m < Mesh->GetMaterials().Num(); ++m)
 					{
-						UMaterialInstance* Inst = Cast<UMaterialInstance>(Mesh->GetMaterials()[m].MaterialInterface);
-						FString MaterialName = (Inst->Parent != nullptr) ? Inst->Parent->GetName() : "Material";
+						UMaterialInterface* Interface = Mesh->GetMaterials()[m].MaterialInterface;
+						UMaterial* ParentMaterial = Interface->GetMaterial();
+						FString MaterialName = ParentMaterial ? ParentMaterial->GetName() : "Material";
 
 						// Material
 						FString MatObjName = ObjectName + "_" + MaterialName;
 
-						if (!GetUniqueResourceName(Inst, MatObjName, ArrayCachedObject, ArrayCachedElement))
+						if (!GetUniqueResourceName(Interface, MatObjName, ArrayCachedObject, ArrayCachedElement))
 						{
 							continue;
 						}
@@ -2629,44 +2637,51 @@ void FCustomizableObjectEditorViewportClient::BakeInstance()
 						}
 
 						FString MatPkgName = FolderDlg->GetAssetPath() + FString("/") + MatObjName;
-						UObject* DupMat = BakeHelper_DuplicateAsset(Inst, MatObjName, MatPkgName, false, ReplacementMap, BakingOverwritePermission);
+						UObject* DupMat = BakeHelper_DuplicateAsset(Interface, MatObjName, MatPkgName, false, ReplacementMap, BakingOverwritePermission);
 						ArrayCachedObject.Add(DupMat);
 						ArrayCachedElement.Add(MatObjName);
 						PackagesToSave.Add(DupMat->GetPackage());
 
-						// Duplicate generated textures
-						if (UMaterialInstanceDynamic* InstDynamic = Cast<UMaterialInstanceDynamic>(DupMat))
+						UMaterialInstance* Inst = Cast<UMaterialInstance>(Interface);
+
+						// Only need to duplicate the generate textures if the original material is a dynamic instance
+						// If the material has Mutable textures, then it will be a dynamic material instance for sure
+						if (Inst)
 						{
-							for (int t = 0; t < Inst->TextureParameterValues.Num(); ++t)
+							// Duplicate generated textures
+							if (UMaterialInstanceDynamic* InstDynamic = Cast<UMaterialInstanceDynamic>(DupMat))
 							{
-								if (Inst->TextureParameterValues[t].ParameterValue)
+								for (int t = 0; t < Inst->TextureParameterValues.Num(); ++t)
 								{
-									UTexture2D* SrcTex = Cast<UTexture2D>(Inst->TextureParameterValues[t].ParameterValue);
-									FString ParameterSanitized = Inst->TextureParameterValues[t].ParameterInfo.Name.ToString();
-									RemoveRestrictedChars(ParameterSanitized);
+									if (Inst->TextureParameterValues[t].ParameterValue)
+									{
+										UTexture2D* SrcTex = Cast<UTexture2D>(Inst->TextureParameterValues[t].ParameterValue);
+										FString ParameterSanitized = Inst->TextureParameterValues[t].ParameterInfo.Name.ToString();
+										RemoveRestrictedChars(ParameterSanitized);
 
-									FString TexObjName = ObjectName + "_" + MaterialName + "_" + ParameterSanitized;
+										FString TexObjName = ObjectName + "_" + MaterialName + "_" + ParameterSanitized;
 
-									if (!GetUniqueResourceName(SrcTex, TexObjName, ArrayCachedObject, ArrayCachedElement))
-									{
-										UTexture* PrevTexture = Cast<UTexture>(ArrayCachedObject[ArrayCachedElement.Find(TexObjName)]);
-										InstDynamic->SetTextureParameterValue(Inst->TextureParameterValues[t].ParameterInfo.Name, PrevTexture);
-										continue;
-									}
+										if (!GetUniqueResourceName(SrcTex, TexObjName, ArrayCachedObject, ArrayCachedElement))
+										{
+											UTexture* PrevTexture = Cast<UTexture>(ArrayCachedObject[ArrayCachedElement.Find(TexObjName)]);
+											InstDynamic->SetTextureParameterValue(Inst->TextureParameterValues[t].ParameterInfo.Name, PrevTexture);
+											continue;
+										}
 
-									if (!ManageBakingAction(AssetPath, TexObjName))
-									{
-										return;
-									}
+										if (!ManageBakingAction(AssetPath, TexObjName))
+										{
+											return;
+										}
 
-									FString TexPkgName = FolderDlg->GetAssetPath() + FString("/") + TexObjName;
-									TMap<UObject*, UObject*> FakeReplacementMap;
-									UTexture2D* DupTex = BakeHelper_CreateAssetTexture(SrcTex, TexObjName, TexPkgName, nullptr, false, FakeReplacementMap, BakingOverwritePermission);
-									ArrayCachedObject.Add(DupTex);
-									ArrayCachedElement.Add(TexObjName);
-									PackagesToSave.Add(DupTex->GetPackage());
+										FString TexPkgName = FolderDlg->GetAssetPath() + FString("/") + TexObjName;
+										TMap<UObject*, UObject*> FakeReplacementMap;
+										UTexture2D* DupTex = BakeHelper_CreateAssetTexture(SrcTex, TexObjName, TexPkgName, nullptr, false, FakeReplacementMap, BakingOverwritePermission);
+										ArrayCachedObject.Add(DupTex);
+										ArrayCachedElement.Add(TexObjName);
+										PackagesToSave.Add(DupTex->GetPackage());
 
-									InstDynamic->SetTextureParameterValue(Inst->TextureParameterValues[t].ParameterInfo.Name, DupTex);
+										InstDynamic->SetTextureParameterValue(Inst->TextureParameterValues[t].ParameterInfo.Name, DupTex);
+									}
 								}
 							}
 						}
