[mutable] Remove most usage of standard C++ lib in the runtime.
#rnx

diff --git a/Engine/Plugins/Experimental/Mutable/Source/CustomizableObject/Public/MuCO/UnrealConversionUtils.h b/Engine/Plugins/Experimental/Mutable/Source/CustomizableObject/Public/MuCO/UnrealConversionUtils.h
--- a/Engine/Plugins/Experimental/Mutable/Source/CustomizableObject/Public/MuCO/UnrealConversionUtils.h
+++ b/Engine/Plugins/Experimental/Mutable/Source/CustomizableObject/Public/MuCO/UnrealConversionUtils.h
@@ -2,7 +2,7 @@
 
 #pragma once
 
-#include "../Private/MuCO/CustomizableObjectSystemPrivate.h"
+#include "MuCO/CustomizableObjectSystemPrivate.h"
 #include "Containers/Array.h"
 #include "HAL/Platform.h"
 #include "MuR/MeshBufferSet.h"
diff --git a/Engine/Plugins/Experimental/Mutable/Source/MutableRuntime/Private/MuR/BlockCompression/Miro/Miro.cpp b/Engine/Plugins/Experimental/Mutable/Source/MutableRuntime/Private/MuR/BlockCompression/Miro/Miro.cpp
--- a/Engine/Plugins/Experimental/Mutable/Source/MutableRuntime/Private/MuR/BlockCompression/Miro/Miro.cpp
+++ b/Engine/Plugins/Experimental/Mutable/Source/MutableRuntime/Private/MuR/BlockCompression/Miro/Miro.cpp
@@ -5,21 +5,13 @@
 #include "Async/ParallelFor.h"
 #include "HAL/UnrealMemory.h"
 #include "Math/UnrealMathSSE.h"
+#include "Math/UnrealMathUtility.h"
 #include "Misc/AssertionMacros.h"
-
-#include <cmath>
-#include <cstdlib>
-#include <memory>
+#include "Templates/UniquePtr.h"
 
 // debug
 //static bool s_debuglog = false;
 
-#if defined(MIRO_SSE)
-#if MIRO_SSE>=2
-#include <pmmintrin.h>
-#endif
-#endif
-
 #ifdef _MSC_VER
 #define MIRO_ALIGN __declspec(align(16))
 #else
@@ -2548,26 +2540,26 @@ namespace impl
 
 		inline vec3i_t floor(vec3f_t color) {
 			vec3i_t result;
-			result.r = static_cast<int>(std::floor(color.r));
-			result.g = static_cast<int>(std::floor(color.g));
-			result.b = static_cast<int>(std::floor(color.b));
+			result.r = static_cast<int>(FMath::Floor(color.r));
+			result.g = static_cast<int>(FMath::Floor(color.g));
+			result.b = static_cast<int>(FMath::Floor(color.b));
 			return result;
 		}
 
 		inline vec3i_t round(vec3f_t color) {
 			vec3i_t result;
-			result.r = static_cast<int>(std::round(color.r));
-			result.g = static_cast<int>(std::round(color.g));
-			result.b = static_cast<int>(std::round(color.b));
+			result.r = static_cast<int>(FMath::RoundToInt32(color.r));
+			result.g = static_cast<int>(FMath::RoundToInt32(color.g));
+			result.b = static_cast<int>(FMath::RoundToInt32(color.b));
 			return result;
 		}
 
 		inline vec4i_t round(vec4f_t color) {
 			vec4i_t result;
-			result.r = static_cast<int>(std::round(color.r));
-			result.g = static_cast<int>(std::round(color.g));
-			result.b = static_cast<int>(std::round(color.b));
-			result.a = static_cast<int>(std::round(color.a));
+			result.r = static_cast<int>(FMath::RoundToInt32(color.r));
+			result.g = static_cast<int>(FMath::RoundToInt32(color.g));
+			result.b = static_cast<int>(FMath::RoundToInt32(color.b));
+			result.a = static_cast<int>(FMath::RoundToInt32(color.a));
 			return result;
 		}
 
@@ -4837,8 +4829,8 @@ namespace impl
 			// stack is to big
 			//        int texels_of_weight[MAX_WEIGHTS_PER_BLOCK][MAX_TEXELS_PER_BLOCK];
 			//        int texelweights_of_weight[MAX_WEIGHTS_PER_BLOCK][MAX_TEXELS_PER_BLOCK];
-			std::unique_ptr<int[]> texels_of_weight(new int[MAX_WEIGHTS_PER_BLOCK * MAX_TEXELS_PER_BLOCK]);
-			std::unique_ptr<int[]> texelweights_of_weight(new int[MAX_WEIGHTS_PER_BLOCK * MAX_TEXELS_PER_BLOCK]);
+			TUniquePtr<int[]> texels_of_weight(new int[MAX_WEIGHTS_PER_BLOCK * MAX_TEXELS_PER_BLOCK]);
+			TUniquePtr<int[]> texelweights_of_weight(new int[MAX_WEIGHTS_PER_BLOCK * MAX_TEXELS_PER_BLOCK]);
 
 			for (i = 0; i < MAX_WEIGHTS_PER_BLOCK; i++)
 			{
@@ -8867,12 +8859,6 @@ namespace miro
 	}
 
 
-	//---------------------------------------------------------------------------------------------
-	const char* GetVersion()
-	{
-		return "0.0.2";
-	}
-
 #if MIRO_INCLUDE_BC
 
 	//---------------------------------------------------------------------------------------------
@@ -10275,7 +10261,7 @@ namespace miro
 
 
 	//---------------------------------------------------------------------------------------------
-	void L_to_ASTC4x4RGBL(uint32_t sx, uint32_t sy, const uint8_t* from, uint8_t* to, int q)
+	void L_to_ASTC4x4RGBL(uint32 sx, uint32 sy, const uint8_t* from, uint8_t* to, int q)
 	{
 		// Almost the same as RGB_to_ASTC4x4RGBL
 		(void)q;
diff --git a/Engine/Plugins/Experimental/Mutable/Source/MutableRuntime/Private/MuR/BlockCompression/Miro/Miro.h b/Engine/Plugins/Experimental/Mutable/Source/MutableRuntime/Private/MuR/BlockCompression/Miro/Miro.h
--- a/Engine/Plugins/Experimental/Mutable/Source/MutableRuntime/Private/MuR/BlockCompression/Miro/Miro.h
+++ b/Engine/Plugins/Experimental/Mutable/Source/MutableRuntime/Private/MuR/BlockCompression/Miro/Miro.h
@@ -2,7 +2,7 @@
 
 #pragma once
 
-#include <cstdint>
+#include "HAL/Platform.h"
 
 // Configure the compression library with these defines:
 
@@ -37,41 +37,39 @@ namespace miro
     extern void initialize();
     extern void finalize();
 
-    //!
-    extern const char* GetVersion();
 
 #if MIRO_INCLUDE_BC
 
     //! BC1 support
-    extern void RGB_to_BC1( uint32_t sx, uint32_t sy, const uint8_t* from, uint8_t* to, int q );
-    extern void RGBA_to_BC1( uint32_t sx, uint32_t sy, const uint8_t* from, uint8_t* to, int q );
-    extern void BC1_to_RGB( uint32_t sx, uint32_t sy, const uint8_t* from, uint8_t* to );
-    extern void BC1_to_RGBA( uint32_t sx, uint32_t sy, const uint8_t* from, uint8_t* to );
+    extern void RGB_to_BC1( uint32 sx, uint32 sy, const uint8* from, uint8* to, int q );
+    extern void RGBA_to_BC1( uint32 sx, uint32 sy, const uint8* from, uint8* to, int q );
+    extern void BC1_to_RGB( uint32 sx, uint32 sy, const uint8* from, uint8* to );
+    extern void BC1_to_RGBA( uint32 sx, uint32 sy, const uint8* from, uint8* to );
 
     //! BC2 support
-    extern void RGB_to_BC2( uint32_t sx, uint32_t sy, const uint8_t* from, uint8_t* to, int q );
-    extern void RGBA_to_BC2( uint32_t sx, uint32_t sy, const uint8_t* from, uint8_t* to, int q );
-    extern void BC2_to_RGB( uint32_t sx, uint32_t sy, const uint8_t* from, uint8_t* to );
-    extern void BC2_to_RGBA( uint32_t sx, uint32_t sy, const uint8_t* from, uint8_t* to );
+    extern void RGB_to_BC2( uint32 sx, uint32 sy, const uint8* from, uint8* to, int q );
+    extern void RGBA_to_BC2( uint32 sx, uint32 sy, const uint8* from, uint8* to, int q );
+    extern void BC2_to_RGB( uint32 sx, uint32 sy, const uint8* from, uint8* to );
+    extern void BC2_to_RGBA( uint32 sx, uint32 sy, const uint8* from, uint8* to );
 
     //! BC3 support
-    extern void RGB_to_BC3( uint32_t sx, uint32_t sy, const uint8_t* from, uint8_t* to, int q );
-    extern void RGBA_to_BC3( uint32_t sx, uint32_t sy, const uint8_t* from, uint8_t* to, int q );
-    extern void BC3_to_RGB( uint32_t sx, uint32_t sy, const uint8_t* from, uint8_t* to );
-    extern void BC3_to_RGBA( uint32_t sx, uint32_t sy, const uint8_t* from, uint8_t* to );
-    extern void BC1_to_BC3( uint32_t sx, uint32_t sy, const uint8_t* from, uint8_t* to, int q );
+    extern void RGB_to_BC3( uint32 sx, uint32 sy, const uint8* from, uint8* to, int q );
+    extern void RGBA_to_BC3( uint32 sx, uint32 sy, const uint8* from, uint8* to, int q );
+    extern void BC3_to_RGB( uint32 sx, uint32 sy, const uint8* from, uint8* to );
+    extern void BC3_to_RGBA( uint32 sx, uint32 sy, const uint8* from, uint8* to );
+    extern void BC1_to_BC3( uint32 sx, uint32 sy, const uint8* from, uint8* to, int q );
 
     //! BC4 support
-    extern void L_to_BC4( uint32_t sx, uint32_t sy, const uint8_t* from, uint8_t* to, int q );
-    extern void BC4_to_L( uint32_t sx, uint32_t sy, const uint8_t* from, uint8_t* to );
-    extern void BC4_to_RGB( uint32_t sx, uint32_t sy, const uint8_t* from, uint8_t* to );
-    extern void BC4_to_RGBA( uint32_t sx, uint32_t sy, const uint8_t* from, uint8_t* to );
+    extern void L_to_BC4( uint32 sx, uint32 sy, const uint8* from, uint8* to, int q );
+    extern void BC4_to_L( uint32 sx, uint32 sy, const uint8* from, uint8* to );
+    extern void BC4_to_RGB( uint32 sx, uint32 sy, const uint8* from, uint8* to );
+    extern void BC4_to_RGBA( uint32 sx, uint32 sy, const uint8* from, uint8* to );
 
     //! BC5 support
-    extern void RGBA_to_BC5( uint32_t sx, uint32_t sy, const uint8_t* from, uint8_t* to, int q );
-    extern void RGB_to_BC5( uint32_t sx, uint32_t sy, const uint8_t* from, uint8_t* to, int q );
-    extern void BC5_to_RGBA( uint32_t sx, uint32_t sy, const uint8_t* from, uint8_t* to );
-    extern void BC5_to_RGB( uint32_t sx, uint32_t sy, const uint8_t* from, uint8_t* to );
+    extern void RGBA_to_BC5( uint32 sx, uint32 sy, const uint8* from, uint8* to, int q );
+    extern void RGB_to_BC5( uint32 sx, uint32 sy, const uint8* from, uint8* to, int q );
+    extern void BC5_to_RGBA( uint32 sx, uint32 sy, const uint8* from, uint8* to );
+    extern void BC5_to_RGB( uint32 sx, uint32 sy, const uint8* from, uint8* to );
 
 #endif
 
@@ -79,27 +77,27 @@ namespace miro
 #if MIRO_INCLUDE_ASTC
 
     //! ASTC 4x4 RGBA low support
-    extern void RGBA_to_ASTC4x4RGBAL( uint32_t sx, uint32_t sy, const uint8_t* from, uint8_t* to, int q );
-    extern void RGB_to_ASTC4x4RGBAL( uint32_t sx, uint32_t sy, const uint8_t* from, uint8_t* to, int q );
-    extern void ASTC4x4RGBAL_to_RGBA( uint32_t sx, uint32_t sy, const uint8_t* from, uint8_t* to );
-    extern void ASTC4x4RGBAL_to_RGB( uint32_t sx, uint32_t sy, const uint8_t* from, uint8_t* to );
+    extern void RGBA_to_ASTC4x4RGBAL( uint32 sx, uint32 sy, const uint8* from, uint8* to, int q );
+    extern void RGB_to_ASTC4x4RGBAL( uint32 sx, uint32 sy, const uint8* from, uint8* to, int q );
+    extern void ASTC4x4RGBAL_to_RGBA( uint32 sx, uint32 sy, const uint8* from, uint8* to );
+    extern void ASTC4x4RGBAL_to_RGB( uint32 sx, uint32 sy, const uint8* from, uint8* to );
 
     //! ASTC 4x4 RGB low support
-    extern void RGB_to_ASTC4x4RGBL( uint32_t sx, uint32_t sy, const uint8_t* from, uint8_t* to, int q );
-    extern void RGBA_to_ASTC4x4RGBL( uint32_t sx, uint32_t sy, const uint8_t* from, uint8_t* to, int q );
-    extern void ASTC4x4RGBL_to_RGB( uint32_t sx, uint32_t sy, const uint8_t* from, uint8_t* to );
-    extern void ASTC4x4RGBL_to_RGBA( uint32_t sx, uint32_t sy, const uint8_t* from, uint8_t* to );
-	extern void L_to_ASTC4x4RGBL(uint32_t sx, uint32_t sy, const uint8_t* from, uint8_t* to, int q);
+    extern void RGB_to_ASTC4x4RGBL( uint32 sx, uint32 sy, const uint8* from, uint8* to, int q );
+    extern void RGBA_to_ASTC4x4RGBL( uint32 sx, uint32 sy, const uint8* from, uint8* to, int q );
+    extern void ASTC4x4RGBL_to_RGB( uint32 sx, uint32 sy, const uint8* from, uint8* to );
+    extern void ASTC4x4RGBL_to_RGBA( uint32 sx, uint32 sy, const uint8* from, uint8* to );
+	extern void L_to_ASTC4x4RGBL(uint32 sx, uint32 sy, const uint8* from, uint8* to, int q);
     
 	//! Convenience transforms
-    extern void ASTC4x4RGBL_to_ASTC4x4RGBAL( uint32_t sx, uint32_t sy, const uint8_t* from, uint8_t* to );
-    extern void ASTC4x4RGBAL_to_ASTC4x4RGBL( uint32_t sx, uint32_t sy, const uint8_t* from, uint8_t* to );
+    extern void ASTC4x4RGBL_to_ASTC4x4RGBAL( uint32 sx, uint32 sy, const uint8* from, uint8* to );
+    extern void ASTC4x4RGBAL_to_ASTC4x4RGBL( uint32 sx, uint32 sy, const uint8* from, uint8* to );
 
     //! ASTC 4x4 RG low support
-    extern void RGB_to_ASTC4x4RGL( uint32_t sx, uint32_t sy, const uint8_t* from, uint8_t* to, int q );
-    extern void RGBA_to_ASTC4x4RGL( uint32_t sx, uint32_t sy, const uint8_t* from, uint8_t* to, int q );
-    extern void ASTC4x4RGL_to_RGB( uint32_t sx, uint32_t sy, const uint8_t* from, uint8_t* to );
-    extern void ASTC4x4RGL_to_RGBA( uint32_t sx, uint32_t sy, const uint8_t* from, uint8_t* to );
+    extern void RGB_to_ASTC4x4RGL( uint32 sx, uint32 sy, const uint8* from, uint8* to, int q );
+    extern void RGBA_to_ASTC4x4RGL( uint32 sx, uint32 sy, const uint8* from, uint8* to, int q );
+    extern void ASTC4x4RGL_to_RGB( uint32 sx, uint32 sy, const uint8* from, uint8* to );
+    extern void ASTC4x4RGL_to_RGBA( uint32 sx, uint32 sy, const uint8* from, uint8* to );
 
 #endif
 
diff --git a/Engine/Plugins/Experimental/Mutable/Source/MutableRuntime/Private/MuR/CodeRunner.cpp b/Engine/Plugins/Experimental/Mutable/Source/MutableRuntime/Private/MuR/CodeRunner.cpp
--- a/Engine/Plugins/Experimental/Mutable/Source/MutableRuntime/Private/MuR/CodeRunner.cpp
+++ b/Engine/Plugins/Experimental/Mutable/Source/MutableRuntime/Private/MuR/CodeRunner.cpp
@@ -81,8 +81,6 @@
 #include "Templates/Tuple.h"
 #include "Trace/Detail/Channel.h"
 
-#include <memory>
-#include <utility>
 
 namespace mu
 {
@@ -178,7 +176,7 @@ namespace mu
             OP::ADDRESS resultAt = value ? args.yes : args.no;
 
             // Schedule the end of this instruction if necessary
-            AddOp( SCHEDULED_OP( item.at, item, 2, (uint32_t)value),
+            AddOp( SCHEDULED_OP( item.at, item, 2, (uint32)value),
 				SCHEDULED_OP( resultAt, item) );
 
             break;
@@ -224,7 +222,7 @@ namespace mu
 	{
 		OP_TYPE type = pModel->GetPrivate()->m_program.GetOpType(item.at);
 
-		const uint8_t* data = pModel->GetPrivate()->m_program.GetOpArgsPointer(item.at);
+		const uint8* data = pModel->GetPrivate()->m_program.GetOpArgsPointer(item.at);
 
 		OP::ADDRESS VarAddress;
 		FMemory::Memcpy(&VarAddress, data, sizeof(OP::ADDRESS));
@@ -234,9 +232,9 @@ namespace mu
 		FMemory::Memcpy(&DefAddress, data, sizeof(OP::ADDRESS));
 		data += sizeof(OP::ADDRESS);
 
-		uint32_t CaseCount;
-		FMemory::Memcpy(&CaseCount, data, sizeof(uint32_t));
-		data += sizeof(uint32_t);
+		uint32 CaseCount;
+		FMemory::Memcpy(&CaseCount, data, sizeof(uint32));
+		data += sizeof(uint32);
 
 		switch (item.stage)
 		{
@@ -275,11 +273,11 @@ namespace mu
 			int var = GetMemory().GetInt(CACHE_ADDRESS(VarAddress, item));
 
 			OP::ADDRESS valueAt = DefAddress;
-			for (uint32_t C = 0; C < CaseCount; ++C)
+			for (uint32 C = 0; C < CaseCount; ++C)
 			{
-				int32_t Condition;
-				FMemory::Memcpy(&Condition, data, sizeof(int32_t));
-				data += sizeof(int32_t);
+				int32 Condition;
+				FMemory::Memcpy(&Condition, data, sizeof(int32));
+				data += sizeof(int32);
 
 				OP::ADDRESS At;
 				FMemory::Memcpy(&At, data, sizeof(OP::ADDRESS));
@@ -335,7 +333,7 @@ namespace mu
     //---------------------------------------------------------------------------------------------
     void CodeRunner::RunCode_Instance( SCHEDULED_OP& item,
                                   const Model* pModel,
-                                  uint32_t lodMask
+                                  uint32 lodMask
                                   )
     {
         //MUTABLE_CPUPROFILER_SCOPE(RunCode_Instance);
@@ -1507,7 +1505,7 @@ namespace mu
                 Ptr<const Mesh> pSource = GetMemory().GetMesh( CACHE_ADDRESS(args.source,item) );
                 Ptr<const Mesh> pFormat = GetMemory().GetMesh( CACHE_ADDRESS(args.format,item) );
 
-                uint8_t flags = args.buffers;
+                uint8 flags = args.buffers;
                 MeshPtr pResult;
                 pResult = MeshFormat( pSource.get(),
                                       pFormat.get(),
@@ -1537,19 +1535,19 @@ namespace mu
 
         case OP_TYPE::ME_EXTRACTLAYOUTBLOCK:
         {
-            const uint8_t* data = pModel->GetPrivate()->m_program.GetOpArgsPointer(item.at);
+            const uint8* data = pModel->GetPrivate()->m_program.GetOpArgsPointer(item.at);
 
             mu::OP::ADDRESS source;
             memcpy( &source, data, sizeof(OP::ADDRESS) );
             data += sizeof(OP::ADDRESS);
 
-            uint16_t layout;
-            memcpy( &layout, data, sizeof(uint16_t) );
-            data += sizeof(uint16_t);
+            uint16 layout;
+            memcpy( &layout, data, sizeof(uint16) );
+            data += sizeof(uint16);
 
-            uint16_t blockCount;
-            memcpy( &blockCount, data, sizeof(uint16_t) );
-            data += sizeof(uint16_t);
+            uint16 blockCount;
+            memcpy( &blockCount, data, sizeof(uint16) );
+            data += sizeof(uint16);
 
             switch (item.stage)
             {
@@ -1570,8 +1568,8 @@ namespace mu
                 Ptr<const Mesh> pSource = GetMemory().GetMesh( CACHE_ADDRESS(source,item) );
 
                 // Access with memcpy necessary for unaligned arm issues.
-                uint32_t blocks[1024];
-                memcpy( blocks, data, sizeof(uint32_t)*std::min(1024,int(blockCount)) );
+                uint32 blocks[1024];
+                memcpy( blocks, data, sizeof(uint32)*std::min(1024,int(blockCount)) );
 
                 MeshPtr pResult;
                 pResult = MeshExtractLayoutBlock( pSource.get(),
@@ -1715,7 +1713,7 @@ namespace mu
                                     pModel->GetPrivate()->m_program.m_constantStrings.size() );
 
                     SHAPE selectionShape;
-                    selectionShape.type = (uint8_t)SHAPE::Type::None;
+                    selectionShape.type = (uint8)SHAPE::Type::None;
                     const string& selectionBone = pModel->GetPrivate()->m_program.m_constantStrings[args.vertexSelectionShapeOrBone];
 					pResult = MeshClipMorphPlane(pSource.get(), origin, normal, args.dist, args.factor, morphShape.size[0], morphShape.size[1], morphShape.size[2], selectionShape, selectionBone, args.maxBoneRadius);
                 }
@@ -1723,7 +1721,7 @@ namespace mu
                 {
                     // No vertex selection
                     SHAPE selectionShape;
-                    selectionShape.type = (uint8_t)SHAPE::Type::None;
+                    selectionShape.type = (uint8)SHAPE::Type::None;
                     pResult = MeshClipMorphPlane(pSource.get(), origin, normal, args.dist, args.factor, morphShape.size[0], morphShape.size[1], morphShape.size[2], selectionShape);
                 }
 
@@ -1963,7 +1961,7 @@ namespace mu
 		case OP_TYPE::ME_BINDSHAPE:
 		{
 			OP::MeshBindShapeArgs args = pModel->GetPrivate()->m_program.GetOpArgs<OP::MeshBindShapeArgs>(item.at);
-			const uint8_t* data = pModel->GetPrivate()->m_program.GetOpArgsPointer(item.at);
+			const uint8* data = pModel->GetPrivate()->m_program.GetOpArgsPointer(item.at);
 
 			switch (item.stage)
 			{
@@ -2209,7 +2207,7 @@ namespace mu
         {
             // Decode op
             // TODO: Partial decode for each stage
-            const uint8_t* data = pModel->GetPrivate()->m_program.GetOpArgsPointer(item.at);
+            const uint8* data = pModel->GetPrivate()->m_program.GetOpArgsPointer(item.at);
 
             OP::ADDRESS source;
             FMemory::Memcpy(&source,data,sizeof(OP::ADDRESS)); data += sizeof(OP::ADDRESS);
@@ -2217,10 +2215,10 @@ namespace mu
             vector<SCHEDULED_OP> conditions;
             vector<OP::ADDRESS> masks;
 
-            uint16_t removes;
-			FMemory::Memcpy(&removes,data,sizeof(uint16_t)); data += sizeof(uint16_t);
+            uint16 removes;
+			FMemory::Memcpy(&removes,data,sizeof(uint16)); data += sizeof(uint16);
 
-            for( uint16_t r=0; r<removes; ++r)
+            for( uint16 r=0; r<removes; ++r)
             {
                 OP::ADDRESS condition;
 				FMemory::Memcpy(&condition,data,sizeof(OP::ADDRESS)); data += sizeof(OP::ADDRESS);
@@ -2937,8 +2935,8 @@ namespace mu
                 Ptr<const Image> pBase = GetMemory().GetImage( CACHE_ADDRESS(args.source,item) );
 
                 FImageSize destSize(
-                            uint16_t( pBase->GetSizeX()*args.factor[0] + 0.5f ),
-                            uint16_t( pBase->GetSizeY()*args.factor[1] + 0.5f ) );
+                            uint16( pBase->GetSizeX()*args.factor[0] + 0.5f ),
+                            uint16( pBase->GetSizeY()*args.factor[1] + 0.5f ) );
 
                 //pResult = ImageResize( pBase.get(), destSize );
                 ImagePtr pResult = ImageResizeLinear(
@@ -3655,7 +3653,7 @@ namespace mu
 					{
 						MUTABLE_CPUPROFILER_SCOPE(ImagePatchResize_Emergency);
 
-						FImageSize blockSize((uint16_t)rect.size[0], (uint16_t)rect.size[1]);
+						FImageSize blockSize((uint16)rect.size[0], (uint16)rect.size[1]);
 						pB = ImageResizeLinear(m_pSettings->GetPrivate()->m_imageCompressionQuality,pB.get(), blockSize);
 					}
 
@@ -4966,11 +4964,12 @@ namespace mu
 					pResult = pSource->Clone();
 
 					SCRATCH_LAYOUT_PACK scratch;
-					scratch.blocks.resize(pSource->GetBlockCount());
-					scratch.sorted.resize(pSource->GetBlockCount());
-					scratch.positions.resize(pSource->GetBlockCount());
-					scratch.priorities.resize(pSource->GetBlockCount());
-					scratch.reductions.resize(pSource->GetBlockCount());
+					int32 BlockCount = pSource->GetBlockCount();
+					scratch.blocks.SetNum(BlockCount);
+					scratch.sorted.SetNum(BlockCount);
+					scratch.positions.SetNum(BlockCount);
+					scratch.priorities.SetNum(BlockCount);
+					scratch.reductions.SetNum(BlockCount);
 
 					LayoutPack3(pResult.get(), pSource.get(), &scratch);
 				}
@@ -5032,7 +5031,7 @@ namespace mu
     void CodeRunner::RunCode( SCHEDULED_OP& item,
                               const Parameters* pParams,
                               const Model* pModel,
-                              uint32_t lodMask)
+                              uint32 lodMask)
     {
 		//UE_LOG( LogMutableCore, Log, TEXT("Running :%5d , %d "), item.at, item.stage );
 		check( item.Type == SCHEDULED_OP::EType::Full );
@@ -5134,7 +5133,7 @@ namespace mu
 	void CodeRunner::RunCodeImageDesc(SCHEDULED_OP& item,
 		const Parameters* pParams,
 		const Model* pModel, 
-		uint32_t lodMask
+		uint32 lodMask
 	)
 	{
 		MUTABLE_CPUPROFILER_SCOPE(RunCodeImageDesc);
@@ -5210,7 +5209,7 @@ namespace mu
 
 		case OP_TYPE::IM_SWITCH:
 		{
-			const uint8_t* data = program.GetOpArgsPointer(item.at);
+			const uint8* data = program.GetOpArgsPointer(item.at);
 		
 			OP::ADDRESS VarAddress;
 			FMemory::Memcpy( &VarAddress, data, sizeof(OP::ADDRESS));
@@ -5220,9 +5219,9 @@ namespace mu
 			FMemory::Memcpy( &DefAddress, data, sizeof(OP::ADDRESS));
 			data += sizeof(OP::ADDRESS);
 
-			uint32_t CaseCount;
-			FMemory::Memcpy( &CaseCount, data, sizeof(uint32_t));
-			data += sizeof(uint32_t);
+			uint32 CaseCount;
+			FMemory::Memcpy( &CaseCount, data, sizeof(uint32));
+			data += sizeof(uint32);
 	
 			switch (item.stage)
 			{
@@ -5248,11 +5247,11 @@ namespace mu
 				int var = GetMemory().GetInt(CACHE_ADDRESS(VarAddress, item));
 
 				OP::ADDRESS valueAt = DefAddress;
-				for (uint32_t C = 0; C < CaseCount; ++C)
+				for (uint32 C = 0; C < CaseCount; ++C)
 				{
-					int32_t Condition;
-					FMemory::Memcpy( &Condition, data, sizeof(int32_t) );
-					data += sizeof(int32_t);
+					int32 Condition;
+					FMemory::Memcpy( &Condition, data, sizeof(int32) );
+					data += sizeof(int32);
 					
 					OP::ADDRESS At;
 					FMemory::Memcpy( &At, data, sizeof(OP::ADDRESS) );
@@ -5478,8 +5477,8 @@ namespace mu
 			{
 				FImageDesc& ResultAndBaseDesc = m_heapImageDesc[item.customState];
 				FImageSize destSize(
-					uint16_t(ResultAndBaseDesc.m_size[0] * args.factor[0] + 0.5f),
-					uint16_t(ResultAndBaseDesc.m_size[1] * args.factor[1] + 0.5f));
+					uint16(ResultAndBaseDesc.m_size[0] * args.factor[0] + 0.5f),
+					uint16(ResultAndBaseDesc.m_size[1] * args.factor[1] + 0.5f));
 				ResultAndBaseDesc.m_size = destSize;
 				GetMemory().SetValidDesc(item);
 				break;
@@ -5515,7 +5514,7 @@ namespace mu
 				FIntPoint ImageSizeInPixels = SizeInBlocks * BlockSizeInPixels;
 
 				FImageDesc& ResultAndBaseDesc = m_heapImageDesc[item.customState];
-				FImageSize destSize(uint16_t(ImageSizeInPixels.X), uint16_t(ImageSizeInPixels.Y));
+				FImageSize destSize(uint16(ImageSizeInPixels.X), uint16(ImageSizeInPixels.Y));
 				ResultAndBaseDesc.m_size = destSize;
 				
 				if (args.generateMipmaps)
diff --git a/Engine/Plugins/Experimental/Mutable/Source/MutableRuntime/Private/MuR/CodeRunnerTasked.cpp b/Engine/Plugins/Experimental/Mutable/Source/MutableRuntime/Private/MuR/CodeRunnerTasked.cpp
--- a/Engine/Plugins/Experimental/Mutable/Source/MutableRuntime/Private/MuR/CodeRunnerTasked.cpp
+++ b/Engine/Plugins/Experimental/Mutable/Source/MutableRuntime/Private/MuR/CodeRunnerTasked.cpp
@@ -50,10 +50,6 @@
 #include "Templates/Tuple.h"
 #include "Trace/Detail/Channel.h"
 
-#include <functional>
-#include <memory>
-#include <utility>
-
 
 DECLARE_CYCLE_STAT(TEXT("MutableCoreTask"), STAT_MutableCoreTask, STATGROUP_Game);
 
@@ -1101,7 +1097,7 @@ namespace mu
 					});
 			}
 
-			uint32_t RomSize = program.m_roms[RomIndex].Size;
+			uint32 RomSize = program.m_roms[RomIndex].Size;
 			check(RomSize>0);
 
 			CodeRunner::ROM_LOAD_OP* op = nullptr;
@@ -1221,7 +1217,7 @@ namespace mu
 					});
 			}
 
-			uint32_t RomSize = program.m_roms[RomIndex].Size;
+			uint32 RomSize = program.m_roms[RomIndex].Size;
 			check(RomSize > 0);
 
 			CodeRunner::ROM_LOAD_OP* op = nullptr;
diff --git a/Engine/Plugins/Experimental/Mutable/Source/MutableRuntime/Private/MuR/CodeVisitor.cpp b/Engine/Plugins/Experimental/Mutable/Source/MutableRuntime/Private/MuR/CodeVisitor.cpp
--- a/Engine/Plugins/Experimental/Mutable/Source/MutableRuntime/Private/MuR/CodeVisitor.cpp
+++ b/Engine/Plugins/Experimental/Mutable/Source/MutableRuntime/Private/MuR/CodeVisitor.cpp
@@ -6,10 +6,6 @@
 #include "HAL/PlatformCrt.h"
 #include "MuR/MutableTrace.h"
 
-#include <functional>
-#include <memory>
-#include <utility>
-
 namespace mu
 {
 
diff --git a/Engine/Plugins/Experimental/Mutable/Source/MutableRuntime/Private/MuR/Config.h b/Engine/Plugins/Experimental/Mutable/Source/MutableRuntime/Private/MuR/Config.h
--- a/Engine/Plugins/Experimental/Mutable/Source/MutableRuntime/Private/MuR/Config.h
+++ b/Engine/Plugins/Experimental/Mutable/Source/MutableRuntime/Private/MuR/Config.h
@@ -14,31 +14,10 @@
 
 // This file contains all the possible configuration flags that affect the build of the runtime
 
-//! Add support for OpenGL 4 gpu building operations.
-//! If this is enabled, the .cpp files in the mutable/runtime/impl/gpu-gl4 need to be linked in the
-//! runtime.
-//#define MUTABLE_GPU_GL4
-
-//! Add support for OpenGL ES 2 gpu building operations.
-//! If this is enabled, the .cpp files in the mutable/runtime/impl/gpu-gles2 need to be linked in
-//! the runtime.
-//#define MUTABLE_GPU_GLES2
-
 //! This flag enables collection of statistics in the System objects. If using the waf build system,
 //! the statistics are only collected in the debug and profile versions of the runtime.
 //#define MUTABLE_PROFILE
 
-//! This flag disables the internal asserts of the library. Without this flag, the asserts are only
-//! disabled in release builds. This can be overriden by MUTABLE_EXCEPTIONS
-//#define MUTABLE_NO_ASSERT
-
-//! This flags will enable asserts regardless of MUTABLE_NO_ASSERT but they will thrown an
-//! exception instead of stop the program. This is useful for tools only.
-//#define MUTABLE_EXCEPTIONS
-
-//! Exclude code for multilayer image operations
-//#define MUTABLE_DISABLE_MULTILAYER
-
 //! Exclude support for BC(DXTC) image formats
 //! This may also be set by the build system for some target platforms.
 //#define MUTABLE_EXCLUDE_BC
diff --git a/Engine/Plugins/Experimental/Mutable/Source/MutableRuntime/Private/MuR/ConvertData.h b/Engine/Plugins/Experimental/Mutable/Source/MutableRuntime/Private/MuR/ConvertData.h
--- a/Engine/Plugins/Experimental/Mutable/Source/MutableRuntime/Private/MuR/ConvertData.h
+++ b/Engine/Plugins/Experimental/Mutable/Source/MutableRuntime/Private/MuR/ConvertData.h
@@ -25,8 +25,8 @@ namespace mu
 		case MBF_FLOAT64:
 		{
 			double* pTypedResult = reinterpret_cast<double*>(pResult);
-			uint8_t* pByteResult = reinterpret_cast<uint8_t*>(pResult);
-			const uint8_t* pByteSource = reinterpret_cast<const uint8_t*>(pSource);
+			uint8* pByteResult = reinterpret_cast<uint8*>(pResult);
+			const uint8* pByteSource = reinterpret_cast<const uint8*>(pSource);
 
 			switch (sourceFormat)
 			{
@@ -54,49 +54,49 @@ namespace mu
 
 			case MBF_INT32:
 			{
-				const int32_t* pTypedSource = reinterpret_cast<const int32_t*>(pSource);
+				const int32* pTypedSource = reinterpret_cast<const int32*>(pSource);
 				pTypedResult[channel] = (double)(pTypedSource[channel]);
 				break;
 			}
 
 			case MBF_UINT32:
 			{
-				const uint32_t* pTypedSource = reinterpret_cast<const uint32_t*>(pSource);
+				const uint32* pTypedSource = reinterpret_cast<const uint32*>(pSource);
 				pTypedResult[channel] = (double)(pTypedSource[channel]);
 				break;
 			}
 
 			case MBF_INT16:
 			{
-				const int16_t* pTypedSource = reinterpret_cast<const int16_t*>(pSource);
+				const int16* pTypedSource = reinterpret_cast<const int16*>(pSource);
 				pTypedResult[channel] = (double)(pTypedSource[channel]);
 				break;
 			}
 
 			case MBF_UINT16:
 			{
-				const uint16_t* pTypedSource = reinterpret_cast<const uint16_t*>(pSource);
+				const uint16* pTypedSource = reinterpret_cast<const uint16*>(pSource);
 				pTypedResult[channel] = (double)(pTypedSource[channel]);
 				break;
 			}
 
 			case MBF_INT8:
 			{
-				const int8_t* pTypedSource = reinterpret_cast<const int8_t*>(pSource);
+				const int8* pTypedSource = reinterpret_cast<const int8*>(pSource);
 				pTypedResult[channel] = (double)(pTypedSource[channel]);
 				break;
 			}
 
 			case MBF_UINT8:
 			{
-				const uint8_t* pTypedSource = reinterpret_cast<const uint8_t*>(pSource);
+				const uint8* pTypedSource = reinterpret_cast<const uint8*>(pSource);
 				pTypedResult[channel] = (double)(pTypedSource[channel]);
 				break;
 			}
 
 			case MBF_NINT32:
 			{
-				const int32_t* pTypedSource = reinterpret_cast<const int32_t*>(pSource);
+				const int32* pTypedSource = reinterpret_cast<const int32*>(pSource);
 				pTypedResult[channel] = (double)(pTypedSource[channel]);
 				pTypedResult[channel] /= 65536.0f * 65536.0f / 2.0f;
 				break;
@@ -104,7 +104,7 @@ namespace mu
 
 			case MBF_NUINT32:
 			{
-				const uint32_t* pTypedSource = reinterpret_cast<const uint32_t*>(pSource);
+				const uint32* pTypedSource = reinterpret_cast<const uint32*>(pSource);
 				pTypedResult[channel] = (double)(pTypedSource[channel]);
 				pTypedResult[channel] /= 65536.0f * 65536.0f - 1.0f;
 				break;
@@ -112,7 +112,7 @@ namespace mu
 
 			case MBF_NINT16:
 			{
-				const int16_t* pTypedSource = reinterpret_cast<const int16_t*>(pSource);
+				const int16* pTypedSource = reinterpret_cast<const int16*>(pSource);
 				pTypedResult[channel] = (double)(pTypedSource[channel]);
 				pTypedResult[channel] /= 32768.0f;
 				break;
@@ -120,7 +120,7 @@ namespace mu
 
 			case MBF_NUINT16:
 			{
-				const uint16_t* pTypedSource = reinterpret_cast<const uint16_t*>(pSource);
+				const uint16* pTypedSource = reinterpret_cast<const uint16*>(pSource);
 				pTypedResult[channel] = (double)(pTypedSource[channel]);
 				pTypedResult[channel] /= 65535.0f;
 				break;
@@ -128,7 +128,7 @@ namespace mu
 
 			case MBF_NINT8:
 			{
-				const int8_t* pTypedSource = reinterpret_cast<const int8_t*>(pSource);
+				const int8* pTypedSource = reinterpret_cast<const int8*>(pSource);
 				pTypedResult[channel] = (double)(pTypedSource[channel]);
 				pTypedResult[channel] /= 128.0f;
 				break;
@@ -136,7 +136,7 @@ namespace mu
 
 			case MBF_NUINT8:
 			{
-				const uint8_t* pTypedSource = reinterpret_cast<const uint8_t*>(pSource);
+				const uint8* pTypedSource = reinterpret_cast<const uint8*>(pSource);
 				pTypedResult[channel] = (double)(pTypedSource[channel]);
 				pTypedResult[channel] /= 255.0f;
 				break;
@@ -145,7 +145,7 @@ namespace mu
 			case MBF_PACKEDDIR8:
 			case MBF_PACKEDDIR8_W_TANGENTSIGN:
 			{
-				const uint8_t* pTypedSource = reinterpret_cast<const uint8_t*>(pSource);
+				const uint8* pTypedSource = reinterpret_cast<const uint8*>(pSource);
 				pTypedResult[channel] = (double)(pTypedSource[channel]);
 				pTypedResult[channel] /= 127.5f;
 				pTypedResult[channel] -= 1.0f;
@@ -155,7 +155,7 @@ namespace mu
 			case MBF_PACKEDDIRS8:
 			case MBF_PACKEDDIRS8_W_TANGENTSIGN:
 			{
-				const int8_t* pTypedSource = reinterpret_cast<const int8_t*>(pSource);
+				const int8* pTypedSource = reinterpret_cast<const int8*>(pSource);
 				pTypedResult[channel] = (double)(pTypedSource[channel]);
 				pTypedResult[channel] /= 127.5f;
 				break;
@@ -171,8 +171,8 @@ namespace mu
 		case MBF_FLOAT32:
 		{
 			float* pTypedResult = reinterpret_cast<float*>( pResult );
-            uint8_t* pByteResult = reinterpret_cast<uint8_t*>( pResult );
-            const uint8_t* pByteSource = reinterpret_cast<const uint8_t*>( pSource );
+            uint8* pByteResult = reinterpret_cast<uint8*>( pResult );
+            const uint8* pByteSource = reinterpret_cast<const uint8*>( pSource );
 
 			switch ( sourceFormat )
 			{
@@ -200,49 +200,49 @@ namespace mu
 
 			case MBF_INT32:
 			{
-                const int32_t* pTypedSource = reinterpret_cast<const int32_t*>( pSource );
+                const int32* pTypedSource = reinterpret_cast<const int32*>( pSource );
 				pTypedResult[channel] = (float)(pTypedSource[channel]);
 				break;
 			}
 
 			case MBF_UINT32:
 			{
-                const uint32_t* pTypedSource = reinterpret_cast<const uint32_t*>( pSource );
+                const uint32* pTypedSource = reinterpret_cast<const uint32*>( pSource );
 				pTypedResult[channel] = (float)(pTypedSource[channel]);
 				break;
 			}
 
 			case MBF_INT16:
 			{
-                const int16_t* pTypedSource = reinterpret_cast<const int16_t*>( pSource );
+                const int16* pTypedSource = reinterpret_cast<const int16*>( pSource );
 				pTypedResult[channel] = (float)(pTypedSource[channel]);
 				break;
 			}
 
 			case MBF_UINT16:
 			{
-                const uint16_t* pTypedSource = reinterpret_cast<const uint16_t*>( pSource );
+                const uint16* pTypedSource = reinterpret_cast<const uint16*>( pSource );
 				pTypedResult[channel] = (float)(pTypedSource[channel]);
 				break;
 			}
 
 			case MBF_INT8:
 			{
-                const int8_t* pTypedSource = reinterpret_cast<const int8_t*>( pSource );
+                const int8* pTypedSource = reinterpret_cast<const int8*>( pSource );
 				pTypedResult[channel] = (float)(pTypedSource[channel]);
 				break;
 			}
 
 			case MBF_UINT8:
 			{
-                const uint8_t* pTypedSource = reinterpret_cast<const uint8_t*>( pSource );
+                const uint8* pTypedSource = reinterpret_cast<const uint8*>( pSource );
 				pTypedResult[channel] = (float)(pTypedSource[channel]);
 				break;
 			}
 
 			case MBF_NINT32:
 			{
-                const int32_t* pTypedSource = reinterpret_cast<const int32_t*>( pSource );
+                const int32* pTypedSource = reinterpret_cast<const int32*>( pSource );
 				pTypedResult[channel] = (float)(pTypedSource[channel]);
 				pTypedResult[channel] /= 65536.0f*65536.0f/2.0f;
 				break;
@@ -250,7 +250,7 @@ namespace mu
 
 			case MBF_NUINT32:
 			{
-                const uint32_t* pTypedSource = reinterpret_cast<const uint32_t*>( pSource );
+                const uint32* pTypedSource = reinterpret_cast<const uint32*>( pSource );
 				pTypedResult[channel] = (float)(pTypedSource[channel]);
 				pTypedResult[channel] /= 65536.0f*65536.0f-1.0f;
 				break;
@@ -258,7 +258,7 @@ namespace mu
 
 			case MBF_NINT16:
 			{
-                const int16_t* pTypedSource = reinterpret_cast<const int16_t*>( pSource );
+                const int16* pTypedSource = reinterpret_cast<const int16*>( pSource );
 				pTypedResult[channel] = (float)(pTypedSource[channel]);
 				pTypedResult[channel] /= 32768.0f;
 				break;
@@ -266,7 +266,7 @@ namespace mu
 
 			case MBF_NUINT16:
 			{
-                const uint16_t* pTypedSource = reinterpret_cast<const uint16_t*>( pSource );
+                const uint16* pTypedSource = reinterpret_cast<const uint16*>( pSource );
 				pTypedResult[channel] = (float)(pTypedSource[channel]);
 				pTypedResult[channel] /= 65535.0f;
 				break;
@@ -274,7 +274,7 @@ namespace mu
 
 			case MBF_NINT8:
 			{
-                const int8_t* pTypedSource = reinterpret_cast<const int8_t*>( pSource );
+                const int8* pTypedSource = reinterpret_cast<const int8*>( pSource );
 				pTypedResult[channel] = (float)(pTypedSource[channel]);
 				pTypedResult[channel] /= 128.0f;
 				break;
@@ -282,7 +282,7 @@ namespace mu
 
 			case MBF_NUINT8:
 			{
-                const uint8_t* pTypedSource = reinterpret_cast<const uint8_t*>( pSource );
+                const uint8* pTypedSource = reinterpret_cast<const uint8*>( pSource );
 				pTypedResult[channel] = (float)(pTypedSource[channel]);
 				pTypedResult[channel] /= 255.0f;
 				break;
@@ -291,7 +291,7 @@ namespace mu
             case MBF_PACKEDDIR8:
             case MBF_PACKEDDIR8_W_TANGENTSIGN:
             {
-                const uint8_t* pTypedSource = reinterpret_cast<const uint8_t*>( pSource );
+                const uint8* pTypedSource = reinterpret_cast<const uint8*>( pSource );
                 pTypedResult[channel] = (float)(pTypedSource[channel]);
                 pTypedResult[channel] /= 127.5f;
                 pTypedResult[channel] -= 1.0f;
@@ -301,7 +301,7 @@ namespace mu
             case MBF_PACKEDDIRS8:
             case MBF_PACKEDDIRS8_W_TANGENTSIGN:
             {
-                const int8_t* pTypedSource = reinterpret_cast<const int8_t*>( pSource );
+                const int8* pTypedSource = reinterpret_cast<const int8*>( pSource );
                 pTypedResult[channel] = (float)(pTypedSource[channel]);
                 pTypedResult[channel] /= 127.5f;
                 break;
@@ -337,42 +337,42 @@ namespace mu
 
 			case MBF_UINT32:
 			{
-                const uint32_t* pTypedSource = reinterpret_cast<const uint32_t*>( pSource );
+                const uint32* pTypedSource = reinterpret_cast<const uint32*>( pSource );
 				pTypedResult[channel] = floatToHalf((float)(pTypedSource[channel]));
 				break;
 			}
 
 			case MBF_INT32:
 			{
-                const int32_t* pTypedSource = reinterpret_cast<const int32_t*>( pSource );
+                const int32* pTypedSource = reinterpret_cast<const int32*>( pSource );
 				pTypedResult[channel] = floatToHalf((float)(pTypedSource[channel]));
 				break;
 			}
 
 			case MBF_UINT16:
 			{
-                const uint16_t* pTypedSource = reinterpret_cast<const uint16_t*>( pSource );
+                const uint16* pTypedSource = reinterpret_cast<const uint16*>( pSource );
 				pTypedResult[channel] = floatToHalf((float)(pTypedSource[channel]));
 				break;
 			}
 
 			case MBF_INT16:
 			{
-                const int16_t* pTypedSource = reinterpret_cast<const int16_t*>( pSource );
+                const int16* pTypedSource = reinterpret_cast<const int16*>( pSource );
 				pTypedResult[channel] = floatToHalf((float)(pTypedSource[channel]));
 				break;
 			}
 
 			case MBF_UINT8:
 			{
-                const uint8_t* pTypedSource = reinterpret_cast<const uint8_t*>( pSource );
+                const uint8* pTypedSource = reinterpret_cast<const uint8*>( pSource );
 				pTypedResult[channel] = floatToHalf((float)(pTypedSource[channel]));
 				break;
 			}
 
 			case MBF_INT8:
 			{
-                const int8_t* pTypedSource = reinterpret_cast<const int8_t*>( pSource );
+                const int8* pTypedSource = reinterpret_cast<const int8*>( pSource );
 				pTypedResult[channel] = floatToHalf((float)(pTypedSource[channel]));
 				break;
 			}
@@ -387,19 +387,19 @@ namespace mu
 		//-----------------------------------------------------------------------------------------
 		case MBF_UINT8:
 		{
-            uint8_t* pTypedResult = reinterpret_cast<uint8_t*>( pResult );
+            uint8* pTypedResult = reinterpret_cast<uint8*>( pResult );
 
 			switch ( sourceFormat )
 			{
 			case MBF_FLOAT32:
 			{
 				const float* pTypedSource = reinterpret_cast<const float*>( pSource );
-                pTypedResult[channel] = (uint8_t)
-                        std::min<uint32_t>(
+                pTypedResult[channel] = (uint8)
+                        std::min<uint32>(
 							0xFF,
-                            std::max<uint32_t>(
+                            std::max<uint32>(
 								0,
-                                (uint32_t)(pTypedSource[channel])
+                                (uint32)(pTypedSource[channel])
 								)
 							);
 				break;
@@ -408,12 +408,12 @@ namespace mu
 			case MBF_FLOAT16:
 			{
 				const float16* pTypedSource = reinterpret_cast<const float16*>( pSource );
-                pTypedResult[channel] = (uint8_t)
-                        std::min<uint32_t>(
+                pTypedResult[channel] = (uint8)
+                        std::min<uint32>(
 							0xFF,
-                            std::max<uint32_t>(
+                            std::max<uint32>(
 								0,
-                                (uint32_t)halfToFloat(pTypedSource[channel])
+                                (uint32)halfToFloat(pTypedSource[channel])
 								)
 							);
 				break;
@@ -421,25 +421,25 @@ namespace mu
 
 			case MBF_INT8:
 			{
-                const int8_t* pTypedSource = reinterpret_cast<const int8_t*>( pSource );
-                pTypedResult[channel] =  (uint8_t)std::max<int8_t>( 0, pTypedSource[channel] );
+                const int8* pTypedSource = reinterpret_cast<const int8*>( pSource );
+                pTypedResult[channel] =  (uint8)std::max<int8>( 0, pTypedSource[channel] );
 				break;
 			}
 
 			case MBF_UINT8:
 			{
-                const uint8_t* pTypedSource = reinterpret_cast<const uint8_t*>( pSource );
+                const uint8* pTypedSource = reinterpret_cast<const uint8*>( pSource );
 				pTypedResult[channel] = pTypedSource[channel];
 				break;
 			}
 
 			case MBF_INT16:
 			{
-                const int16_t* pTypedSource = reinterpret_cast<const int16_t*>( pSource );
-                pTypedResult[channel] =  (uint8_t)
-                        std::min<int16_t>(
+                const int16* pTypedSource = reinterpret_cast<const int16*>( pSource );
+                pTypedResult[channel] =  (uint8)
+                        std::min<int16>(
 							0xFF,
-                            std::max<int16_t>(
+                            std::max<int16>(
 								0,
 								pTypedSource[channel]
 								)
@@ -450,11 +450,11 @@ namespace mu
 			case MBF_UINT16:
 			{
 				// Clamp
-                const uint16_t* pTypedSource = reinterpret_cast<const uint16_t*>( pSource );
-                pTypedResult[channel] = (uint8_t)
-                        std::min<uint16_t>(
+                const uint16* pTypedSource = reinterpret_cast<const uint16*>( pSource );
+                pTypedResult[channel] = (uint8)
+                        std::min<uint16>(
 							0xFF,
-                            std::max<uint16_t>(
+                            std::max<uint16>(
 								0,
 								pTypedSource[channel]
 								)
@@ -464,11 +464,11 @@ namespace mu
 
 			case MBF_INT32:
 			{
-                const int32_t* pTypedSource = reinterpret_cast<const int32_t*>( pSource );
-                pTypedResult[channel] =  (uint8_t)
-                        std::min<int32_t>(
+                const int32* pTypedSource = reinterpret_cast<const int32*>( pSource );
+                pTypedResult[channel] =  (uint8)
+                        std::min<int32>(
 							0xFF,
-                            std::max<int32_t>(
+                            std::max<int32>(
 								0,
 								pTypedSource[channel]
 								)
@@ -479,11 +479,11 @@ namespace mu
 			case MBF_UINT32:
 			{
 				// Clamp
-                const uint32_t* pTypedSource = reinterpret_cast<const uint32_t*>( pSource );
-                pTypedResult[channel] = (uint8_t)
-                        std::min<uint32_t>(
+                const uint32* pTypedSource = reinterpret_cast<const uint32*>( pSource );
+                pTypedResult[channel] = (uint8)
+                        std::min<uint32>(
 							0xFF,
-                            std::max<uint32_t>(
+                            std::max<uint32>(
 								0,
 								pTypedSource[channel]
 								)
@@ -509,19 +509,19 @@ namespace mu
 
 		case MBF_UINT16:
 		{
-            uint16_t* pTypedResult = reinterpret_cast<uint16_t*>( pResult );
+            uint16* pTypedResult = reinterpret_cast<uint16*>( pResult );
 
 			switch ( sourceFormat )
 			{
 			case MBF_FLOAT32:
 			{
 				const float* pTypedSource = reinterpret_cast<const float*>( pSource );
-                pTypedResult[channel] = (uint16_t)
-                        std::min<uint32_t>(
+                pTypedResult[channel] = (uint16)
+                        std::min<uint32>(
 							0xFFFF,
-                            std::max<uint32_t>(
+                            std::max<uint32>(
 								0,
-                                (uint32_t)(pTypedSource[channel])
+                                (uint32)(pTypedSource[channel])
 								)
 							);
 				break;
@@ -530,12 +530,12 @@ namespace mu
 			case MBF_FLOAT16:
 			{
 				const float16* pTypedSource = reinterpret_cast<const float16*>( pSource );
-                pTypedResult[channel] = (uint16_t)
-                        std::min<uint32_t>(
+                pTypedResult[channel] = (uint16)
+                        std::min<uint32>(
 							0xFF,
-                            std::max<uint32_t>(
+                            std::max<uint32>(
 								0,
-                                (uint32_t)halfToFloat(pTypedSource[channel])
+                                (uint32)halfToFloat(pTypedSource[channel])
 								)
 							);
 				break;
@@ -543,40 +543,40 @@ namespace mu
 
 			case MBF_UINT8:
 			{
-                const uint8_t* pTypedSource = reinterpret_cast<const uint8_t*>( pSource );
+                const uint8* pTypedSource = reinterpret_cast<const uint8*>( pSource );
 				pTypedResult[channel] =  pTypedSource[channel];
 				break;
 			}
 
 			case MBF_INT8:
 			{
-                const int8_t* pTypedSource = reinterpret_cast<const int8_t*>( pSource );
-                pTypedResult[channel] =  (uint16_t)std::max<int8_t>( 0, pTypedSource[channel] );
+                const int8* pTypedSource = reinterpret_cast<const int8*>( pSource );
+                pTypedResult[channel] =  (uint16)std::max<int8>( 0, pTypedSource[channel] );
 				break;
 			}
 
 			case MBF_UINT16:
 			{
-                const uint16_t* pTypedSource = reinterpret_cast<const uint16_t*>( pSource );
+                const uint16* pTypedSource = reinterpret_cast<const uint16*>( pSource );
 				pTypedResult[channel] =  pTypedSource[channel];
 				break;
 			}
 
 			case MBF_INT16:
 			{
-                const int16_t* pTypedSource = reinterpret_cast<const int16_t*>( pSource );
-                pTypedResult[channel] =  (uint16_t)std::max<int16_t>( 0, pTypedSource[channel] );
+                const int16* pTypedSource = reinterpret_cast<const int16*>( pSource );
+                pTypedResult[channel] =  (uint16)std::max<int16>( 0, pTypedSource[channel] );
 				break;
 			}
 
 			case MBF_UINT32:
 			{
 				// Clamp
-                const uint32_t* pTypedSource = reinterpret_cast<const uint32_t*>( pSource );
-                pTypedResult[channel] = (uint16_t)
-                        std::min<uint32_t>(
+                const uint32* pTypedSource = reinterpret_cast<const uint32*>( pSource );
+                pTypedResult[channel] = (uint16)
+                        std::min<uint32>(
 							0xFFFF,
-                            std::max<uint32_t>(
+                            std::max<uint32>(
 								0,
 								pTypedSource[channel]
 								)
@@ -587,11 +587,11 @@ namespace mu
 			case MBF_INT32:
 			{
 				// Clamp
-                const int32_t* pTypedSource = reinterpret_cast<const int32_t*>( pSource );
-                pTypedResult[channel] = (uint16_t)
-                        std::min<int32_t>(
+                const int32* pTypedSource = reinterpret_cast<const int32*>( pSource );
+                pTypedResult[channel] = (uint16)
+                        std::min<int32>(
 							0xFFFF,
-                            std::max<int32_t>(
+                            std::max<int32>(
 								0,
 								pTypedSource[channel]
 								)
@@ -617,7 +617,7 @@ namespace mu
 
 		case MBF_UINT32:
 		{
-            uint32_t* pTypedResult = reinterpret_cast<uint32_t*>( pResult );
+            uint32* pTypedResult = reinterpret_cast<uint32*>( pResult );
 
 			switch ( sourceFormat )
 			{
@@ -625,11 +625,11 @@ namespace mu
 			{
 				const float* pTypedSource = reinterpret_cast<const float*>( pSource );
 				pTypedResult[channel] =
-                        std::min<uint32_t>(
+                        std::min<uint32>(
 							0xFFFF,
-                            std::max<uint32_t>(
+                            std::max<uint32>(
 								0,
-                                (uint32_t)(pTypedSource[channel])
+                                (uint32)(pTypedSource[channel])
 								)
 							);
 				break;
@@ -639,11 +639,11 @@ namespace mu
 			{
 				const float16* pTypedSource = reinterpret_cast<const float16*>( pSource );
 				pTypedResult[channel] =
-                        std::min<uint32_t>(
+                        std::min<uint32>(
 							0xFF,
-                            std::max<uint32_t>(
+                            std::max<uint32>(
 								0,
-                                (uint32_t)halfToFloat(pTypedSource[channel])
+                                (uint32)halfToFloat(pTypedSource[channel])
 								)
 							);
 				break;
@@ -651,35 +651,35 @@ namespace mu
 
 			case MBF_UINT8:
 			{
-                const uint8_t* pTypedSource = reinterpret_cast<const uint8_t*>( pSource );
+                const uint8* pTypedSource = reinterpret_cast<const uint8*>( pSource );
 				pTypedResult[channel] =  pTypedSource[channel];
 				break;
 			}
 
 			case MBF_INT8:
 			{
-                const int8_t* pTypedSource = reinterpret_cast<const int8_t*>( pSource );
-                pTypedResult[channel] =  (uint16_t)std::max<int8_t>( 0, pTypedSource[channel] );
+                const int8* pTypedSource = reinterpret_cast<const int8*>( pSource );
+                pTypedResult[channel] =  (uint16)std::max<int8>( 0, pTypedSource[channel] );
 				break;
 			}
 
 			case MBF_UINT16:
 			{
-                const uint16_t* pTypedSource = reinterpret_cast<const uint16_t*>( pSource );
+                const uint16* pTypedSource = reinterpret_cast<const uint16*>( pSource );
 				pTypedResult[channel] =  pTypedSource[channel];
 				break;
 			}
 
 			case MBF_INT16:
 			{
-                const int16_t* pTypedSource = reinterpret_cast<const int16_t*>( pSource );
-                pTypedResult[channel] =  (uint16_t)std::max<int16_t>( 0, pTypedSource[channel] );
+                const int16* pTypedSource = reinterpret_cast<const int16*>( pSource );
+                pTypedResult[channel] =  (uint16)std::max<int16>( 0, pTypedSource[channel] );
 				break;
 			}
 
 			case MBF_UINT32:
 			{
-                const uint32_t* pTypedSource = reinterpret_cast<const uint32_t*>( pSource );
+                const uint32* pTypedSource = reinterpret_cast<const uint32*>( pSource );
 				pTypedResult[channel] = pTypedSource[channel];
 				break;
 			}
@@ -687,8 +687,8 @@ namespace mu
 			case MBF_INT32:
 			{
 				// Clamp
-                const int32_t* pTypedSource = reinterpret_cast<const int32_t*>( pSource );
-                pTypedResult[channel] = (uint32_t)std::max<int32_t>( 0, pTypedSource[channel] );
+                const int32* pTypedSource = reinterpret_cast<const int32*>( pSource );
+                pTypedResult[channel] = (uint32)std::max<int32>( 0, pTypedSource[channel] );
 				break;
 			}
 
@@ -710,19 +710,19 @@ namespace mu
 
 		case MBF_INT8:
 		{
-            int8_t* pTypedResult = reinterpret_cast<int8_t*>( pResult );
+            int8* pTypedResult = reinterpret_cast<int8*>( pResult );
 
 			switch ( sourceFormat )
 			{
 			case MBF_FLOAT32:
 			{
 				const float* pTypedSource = reinterpret_cast<const float*>( pSource );
-                pTypedResult[channel] = (int8_t)
-                        std::min<int32_t>(
+                pTypedResult[channel] = (int8)
+                        std::min<int32>(
 							127,
-                            std::max<int32_t>(
+                            std::max<int32>(
 								-128,
-                                (int32_t)(pTypedSource[channel])
+                                (int32)(pTypedSource[channel])
 								)
 							);
 				break;
@@ -731,12 +731,12 @@ namespace mu
 			case MBF_FLOAT16:
 			{
 				const float16* pTypedSource = reinterpret_cast<const float16*>( pSource );
-                pTypedResult[channel] = (int8_t)
-                        std::min<int32_t>(
+                pTypedResult[channel] = (int8)
+                        std::min<int32>(
 							127,
-                            std::max<int32_t>(
+                            std::max<int32>(
 								-128,
-                                (int32_t)halfToFloat(pTypedSource[channel])
+                                (int32)halfToFloat(pTypedSource[channel])
 								)
 							);
 				break;
@@ -744,7 +744,7 @@ namespace mu
 
 			case MBF_INT8:
 			{
-                const int8_t* pTypedSource = reinterpret_cast<const int8_t*>( pSource );
+                const int8* pTypedSource = reinterpret_cast<const int8*>( pSource );
 				pTypedResult[channel] = pTypedSource[channel];
 				break;
 			}
@@ -767,19 +767,19 @@ namespace mu
 
 		case MBF_INT16:
 		{
-            int16_t* pTypedResult = reinterpret_cast<int16_t*>( pResult );
+            int16* pTypedResult = reinterpret_cast<int16*>( pResult );
 
 			switch ( sourceFormat )
 			{
 			case MBF_FLOAT32:
 			{
 				const float* pTypedSource = reinterpret_cast<const float*>( pSource );
-                pTypedResult[channel] = (int16_t)
-                        std::min<int32_t>(
+                pTypedResult[channel] = (int16)
+                        std::min<int32>(
 							32767,
-                            std::max<int32_t>(
+                            std::max<int32>(
 								-32768,
-                                (int32_t)(pTypedSource[channel])
+                                (int32)(pTypedSource[channel])
 								)
 							);
 				break;
@@ -788,12 +788,12 @@ namespace mu
 			case MBF_FLOAT16:
 			{
 				const float16* pTypedSource = reinterpret_cast<const float16*>( pSource );
-                pTypedResult[channel] = (int16_t)
-                        std::min<int32_t>(
+                pTypedResult[channel] = (int16)
+                        std::min<int32>(
 							32767,
-                            std::max<int32_t>(
+                            std::max<int32>(
 								-32768,
-                                (int32_t)halfToFloat(pTypedSource[channel])
+                                (int32)halfToFloat(pTypedSource[channel])
 								)
 							);
 				break;
@@ -801,35 +801,35 @@ namespace mu
 
 			case MBF_INT8:
 			{
-                const int8_t* pTypedSource = reinterpret_cast<const int8_t*>( pSource );
-                pTypedResult[channel] = (int16_t)pTypedSource[channel];
+                const int8* pTypedSource = reinterpret_cast<const int8*>( pSource );
+                pTypedResult[channel] = (int16)pTypedSource[channel];
 				break;
 			}
 
 			case MBF_UINT8:
 			{
-                const uint8_t* pTypedSource = reinterpret_cast<const uint8_t*>( pSource );
-                pTypedResult[channel] = (int16_t)pTypedSource[channel];
+                const uint8* pTypedSource = reinterpret_cast<const uint8*>( pSource );
+                pTypedResult[channel] = (int16)pTypedSource[channel];
 				break;
 			}
 
 			case MBF_UINT16:
 			{
-                const uint16_t* pTypedSource = reinterpret_cast<const uint16_t*>( pSource );
-                pTypedResult[channel] = (int16_t)
-                        std::min<int32_t>(
-                            32767, (int32_t)pTypedSource[channel]
+                const uint16* pTypedSource = reinterpret_cast<const uint16*>( pSource );
+                pTypedResult[channel] = (int16)
+                        std::min<int32>(
+                            32767, (int32)pTypedSource[channel]
 							);
 				break;
 			}
 
 			case MBF_INT32:
 			{
-                const int32_t* pTypedSource = reinterpret_cast<const int32_t*>( pSource );
-                pTypedResult[channel] = (int16_t)
-                        std::min<int32_t>(
+                const int32* pTypedSource = reinterpret_cast<const int32*>( pSource );
+                pTypedResult[channel] = (int16)
+                        std::min<int32>(
 							32767,
-                            std::max<int32_t>(
+                            std::max<int32>(
 								-32768,
 								pTypedSource[channel]
 								)
@@ -839,9 +839,9 @@ namespace mu
 
 			case MBF_UINT32:
 			{
-                const uint32_t* pTypedSource = reinterpret_cast<const uint32_t*>( pSource );
-                pTypedResult[channel] = (int16_t)
-                        std::min<int32_t>(
+                const uint32* pTypedSource = reinterpret_cast<const uint32*>( pSource );
+                pTypedResult[channel] = (int16)
+                        std::min<int32>(
 							32767, pTypedSource[channel]
 							);
 				break;
@@ -865,62 +865,62 @@ namespace mu
 
 		case MBF_INT32:
 		{
-            int32_t* pTypedResult = reinterpret_cast<int32_t*>( pResult );
+            int32* pTypedResult = reinterpret_cast<int32*>( pResult );
 
 			switch ( sourceFormat )
 			{
 			case MBF_FLOAT32:
 			{
 				const float* pTypedSource = reinterpret_cast<const float*>( pSource );
-                pTypedResult[channel] = (int32_t)(pTypedSource[channel]);
+                pTypedResult[channel] = (int32)(pTypedSource[channel]);
 				break;
 			}
 
 			case MBF_FLOAT16:
 			{
 				const float16* pTypedSource = reinterpret_cast<const float16*>( pSource );
-                pTypedResult[channel] = (int32_t)halfToFloat(pTypedSource[channel]);
+                pTypedResult[channel] = (int32)halfToFloat(pTypedSource[channel]);
 				break;
 			}
 
 			case MBF_INT8:
 			{
-                const int8_t* pTypedSource = reinterpret_cast<const int8_t*>( pSource );
-                pTypedResult[channel] = (int32_t)pTypedSource[channel];
+                const int8* pTypedSource = reinterpret_cast<const int8*>( pSource );
+                pTypedResult[channel] = (int32)pTypedSource[channel];
 				break;
 			}
 
 			case MBF_UINT8:
 			{
-                const uint8_t* pTypedSource = reinterpret_cast<const uint8_t*>( pSource );
-                pTypedResult[channel] = (int32_t)pTypedSource[channel];
+                const uint8* pTypedSource = reinterpret_cast<const uint8*>( pSource );
+                pTypedResult[channel] = (int32)pTypedSource[channel];
 				break;
 			}
 
 			case MBF_INT16:
 			{
-                const int16_t* pTypedSource = reinterpret_cast<const int16_t*>( pSource );
-                pTypedResult[channel] = (int32_t)pTypedSource[channel];
+                const int16* pTypedSource = reinterpret_cast<const int16*>( pSource );
+                pTypedResult[channel] = (int32)pTypedSource[channel];
 				break;
 			}
 
 			case MBF_UINT16:
 			{
-                const uint16_t* pTypedSource = reinterpret_cast<const uint16_t*>( pSource );
-                pTypedResult[channel] = (int32_t)pTypedSource[channel];
+                const uint16* pTypedSource = reinterpret_cast<const uint16*>( pSource );
+                pTypedResult[channel] = (int32)pTypedSource[channel];
 				break;
 			}
 
 			case MBF_UINT32:
 			{
-                const uint32_t* pTypedSource = reinterpret_cast<const uint32_t*>( pSource );
-                pTypedResult[channel] = (int32_t)pTypedSource[channel];
+                const uint32* pTypedSource = reinterpret_cast<const uint32*>( pSource );
+                pTypedResult[channel] = (int32)pTypedSource[channel];
 				break;
 			}
 
 			case MBF_INT32:
 			{
-                const int32_t* pTypedSource = reinterpret_cast<const int32_t*>( pSource );
+                const int32* pTypedSource = reinterpret_cast<const int32*>( pSource );
                 pTypedResult[channel] = pTypedSource[channel];
 				break;
 			}
@@ -944,20 +944,20 @@ namespace mu
 		//-----------------------------------------------------------------------------------------
 		case MBF_NUINT8:
 		{
-            uint8_t* pTypedResult = reinterpret_cast<uint8_t*>( pResult );
+            uint8* pTypedResult = reinterpret_cast<uint8*>( pResult );
 
 			switch ( sourceFormat )
 			{
 			case MBF_NUINT8:
 			{
-				auto pTypedSource = reinterpret_cast<const uint8_t*>(pSource);
+				auto pTypedSource = reinterpret_cast<const uint8*>(pSource);
 				pTypedResult[channel] = pTypedSource[channel];
 				break;
 			}
 
 			case MBF_NUINT16:
 			{
-				const uint16_t* pTypedSource = reinterpret_cast<const uint16_t*>(pSource);
+				const uint16* pTypedSource = reinterpret_cast<const uint16*>(pSource);
 				pTypedResult[channel] = pTypedSource[channel] / (65535 / 255);
 				break;
 			}
@@ -965,12 +965,12 @@ namespace mu
 			case MBF_FLOAT32:
 			{
 				const float* pTypedSource = reinterpret_cast<const float*>(pSource);
-				pTypedResult[channel] = (uint8_t)
-					std::min<uint32_t>(
+				pTypedResult[channel] = (uint8)
+					std::min<uint32>(
 						0xFF,
-						std::max<uint32_t>(
+						std::max<uint32>(
 							0,
-							(uint32_t)(((float)0xFF)*pTypedSource[channel] + 0.5f)
+							(uint32)(((float)0xFF)*pTypedSource[channel] + 0.5f)
 							)
 						);
 				break;
@@ -979,12 +979,12 @@ namespace mu
 			case MBF_FLOAT16:
 			{
 				const float16* pTypedSource = reinterpret_cast<const float16*>( pSource );
-                pTypedResult[channel] = (uint8_t)
-                        std::min<uint32_t>(
+                pTypedResult[channel] = (uint8)
+                        std::min<uint32>(
 							0xFF,
-                            std::max<uint32_t>(
+                            std::max<uint32>(
 								0,
-                                (uint32_t)(((float)0xFF)*halfToFloat(pTypedSource[channel])+0.5f)
+                                (uint32)(((float)0xFF)*halfToFloat(pTypedSource[channel])+0.5f)
 								)
 							);
 				break;
@@ -1009,20 +1009,20 @@ namespace mu
 
 		case MBF_NUINT16:
 		{
-            uint16_t* pTypedResult = reinterpret_cast<uint16_t*>( pResult );
+            uint16* pTypedResult = reinterpret_cast<uint16*>( pResult );
 
 			switch ( sourceFormat )
 			{
 			case MBF_NUINT16:
 			{
-				auto pTypedSource = reinterpret_cast<const uint16_t*>(pSource);
+				auto pTypedSource = reinterpret_cast<const uint16*>(pSource);
 				pTypedResult[channel] = pTypedSource[channel];
 				break;
 			}
 
 			case MBF_NUINT8:
 			{
-				const uint8_t* pTypedSource = reinterpret_cast<const uint8_t*>(pSource);
+				const uint8* pTypedSource = reinterpret_cast<const uint8*>(pSource);
 				pTypedResult[channel] = pTypedSource[channel] * (65535 / 255);
 				break;
 			}
@@ -1030,12 +1030,12 @@ namespace mu
 			case MBF_FLOAT32:
 			{
 				const float* pTypedSource = reinterpret_cast<const float*>( pSource );
-                pTypedResult[channel] = (uint16_t)
-                        std::min<uint32_t>(
+                pTypedResult[channel] = (uint16)
+                        std::min<uint32>(
 							0xFFFF,
-                            std::max<uint32_t>(
+                            std::max<uint32>(
 								0,
-                                (uint32_t)(((float)0xFFFF)*pTypedSource[channel]+0.5f)
+                                (uint32)(((float)0xFFFF)*pTypedSource[channel]+0.5f)
 								)
 							);
 				break;
@@ -1044,12 +1044,12 @@ namespace mu
 			case MBF_FLOAT16:
 			{
 				const float16* pTypedSource = reinterpret_cast<const float16*>( pSource );
-                pTypedResult[channel] = (uint16_t)
-                        std::min<uint32_t>(
+                pTypedResult[channel] = (uint16)
+                        std::min<uint32>(
 							0xFFFF,
-                            std::max<uint32_t>(
+                            std::max<uint32>(
 								0,
-                                (uint32_t)(((float)0xFFFF)*halfToFloat(pTypedSource[channel])+0.5f)
+                                (uint32)(((float)0xFFFF)*halfToFloat(pTypedSource[channel])+0.5f)
 								)
 							);
 				break;
@@ -1074,19 +1074,19 @@ namespace mu
 
 		case MBF_NUINT32:
 		{
-            uint32_t* pTypedResult = reinterpret_cast<uint32_t*>( pResult );
+            uint32* pTypedResult = reinterpret_cast<uint32*>( pResult );
 
 			switch ( sourceFormat )
 			{
 			case MBF_FLOAT32:
 			{
 				const float* pTypedSource = reinterpret_cast<const float*>( pSource );
-                pTypedResult[channel] = (uint32_t)
-                        std::min<uint32_t>(
+                pTypedResult[channel] = (uint32)
+                        std::min<uint32>(
 							0xFFFFFFFF,
-                            std::max<uint32_t>(
+                            std::max<uint32>(
 								0,
-                                (uint32_t)(((float)0xFFFFFFFF)*pTypedSource[channel])
+                                (uint32)(((float)0xFFFFFFFF)*pTypedSource[channel])
 								)
 							);
 				break;
@@ -1095,12 +1095,12 @@ namespace mu
 			case MBF_FLOAT16:
 			{
 				const float16* pTypedSource = reinterpret_cast<const float16*>( pSource );
-                pTypedResult[channel] = (uint32_t)
-                        std::min<uint32_t>(
+                pTypedResult[channel] = (uint32)
+                        std::min<uint32>(
 							0xFFFFFFFF,
-                            std::max<uint32_t>(
+                            std::max<uint32>(
 								0,
-                                (uint32_t)(((float)0xFFFFFFFF)*halfToFloat(pTypedSource[channel])+0.5f)
+                                (uint32)(((float)0xFFFFFFFF)*halfToFloat(pTypedSource[channel])+0.5f)
 								)
 							);
 				break;
@@ -1125,19 +1125,19 @@ namespace mu
 
 		case MBF_NINT8:
 		{
-            int8_t* pTypedResult = reinterpret_cast<int8_t*>( pResult );
+            int8* pTypedResult = reinterpret_cast<int8*>( pResult );
 
 			switch ( sourceFormat )
 			{
 			case MBF_FLOAT32:
 			{
 				const float* pTypedSource = reinterpret_cast<const float*>( pSource );
-                pTypedResult[channel] = (int8_t)
-                        std::min<int32_t>(
+                pTypedResult[channel] = (int8)
+                        std::min<int32>(
 							127,
-                            std::max<int32_t>(
+                            std::max<int32>(
 								-128,
-                                (int32_t)(128.0f*pTypedSource[channel]+0.5f)
+                                (int32)(128.0f*pTypedSource[channel]+0.5f)
 								)
 							);
 				break;
@@ -1146,12 +1146,12 @@ namespace mu
 			case MBF_FLOAT16:
 			{
 				const float16* pTypedSource = reinterpret_cast<const float16*>( pSource );
-                pTypedResult[channel] = (int8_t)
-                        std::min<int32_t>(
+                pTypedResult[channel] = (int8)
+                        std::min<int32>(
 							127,
-                            std::max<int32_t>(
+                            std::max<int32>(
 								-128,
-                                (int32_t)(128.0f*halfToFloat(pTypedSource[channel])+0.5f)
+                                (int32)(128.0f*halfToFloat(pTypedSource[channel])+0.5f)
 								)
 							);
 				break;
@@ -1176,19 +1176,19 @@ namespace mu
 
 		case MBF_NINT16:
 		{
-            int16_t* pTypedResult = reinterpret_cast<int16_t*>( pResult );
+            int16* pTypedResult = reinterpret_cast<int16*>( pResult );
 
 			switch ( sourceFormat )
 			{
 			case MBF_FLOAT32:
 			{
 				const float* pTypedSource = reinterpret_cast<const float*>( pSource );
-                pTypedResult[channel] = (int16_t)
-                        std::min<int32_t>(
+                pTypedResult[channel] = (int16)
+                        std::min<int32>(
 							32767,
-                            std::max<int32_t>(
+                            std::max<int32>(
 								-32768,
-                                (int32_t)(32768.0f*pTypedSource[channel]+0.5f)
+                                (int32)(32768.0f*pTypedSource[channel]+0.5f)
 								)
 							);
 				break;
@@ -1197,12 +1197,12 @@ namespace mu
 			case MBF_FLOAT16:
 			{
 				const float16* pTypedSource = reinterpret_cast<const float16*>( pSource );
-                pTypedResult[channel] = (int16_t)
-                        std::min<int32_t>(
+                pTypedResult[channel] = (int16)
+                        std::min<int32>(
 							32767,
-                            std::max<int32_t>(
+                            std::max<int32>(
 								-32768,
-                                (int32_t)(32768.0f*halfToFloat(pTypedSource[channel])+0.5f)
+                                (int32)(32768.0f*halfToFloat(pTypedSource[channel])+0.5f)
 								)
 							);
 				break;
@@ -1228,21 +1228,21 @@ namespace mu
 
 		case MBF_NINT32:
 		{
-            int32_t* pTypedResult = reinterpret_cast<int32_t*>( pResult );
+            int32* pTypedResult = reinterpret_cast<int32*>( pResult );
 
 			switch ( sourceFormat )
 			{
 			case MBF_FLOAT32:
 			{
 				const float* pTypedSource = reinterpret_cast<const float*>( pSource );
-                pTypedResult[channel] = (int32_t)(2147483648.0f*pTypedSource[channel]+0.5f);
+                pTypedResult[channel] = (int32)(2147483648.0f*pTypedSource[channel]+0.5f);
 				break;
 			}
 
 			case MBF_FLOAT16:
 			{
 				const float16* pTypedSource = reinterpret_cast<const float16*>( pSource );
-                pTypedResult[channel] = (int32_t)(2147483648.0f*halfToFloat(pTypedSource[channel])+0.5f);
+                pTypedResult[channel] = (int32)(2147483648.0f*halfToFloat(pTypedSource[channel])+0.5f);
 				break;
 			}
 
@@ -1266,15 +1266,15 @@ namespace mu
         case MBF_PACKEDDIR8:
         case MBF_PACKEDDIR8_W_TANGENTSIGN:
         {
-            uint8_t* pTypedResult = reinterpret_cast<uint8_t*>( pResult );
+            uint8* pTypedResult = reinterpret_cast<uint8*>( pResult );
 
             switch ( sourceFormat )
             {
             case MBF_PACKEDDIR8:
             case MBF_PACKEDDIR8_W_TANGENTSIGN:
             {
-                const uint8_t* pTypedSource = reinterpret_cast<const uint8_t*>( pSource );
-                uint8_t source = pTypedSource[channel];
+                const uint8* pTypedSource = reinterpret_cast<const uint8*>( pSource );
+                uint8 source = pTypedSource[channel];
                 pTypedResult[channel] = source;
                 break;
             }
@@ -1285,7 +1285,7 @@ namespace mu
                 float source = pTypedSource[channel];
                 source = (source*0.5f+0.5f)*255.0f;
                 pTypedResult[channel] =
-                        (uint8_t)std::min<float>( 255.0f, std::max<float>( 0.0f, source ) );
+                        (uint8)std::min<float>( 255.0f, std::max<float>( 0.0f, source ) );
                 break;
             }
 
@@ -1299,15 +1299,15 @@ namespace mu
         case MBF_PACKEDDIRS8:
         case MBF_PACKEDDIRS8_W_TANGENTSIGN:
         {
-            int8_t* pTypedResult = reinterpret_cast<int8_t*>( pResult );
+            int8* pTypedResult = reinterpret_cast<int8*>( pResult );
 
             switch ( sourceFormat )
             {
             case MBF_PACKEDDIRS8:
             case MBF_PACKEDDIRS8_W_TANGENTSIGN:
             {
-                const int8_t* pTypedSource = reinterpret_cast<const int8_t*>( pSource );
-                int8_t source = pTypedSource[channel];
+                const int8* pTypedSource = reinterpret_cast<const int8*>( pSource );
+                int8 source = pTypedSource[channel];
                 pTypedResult[channel] = source;
                 break;
             }
@@ -1318,7 +1318,7 @@ namespace mu
                 float source = pTypedSource[channel];
                 source = source*0.5f*255.0f;
                 pTypedResult[channel] =
-                        (int8_t)std::min<float>( 127.0f, std::max<float>( -128.0f, source ) );
+                        (int8)std::min<float>( 127.0f, std::max<float>( -128.0f, source ) );
                 break;
             }
 
diff --git a/Engine/Plugins/Experimental/Mutable/Source/MutableRuntime/Private/MuR/ImageRLE.cpp b/Engine/Plugins/Experimental/Mutable/Source/MutableRuntime/Private/MuR/ImageRLE.cpp
--- a/Engine/Plugins/Experimental/Mutable/Source/MutableRuntime/Private/MuR/ImageRLE.cpp
+++ b/Engine/Plugins/Experimental/Mutable/Source/MutableRuntime/Private/MuR/ImageRLE.cpp
@@ -6,24 +6,22 @@
 #include "Misc/AssertionMacros.h"
 #include "MuR/MemoryPrivate.h"
 
-#include <memory>
-#include <utility>
 
 namespace mu
 {
 
     //---------------------------------------------------------------------------------------------
-    uint32_t CompressRLE_L( int32 width, int32 rows, const uint8_t* pBaseData, uint8_t* destData, uint32_t destDataSize )
+    uint32 CompressRLE_L( int32 width, int32 rows, const uint8* pBaseData, uint8* destData, uint32 destDataSize )
     {
-		const uint8_t* InitialBaseData = pBaseData;
+		const uint8* InitialBaseData = pBaseData;
 
-        uint8_t* rle = destData;
+        uint8* rle = destData;
 
         size_t maxSize = destDataSize;
 
-        // The first uint32_t will be the total mip size.
+        // The first uint32 will be the total mip size.
         // Then there is an offset from the initial data pointer for each line.
-        uint32_t offset = sizeof( uint32_t ) * ( rows + 1 );
+        uint32 offset = sizeof( uint32 ) * ( rows + 1 );
 
         // Could happen in an image of 1x100, for example.
         if ( offset >= maxSize )
@@ -33,16 +31,16 @@ namespace mu
 
         for ( int r = 0; r < rows; ++r )
         {
-            uint32_t* pOffset = (uint32_t*)&rle[sizeof( uint32_t ) * ( r + 1 )];
-            FMemory::Memmove( pOffset, &offset, sizeof( uint32_t ) );
+            uint32* pOffset = (uint32*)&rle[sizeof( uint32 ) * ( r + 1 )];
+            FMemory::Memmove( pOffset, &offset, sizeof( uint32 ) );
 
-            const uint8_t* pBaseRowEnd = pBaseData + width;
+            const uint8* pBaseRowEnd = pBaseData + width;
             while ( pBaseData != pBaseRowEnd )
             {
                 // Count equal pixels
-                uint8_t equalPixel = *pBaseData;
+                uint8 equalPixel = *pBaseData;
 
-                uint16_t equal = 0;
+                uint16 equal = 0;
                 while ( pBaseData != pBaseRowEnd && equal < 65535 && pBaseData[0] == equalPixel )
                 {
                     pBaseData++;
@@ -50,8 +48,8 @@ namespace mu
                 }
 
                 // Count different pixels
-                uint8_t different = 0;
-                const uint8_t* pDifferentPixels = pBaseData;
+                uint8 different = 0;
+                const uint8* pDifferentPixels = pBaseData;
                 while ( pBaseData < pBaseRowEnd - 1 && different < 255 &&
                         // Last in the row, or different from next
                         ( pBaseData == pBaseRowEnd || pBaseData[0] != pBaseData[1] ) )
@@ -65,7 +63,7 @@ namespace mu
                 {
                     return 0;
                 }
-                FMemory::Memmove( &rle[offset], &equal, sizeof( uint16_t ) );
+                FMemory::Memmove( &rle[offset], &equal, sizeof( uint16 ) );
                 offset += 2;
                 rle[offset] = different;
                 ++offset;
@@ -87,7 +85,7 @@ namespace mu
             }
         }
 
-        uint32_t* pTotalSize = (uint32_t*)rle;
+        uint32* pTotalSize = (uint32*)rle;
         *pTotalSize = offset;
 
 #ifdef MUTABLE_DEBUG_RLE		
@@ -120,7 +118,7 @@ namespace mu
 
 
     //---------------------------------------------------------------------------------------------
-    uint32_t UncompressRLE_L( int32 width, int32 rows, const uint8* pStartBaseData, uint8* pStartDestData )
+    uint32 UncompressRLE_L( int32 width, int32 rows, const uint8* pStartBaseData, uint8* pStartDestData )
     {
 		const uint8* pBaseData = pStartBaseData;
 		uint8* pDestData = pStartDestData;
@@ -137,10 +135,10 @@ namespace mu
                 FMemory::Memmove(&equal, pBaseData, sizeof(uint16));
                 pBaseData += 2;
 
-                uint8_t different = *pBaseData;
+                uint8 different = *pBaseData;
                 ++pBaseData;
 
-                uint8_t equalPixel = *pBaseData;
+                uint8 equalPixel = *pBaseData;
                 ++pBaseData;
 
                 if (equal)
@@ -167,27 +165,27 @@ namespace mu
     }
 
     //---------------------------------------------------------------------------------------------
-    uint32_t CompressRLE_L1( int32 width, int32 rows,
-                             const uint8_t* pBaseData,
-                             uint8_t* destData,
-                             uint32_t destDataSize )
+    uint32 CompressRLE_L1( int32 width, int32 rows,
+                             const uint8* pBaseData,
+                             uint8* destData,
+                             uint32 destDataSize )
     {
         vector<int8_t> rle;
         rle.reserve(  (width * rows) / 2 );
 
-        uint32_t offset = sizeof(uint32_t)*(rows+1);
+        uint32 offset = sizeof(uint32)*(rows+1);
         rle.resize( offset );
 
         for ( int r=0; r<rows; ++r )
         {
-            uint32_t* pOffset = (uint32_t*) &rle[ sizeof(uint32_t)*(r+1) ];
+            uint32* pOffset = (uint32*) &rle[ sizeof(uint32)*(r+1) ];
             *pOffset = offset;
 
-            const uint8_t* pBaseRowEnd = pBaseData + width;
+            const uint8* pBaseRowEnd = pBaseData + width;
             while ( pBaseData!=pBaseRowEnd )
             {
                 // Count 0 pixels
-                uint16_t zeroPixels = 0;
+                uint16 zeroPixels = 0;
                 while ( pBaseData!=pBaseRowEnd
                         && !*pBaseData )
                 {
@@ -196,7 +194,7 @@ namespace mu
                 }
 
                 // Count 1 pixels
-                uint16_t onePixels = 0;
+                uint16 onePixels = 0;
                 while ( pBaseData!=pBaseRowEnd
                         && *pBaseData )
                 {
@@ -206,10 +204,10 @@ namespace mu
 
                 // Copy block
                 rle.resize( rle.size()+4 );
-                FMemory::Memmove(&rle[ offset ], &zeroPixels, sizeof(uint16_t));
+                FMemory::Memmove(&rle[ offset ], &zeroPixels, sizeof(uint16));
                 offset += 2;
 
-                FMemory::Memmove(&rle[ offset ], &onePixels, sizeof(uint16_t));
+                FMemory::Memmove(&rle[ offset ], &onePixels, sizeof(uint16));
                 offset += 2;
             }
         }
@@ -220,7 +218,7 @@ namespace mu
             return 0;
         }
 
-        uint32_t* pTotalSize = (uint32_t*) &rle[ 0 ];
+        uint32* pTotalSize = (uint32*) &rle[ 0 ];
         *pTotalSize = offset;
 
         FMemory::Memmove( destData, &rle[0], offset );
@@ -231,24 +229,24 @@ namespace mu
 
 
     //---------------------------------------------------------------------------------------------
-    uint32_t UncompressRLE_L1( int width, int rows, const uint8_t* pStartBaseData, uint8_t* pDestData )
+    uint32 UncompressRLE_L1( int width, int rows, const uint8* pStartBaseData, uint8* pDestData )
     {
-        const uint8_t* pBaseData = pStartBaseData;
-        pBaseData += sizeof(uint32_t); // Total mip size
-        pBaseData += rows*sizeof(uint32_t);
+        const uint8* pBaseData = pStartBaseData;
+        pBaseData += sizeof(uint32); // Total mip size
+        pBaseData += rows*sizeof(uint32);
 
         for ( int r=0; r<rows; ++r )
         {
-            const uint8_t* pDestRowEnd = pDestData + width;
+            const uint8* pDestRowEnd = pDestData + width;
             while ( pDestData!=pDestRowEnd )
             {
                 // Decode header
-                uint16_t zeroPixels = 0;
-                FMemory::Memmove(&zeroPixels, pBaseData, sizeof(uint16_t));
+                uint16 zeroPixels = 0;
+                FMemory::Memmove(&zeroPixels, pBaseData, sizeof(uint16));
                 pBaseData += 2;
 
-                uint16_t onePixels = 0;
-                FMemory::Memmove(&onePixels, pBaseData, sizeof(uint16_t));
+                uint16 onePixels = 0;
+                FMemory::Memmove(&onePixels, pBaseData, sizeof(uint16));
                 pBaseData += 2;
 
                 if (zeroPixels)
@@ -266,36 +264,36 @@ namespace mu
         }
 
         size_t totalSize = pBaseData-pStartBaseData;
-        check( totalSize==*(uint32_t*)pStartBaseData );
+        check( totalSize==*(uint32*)pStartBaseData );
 
-        return (uint32_t)totalSize;
+        return (uint32)totalSize;
     }
 
 
     //---------------------------------------------------------------------------------------------
     void CompressRLE_RGBA( int width, int rows,
-                           const uint8_t* pBaseDataByte,
+                           const uint8* pBaseDataByte,
                            TArray<uint8>& destData )
     {
         // TODO: Support for compression from compressed data size, like L_RLE formats.
         vector<int8_t> rle;
         rle.reserve(  (width*rows) );
 
-        const uint32_t* pBaseData = (const uint32_t*)pBaseDataByte;
+        const uint32* pBaseData = (const uint32*)pBaseDataByte;
         rle.resize( rows*4 );
-        uint32_t offset = sizeof(uint32_t)*rows;
+        uint32 offset = sizeof(uint32)*rows;
         for ( int r=0; r<rows; ++r )
         {
-            uint32_t* pOffset = (uint32_t*) &rle[ sizeof(uint32_t)*r ];
-            FMemory::Memmove(pOffset, &offset, sizeof(uint32_t));
+            uint32* pOffset = (uint32*) &rle[ sizeof(uint32)*r ];
+            FMemory::Memmove(pOffset, &offset, sizeof(uint32));
 
-            const uint32_t* pBaseRowEnd = pBaseData + width;
+            const uint32* pBaseRowEnd = pBaseData + width;
             while ( pBaseData!=pBaseRowEnd )
             {
                 // Count equal pixels
-                uint32_t equalPixel = *pBaseData;
+                uint32 equalPixel = *pBaseData;
 
-                uint16_t equal = 0;
+                uint16 equal = 0;
                 while ( pBaseData<pBaseRowEnd-3 && equal<65535
                         && pBaseData[0]==equalPixel && pBaseData[1]==equalPixel
                         && pBaseData[2]==equalPixel && pBaseData[3]==equalPixel )
@@ -305,8 +303,8 @@ namespace mu
                 }
 
                 // Count different pixels
-                uint16_t different = 0;
-                const uint32_t* pDifferentPixels = pBaseData;
+                uint16 different = 0;
+                const uint32* pDifferentPixels = pBaseData;
                 while ( pBaseData!=pBaseRowEnd
                         &&
                         different<65535
@@ -325,18 +323,18 @@ namespace mu
 
                 // Copy header
                 rle.resize( rle.size()+8 );
-                FMemory::Memmove(&rle[ offset ], &equal, sizeof(uint16_t));
+                FMemory::Memmove(&rle[ offset ], &equal, sizeof(uint16));
                 offset += 2;
-                FMemory::Memmove(&rle[ offset ], &different, sizeof(uint16_t));
+                FMemory::Memmove(&rle[ offset ], &different, sizeof(uint16));
                 offset += 2;
-                FMemory::Memmove(&rle[ offset ], &equalPixel, sizeof(uint32_t));
+                FMemory::Memmove(&rle[ offset ], &equalPixel, sizeof(uint32));
                 offset += 4;
 
                 // Copy the different pixels
 				if (different)
 				{
 					// If we are at the end of a row, maybe there isn't a block of 4 pixels
-					uint16_t BytesToCopy = FMath::Min(different * 4 * 4, uint16_t(pBaseRowEnd - pDifferentPixels) * 4);
+					uint16 BytesToCopy = FMath::Min(different * 4 * 4, uint16(pBaseRowEnd - pDifferentPixels) * 4);
 
 					rle.resize( rle.size()+ BytesToCopy);
                     FMemory::Memmove( &rle[offset], pDifferentPixels, BytesToCopy);
@@ -354,30 +352,30 @@ namespace mu
 
 
     //---------------------------------------------------------------------------------------------
-    void UncompressRLE_RGBA( int width, int rows, const uint8_t* pBaseData, uint8_t* pDestDataB )
+    void UncompressRLE_RGBA( int width, int rows, const uint8* pBaseData, uint8* pDestDataB )
     {
-        uint32_t* pDestData = reinterpret_cast<uint32_t*>( pDestDataB );
+        uint32* pDestData = reinterpret_cast<uint32*>( pDestDataB );
 
-        pBaseData += rows*sizeof(uint32_t);
+        pBaseData += rows*sizeof(uint32);
 
         int pendingPixels = width*rows;
 
         for ( int r=0; r<rows; ++r )
         {
-            const uint32_t* pDestRowEnd = pDestData + width;
+            const uint32* pDestRowEnd = pDestData + width;
             while ( pDestData!=pDestRowEnd )
             {
                 // Decode header
-                uint16_t equal = 0;
-                FMemory::Memmove(&equal, pBaseData, sizeof(uint16_t));
+                uint16 equal = 0;
+                FMemory::Memmove(&equal, pBaseData, sizeof(uint16));
                 pBaseData += 2;
 
-                uint16_t different = 0;
-                FMemory::Memmove(&different, pBaseData, sizeof(uint16_t));
+                uint16 different = 0;
+                FMemory::Memmove(&different, pBaseData, sizeof(uint16));
                 pBaseData += 2;
 
-                uint32_t equalPixel = 0;
-                FMemory::Memmove(&equalPixel, pBaseData, sizeof(uint32_t));
+                uint32 equalPixel = 0;
+                FMemory::Memmove(&equalPixel, pBaseData, sizeof(uint32));
                 pBaseData += 4;
 
                 check((equal+different)*4<=pendingPixels);
@@ -392,7 +390,7 @@ namespace mu
 				if (different)
 				{
 					// If we are at the end of a row, maybe there isn't a block of 4 pixels
-					uint16_t PixelsToCopy = FMath::Min(uint16_t(different * 4), uint16_t(pDestRowEnd - pDestData));
+					uint16 PixelsToCopy = FMath::Min(uint16(different * 4), uint16(pDestRowEnd - pDestData));
 
 					FMemory::Memmove( pDestData, pBaseData, PixelsToCopy *4 );
 					pDestData += PixelsToCopy;
@@ -409,7 +407,7 @@ namespace mu
     //---------------------------------------------------------------------------------------------
     struct UINT24
     {
-        uint8_t d[3];
+        uint8 d[3];
 
         bool operator==( const UINT24& o ) const
         {
@@ -426,7 +424,7 @@ namespace mu
 
     //---------------------------------------------------------------------------------------------
     void CompressRLE_RGB( int width, int rows,
-                          const uint8_t* pBaseDataByte,
+                          const uint8* pBaseDataByte,
                           TArray<uint8>& destData )
     {
         vector<int8_t> rle;
@@ -434,10 +432,10 @@ namespace mu
 
         const UINT24* pBaseData = (const UINT24*)pBaseDataByte;
         rle.resize( rows*4 );
-        uint32_t offset = sizeof(uint32_t)*rows;
+        uint32 offset = sizeof(uint32)*rows;
         for ( int r=0; r<rows; ++r )
         {
-            uint32_t* pOffset = (uint32_t*) &rle[ sizeof(uint32_t)*r ];
+            uint32* pOffset = (uint32*) &rle[ sizeof(uint32)*r ];
             *pOffset = offset;
 
             const UINT24* pBaseRowEnd = pBaseData + width;
@@ -446,7 +444,7 @@ namespace mu
                 // Count equal pixels
                 UINT24 equalPixel = *pBaseData;
 
-                uint16_t equal = 0;
+                uint16 equal = 0;
                 while ( pBaseData<pBaseRowEnd-3 && equal<65535
                         && pBaseData[0]==equalPixel && pBaseData[1]==equalPixel
                         && pBaseData[2]==equalPixel && pBaseData[3]==equalPixel )
@@ -456,7 +454,7 @@ namespace mu
                 }
 
                 // Count different pixels
-                uint16_t different = 0;
+                uint16 different = 0;
                 const UINT24* pDifferentPixels = pBaseData;
                 while ( pBaseData!=pBaseRowEnd
                         &&
@@ -476,9 +474,9 @@ namespace mu
 
                 // Copy header
                 rle.resize( rle.size()+8 );
-                FMemory::Memmove( &rle[offset], &equal, sizeof(uint16_t) );
+                FMemory::Memmove( &rle[offset], &equal, sizeof(uint16) );
                 offset += 2;
-                FMemory::Memmove( &rle[offset], &different, sizeof(uint16_t) );
+                FMemory::Memmove( &rle[offset], &different, sizeof(uint16) );
                 offset += 2;
                 FMemory::Memmove( &rle[offset], &equalPixel, sizeof(UINT24) );
                 offset += 4;
@@ -487,7 +485,7 @@ namespace mu
 				if (different)
 				{
 					// If we are at the end of a row, maybe there isn't a block of 4 pixels
-					uint16_t BytesToCopy = FMath::Min(different * 4 * 3, uint16_t(pBaseRowEnd- pDifferentPixels)*3 );
+					uint16 BytesToCopy = FMath::Min(different * 4 * 3, uint16(pBaseRowEnd- pDifferentPixels)*3 );
 
 					rle.resize( rle.size()+BytesToCopy );
                     FMemory::Memmove( &rle[offset], pDifferentPixels, BytesToCopy );
@@ -505,11 +503,11 @@ namespace mu
 
 
     //---------------------------------------------------------------------------------------------
-    void UncompressRLE_RGB( int width, int rows, const uint8_t* pBaseData, uint8_t* pDestDataB )
+    void UncompressRLE_RGB( int width, int rows, const uint8* pBaseData, uint8* pDestDataB )
     {
         UINT24* pDestData = reinterpret_cast<UINT24*>( pDestDataB );
 
-        pBaseData += rows*sizeof(uint32_t);
+        pBaseData += rows*sizeof(uint32);
 
         for ( int r=0; r<rows; ++r )
         {
@@ -517,10 +515,10 @@ namespace mu
             while ( pDestData!=pDestRowEnd )
             {
                 // Decode header
-                uint16_t equal = *(const uint16_t*)pBaseData;
+                uint16 equal = *(const uint16*)pBaseData;
                 pBaseData += 2;
 
-                uint16_t different = *(const uint16_t*)pBaseData;
+                uint16 different = *(const uint16*)pBaseData;
                 pBaseData += 2;
 
                 UINT24 equalPixel = *(const UINT24*)pBaseData;
@@ -535,7 +533,7 @@ namespace mu
 				if (different)
 				{
 					// If we are at the end of a row, maybe there isn't a block of 4 pixels
-					uint16_t PixelsToCopy = FMath::Min(uint16_t(different * 4), uint16_t(pDestRowEnd - pDestData));
+					uint16 PixelsToCopy = FMath::Min(uint16(different * 4), uint16(pDestRowEnd - pDestData));
 
 					FMemory::Memmove( pDestData, pBaseData, PixelsToCopy*3);
 					pDestData += PixelsToCopy;
diff --git a/Engine/Plugins/Experimental/Mutable/Source/MutableRuntime/Private/MuR/Instance.cpp b/Engine/Plugins/Experimental/Mutable/Source/MutableRuntime/Private/MuR/Instance.cpp
--- a/Engine/Plugins/Experimental/Mutable/Source/MutableRuntime/Private/MuR/Instance.cpp
+++ b/Engine/Plugins/Experimental/Mutable/Source/MutableRuntime/Private/MuR/Instance.cpp
@@ -9,10 +9,6 @@
 #include "MuR/MemoryPrivate.h"
 #include "MuR/MutableMath.h"
 
-#include <memory>
-#include <utility>
-
-
 namespace mu
 {
 
diff --git a/Engine/Plugins/Experimental/Mutable/Source/MutableRuntime/Private/MuR/Memory.cpp b/Engine/Plugins/Experimental/Mutable/Source/MutableRuntime/Private/MuR/Memory.cpp
--- a/Engine/Plugins/Experimental/Mutable/Source/MutableRuntime/Private/MuR/Memory.cpp
+++ b/Engine/Plugins/Experimental/Mutable/Source/MutableRuntime/Private/MuR/Memory.cpp
@@ -3,20 +3,15 @@
 #include "MuR/BlockCompression/Miro/Miro.h"
 #include "MuR/MemoryPrivate.h"
 #include "MuR/MutableMemory.h"
-
-#include <stdint.h>
-#include <cstdlib>
+#include "HAL/MallocAnsi.h"
 
 namespace mu
 {
 
-	//! This would be nice, but apparently several compilers don't support it.
-#define MUTABLE_USE_STD_ALIGNED_ALLOC   1
-
 
 //! Static pointers to the user provided malloc and free methods.
 //! If they are null (default) the versions on the standard library will be used.
-	static void* (*s_custom_malloc)(size_t, uint32_t) = nullptr;
+	static void* (*s_custom_malloc)(size_t, uint32) = nullptr;
 	static void(*s_custom_free)(void*) = nullptr;
 
 
@@ -24,86 +19,6 @@ namespace mu
 	static int s_finalized = 0;
 
 
-	//---------------------------------------------------------------------------------------------
-#if MUTABLE_USE_STD_ALIGNED_ALLOC
-
-	inline void* mutable_aligned_alloc(size_t size, size_t align)
-	{
-#if defined(_WIN32) && _WIN32
-		return _aligned_malloc(size, align);
-#elif defined(MUTABLE_PLATFORM_ANDROID)
-		return memalign(align, size);
-#else
-		return aligned_alloc(align, size);
-#endif
-	}
-
-	inline void mutable_aligned_free(void* p)
-	{
-#if defined(_WIN32) && _WIN32
-		_aligned_free(p);
-#else
-		free(p);
-#endif
-	}
-
-#else
-
-	typedef uint16_t offset_t;
-#define PTR_OFFSET_SZ sizeof( offset_t ) 
-#ifndef align_up
-#define align_up( num, align ) ( ( ( num ) + ( (align)-1 ) ) & ~( (align)-1 ) )
-#endif
-
-	inline void* mutable_aligned_alloc(size_t size, size_t align)
-	{
-		void* ptr = nullptr;
-
-		// We want it to be a power of two since align_up operates on powers of two
-		assert((align & (align - 1)) == 0);
-
-		if (align && size)
-		{
-			// We know we have to fit an offset value
-			// We also allocate extra bytes to ensure we
-			// can meet the alignment
-			size_t hdr_size = PTR_OFFSET_SZ + (align - 1);
-			void* p = malloc(size + hdr_size);
-
-			if (p)
-			{
-				// Add the offset size to malloc's pointer (we will always store that)
-				// Then align the resulting value to the
-				// target alignment
-				ptr = (void*)align_up(((uintptr_t)p + PTR_OFFSET_SZ), align);
-
-				// Calculate the offset and store it
-				// behind our aligned pointer
-				*((offset_t*)ptr - 1) = (offset_t)((uintptr_t)ptr - (uintptr_t)p);
-
-			} // else NULL, could not malloc
-		}     // else NULL, invalid arguments
-
-		return ptr;
-	}
-
-
-	void mutable_aligned_free(void* ptr)
-	{
-		assert(ptr);
-
-		// Walk backwards from the passed-in pointer to get the pointer offset. We convert to an 
-		// offset_t pointer and rely on pointer math to get the data
-		offset_t offset = *((offset_t*)ptr - 1);
-
-		// Once we have the offset, we can get our original pointer and call free
-		void* p = (void*)((uint8_t*)ptr - offset);
-		free(p);
-	}
-
-#endif
-
-
 	//---------------------------------------------------------------------------------------------
 	//! Call the system malloc or the user-provided replacement.
 	//---------------------------------------------------------------------------------------------
@@ -121,7 +36,7 @@ namespace mu
 		return pMem;
 	}
 
-	inline void* lowerlevel_malloc_aligned(size_t bytes, uint32_t alignment)
+	inline void* lowerlevel_malloc_aligned(size_t bytes, uint32 alignment)
 	{
 		void* pMem = 0;
 		if (s_custom_malloc)
@@ -132,7 +47,7 @@ namespace mu
 		{
 			size_t padding = bytes % alignment;
 			bytes += padding ? alignment - padding : 0;
-			pMem = mutable_aligned_alloc(bytes, alignment);
+			pMem = AnsiMalloc(bytes, alignment);
 		}
 		return pMem;
 	}
@@ -162,7 +77,7 @@ namespace mu
 		}
 		else
 		{
-			mutable_aligned_free(ptr);
+			AnsiFree(ptr);
 		}
 	}
 
@@ -186,7 +101,7 @@ namespace mu
 	}
 
 
-	void* mutable_malloc_aligned(size_t size, uint32_t alignment)
+	void* mutable_malloc_aligned(size_t size, uint32 alignment)
 	{
 		if (!size)
 		{
@@ -237,7 +152,7 @@ namespace mu
 	//-------------------------------------------------------------------------------------------------
 	void Initialize
 	(
-		void* (*customMalloc)(size_t, uint32_t),
+		void* (*customMalloc)(size_t, uint32),
 		void (*customFree)(void*)
 	)
 	{
diff --git a/Engine/Plugins/Experimental/Mutable/Source/MutableRuntime/Private/MuR/MemoryPrivate.h b/Engine/Plugins/Experimental/Mutable/Source/MutableRuntime/Private/MuR/MemoryPrivate.h
--- a/Engine/Plugins/Experimental/Mutable/Source/MutableRuntime/Private/MuR/MemoryPrivate.h
+++ b/Engine/Plugins/Experimental/Mutable/Source/MutableRuntime/Private/MuR/MemoryPrivate.h
@@ -2,7 +2,7 @@
 
 #pragma once
 
-#include "../Private/MuR/Platform.h"
+#include "MuR/Platform.h"
 #include "MuR/MutableMemory.h"
 
 #include <unordered_map>
diff --git a/Engine/Plugins/Experimental/Mutable/Source/MutableRuntime/Private/MuR/Mesh.cpp b/Engine/Plugins/Experimental/Mutable/Source/MutableRuntime/Private/MuR/Mesh.cpp
--- a/Engine/Plugins/Experimental/Mutable/Source/MutableRuntime/Private/MuR/Mesh.cpp
+++ b/Engine/Plugins/Experimental/Mutable/Source/MutableRuntime/Private/MuR/Mesh.cpp
@@ -10,7 +10,7 @@
 #include "MuR/MeshPrivate.h"
 #include "MuR/MutableTrace.h"
 
-#include <utility>
+#include <limits>
 
 namespace mu
 {
@@ -799,8 +799,8 @@ void Mesh::GetVertexMap
     for ( int ov=0; ov<otherVertexCount; ++ov )
     {
         float v = (*itop)[MUTABLE_BUCKET_CHANNEL];
-        rangeMin = std::min( rangeMin, v );
-        rangeMax = std::max( rangeMax, v );
+        rangeMin = FMath::Min( rangeMin, v );
+        rangeMax = FMath::Max( rangeMax, v );
         ++itop;
     }
     rangeMin -= tolerance;
diff --git a/Engine/Plugins/Experimental/Mutable/Source/MutableRuntime/Private/MuR/Model.cpp b/Engine/Plugins/Experimental/Mutable/Source/MutableRuntime/Private/MuR/Model.cpp
--- a/Engine/Plugins/Experimental/Mutable/Source/MutableRuntime/Private/MuR/Model.cpp
+++ b/Engine/Plugins/Experimental/Mutable/Source/MutableRuntime/Private/MuR/Model.cpp
@@ -20,9 +20,6 @@
 #include "MuR/Types.h"
 #include "Templates/Tuple.h"
 
-#include <memory>
-#include <string>
-#include <utility>
 
 namespace mu
 {
@@ -89,7 +86,7 @@ namespace mu
     void PROGRAM::LogHistogram() const
     {
 #if 0
-        uint64_t countPerType[(int)OP_TYPE::COUNT];
+        uint64 countPerType[(int)OP_TYPE::COUNT];
         mutable_memset(countPerType,0,sizeof(countPerType));
 
         for ( const auto& o: m_opAddress )
@@ -98,14 +95,14 @@ namespace mu
             countPerType[(int)type]++;
         }
 
-        vector< pair<uint64_t,OP_TYPE> > sorted((int)OP_TYPE::COUNT);
+        vector< pair<uint64,OP_TYPE> > sorted((int)OP_TYPE::COUNT);
         for (int i=0; i<(int)OP_TYPE::COUNT; ++i)
         {
             sorted[i].second = (OP_TYPE)i;
             sorted[i].first = countPerType[i];
         }
 
-        std::sort(sorted.begin(),sorted.end(), []( const pair<uint64_t,OP_TYPE>& a, const pair<uint64_t,OP_TYPE>& b )
+        std::sort(sorted.begin(),sorted.end(), []( const pair<uint64,OP_TYPE>& a, const pair<uint64,OP_TYPE>& b )
         {
             return a.first>b.first;
         });
@@ -489,21 +486,21 @@ namespace mu
                                   const map< vector<int>, PARAMETER_VALUE >& multi )
     {
         parameterValuesBlob.resize( pos+4 );
-        uint32_t s = uint32_t( multi.size() );
-        memcpy( &parameterValuesBlob[pos], &s, sizeof(uint32_t) );
+        uint32 s = uint32( multi.size() );
+        memcpy( &parameterValuesBlob[pos], &s, sizeof(uint32) );
         pos+=4;
 
         for(const auto& v: multi)
         {
-            uint32_t ds = uint32_t( v.first.size() );
+            uint32 ds = uint32( v.first.size() );
 
             parameterValuesBlob.resize( pos + 4 + ds*4 );
 
-            memcpy( &parameterValuesBlob[pos], &s, sizeof(uint32_t) );
+            memcpy( &parameterValuesBlob[pos], &s, sizeof(uint32) );
             pos+=4;
 
-            memcpy( &parameterValuesBlob[pos],v.first.data(), ds*sizeof(int32_t) );
-            pos += ds*sizeof(int32_t);
+            memcpy( &parameterValuesBlob[pos],v.first.data(), ds*sizeof(int32) );
+            pos += ds*sizeof(int32);
         }
 
         return pos;
@@ -511,7 +508,7 @@ namespace mu
 
 
     //---------------------------------------------------------------------------------------------
-    uint32_t Model::Private::GetResourceKey( uint32_t paramListIndex,
+    uint32 Model::Private::GetResourceKey( uint32 paramListIndex,
                                              OP::ADDRESS rootAt,
                                              const Parameters* pParams )
     {
@@ -556,7 +553,7 @@ namespace mu
                 break;
 
             case PARAMETER_TYPE::T_INT:
-                dataSize = sizeof(int32_t);
+                dataSize = sizeof(int32);
                 parameterValuesBlob.resize( pos+dataSize );
                 memcpy( &parameterValuesBlob[pos],
                         &pParams->GetPrivate()->m_values[param].m_int,
@@ -704,7 +701,7 @@ namespace mu
         }
 
         // Generate a new id
-        uint32_t newId = ++m_lastResourceKeyId;
+        uint32 newId = ++m_lastResourceKeyId;
         RESOURCE_KEY newKey;
         newKey.m_id = newId;
         newKey.m_lastRequestId = m_lastResourceResquestId;
@@ -800,14 +797,14 @@ namespace mu
 
 
     //---------------------------------------------------------------------------------------------
-    int64_t ModelParametersGenerator::GetInstanceCount()
+    int64 ModelParametersGenerator::GetInstanceCount()
     {
         return m_pD->m_instanceCount;
     }
 
 
     //---------------------------------------------------------------------------------------------
-    ParametersPtr ModelParametersGenerator::GetInstance( int64_t index,
+    ParametersPtr ModelParametersGenerator::GetInstance( int64 index,
                                                          float (*randomGenerator )() )
     {
 		LLM_SCOPE_BYNAME(TEXT("MutableRuntime"));
@@ -852,7 +849,7 @@ namespace mu
         else
         {
             int paramCount = (int)m_pD->m_pModel->GetPrivate()->m_program.m_parameters.size();
-            int64_t currentInstance = index;
+            int64 currentInstance = index;
             for (int i=0;i<paramCount;++i)
             {
                 const PARAMETER_DESC& param = m_pD->m_pModel->GetPrivate()->m_program.m_parameters[ i ];
diff --git a/Engine/Plugins/Experimental/Mutable/Source/MutableRuntime/Private/MuR/ModelPrivate.h b/Engine/Plugins/Experimental/Mutable/Source/MutableRuntime/Private/MuR/ModelPrivate.h
--- a/Engine/Plugins/Experimental/Mutable/Source/MutableRuntime/Private/MuR/ModelPrivate.h
+++ b/Engine/Plugins/Experimental/Mutable/Source/MutableRuntime/Private/MuR/ModelPrivate.h
@@ -49,7 +49,7 @@ namespace mu
     {
         size_t pos = code.size();
         code.resize( code.size()+sizeof(DATA) );
-        memcpy( &code[pos], &data, sizeof(DATA) );
+		FMemory::Memcpy (&code[pos], &data, sizeof(DATA));
     }
 
 
@@ -614,7 +614,7 @@ namespace mu
 
             OP_TYPE result;
             uint64 byteCodeAddress = m_opAddress[at];
-            memcpy( &result, &m_byteCode[byteCodeAddress], sizeof(OP_TYPE) );
+			FMemory::Memcpy( &result, &m_byteCode[byteCodeAddress], sizeof(OP_TYPE) );
             check( result<OP_TYPE::COUNT);
             return result;
         }
@@ -625,7 +625,7 @@ namespace mu
             ARGS result;
             uint64 byteCodeAddress = m_opAddress[at];
             byteCodeAddress += sizeof(OP_TYPE);
-            memcpy( &result, &m_byteCode[byteCodeAddress], sizeof(ARGS));
+			FMemory::Memcpy( &result, &m_byteCode[byteCodeAddress], sizeof(ARGS));
             return result;
         }
 
diff --git a/Engine/Plugins/Experimental/Mutable/Source/MutableRuntime/Private/MuR/MutableMath.h b/Engine/Plugins/Experimental/Mutable/Source/MutableRuntime/Private/MuR/MutableMath.h
--- a/Engine/Plugins/Experimental/Mutable/Source/MutableRuntime/Private/MuR/MutableMath.h
+++ b/Engine/Plugins/Experimental/Mutable/Source/MutableRuntime/Private/MuR/MutableMath.h
@@ -3,8 +3,8 @@
 #pragma once
 
 #include "MuR/Types.h"
-#include "../Private/MuR/SerialisationPrivate.h"
-#include "../Private/MuR/Platform.h"
+#include "MuR/SerialisationPrivate.h"
+#include "MuR/Platform.h"
 #include "Math/Matrix.h"
 #include "Math/IntVector.h"
 
@@ -35,17 +35,6 @@ namespace mu
         return r;
     }
 
-    //---------------------------------------------------------------------------------------------
-    //! Base 2 logarithm
-    //---------------------------------------------------------------------------------------------
-    inline int log2( int v )
-    {
-        // TODO: There are one-line asm versions for most architectures.
-        int r = 0;
-        while (v >>= 1) ++r;
-        return r;
-    }
-
 
     //---------------------------------------------------------------------------------------------
     //!
@@ -53,29 +42,17 @@ namespace mu
     template<class T>
     T clamp( const T mi, const T ma, const T v )
     {
-        return std::max( mi, std::min( ma, v ) );
-    }
-
-
-    //---------------------------------------------------------------------------------------------
-    // Robust floating point comparisons:
-    // http://realtimecollisiondetection.net/blog/?p=89
-    //---------------------------------------------------------------------------------------------
-    inline bool equal(const float f0, const float f1,
-                      const float epsilon = std::numeric_limits<float>::epsilon() )
-    {
-        //return fabs(f0-f1) <= epsilon;
-        return fabs(f0-f1) <= epsilon * std::max(1.0f, std::max( fabsf(f0), fabsf(f1) ) );
+        return FMath::Max( mi, FMath::Min( ma, v ) );
     }
 
 
 	//---------------------------------------------------------------------------------------------
 	// 16-bit floating point number support.
 	//---------------------------------------------------------------------------------------------
-    typedef uint16_t float16;
+    typedef uint16 float16;
 
 	//---------------------------------------------------------------------------------------------
-    inline uint32_t halfToFloatI( float16 y )
+    inline uint32 halfToFloatI( float16 y )
 	{
 		int s = (y >> 15) & 0x00000001;
 		int e = (y >> 10) & 0x0000001f;
@@ -123,7 +100,7 @@ namespace mu
 		union
 		{
 			float f;
-            uint32_t i;
+            uint32 i;
 		} v;
 
 		v.i = halfToFloatI(y);
@@ -132,7 +109,7 @@ namespace mu
 
 
 	//---------------------------------------------------------------------------------------------
-    inline float16 floatToHalfI( uint32_t i )
+    inline float16 floatToHalfI( uint32 i )
 	{
 		int s =  (i >> 16) & 0x00008000;
 		int e = ((i >> 23) & 0x000000ff) - (127 - 15);
@@ -183,7 +160,7 @@ namespace mu
 		union
 		{
 			float f;
-            uint32_t i;
+            uint32 i;
 		} v;
 
 		v.f = i;
@@ -388,7 +365,7 @@ namespace mu
             V r;
             for (size_t i = 0; i < DIM; ++i)
             {
-                r.m[i] = std::min<>(a.m[i], b.m[i]);
+                r.m[i] = FMath::Min(a.m[i], b.m[i]);
             }
             return r;
         }
@@ -923,7 +900,7 @@ namespace mu
 	{
 		union MUTABLE_FLOAT_INT
 		{
-            int32_t i;
+            int32 i;
 			float y;
 		} fi;
 		float x2;
@@ -1066,7 +1043,7 @@ namespace mu
             for (int d = 0; d < VECTOR::GetDim(); ++d)
             {
                 result = result &&
-                         std::abs((min[d] * 2 + size[d]) - (other.min[d] * 2 + other.size[d])) <=
+                         FMath::Abs((min[d] * 2 + size[d]) - (other.min[d] * 2 + other.size[d])) <=
                              (size[d] + other.size[d]);
             }
 
@@ -1082,7 +1059,7 @@ namespace mu
             for (int d = 0; d < VECTOR::GetDim(); ++d)
             {
                 result = result &&
-                         std::abs((min[d] * 2 + size[d]) - (other.min[d] * 2 + other.size[d])) <
+                         FMath::Abs((min[d] * 2 + size[d]) - (other.min[d] * 2 + other.size[d])) <
                              (size[d] + other.size[d]);
             }
 
@@ -1097,9 +1074,9 @@ namespace mu
             box<VECTOR> result;
             for (int d = 0; d < VECTOR::GetDim(); ++d)
             {
-                result.min[d] = std::max(min[d], other.min[d]);
+                result.min[d] = FMath::Max(min[d], other.min[d]);
                 result.size[d] =
-                    std::min(typename VECTOR::SCALAR(min[d] + size[d] - result.min[d]),
+                    FMath::Min(typename VECTOR::SCALAR(min[d] + size[d] - result.min[d]),
                              typename VECTOR::SCALAR(other.min[d] + other.size[d] - result.min[d]));
             }
             return result;
@@ -1243,9 +1220,9 @@ namespace mu
 		template<int DIM2>
 		explicit mat( const mat<SCALAR,DIM2>& other )
 		{
-			for (int i=0; i<std::min(DIM,DIM2); ++i)
+			for (int i=0; i<FMath::Min(DIM,DIM2); ++i)
 			{
-				for (int j=0; j<std::min(DIM,DIM2) ;++j)
+				for (int j=0; j<FMath::Min(DIM,DIM2) ;++j)
 				{
 					m[i][j] = other[i][j];
 				}
@@ -1264,7 +1241,7 @@ namespace mu
 		inline static M Scale( const vec<SCALAR,DIM2>& s )
 		{
 			M r;
-			for (int i=0; i<std::min(DIM,DIM2); ++i)
+			for (int i=0; i<FMath::Min(DIM,DIM2); ++i)
 			{
 				r.m[i][i] = s[i];
 			}
@@ -1285,7 +1262,7 @@ namespace mu
 		{
 			M r = Identity();
 
-			for (int i=0; i<std::min(DIM,DIM2); ++i)
+			for (int i=0; i<FMath::Min(DIM,DIM2); ++i)
 			{
 				r.m[i][DIM-1] = pos[i];
 			}
diff --git a/Engine/Plugins/Experimental/Mutable/Source/MutableRuntime/Private/MuR/Octree.h b/Engine/Plugins/Experimental/Mutable/Source/MutableRuntime/Private/MuR/Octree.h
--- a/Engine/Plugins/Experimental/Mutable/Source/MutableRuntime/Private/MuR/Octree.h
+++ b/Engine/Plugins/Experimental/Mutable/Source/MutableRuntime/Private/MuR/Octree.h
@@ -7,7 +7,6 @@
 #include "MuR/MutableMath.h"
 #include "MuR/Ptr.h"
 
-#include <vector>
 
 #define MAX_DIVISIONS 5
 #define MAX_VECTORS 20
@@ -17,7 +16,7 @@ namespace mu
 	struct VertexInfo
 	{
 		vec3f vector;
-		int index;
+		int32 index;
 	};
 
 	class Octree
@@ -47,25 +46,15 @@ namespace mu
 					children[i] = nullptr;
 				}	
 			}
-
-            for (size_t i = 0; i < elements.size(); ++i)
-			{
-				delete elements[i];
-				elements[i] = nullptr;
-			}
 		}
 
 		void InsertElement(vec3f v, int idx)
 		{
 			if (!children[0])
 			{
-				VertexInfo* cur = new VertexInfo();
-				cur->vector = v;
-				cur->index = idx;
-
-				elements.push_back(cur);
+				elements.Add({ v, idx });
 
-				if (elements.size() > MAX_VECTORS && division < MAX_DIVISIONS)
+				if (elements.Num() > MAX_VECTORS && division < MAX_DIVISIONS)
 				{
 					Divide();
 				}
@@ -114,18 +103,12 @@ namespace mu
 				children[3 + i * 4] = new Octree(nMin, nMax, padding, division + 1);
 			}
 
-            for (size_t i = 0; i < elements.size(); ++i)
+            for (size_t i = 0; i < elements.Num(); ++i)
 			{
-				InsertElement(elements[i]->vector, elements[i]->index);
+				InsertElement(elements[i].vector, elements[i].index);
 			}
-			
-            for (size_t i = 0; i < elements.size(); ++i)
-			{
-				delete elements[i];
-				elements[i] = nullptr;
-			}
-
-			elements.clear();
+		            
+			elements.Empty();
 		}
 
 		bool Collides(Octree* child, vec3f v)
@@ -152,13 +135,13 @@ namespace mu
 
 			if (!children[0])
 			{
-                for (size_t i = 0; i < elements.size(); ++i)
+                for (size_t i = 0; i < elements.Num(); ++i)
 				{
-					vec3f r = elements[i]->vector - v;
+					vec3f r = elements[i].vector - v;
 					
-					if (idx > elements[i]->index && dot(r,r) <= dist*dist)
+					if (idx > elements[i].index && dot(r,r) <= dist*dist)
 					{
-						ret = elements[i]->index;
+						ret = elements[i].index;
 						break;
 					}
 				}
@@ -189,7 +172,7 @@ namespace mu
 
 		float padding = 0.0f;
 
-		std::vector<VertexInfo*> elements;
+		TArray<VertexInfo> elements;
 
 		Octree* children[8];
 	};
diff --git a/Engine/Plugins/Experimental/Mutable/Source/MutableRuntime/Private/MuR/OpImageProject.cpp b/Engine/Plugins/Experimental/Mutable/Source/MutableRuntime/Private/MuR/OpImageProject.cpp
--- a/Engine/Plugins/Experimental/Mutable/Source/MutableRuntime/Private/MuR/OpImageProject.cpp
+++ b/Engine/Plugins/Experimental/Mutable/Source/MutableRuntime/Private/MuR/OpImageProject.cpp
@@ -23,11 +23,8 @@
 #include "MuR/Ptr.h"
 #include "MuR/Raster.h"
 
-#include <functional>
-#include <memory>
 #include <queue>
 #include <unordered_set>
-#include <utility>
 
 
 namespace mu
@@ -43,8 +40,8 @@ namespace
 
 		RaycastPixelProcessor( const PROJECTOR& projector,
 							   const Image* pSource,
-                               const uint8_t* pTargetData,
-                               const uint8_t* pMaskData,
+                               const uint8* pTargetData,
+                               const uint8* pMaskData,
 							   float fadeStart,
 							   float fadeEnd )
 		{
@@ -85,7 +82,7 @@ namespace
 
             if ( factor>0 && m_pMaskData )
 			{
-                uint8_t maskFactor = m_pMaskData[ (pBufferPos-m_pTargetData) / m_pixelSize ];
+                uint8 maskFactor = m_pMaskData[ (pBufferPos-m_pTargetData) / m_pixelSize ];
 				factor = float(maskFactor) / 255.0f;
 			}
 
@@ -117,7 +114,7 @@ namespace
 				if ( u>=0 && u<1 && v>=0 && v<1 )
 				{
 
-                    const uint8_t* pPixel = m_pSource->GetData();
+                    const uint8* pPixel = m_pSource->GetData();
 					// TODO: clamp?
 					pPixel += ( m_pSource->GetSizeX() * int(m_pSource->GetSizeY() * v)
 								+ int( m_pSource->GetSizeX() * u ) )
@@ -127,21 +124,21 @@ namespace
 					switch ( m_pSource->GetFormat() )
 					{
 					case EImageFormat::IF_L_UBYTE:
-                        pBufferPos[0] = uint8_t( pPixel[0] * factor );
+                        pBufferPos[0] = uint8( pPixel[0] * factor );
 						break;
 
 					case EImageFormat::IF_RGB_UBYTE:
-                        pBufferPos[0] = uint8_t( pPixel[0] * factor );
-                        pBufferPos[1] = uint8_t( pPixel[1] * factor );
-                        pBufferPos[2] = uint8_t( pPixel[2] * factor );
+                        pBufferPos[0] = uint8( pPixel[0] * factor );
+                        pBufferPos[1] = uint8( pPixel[1] * factor );
+                        pBufferPos[2] = uint8( pPixel[2] * factor );
 						break;
 
                     case EImageFormat::IF_BGRA_UBYTE:
                     case EImageFormat::IF_RGBA_UBYTE:
-                        pBufferPos[0] = uint8_t( pPixel[0] * factor );
-                        pBufferPos[1] = uint8_t( pPixel[1] * factor );
-                        pBufferPos[2] = uint8_t( pPixel[2] * factor );
-                        pBufferPos[3] = uint8_t( pPixel[3] * factor );
+                        pBufferPos[0] = uint8( pPixel[0] * factor );
+                        pBufferPos[1] = uint8( pPixel[1] * factor );
+                        pBufferPos[2] = uint8( pPixel[2] * factor );
+                        pBufferPos[3] = uint8( pPixel[3] * factor );
 						break;
 
 					default:
@@ -164,8 +161,8 @@ namespace
         //! Cosine of the fading angle range
         float m_fadeStartCos, m_fadeEndCos;
 
-        const uint8_t* m_pTargetData;
-        const uint8_t* m_pMaskData;
+        const uint8* m_pTargetData;
+        const uint8* m_pMaskData;
 		int m_pixelSize;
 
 	};
@@ -179,8 +176,8 @@ namespace
 
 		RaycastPixelProcessor_UBYTE( const PROJECTOR& projector,
 									 const Image* pSource,
-                                     const uint8_t* pTargetData,
-                                     const uint8_t* pMaskData,
+                                     const uint8* pTargetData,
+                                     const uint8* pMaskData,
 									 float fadeStart,
 									 float fadeEnd )
 		{
@@ -213,7 +210,7 @@ namespace
 		//-----------------------------------------------------------------------------------------
         inline void ProcessPixel( unsigned char* pBufferPos, float varying[4] ) const
 		{
-            int32_t factor_8 = 256;
+            int32 factor_8 = 256;
 
             // depth clamp
             if (varying[2]<0.0f || varying[2]>1.0f)
@@ -231,7 +228,7 @@ namespace
             {
                 if ( factor_8>0 && m_pMaskData )
                 {
-                    uint8_t maskFactor = m_pMaskData[ (pBufferPos-m_pTargetData) / PIXEL_SIZE ];
+                    uint8 maskFactor = m_pMaskData[ (pBufferPos-m_pTargetData) / PIXEL_SIZE ];
                     factor_8 = maskFactor;
                 }
 
@@ -241,7 +238,7 @@ namespace
                     {
                         factor_8 = ( factor_8
                                      *
-                                     int32_t( 256.0f * ( angleCos - m_fadeEndCos ) / ( m_fadeStartCos - m_fadeEndCos ) )
+                                     int32( 256.0f * ( angleCos - m_fadeEndCos ) / ( m_fadeStartCos - m_fadeEndCos ) )
                                      ) >> 8;
                     }
                     else if ( angleCos<=m_fadeEndCos )
@@ -254,7 +251,7 @@ namespace
             {
                 if ( angleCos<m_fadeStartCos && angleCos>m_fadeEndCos )
                 {
-                    factor_8 = int32_t( 256.0f * ( angleCos - m_fadeEndCos ) / ( m_fadeStartCos - m_fadeEndCos ) );
+                    factor_8 = int32( 256.0f * ( angleCos - m_fadeEndCos ) / ( m_fadeStartCos - m_fadeEndCos ) );
                 }
                 else if ( angleCos<=m_fadeEndCos )
                 {
@@ -271,7 +268,7 @@ namespace
 				if ( u>=0 && u<1 && v>=0 && v<1 )
 				{
 
-                    const uint8_t* pPixel = m_pSourceData;
+                    const uint8* pPixel = m_pSourceData;
 					// TODO: clamp?
 					pPixel += ( m_sourceSizeX * int(m_sourceSizeY * v)
 								+ int( m_sourceSizeX * u ) )
@@ -280,14 +277,14 @@ namespace
 					// Write result
 					for ( int i=0; i<PIXEL_SIZE; ++i )
 					{
-						pBufferPos[i] = uint8_t( ( pPixel[i] * factor_8 ) >> 8 );
+						pBufferPos[i] = uint8( ( pPixel[i] * factor_8 ) >> 8 );
 					}
 				}
 			}
 		}
 
 
-        const uint8_t* m_pSourceData;
+        const uint8* m_pSourceData;
 		int m_sourceSizeX, m_sourceSizeY;
 
 		vec3<float> m_direction;
@@ -299,8 +296,8 @@ namespace
         //! Cosine of the fading angle range
         float m_fadeStartCos, m_fadeEndCos;
 
-        const uint8_t* m_pTargetData;
-        const uint8_t* m_pMaskData;
+        const uint8* m_pTargetData;
+        const uint8* m_pMaskData;
 	};
 
 
@@ -310,8 +307,8 @@ namespace
     public:
 
         RasterProjectedPixelProcessor( const Image* pSource,
-                               const uint8_t* pTargetData,
-                               const uint8_t* pMaskData,
+                               const uint8* pTargetData,
+                               const uint8* pMaskData,
                                float fadeStart,
                                float fadeEnd )
         {
@@ -342,7 +339,7 @@ namespace
 
             if ( factor>0 && m_pMaskData )
             {
-                uint8_t maskFactor = m_pMaskData[ (pBufferPos-m_pTargetData) / m_pixelSize ];
+                uint8 maskFactor = m_pMaskData[ (pBufferPos-m_pTargetData) / m_pixelSize ];
                 factor = float(maskFactor) / 255.0f;
             }
 
@@ -370,7 +367,7 @@ namespace
                 if ( u>=0 && u<1 && v>=0 && v<1 )
                 {
 
-                    const uint8_t* pPixel = m_pSource->GetData();
+                    const uint8* pPixel = m_pSource->GetData();
                     // TODO: clamp?
                     pPixel += ( m_pSource->GetSizeX() * int(m_pSource->GetSizeY() * v)
                                 + int( m_pSource->GetSizeX() * u ) )
@@ -380,21 +377,21 @@ namespace
                     switch ( m_pSource->GetFormat() )
                     {
                     case EImageFormat::IF_L_UBYTE:
-                        pBufferPos[0] = uint8_t( pPixel[0] * factor );
+                        pBufferPos[0] = uint8( pPixel[0] * factor );
                         break;
 
                     case EImageFormat::IF_RGB_UBYTE:
-                        pBufferPos[0] = uint8_t( pPixel[0] * factor );
-                        pBufferPos[1] = uint8_t( pPixel[1] * factor );
-                        pBufferPos[2] = uint8_t( pPixel[2] * factor );
+                        pBufferPos[0] = uint8( pPixel[0] * factor );
+                        pBufferPos[1] = uint8( pPixel[1] * factor );
+                        pBufferPos[2] = uint8( pPixel[2] * factor );
                         break;
 
                     case EImageFormat::IF_BGRA_UBYTE:
                     case EImageFormat::IF_RGBA_UBYTE:
-                        pBufferPos[0] = uint8_t( pPixel[0] * factor );
-                        pBufferPos[1] = uint8_t( pPixel[1] * factor );
-                        pBufferPos[2] = uint8_t( pPixel[2] * factor );
-                        pBufferPos[3] = uint8_t( pPixel[3] * factor );
+                        pBufferPos[0] = uint8( pPixel[0] * factor );
+                        pBufferPos[1] = uint8( pPixel[1] * factor );
+                        pBufferPos[2] = uint8( pPixel[2] * factor );
+                        pBufferPos[3] = uint8( pPixel[3] * factor );
                         break;
 
                     default:
@@ -412,8 +409,8 @@ namespace
         //! Cosine of the fading angle range
         float m_fadeStartCos, m_fadeEndCos;
 
-        const uint8_t* m_pTargetData;
-        const uint8_t* m_pMaskData;
+        const uint8* m_pTargetData;
+        const uint8* m_pMaskData;
         int m_pixelSize;
 
     };
@@ -426,8 +423,8 @@ namespace
     public:
 
         RasterProjectedPixelProcessor_UBYTE( const Image* pSource,
-                                     const uint8_t* pTargetData,
-                                     const uint8_t* pMaskData,
+                                     const uint8* pTargetData,
+                                     const uint8* pMaskData,
                                      float fadeStart,
                                      float fadeEnd )
         {
@@ -450,7 +447,7 @@ namespace
         //-----------------------------------------------------------------------------------------
         inline void ProcessPixel( unsigned char* pBufferPos, float varying[4] ) const
         {
-            int32_t factor_8 = 256;
+            int32 factor_8 = 256;
 
             // depth clamp
             if (varying[2]<0.0f || varying[2]>1.0f)
@@ -464,7 +461,7 @@ namespace
             {
                 if ( factor_8>0 && m_pMaskData )
                 {
-                    uint8_t maskFactor = m_pMaskData[ (pBufferPos-m_pTargetData) / PIXEL_SIZE ];
+                    uint8 maskFactor = m_pMaskData[ (pBufferPos-m_pTargetData) / PIXEL_SIZE ];
                     factor_8 = maskFactor;
                 }
 
@@ -474,7 +471,7 @@ namespace
                     {
                         factor_8 = ( factor_8
                                      *
-                                     int32_t( 256.0f * ( angleCos - m_fadeEndCos ) / ( m_fadeStartCos - m_fadeEndCos ) )
+                                     int32( 256.0f * ( angleCos - m_fadeEndCos ) / ( m_fadeStartCos - m_fadeEndCos ) )
                                      ) >> 8;
                     }
                     else if ( angleCos<=m_fadeEndCos )
@@ -487,7 +484,7 @@ namespace
             {
                 if ( angleCos<m_fadeStartCos && angleCos>m_fadeEndCos )
                 {
-                    factor_8 = int32_t( 256.0f * ( angleCos - m_fadeEndCos ) / ( m_fadeStartCos - m_fadeEndCos ) );
+                    factor_8 = int32( 256.0f * ( angleCos - m_fadeEndCos ) / ( m_fadeStartCos - m_fadeEndCos ) );
                 }
                 else if ( angleCos<=m_fadeEndCos )
                 {
@@ -504,7 +501,7 @@ namespace
                 if ( u>=0 && u<1 && v>=0 && v<1 )
                 {
 
-                    const uint8_t* pPixel = m_pSourceData;
+                    const uint8* pPixel = m_pSourceData;
                     // TODO: clamp?
                     pPixel += ( m_sourceSizeX * int(m_sourceSizeY * v)
                                 + int( m_sourceSizeX * u ) )
@@ -513,21 +510,21 @@ namespace
                     // Write result
                     for ( int i=0; i<PIXEL_SIZE; ++i )
                     {
-                        pBufferPos[i] = uint8_t( ( pPixel[i] * factor_8 ) >> 8 );
+                        pBufferPos[i] = uint8( ( pPixel[i] * factor_8 ) >> 8 );
                     }
                 }
             }
         }
 
 
-        const uint8_t* m_pSourceData;
+        const uint8* m_pSourceData;
         int m_sourceSizeX, m_sourceSizeY;
 
         //! Cosine of the fading angle range
         float m_fadeStartCos, m_fadeEndCos;
 
-        const uint8_t* m_pTargetData;
-        const uint8_t* m_pMaskData;
+        const uint8* m_pTargetData;
+        const uint8* m_pMaskData;
     };
 
 
@@ -538,8 +535,8 @@ namespace
     public:
 
         RasterCylindricalProjectedPixelProcessor_UBYTE( const Image* pSource,
-                                     const uint8_t* pTargetData,
-                                     const uint8_t* pMaskData,
+                                     const uint8* pTargetData,
+                                     const uint8* pMaskData,
                                      float fadeStart,
                                      float fadeEnd,
                                                         float projectionAngle )
@@ -564,7 +561,7 @@ namespace
         //-----------------------------------------------------------------------------------------
         inline void ProcessPixel( unsigned char* pBufferPos, float varying[4] ) const
         {
-            int32_t factor_8 = 256;
+            int32 factor_8 = 256;
 
             // Position in unit cylinder space
             float x_cyl = varying[0];
@@ -584,7 +581,7 @@ namespace
             {
                 if ( factor_8>0 && m_pMaskData )
                 {
-                    uint8_t maskFactor = m_pMaskData[ (pBufferPos-m_pTargetData) / PIXEL_SIZE ];
+                    uint8 maskFactor = m_pMaskData[ (pBufferPos-m_pTargetData) / PIXEL_SIZE ];
                     factor_8 = maskFactor;
                 }
 
@@ -594,7 +591,7 @@ namespace
                     {
                         factor_8 = ( factor_8
                                      *
-                                     int32_t( 256.0f * ( angleCos - m_fadeEndCos ) / ( m_fadeStartCos - m_fadeEndCos ) )
+                                     int32( 256.0f * ( angleCos - m_fadeEndCos ) / ( m_fadeStartCos - m_fadeEndCos ) )
                                      ) >> 8;
                     }
                     else if ( angleCos<=m_fadeEndCos )
@@ -607,7 +604,7 @@ namespace
             {
                 if ( angleCos<m_fadeStartCos && angleCos>m_fadeEndCos )
                 {
-                    factor_8 = int32_t( 256.0f * ( angleCos - m_fadeEndCos ) / ( m_fadeStartCos - m_fadeEndCos ) );
+                    factor_8 = int32( 256.0f * ( angleCos - m_fadeEndCos ) / ( m_fadeStartCos - m_fadeEndCos ) );
                 }
                 else if ( angleCos<=m_fadeEndCos )
                 {
@@ -625,7 +622,7 @@ namespace
                 if ( u>=0 && u<1 && v>=0 && v<1 )
                 {
 
-                    const uint8_t* pPixel = m_pSourceData;
+                    const uint8* pPixel = m_pSourceData;
                     // TODO: clamp?
                     pPixel += ( m_sourceSizeX * int(m_sourceSizeY * v)
                                 + int( m_sourceSizeX * u ) )
@@ -634,14 +631,14 @@ namespace
                     // Write result
                     for ( int i=0; i<PIXEL_SIZE; ++i )
                     {
-                        pBufferPos[i] = uint8_t( ( pPixel[i] * factor_8 ) >> 8 );
+                        pBufferPos[i] = uint8( ( pPixel[i] * factor_8 ) >> 8 );
                     }
                 }
             }
         }
 
 
-        const uint8_t* m_pSourceData;
+        const uint8* m_pSourceData;
         int m_sourceSizeX, m_sourceSizeY;
 
         //! Cosine of the fading angle range
@@ -649,8 +646,8 @@ namespace
 
         float m_projectionAngle;
 
-        const uint8_t* m_pTargetData;
-        const uint8_t* m_pMaskData;
+        const uint8* m_pTargetData;
+        const uint8* m_pMaskData;
     };
 
 
@@ -721,7 +718,7 @@ void ImageRasterProjected_Generic( const Mesh* pMesh,
     UntypedMeshBufferIteratorConst indIt( pMesh->GetIndexBuffers(), MBS_VERTEXINDEX, 0 );
     for ( int i=0; i<faceCount*3; ++i )
     {
-        uint32_t index=0;
+        uint32 index=0;
         ConvertData( 0, &index, MBF_UINT32, indIt.ptr(), indIt.GetFormat() );
 
         indices[i] = index;
@@ -833,7 +830,7 @@ struct OPTIMISED_VERTEX_WRAPPING
     vec2<float> uv;
     vec3<float> pos;
     vec3<float> nor;
-    uint32_t layoutBlock;
+    uint32 layoutBlock;
 };
 
 static_assert( sizeof(OPTIMISED_VERTEX_WRAPPING)==36, "UNEXPECTED_STRUCT_SIZE" );
@@ -885,7 +882,7 @@ void ImageRasterProjected_Optimised( const Mesh* pMesh,
 
     // Get the indices
     check( pMesh->GetIndexBuffers().GetElementSize(0)==4 );
-    auto pIndices = reinterpret_cast<const uint32_t*>( pMesh->GetIndexBuffers().GetBufferData(0) );
+    auto pIndices = reinterpret_cast<const uint32*>( pMesh->GetIndexBuffers().GetBufferData(0) );
 
 
     // The mesh is supposed to contain only the faces in the selected block
@@ -962,7 +959,7 @@ void ImageRasterProjected_OptimisedWrapping( const Mesh* pMesh,
         scratch->culledVertex[v] = pVertices[v].nor[0] < fadeEndCos;
 
         // Cull vertices that don't belong to the current layout block.
-        if (pVertices[v].layoutBlock!=uint32_t(block))
+        if (pVertices[v].layoutBlock!=uint32(block))
         {
             scratch->culledVertex[v] = true;
         }
@@ -970,7 +967,7 @@ void ImageRasterProjected_OptimisedWrapping( const Mesh* pMesh,
 
     // Get the indices
     check( pMesh->GetIndexBuffers().GetElementSize(0)==4 );
-    auto pIndices = reinterpret_cast<const uint32_t*>( pMesh->GetIndexBuffers().GetBufferData(0) );
+    auto pIndices = reinterpret_cast<const uint32*>( pMesh->GetIndexBuffers().GetBufferData(0) );
 
 
     // The mesh is supposed to contain only the faces in the selected block
@@ -1580,7 +1577,7 @@ bool testPointsAreInOppositeSidesOfEdge( const vec2f& pointA, const vec2f& point
 struct PROJECTED_VERTEX
 {   
     float pos0 = 0, pos1 = 0, pos2 = 0;
-    uint32_t mask3 = 0;
+    uint32 mask3 = 0;
 
 	inline vec2<float> xy() const { return vec2<float>(pos0,pos1); }
 };
@@ -1590,16 +1587,16 @@ static_assert( sizeof(PROJECTED_VERTEX)==16, "Unexpected struct size" );
 
 //-------------------------------------------------------------------------------------------------
 void MeshProject_Optimised_Planar( const OPTIMISED_VERTEX* pVertices, int vertexCount,
-                                   const uint32_t* pIndices, int faceCount,
+                                   const uint32* pIndices, int faceCount,
                                    const vec3f& projectorPosition, const vec3f& projectorDirection,
                                    const vec3f& projectorSide, const vec3f& projectorUp,
                                    const vec3f& projectorScale,
                                    OPTIMISED_VERTEX* pResultVertices, int& currentVertex,
-                                   uint32_t* pResultIndices, int& currentIndex  )
+                                   uint32* pResultIndices, int& currentIndex  )
 {
 	MUTABLE_CPUPROFILER_SCOPE(MeshProject_Optimised_Planar)
 
-    vector<int32_t> oldToNewVertex(vertexCount,-1);
+    vector<int32> oldToNewVertex(vertexCount,-1);
 
     vector<PROJECTED_VERTEX> projectedPositions(vertexCount);
 
@@ -1611,7 +1608,7 @@ void MeshProject_Optimised_Planar( const OPTIMISED_VERTEX* pVertices, int vertex
         float z = dot( pVertices[v].pos-projectorPosition, projectorDirection ) / projectorScale[2];
 
         // Plane mask with bits for each plane discarding the vertex
-        uint32_t planeMask =
+        uint32 planeMask =
                 ((x<0.0f)<<0) |
                 ((x>1.0f)<<1) |
                 ((y<0.0f)<<2) |
@@ -1645,7 +1642,7 @@ void MeshProject_Optimised_Planar( const OPTIMISED_VERTEX* pVertices, int vertex
             // This face is required.
             for (int v=0;v<3;++v)
             {
-                uint32_t i = pIndices[f*3+v];
+                uint32 i = pIndices[f*3+v];
                 if (oldToNewVertex[i]<0)
                 {
                     pResultVertices[currentVertex] = pVertices[i];
@@ -1672,16 +1669,16 @@ void MeshProject_Optimised_Planar( const OPTIMISED_VERTEX* pVertices, int vertex
 
 //-------------------------------------------------------------------------------------------------
 void MeshProject_Optimised_Cylindrical( const OPTIMISED_VERTEX* pVertices, int vertexCount,
-                                        const uint32_t* pIndices, int faceCount,
+                                        const uint32* pIndices, int faceCount,
                                         const vec3f& projectorPosition, const vec3f& projectorDirection,
                                         const vec3f& projectorSide, const vec3f& projectorUp,
                                         const vec3f& projectorScale,
                                         OPTIMISED_VERTEX* pResultVertices, int& currentVertex,
-                                        uint32_t* pResultIndices, int& currentIndex  )
+                                        uint32* pResultIndices, int& currentIndex  )
 {
 	MUTABLE_CPUPROFILER_SCOPE(MeshProject_Optimised_Cylindrical)
 
-	vector<int32_t> oldToNewVertex(vertexCount,-1);
+	vector<int32> oldToNewVertex(vertexCount,-1);
 
     vector<PROJECTED_VERTEX> projectedPositions(vertexCount);
 
@@ -1706,7 +1703,7 @@ void MeshProject_Optimised_Cylindrical( const OPTIMISED_VERTEX* pVertices, int v
         projectedPositions[v].pos0 = x;
         projectedPositions[v].pos1 = vertexPos_cylinder.y() / radius;
         projectedPositions[v].pos2 = vertexPos_cylinder.z() / radius;
-        uint32_t planeMask =
+        uint32 planeMask =
                 ((x<0.0f)<<0) |
                 ((x>1.0f)<<1) |
                 ((r2>=(radius*radius))<<2);
@@ -1737,7 +1734,7 @@ void MeshProject_Optimised_Cylindrical( const OPTIMISED_VERTEX* pVertices, int v
             // This face is required.
             for (int v=0;v<3;++v)
             {
-                uint32_t i = pIndices[f*3+v];
+                uint32 i = pIndices[f*3+v];
                 if (oldToNewVertex[i]<0)
                 {
                     pResultVertices[currentVertex] = pVertices[i];
@@ -1769,7 +1766,7 @@ void MeshProject_Optimised_Wrapping( const Mesh* pMesh,
                                      const vec3f& projectorSide, const vec3f& projectorUp,
                                      const vec3f& projectorScale,
                                      OPTIMISED_VERTEX_WRAPPING* pResultVertices, int& currentVertex,
-                                     uint32_t* pResultIndices, int& currentIndex )
+                                     uint32* pResultIndices, int& currentIndex )
 {
 	MUTABLE_CPUPROFILER_SCOPE(MeshProject_Optimised_Wrapping)
 
@@ -1780,7 +1777,7 @@ void MeshProject_Optimised_Wrapping( const Mesh* pMesh,
 
     // Get the indices
     check( pMesh->GetIndexBuffers().GetElementSize(0)==4 );
-    auto pIndices = reinterpret_cast<const uint32_t*>( pMesh->GetIndexBuffers().GetBufferData(0) );
+    auto pIndices = reinterpret_cast<const uint32*>( pMesh->GetIndexBuffers().GetBufferData(0) );
     int faceCount = pMesh->GetFaceCount();
 
     vector<PROJECTED_VERTEX> projectedPositions(vertexCount);
@@ -2434,7 +2431,7 @@ void MeshProject_Optimised_Wrapping( const Mesh* pMesh,
         }
     }
 
-    vector<int32_t> oldToNewVertex(vertexCount,-1);
+    vector<int32> oldToNewVertex(vertexCount,-1);
 
     // Add the projected face
     for(int f : processedFaces)
@@ -2576,7 +2573,7 @@ MeshPtr MeshProject_Optimised( const Mesh* pMesh,
         pResult = CreateMeshOptimisedForProjection( layout );
         pResult->GetVertexBuffers().SetElementCount(vertexCount);
         pResult->GetIndexBuffers().SetElementCount(indexCount);
-        auto pResultIndices = reinterpret_cast<uint32_t*>( pResult->GetIndexBuffers().GetBufferData(0) );
+        auto pResultIndices = reinterpret_cast<uint32*>( pResult->GetIndexBuffers().GetBufferData(0) );
         auto pResultVertices = reinterpret_cast<OPTIMISED_VERTEX*>( pResult->GetVertexBuffers().GetBufferData(0) );
 
         // Get the vertices
@@ -2585,7 +2582,7 @@ MeshPtr MeshProject_Optimised( const Mesh* pMesh,
 
         // Get the indices
         check( pMesh->GetIndexBuffers().GetElementSize(0)==4 );
-        auto pIndices = reinterpret_cast<const uint32_t*>( pMesh->GetIndexBuffers().GetBufferData(0) );
+        auto pIndices = reinterpret_cast<const uint32*>( pMesh->GetIndexBuffers().GetBufferData(0) );
         int faceCount = pMesh->GetFaceCount();
 
         if (projector.type==PROJECTOR_TYPE::PLANAR)
@@ -2618,7 +2615,7 @@ MeshPtr MeshProject_Optimised( const Mesh* pMesh,
         pResult = CreateMeshOptimisedForWrappingProjection( layout );
         pResult->GetVertexBuffers().SetElementCount(vertexCount);
         pResult->GetIndexBuffers().SetElementCount(indexCount);
-        auto pResultIndices = reinterpret_cast<uint32_t*>( pResult->GetIndexBuffers().GetBufferData(0) );
+        auto pResultIndices = reinterpret_cast<uint32*>( pResult->GetIndexBuffers().GetBufferData(0) );
         auto pResultVertices = reinterpret_cast<OPTIMISED_VERTEX_WRAPPING*>( pResult->GetVertexBuffers().GetBufferData(0) );
 
         MeshProject_Optimised_Wrapping( pMesh,
diff --git a/Engine/Plugins/Experimental/Mutable/Source/MutableRuntime/Private/MuR/OpImageProject.h b/Engine/Plugins/Experimental/Mutable/Source/MutableRuntime/Private/MuR/OpImageProject.h
--- a/Engine/Plugins/Experimental/Mutable/Source/MutableRuntime/Private/MuR/OpImageProject.h
+++ b/Engine/Plugins/Experimental/Mutable/Source/MutableRuntime/Private/MuR/OpImageProject.h
@@ -5,19 +5,17 @@
 #include "MuR/MemoryPrivate.h"
 #include "MuR/Mesh.h"
 
-#include <stdint.h>
-
 
 namespace mu
 {
-class Image;
-struct PROJECTOR;
-template <int NUM_INTERPOLATORS> class RasterVertex;
+	class Image;
+	struct PROJECTOR;
+	template <int NUM_INTERPOLATORS> class RasterVertex;
 
     struct SCRATCH_IMAGE_PROJECT
     {
         vector< RasterVertex<4> > vertices;
-        vector<uint8_t> culledVertex;
+        vector<uint8> culledVertex;
     };
 
 
diff --git a/Engine/Plugins/Experimental/Mutable/Source/MutableRuntime/Private/MuR/OpLayoutPack.h b/Engine/Plugins/Experimental/Mutable/Source/MutableRuntime/Private/MuR/OpLayoutPack.h
--- a/Engine/Plugins/Experimental/Mutable/Source/MutableRuntime/Private/MuR/OpLayoutPack.h
+++ b/Engine/Plugins/Experimental/Mutable/Source/MutableRuntime/Private/MuR/OpLayoutPack.h
@@ -5,7 +5,7 @@
 #include "MuR/Layout.h"
 #include "MuR/Platform.h"
 
-#include <algorithm>
+#include <limits>
 
 
 namespace mu
@@ -61,7 +61,7 @@ namespace mu
 	// 		int thisArea = blocks[index][0] * blocks[index][1];
 
 	// 		int p;
-	// 		for ( p=0; p<(int)sorted.size(); ++p )
+	// 		for ( p=0; p<(int)sorted.Num(); ++p )
 	// 		{
 	// 			int pos = sorted[p];
 	// 			int posHeight = blocks[pos][1];
@@ -81,7 +81,7 @@ namespace mu
 	// 	vector< vec2<int> > positions( blockCount );
 	// 	maxY = 0;
 
-	// 	for ( int p=0; p<(int)sorted.size(); ++p )
+	// 	for ( int p=0; p<(int)sorted.Num(); ++p )
 	// 	{
 	// 		int index = sorted[p];
 
@@ -195,19 +195,19 @@ namespace mu
     //---------------------------------------------------------------------------------------------
     struct SCRATCH_LAYOUT_PACK
     {
-        vector< vec2<int> > blocks;
-        vector< LAY_BLOCK > sorted;
-        vector< vec2<int> > positions;
-		vector< int > priorities;
-		vector< vec2<int> > reductions;
+        TArray< vec2<int> > blocks;
+		TArray< LAY_BLOCK > sorted;
+		TArray< vec2<int> > positions;
+		TArray< int > priorities;
+		TArray< vec2<int> > reductions;
     };
 
 
     inline string DebugGetBlockAt( const SCRATCH_LAYOUT_PACK& scratch,
-                                   const vector<uint8_t>& packedFlag,
+                                   const TArray<uint8_t>& packedFlag,
                                    int x, int y )
     {
-        for ( size_t b=0; b<packedFlag.size(); ++b )
+        for ( size_t b=0; b<packedFlag.Num(); ++b )
         {
             if (packedFlag[b])
             {
@@ -236,7 +236,7 @@ namespace mu
     //     check( pResult->GetBlockCount() == pSourceLayout->GetBlockCount() );
 
     //     int blockCount = pSourceLayout->GetBlockCount();
-    //     check( (int)scratch->blocks.size()==blockCount );
+    //     check( (int)scratch->blocks.Num()==blockCount );
 
     //     // Look for the maximum block sizes on the layout and the total area
     //     int maxX = 0;
@@ -272,14 +272,14 @@ namespace mu
     //     }
 
     //     // Sort by height, area
-    //     check( (int)scratch->sorted.size()==blockCount );
+    //     check( (int)scratch->sorted.Num()==blockCount );
     //     for ( int index=0; index<blockCount; ++index )
     //     {
     //         scratch->sorted[index] = LAY_BLOCK( index, scratch->blocks[index], scratch->priorities[index] );
     //     }
     //     std::sort( scratch->sorted.begin(), scratch->sorted.end(), CompareBlocks );
 
-    //     check( (int)scratch->sorted.size()==blockCount );
+    //     check( (int)scratch->sorted.Num()==blockCount );
 
     //     bool fits = false;
     //     int iterations = 0;
@@ -294,7 +294,7 @@ namespace mu
     //         mutable_memset( horizon, 0, 256*sizeof(int16_t) );
     //         maxY = 0;
 
-    //         for ( size_t p=0; p<scratch->sorted.size(); ++p )
+    //         for ( size_t p=0; p<scratch->sorted.Num(); ++p )
     //         {
     //             // Seek for the lowest span where the block fits
     //             int currentLevel = std::numeric_limits<int>::max();
@@ -442,7 +442,7 @@ namespace mu
 
 		// Number of blocks alrady packed
 		size_t packed = 0;
-		vector<uint8_t> packedFlag(scratch->sorted.size(), 0);
+		vector<uint8_t> packedFlag(scratch->sorted.Num(), 0);
 
 		// Pack with fixed horizontal size
 		check(*maxX < 256);
@@ -454,7 +454,7 @@ namespace mu
 
 		int iterations = 0;
 
-		while (packed < scratch->sorted.size() || iterations > 5000)
+		while (packed < scratch->sorted.Num() || iterations > 5000)
 		{
 			++iterations;
 
@@ -464,7 +464,7 @@ namespace mu
 			int bestWithHole = -1;
 			int bestWithHoleX = -1;
 			int bestWithHoleLevel = -1;
-			for (size_t candidate = 0; candidate < scratch->sorted.size(); ++candidate)
+			for (size_t candidate = 0; candidate < scratch->sorted.Num(); ++candidate)
 			{
 				// Skip it if we packed it already
 				if (packedFlag[candidate]) continue;
@@ -494,12 +494,12 @@ namespace mu
 					// Does it make an unfillable hole with the top or side?
 					int minX = std::numeric_limits<int>::max();
 					int minY = std::numeric_limits<int>::max();
-					for (size_t b = 0; b < scratch->sorted.size(); ++b)
+					for (size_t b = 0; b < scratch->sorted.Num(); ++b)
 					{
 						if (!packedFlag[b] && b != candidate)
 						{
-							minX = std::min(minX, scratch->sorted[b].size[0]);
-							minY = std::min(minY, scratch->sorted[b].size[1]);
+							minX = FMath::Min(minX, scratch->sorted[b].size[0]);
+							minY = FMath::Min(minY, scratch->sorted[b].size[1]);
 						}
 					}
 
@@ -572,7 +572,7 @@ namespace mu
 
 			// Store
 			scratch->positions[scratch->sorted[best].index] = vec2<int>(bestX, bestLevel);
-			*maxY = std::max(*maxY, bestLevel + scratch->sorted[best].size[1]);
+			*maxY = FMath::Max(*maxY, bestLevel + scratch->sorted[best].size[1]);
 
 			if (packStrategy == EPackStrategy::FIXED_LAYOUT && *maxY > layoutSizeY)
 			{
@@ -595,7 +595,7 @@ namespace mu
 			//                AXE_LOG("layout",Warning,"--------------------------------------------------");
 		}
 
-		if (packed < scratch->sorted.size())
+		if (packed < scratch->sorted.Num())
 		{
 			fits = false;
 		}
@@ -615,7 +615,7 @@ namespace mu
         check( pResult->GetBlockCount() == pSourceLayout->GetBlockCount() );
 
         int blockCount = pSourceLayout->GetBlockCount();
-        check( (int)scratch->blocks.size()==blockCount );
+        check( (int)scratch->blocks.Num()==blockCount );
 
 		//Getting maximum layout grid size:
 		int layoutSizeX, layoutSizeY;
@@ -644,8 +644,8 @@ namespace mu
 				usePriority = true;
 			}
 
-            maxX = std::max( maxX, b.size[0] );
-            maxY = std::max( maxY, b.size[1] );
+            maxX = FMath::Max( maxX, b.size[0] );
+            maxY = FMath::Max( maxY, b.size[1] );
 
             area += b.size[0] * b.size[1];
 
@@ -717,7 +717,7 @@ namespace mu
 				//recalculating area and maximum block sizes
 				for (int index = 0; index < blockCount; ++index)
 				{
-					maxX = std::max(maxX, scratch->blocks[index][0]);
+					maxX = FMath::Max(maxX, scratch->blocks[index][0]);
 					area += scratch->blocks[index][0] * scratch->blocks[index][1];
 				}
 			}
@@ -739,7 +739,7 @@ namespace mu
 				//recalculating area and maximum block sizes
 				for (int index = 0; index < blockCount; ++index)
 				{
-					maxY = std::max(maxY, scratch->blocks[index][1]);
+					maxY = FMath::Max(maxY, scratch->blocks[index][1]);
 					area += scratch->blocks[index][0] * scratch->blocks[index][1];
 				}
 			}
@@ -765,19 +765,19 @@ namespace mu
 		int r_it = 0;
 
         // Sort by height, area
-        check( (int)scratch->sorted.size()==blockCount );
+        check( (int)scratch->sorted.Num()==blockCount );
         for ( int index=0; index<blockCount; ++index )
         {
             scratch->sorted[index] = LAY_BLOCK( index, scratch->blocks[index], scratch->priorities[index] );
         }
 
-		std::sort(scratch->sorted.begin(), scratch->sorted.end(), CompareBlocks);
+		scratch->sorted.Sort( CompareBlocks);
 		
 		if(pSourceLayout->GetLayoutPackingStrategy() == EPackStrategy::FIXED_LAYOUT)
 		{
 			if (usePriority)	//Sort by priority
 			{
-				std::sort(scratch->sorted.begin(), scratch->sorted.end(), CompareBlocksPriority);
+				scratch->sorted.Sort(CompareBlocksPriority);
 			}
 
 			//Shrink blocks in case we do not have enough space to pack everything
@@ -795,7 +795,7 @@ namespace mu
 			//Sort by height&area before packing
 			if (usePriority)
 			{
-				std::sort(scratch->sorted.begin(), scratch->sorted.end(), CompareBlocks);
+				scratch->sorted.Sort(CompareBlocks);
 			}
 
 			//Try to pack everything
@@ -806,7 +806,7 @@ namespace mu
 				//Sort by priority before shrink
 				if (usePriority)
 				{
-					std::sort(scratch->sorted.begin(), scratch->sorted.end(), CompareBlocksPriority);
+					scratch->sorted.Sort(CompareBlocksPriority);
 				}
 
 				ReduceBlock(blockCount, &area, &r_it, scratch);
diff --git a/Engine/Plugins/Experimental/Mutable/Source/MutableRuntime/Private/MuR/OpMeshApplyPose.h b/Engine/Plugins/Experimental/Mutable/Source/MutableRuntime/Private/MuR/OpMeshApplyPose.h
--- a/Engine/Plugins/Experimental/Mutable/Source/MutableRuntime/Private/MuR/OpMeshApplyPose.h
+++ b/Engine/Plugins/Experimental/Mutable/Source/MutableRuntime/Private/MuR/OpMeshApplyPose.h
@@ -7,8 +7,6 @@
 #include "MuR/Platform.h"
 #include "MuR/OpMeshChartDifference.h"
 
-#include <memory>
-
 
 namespace mu
 {
diff --git a/Engine/Plugins/Experimental/Mutable/Source/MutableRuntime/Private/MuR/OpMeshClipWithMesh.h b/Engine/Plugins/Experimental/Mutable/Source/MutableRuntime/Private/MuR/OpMeshClipWithMesh.h
--- a/Engine/Plugins/Experimental/Mutable/Source/MutableRuntime/Private/MuR/OpMeshClipWithMesh.h
+++ b/Engine/Plugins/Experimental/Mutable/Source/MutableRuntime/Private/MuR/OpMeshClipWithMesh.h
@@ -16,8 +16,6 @@
 
 #include "Math/UnrealMathUtility.h"
 
-#include <memory>
-
 
 namespace mu
 {
diff --git a/Engine/Plugins/Experimental/Mutable/Source/MutableRuntime/Private/MuR/OpMeshFormat.cpp b/Engine/Plugins/Experimental/Mutable/Source/MutableRuntime/Private/MuR/OpMeshFormat.cpp
--- a/Engine/Plugins/Experimental/Mutable/Source/MutableRuntime/Private/MuR/OpMeshFormat.cpp
+++ b/Engine/Plugins/Experimental/Mutable/Source/MutableRuntime/Private/MuR/OpMeshFormat.cpp
@@ -19,9 +19,6 @@
 #include "MuR/Ptr.h"
 #include "MuR/RefCounted.h"
 
-#include <memory>
-#include <utility>
-
 
 namespace mu
 {
diff --git a/Engine/Plugins/Experimental/Mutable/Source/MutableRuntime/Private/MuR/Operations.h b/Engine/Plugins/Experimental/Mutable/Source/MutableRuntime/Private/MuR/Operations.h
--- a/Engine/Plugins/Experimental/Mutable/Source/MutableRuntime/Private/MuR/Operations.h
+++ b/Engine/Plugins/Experimental/Mutable/Source/MutableRuntime/Private/MuR/Operations.h
@@ -1028,7 +1028,7 @@ namespace mu
         {
             if (o.type!=type) return false;
             //if (o.flags!=flags) return false;
-            if (mutable_memcmp(&o.args, &args, sizeof(ARGS))) return false;
+            if (FMemory::Memcmp(&o.args, &args, sizeof(ARGS))) return false;
             return true;
         }
 
@@ -1038,7 +1038,7 @@ namespace mu
             if (o.type>type) return false;
 //            if (o.flags<flags) return true;
 //            if (o.flags>flags) return false;
-            if (mutable_memcmp(&o.args, &args, sizeof(ARGS))<0) return true;
+            if (FMemory::Memcmp(&o.args, &args, sizeof(ARGS))<0) return true;
             return false;
         }
 
diff --git a/Engine/Plugins/Experimental/Mutable/Source/MutableRuntime/Private/MuR/Parameters.cpp b/Engine/Plugins/Experimental/Mutable/Source/MutableRuntime/Private/MuR/Parameters.cpp
--- a/Engine/Plugins/Experimental/Mutable/Source/MutableRuntime/Private/MuR/Parameters.cpp
+++ b/Engine/Plugins/Experimental/Mutable/Source/MutableRuntime/Private/MuR/Parameters.cpp
@@ -15,9 +15,6 @@
 #include "MuR/SerialisationPrivate.h"
 #include "MuR/System.h"
 
-#include <memory>
-#include <utility>
-
 
 namespace mu
 {
diff --git a/Engine/Plugins/Experimental/Mutable/Source/MutableRuntime/Private/MuR/ParametersPrivate.h b/Engine/Plugins/Experimental/Mutable/Source/MutableRuntime/Private/MuR/ParametersPrivate.h
--- a/Engine/Plugins/Experimental/Mutable/Source/MutableRuntime/Private/MuR/ParametersPrivate.h
+++ b/Engine/Plugins/Experimental/Mutable/Source/MutableRuntime/Private/MuR/ParametersPrivate.h
@@ -303,7 +303,7 @@ namespace mu
         //!
         bool operator==( const PARAMETER_VALUE& other ) const
         {
-            return mutable_memcmp(this, &other, sizeof(PARAMETER_VALUE))==0;
+            return FMemory::Memcmp(this, &other, sizeof(PARAMETER_VALUE))==0;
         }
 
     };
diff --git a/Engine/Plugins/Experimental/Mutable/Source/MutableRuntime/Private/MuR/Platform.h b/Engine/Plugins/Experimental/Mutable/Source/MutableRuntime/Private/MuR/Platform.h
--- a/Engine/Plugins/Experimental/Mutable/Source/MutableRuntime/Private/MuR/Platform.h
+++ b/Engine/Plugins/Experimental/Mutable/Source/MutableRuntime/Private/MuR/Platform.h
@@ -15,91 +15,10 @@
 
 #include "MuR/Types.h"
 #include "HAL/UnrealMemory.h"
-#include <cstdint>
 
 MUTABLERUNTIME_API DECLARE_LOG_CATEGORY_EXTERN(LogMutableCore, Log, All);
 
 
-#ifdef __ANDROID__
-    #define MUTABLE_PLATFORM_ANDROID
-#elif defined(_XBOX_ONE)
-	#define MUTABLE_PLATFORM_XBOXONE
-#elif defined(__PS4__)
-	#define MUTABLE_PLATFORM_PS4
-#elif defined(__PS5__)
-	#define MUTABLE_PLATFORM_PS5
-#elif defined(__SWITCH__)
-    #define MUTABLE_PLATFORM_SWITCH
-#elif defined(_WIN64) && _WIN64
-    #define MUTABLE_PLATFORM_WIN64
-#elif defined(_WIN32) && _WIN32
-    #define MUTABLE_PLATFORM_WIN32
-#elif defined(__APPLE__) && __APPLE__
-    #include "TargetConditionals.h"
-    #if TARGET_IPHONE_SIMULATOR
-        #define MUTABLE_PLATFORM_IOS
-    #elif TARGET_OS_IPHONE
-        #define MUTABLE_PLATFORM_IOS
-    #elif TARGET_OS_MAC
-        #define MUTABLE_PLATFORM_OSX
-    #else
-        // Unsupported platform
-        #warning "Unidentified platform."
-    #endif
-#elif defined(__linux) && __linux
-    #define MUTABLE_PLATFORM_LINUX
-#elif defined(__unix) && __unix // all unices not caught above
-    #define MUTABLE_PLATFORM_LINUX
-#elif defined(__posix) && __posix
-    #define MUTABLE_PLATFORM_LINUX
-#endif
-
-
-#if defined(MUTABLE_PLATFORM_XBOXONE)
-
-	#define NOMINMAX
-	#include <xdk.h>
-	#include <wrl.h>
-
-#elif defined(MUTABLE_PLATFORM_WIN64) || defined(MUTABLE_PLATFORM_WIN32)
-
-	#ifdef _MSC_VER
-
-		#ifndef _HAS_EXCEPTIONS
-			#define _HAS_EXCEPTIONS 0
-		#endif
-
-		#ifndef WIN32_LEAN_AND_MEAN
-			#define WIN32_LEAN_AND_MEAN
-		#endif
-
-        #define NOMINMAX
-        //#include <windows.h>
-
-		// Some warnings that are more annoying than useful
-		#pragma warning( disable : 4127 )
-		#pragma warning( disable : 6385 )
-		#pragma warning( disable : 6287 )
-
-
-    #else
-
-        // Some unix environment under windows.
-        #include <sys/time.h>
-
-	#endif
-
-#elif defined(MUTABLE_PLATFORM_ANDROID)
-
-    #include <sys/time.h>
-
-#else
-
-    #include <sys/time.h>
-
-#endif
-
-
 //! Unify debug defines
 #if !defined(MUTABLE_DEBUG)
     #if !defined(NDEBUG) || defined(_DEBUG)
@@ -111,25 +30,6 @@ MUTABLERUNTIME_API DECLARE_LOG_CATEGORY_EXTERN(LogMutableCore, Log, All);
 namespace mu
 {
 
-    inline void* mutable_system_malloc( size_t size, uint32_t alignment )
-    {        
-        return FMemory::Malloc(size,alignment);
-    }
-
-
-    inline void mutable_system_free( void* ptr )
-    {
-		FMemory::Free(ptr);
-    }
-
-
-	//---------------------------------------------------------------------------------------------
-	inline int mutable_memcmp ( const void * a, const void* b, size_t num )
-	{
-		return memcmp( a, b, num );
-	}
-
-
 	//---------------------------------------------------------------------------------------------
 	inline int mutable_vsnprintf
 		( char *buffer, size_t sizeOfBuffer, const char *format, va_list argptr )
@@ -165,48 +65,4 @@ namespace mu
 		#endif
 	}
 
-
-	//---------------------------------------------------------------------------------------------
-	//! Simple profiling timer.
-	//! It should only be used in debug and profile builds.
-	//---------------------------------------------------------------------------------------------
-	struct TIMER
-	{
-	public:
-
-		//!
-		inline void Start()
-		{
-#ifdef MUTABLE_PROFILE
-			m_start = int64( FPlatformTime::Seconds()*1000000.0 );
-#endif
-		}
-
-		//! Stop counting and return the ellapsed time since Start in milliseconds.
-		inline int GetMilliseconds() const
-		{
-#ifdef MUTABLE_PROFILE
-			return int( (int64(FPlatformTime::Seconds() * 1000000.0) - m_start)/1000 );
-#else
-			return 0;
-#endif
-		}
-
-		//! Stop counting and return the ellapsed time since Start in microseconds.
-		inline int GetMicroseconds() const
-		{
-#ifdef MUTABLE_PROFILE
-			return int((int64(FPlatformTime::Seconds() * 1000000.0) - m_start));
-#else
-			return 0;
-#endif
-		}
-
-	private:
-
-#ifdef MUTABLE_PROFILE
-	int64_t m_start=0;
-#endif
-	};
-
 }
diff --git a/Engine/Plugins/Experimental/Mutable/Source/MutableRuntime/Private/MuR/Serialisation.cpp b/Engine/Plugins/Experimental/Mutable/Source/MutableRuntime/Private/MuR/Serialisation.cpp
--- a/Engine/Plugins/Experimental/Mutable/Source/MutableRuntime/Private/MuR/Serialisation.cpp
+++ b/Engine/Plugins/Experimental/Mutable/Source/MutableRuntime/Private/MuR/Serialisation.cpp
@@ -11,8 +11,7 @@
 #include "MuR/MemoryPrivate.h"
 #include "MuR/SerialisationPrivate.h"
 
-#include <memory>
-#include <utility>
+#include <limits>
 
 
 namespace mu
@@ -50,7 +49,7 @@ namespace mu
             check( size<std::numeric_limits<size_t>::max() );
 
             const uint8_t* pSource = ((const uint8_t*)(m_pD->m_pBuffer))+m_pD->m_pos;
-            memcpy( pData, pSource, (size_t)size );
+            FMemory::Memcpy( pData, pSource, (size_t)size );
             m_pD->m_pos += size;
         }
     }
@@ -90,7 +89,7 @@ namespace mu
 
             uint64 pos = m_pD->m_buffer.Num();
             m_pD->m_buffer.SetNum( pos + size, false );
-            memcpy( &m_pD->m_buffer[pos], pData, size );
+			FMemory::Memcpy( &m_pD->m_buffer[pos], pData, size );
         }
     }
 
diff --git a/Engine/Plugins/Experimental/Mutable/Source/MutableRuntime/Private/MuR/Settings.cpp b/Engine/Plugins/Experimental/Mutable/Source/MutableRuntime/Private/MuR/Settings.cpp
--- a/Engine/Plugins/Experimental/Mutable/Source/MutableRuntime/Private/MuR/Settings.cpp
+++ b/Engine/Plugins/Experimental/Mutable/Source/MutableRuntime/Private/MuR/Settings.cpp
@@ -40,7 +40,7 @@ namespace mu
 
 
     //---------------------------------------------------------------------------------------------
-    void Settings::SetStreamingCache( uint64_t bytes )
+    void Settings::SetStreamingCache( uint64 bytes )
     {
         m_pD->m_streamingCacheBytes = bytes;
     }
diff --git a/Engine/Plugins/Experimental/Mutable/Source/MutableRuntime/Private/MuR/SettingsPrivate.h b/Engine/Plugins/Experimental/Mutable/Source/MutableRuntime/Private/MuR/SettingsPrivate.h
--- a/Engine/Plugins/Experimental/Mutable/Source/MutableRuntime/Private/MuR/SettingsPrivate.h
+++ b/Engine/Plugins/Experimental/Mutable/Source/MutableRuntime/Private/MuR/SettingsPrivate.h
@@ -14,8 +14,8 @@ namespace mu
     {
     public:
         bool m_profile = false;
-        uint64_t m_streamingCacheBytes = 0;
-        int m_imageCompressionQuality = 0;
+        uint64 m_streamingCacheBytes = 0;
+        int32 m_imageCompressionQuality = 0;
     };
 
 
diff --git a/Engine/Plugins/Experimental/Mutable/Source/MutableRuntime/Private/MuR/System.cpp b/Engine/Plugins/Experimental/Mutable/Source/MutableRuntime/Private/MuR/System.cpp
--- a/Engine/Plugins/Experimental/Mutable/Source/MutableRuntime/Private/MuR/System.cpp
+++ b/Engine/Plugins/Experimental/Mutable/Source/MutableRuntime/Private/MuR/System.cpp
@@ -33,9 +33,6 @@
 #include "Templates/Tuple.h"
 #include "Trace/Detail/Channel.h"
 
-#include <functional>
-#include <memory>
-#include <utility>
 
 namespace mu
 {
diff --git a/Engine/Plugins/Experimental/Mutable/Source/MutableRuntime/Public/MuR/Instance.h b/Engine/Plugins/Experimental/Mutable/Source/MutableRuntime/Public/MuR/Instance.h
--- a/Engine/Plugins/Experimental/Mutable/Source/MutableRuntime/Public/MuR/Instance.h
+++ b/Engine/Plugins/Experimental/Mutable/Source/MutableRuntime/Public/MuR/Instance.h
@@ -40,7 +40,7 @@ namespace mu
 		//-----------------------------------------------------------------------------------------
 
         //! Type for the instance unique identifiers.
-        typedef uint32_t ID;
+        typedef uint32 ID;
 
         //! Get a unique identifier for this instance. It doesn't change during the entire
         //! lifecycle of each instance. This identifier can be used in the System methods to update
diff --git a/Engine/Plugins/Experimental/Mutable/Source/MutableRuntime/Public/MuR/Layout.h b/Engine/Plugins/Experimental/Mutable/Source/MutableRuntime/Public/MuR/Layout.h
--- a/Engine/Plugins/Experimental/Mutable/Source/MutableRuntime/Public/MuR/Layout.h
+++ b/Engine/Plugins/Experimental/Mutable/Source/MutableRuntime/Public/MuR/Layout.h
@@ -21,7 +21,7 @@ namespace mu
 	typedef Ptr<const Layout> LayoutPtrConst;
 
 	//! Types of layout packing strategies 
-	enum class EPackStrategy : uint32_t
+	enum class EPackStrategy : uint32
 	{
 		RESIZABLE_LAYOUT,
 		FIXED_LAYOUT
@@ -136,7 +136,7 @@ namespace mu
 		//-----------------------------------------------------------------------------------------
 		struct BLOCK
 		{
-			BLOCK(vec2<uint16_t> min = vec2<uint16_t>(), vec2<uint16_t> size = vec2<uint16_t>())
+			BLOCK(vec2<uint16> min = vec2<uint16>(), vec2<uint16> size = vec2<uint16>())
 			{
 				m_min = min;
 				m_size = size;
@@ -144,14 +144,14 @@ namespace mu
 				m_priority = 0;
 			}
 
-			vec2<uint16_t> m_min;
-			vec2<uint16_t> m_size;
+			vec2<uint16> m_min;
+			vec2<uint16> m_size;
 
 			//! Absolute id used to control merging of various layouts
-			int32_t m_id;
+			int32 m_id;
 
 			//! Priority value to control the shrink texture layout strategy
-			int32_t m_priority;
+			int32 m_priority;
 
 
 			//!
@@ -192,9 +192,9 @@ namespace mu
 
 
 		//!
-		vec2<uint16_t> m_size;
+		vec2<uint16> m_size;
 
-		vec2<uint16_t> m_maxsize;
+		vec2<uint16> m_maxsize;
 
 		//!
 		TArray<BLOCK> m_blocks;
@@ -205,20 +205,20 @@ namespace mu
 		//!
 		void Serialise(OutputArchive& arch) const
 		{
-			uint32_t ver = 3;
+			uint32 ver = 3;
 			arch << ver;
 
 			arch << m_size;
 			arch << m_blocks;
 
 			arch << m_maxsize;
-			arch << uint32_t(m_strategy);
+			arch << uint32(m_strategy);
 		}
 
 		//!
 		void Unserialise(InputArchive& arch)
 		{
-			uint32_t ver;
+			uint32 ver;
 			arch >> ver;
 			check(ver == 3);
 
@@ -227,7 +227,7 @@ namespace mu
 			arch >> m_blocks;
 			arch >> m_maxsize;
 
-			uint32_t temp;
+			uint32 temp;
 			arch >> temp;
 			m_strategy = EPackStrategy(temp);
 		}
@@ -237,7 +237,7 @@ namespace mu
 
 		//! Find a block by id. This converts the "absolute" id to a relative index to the layout
 		//! blocks. Return -1 if not found.
-		int32 FindBlock(int32_t id) const;
+		int32 FindBlock(int32 id) const;
 
 		//! Return true if the layout is a single block filling all area.
 		bool IsSingleBlockAndFull() const;
diff --git a/Engine/Plugins/Experimental/Mutable/Source/MutableRuntime/Public/MuR/Mesh.h b/Engine/Plugins/Experimental/Mutable/Source/MutableRuntime/Public/MuR/Mesh.h
--- a/Engine/Plugins/Experimental/Mutable/Source/MutableRuntime/Public/MuR/Mesh.h
+++ b/Engine/Plugins/Experimental/Mutable/Source/MutableRuntime/Public/MuR/Mesh.h
@@ -21,8 +21,6 @@
 #include "MuR/Skeleton.h"
 #include "Templates/Tuple.h"
 
-#include <memory>
-
 class FString;
 
 namespace mu
@@ -56,11 +54,11 @@ namespace mu
 			m_id = 0;
 		}
 
-		int32_t m_firstVertex;
-		int32_t m_vertexCount;
-		int32_t m_firstIndex;
-		int32_t m_indexCount;
-		uint32_t m_id;
+		int32 m_firstVertex;
+		int32 m_vertexCount;
+		int32 m_firstIndex;
+		int32 m_indexCount;
+		uint32 m_id;
 
 
 		//!
@@ -179,7 +177,7 @@ namespace mu
 
         //! Return an internal id that can be used to match mesh surfaces and instance surfaces.
         //! Only valid for meshes that are part of instances.
-        uint32_t GetSurfaceId( int surfaceIndex ) const;
+        uint32 GetSurfaceId( int surfaceIndex ) const;
 
         //! \}
 
@@ -281,7 +279,7 @@ namespace mu
 
         //! Get an internal identifier used to reference this mesh in operations like deferred
         //! mesh building, or instance updating.
-        uint32_t GetId() const;
+        uint32 GetId() const;
 
 
     protected:
@@ -294,7 +292,7 @@ namespace mu
 
 		//! Non-persistent internal id unique for a mesh generated for a specific state and
 		//! parameter values.
-		mutable uint32_t m_internalId = 0;
+		mutable uint32 m_internalId = 0;
 
 		//!
 		FMeshBufferSet m_VertexBuffers;
@@ -311,7 +309,7 @@ namespace mu
 		//! This is bit-mask on the STATIC_MESH_FORMATS enumeration, marking what static formats
 		//! are compatible with this one. Usually precalculated at model compilation time.
 		//! It should be reset after any operation that modifies the format.
-		mutable uint32_t m_staticFormatFlags = 0;
+		mutable uint32 m_staticFormatFlags = 0;
 
 		TArray<MESH_SURFACE> m_surfaces;
 
@@ -331,7 +329,7 @@ namespace mu
 
 			inline void Serialise(OutputArchive& arch) const
 			{
-				const int32_t ver = 0;
+				const int32 ver = 0;
 				arch << ver;
 
 				arch << m_name;
@@ -340,7 +338,7 @@ namespace mu
 
 			inline void Unserialise(InputArchive& arch)
 			{
-				int32_t ver = 0;
+				int32 ver = 0;
 				arch >> ver;
 				check(ver == 0);
 
@@ -371,7 +369,7 @@ namespace mu
 
 			inline void Serialise(OutputArchive& arch) const
 			{
-				const int32_t ver = 0;
+				const int32 ver = 0;
 				arch << ver;
 
 				arch << m_boneName;
@@ -381,7 +379,7 @@ namespace mu
 
 			inline void Unserialise(InputArchive& arch)
 			{
-				int32_t ver = 0;
+				int32 ver = 0;
 				arch >> ver;
 				check(ver == 0);
 
@@ -409,7 +407,7 @@ namespace mu
 		//!
 		inline void Serialise(OutputArchive& arch) const
 		{
-			uint32_t ver = 13;
+			uint32 ver = 13;
 			arch << ver;
 
 			arch << m_IndexBuffers;
@@ -433,7 +431,7 @@ namespace mu
 		//!
 		inline void Unserialise(InputArchive& arch)
 		{
-			uint32_t ver;
+			uint32 ver;
 			arch >> ver;
 			check(ver <= 13);
 
@@ -525,7 +523,7 @@ namespace mu
 		bool IsSimilar(const Mesh& o) const;
 
 		//! Compare the vertex attributes to check if they match.
-		bool IsSameVertex( uint32_t v, const Mesh& other, uint32_t otherVertexIndex, float tolerance = 1e-5f) const;
+		bool IsSameVertex( uint32 v, const Mesh& other, uint32 otherVertexIndex, float tolerance = 1e-5f) const;
 
 
 		//! Make a map from the vertices in this mesh to thefirst matching vertex of the given
@@ -554,7 +552,7 @@ namespace mu
 		bool HasFace( const Mesh& other, int otherFaceIndex, const VERTEX_MATCH_MAP& vertexMap ) const;
 
 		//! Compare the vertex attributes to check if they match.
-		vec3<uint32_t> GetFaceVertexIndices(int f) const;
+		vec3<uint32> GetFaceVertexIndices(int f) const;
 
 		//! Return true if the given mesh has the same vertex and index formats, and in the same
 		//! buffer structure.
diff --git a/Engine/Plugins/Experimental/Mutable/Source/MutableRuntime/Public/MuR/MeshBufferSet.h b/Engine/Plugins/Experimental/Mutable/Source/MutableRuntime/Public/MuR/MeshBufferSet.h
--- a/Engine/Plugins/Experimental/Mutable/Source/MutableRuntime/Public/MuR/MeshBufferSet.h
+++ b/Engine/Plugins/Experimental/Mutable/Source/MutableRuntime/Public/MuR/MeshBufferSet.h
@@ -372,7 +372,7 @@ namespace mu
 
 		//! Copy an element from one position to another, overwriting the other element.
 		//! Both positions must be valid, buffer size won't change.
-		void CopyElement(uint32_t fromIndex, uint32_t toIndex);
+		void CopyElement(uint32 fromIndex, uint32 toIndex);
 
 		//! Compare the format of the two buffers at index buffer and return true if they match.
 		bool HasSameFormat(int32 buffer, const FMeshBufferSet& pOther) const;
diff --git a/Engine/Plugins/Experimental/Mutable/Source/MutableRuntime/Public/MuR/Model.h b/Engine/Plugins/Experimental/Mutable/Source/MutableRuntime/Public/MuR/Model.h
--- a/Engine/Plugins/Experimental/Mutable/Source/MutableRuntime/Public/MuR/Model.h
+++ b/Engine/Plugins/Experimental/Mutable/Source/MutableRuntime/Public/MuR/Model.h
@@ -6,7 +6,6 @@
 #include "MuR/Ptr.h"
 #include "MuR/RefCounted.h"
 
-#include <stdint.h>
 #include <functional>
 
 namespace mu { class InputArchive; }
@@ -18,8 +17,8 @@ namespace mu { class OutputArchive; }
 //! data structures. Compiled models are not necessarily compatible when the runtime is updated,
 //! so this version number can be used externally to verify this. It is not used internally, and
 //! serializing models from different versions than this runtime will probably result in a crash.
-#define MUTABLE_COMPILED_MODEL_CODE_VERSION		uint32_t( 47 )
-#define MUTABLE_PARAMETERS_VERSION              uint32_t( 1 )
+#define MUTABLE_COMPILED_MODEL_CODE_VERSION		uint32( 47 )
+#define MUTABLE_PARAMETERS_VERSION              uint32( 1 )
 
 
 namespace mu
@@ -149,13 +148,13 @@ namespace mu
         ModelParametersGenerator( const Model* pModel, System* pSystem, bool considerRelevancy=true );
 
         //! Return the number of different possible instances that can be built from the model.
-        int64_t GetInstanceCount();
+        int64 GetInstanceCount();
 
         //! Return the parameters of one of the possible instances of the model.
         //! \param index is an index in the range of 0 to GetInstanceCount-1
         //! \param randomGenerator if not null, this function will be used to generate random values
         //! for the continuous parameters of the instance.
-        ParametersPtr GetInstance( int64_t index, float (*randomGenerator )() );
+        ParametersPtr GetInstance( int64 index, float (*randomGenerator )() );
 
         //! Return the parameters of one of a random instance of the model.
         //! \param randomGenerator This function will be used to generate random values
diff --git a/Engine/Plugins/Experimental/Mutable/Source/MutableRuntime/Public/MuR/MutableMemory.h b/Engine/Plugins/Experimental/Mutable/Source/MutableRuntime/Public/MuR/MutableMemory.h
--- a/Engine/Plugins/Experimental/Mutable/Source/MutableRuntime/Public/MuR/MutableMemory.h
+++ b/Engine/Plugins/Experimental/Mutable/Source/MutableRuntime/Public/MuR/MutableMemory.h
@@ -4,9 +4,6 @@
 
 #include "CoreTypes.h"
 
-#include <cstdint>
-#include <cstring>	// for size_t in some platforms
-
 namespace mu
 {
 	//! \defgroup runtime Run-time classes
@@ -26,7 +23,7 @@ namespace mu
 	//! \ingroup runtime
 	MUTABLERUNTIME_API extern void Initialize
 		(
-            void* (*customMalloc)( size_t, uint32_t ),
+            void* (*customMalloc)( size_t, uint32 ),
             void (*customFree)( void* )
 		);
 
@@ -38,7 +35,7 @@ namespace mu
 	//! Memory management functions to be used inside the library. No other memory allocation is
 	//! allowed.
 	MUTABLERUNTIME_API extern void* mutable_malloc(size_t size);
-	MUTABLERUNTIME_API extern void* mutable_malloc_aligned(size_t size, uint32_t alignment);
+	MUTABLERUNTIME_API extern void* mutable_malloc_aligned(size_t size, uint32 alignment);
 	MUTABLERUNTIME_API extern void mutable_free(void* ptr);
 	MUTABLERUNTIME_API extern void mutable_free(void* ptr, size_t size);
 	MUTABLERUNTIME_API extern void mutable_free_aligned(void* ptr, size_t size);
diff --git a/Engine/Plugins/Experimental/Mutable/Source/MutableRuntime/Public/MuR/MutableString.h b/Engine/Plugins/Experimental/Mutable/Source/MutableRuntime/Public/MuR/MutableString.h
--- a/Engine/Plugins/Experimental/Mutable/Source/MutableRuntime/Public/MuR/MutableString.h
+++ b/Engine/Plugins/Experimental/Mutable/Source/MutableRuntime/Public/MuR/MutableString.h
@@ -9,8 +9,6 @@
 #include "MuR/Serialisation.h"
 #include "MuR/SerialisationPrivate.h"
 
-#include <stdint.h>
-#include <memory>
 
 namespace mu
 {
@@ -54,7 +52,7 @@ namespace mu
 		//!
 		void Serialise(OutputArchive& arch) const
 		{
-			uint32_t ver = 0;
+			uint32 ver = 0;
 			arch << ver;
 
 			arch << m_value;
@@ -63,7 +61,7 @@ namespace mu
 		//!
 		void Unserialise(InputArchive& arch)
 		{
-			uint32_t ver;
+			uint32 ver;
 			arch >> ver;
 			check(ver <= 0);
 
diff --git a/Engine/Plugins/Experimental/Mutable/Source/MutableRuntime/Public/MuR/PhysicsBody.h b/Engine/Plugins/Experimental/Mutable/Source/MutableRuntime/Public/MuR/PhysicsBody.h
--- a/Engine/Plugins/Experimental/Mutable/Source/MutableRuntime/Public/MuR/PhysicsBody.h
+++ b/Engine/Plugins/Experimental/Mutable/Source/MutableRuntime/Public/MuR/PhysicsBody.h
@@ -15,7 +15,6 @@
 #include "MuR/RefCounted.h"
 #include "MuR/Serialisation.h"
 
-#include <memory>
 
 namespace mu
 {
@@ -32,7 +31,7 @@ namespace mu
 		//!
 		inline void Serialise( OutputArchive& arch ) const
 		{
-			const uint32_t ver = 0;
+			const uint32 ver = 0;
 			arch << ver;
 
 			arch << Name;
@@ -42,7 +41,7 @@ namespace mu
 		//!
 		inline void Unserialise( InputArchive& arch )
 		{
-			uint32_t ver;
+			uint32 ver;
 			arch >> ver;
 			check(ver == 0);
 
@@ -66,7 +65,7 @@ namespace mu
 		{
 			FBodyShape::Serialise(arch);
 
-			const uint32_t ver = 0;
+			const uint32 ver = 0;
 			arch << ver;
 
 			arch << Position;
@@ -78,7 +77,7 @@ namespace mu
 		{
 			FBodyShape::Unserialise(arch);
 
-			uint32_t ver;
+			uint32 ver;
 			arch >> ver;
 			check(ver == 0);
 
@@ -103,7 +102,7 @@ namespace mu
 		{
 			FBodyShape::Serialise( arch );
 
-			const uint32_t ver = 0;
+			const uint32 ver = 0;
 			arch << ver;
 
 			arch << Position;
@@ -116,7 +115,7 @@ namespace mu
 		{
 			FBodyShape::Unserialise( arch );
 
-			uint32_t ver;
+			uint32 ver;
 			arch >> ver;
 			check(ver == 0);
 
@@ -147,7 +146,7 @@ namespace mu
 		{
 			FBodyShape::Serialise( arch );
 
-			const uint32_t ver = 0;
+			const uint32 ver = 0;
 			arch << ver;
 
 			arch << Position;
@@ -161,7 +160,7 @@ namespace mu
 		{
 			FBodyShape::Unserialise( arch );
 
-			uint32_t ver;
+			uint32 ver;
 			arch >> ver;
 			check(ver == 0);
 
@@ -195,7 +194,7 @@ namespace mu
 		{
 			FBodyShape::Serialise( arch );
 
-			const uint32_t ver = 0;
+			const uint32 ver = 0;
 			arch << ver;
 
 			arch << Position;
@@ -210,7 +209,7 @@ namespace mu
 		{
 			FBodyShape::Unserialise( arch );
 			
-			uint32_t ver;
+			uint32 ver;
 			arch >> ver;
 			check(ver == 0);
 
@@ -242,7 +241,7 @@ namespace mu
 		{
 			FBodyShape::Serialise( arch );
 	
-		 	uint32_t ver = 0;
+		 	uint32 ver = 0;
 		 	arch << ver;
 
 		 	arch << Vertices;
@@ -253,7 +252,7 @@ namespace mu
 		  //!
 		inline void Unserialise( InputArchive& arch )
 		{
-		 	uint32_t ver;
+		 	uint32 ver;
 		 	arch >> ver;
 		 	check(ver == 0);
 
@@ -283,7 +282,7 @@ namespace mu
 		 //!
 		inline void Serialise( OutputArchive& arch ) const
 		{
-		 	uint32_t ver = 0;
+		 	uint32 ver = 0;
 		 	arch << ver;
 		
 		 	arch << Spheres;
@@ -296,7 +295,7 @@ namespace mu
 		 //!
 		inline void Unserialise( InputArchive& arch )
 		{
-		 	uint32_t ver;
+		 	uint32 ver;
 		 	arch >> ver;
 		 	check(ver == 0);
 		
@@ -433,7 +432,7 @@ namespace mu
         //!
         inline void Serialise( OutputArchive& arch ) const
 		{
-            uint32_t ver = 1;
+            uint32 ver = 1;
 			arch << ver;
         	
         	arch << Bodies;
@@ -445,7 +444,7 @@ namespace mu
         //!
         inline void Unserialise( InputArchive& arch )
 		{
-            uint32_t ver;
+            uint32 ver;
 			arch >> ver;
             check(ver <= 1);
         	
diff --git a/Engine/Plugins/Experimental/Mutable/Source/MutableRuntime/Public/MuR/Settings.h b/Engine/Plugins/Experimental/Mutable/Source/MutableRuntime/Public/MuR/Settings.h
--- a/Engine/Plugins/Experimental/Mutable/Source/MutableRuntime/Public/MuR/Settings.h
+++ b/Engine/Plugins/Experimental/Mutable/Source/MutableRuntime/Public/MuR/Settings.h
@@ -5,8 +5,6 @@
 #include "MuR/Ptr.h"
 #include "MuR/RefCounted.h"
 
-#include <stdint.h>
-
 
 namespace mu
 {
@@ -35,7 +33,7 @@ namespace mu
         //! more streaming and higher instance construction times, but will use less memory while
         //! building objects.
         //! Defaults to 0, which disables the limit.
-        void SetStreamingCache( uint64_t bytes );
+        void SetStreamingCache( uint64 bytes );
 
         //! Set the quality for the image compression algorithms.
         //! \param quality Quality level of the compression. Higher quality increases the instance
diff --git a/Engine/Plugins/Experimental/Mutable/Source/MutableRuntime/Public/MuR/Skeleton.h b/Engine/Plugins/Experimental/Mutable/Source/MutableRuntime/Public/MuR/Skeleton.h
--- a/Engine/Plugins/Experimental/Mutable/Source/MutableRuntime/Public/MuR/Skeleton.h
+++ b/Engine/Plugins/Experimental/Mutable/Source/MutableRuntime/Public/MuR/Skeleton.h
@@ -2,9 +2,8 @@
 
 #pragma once
 
-//#include "MuR/MutableMath.h"
-#include "../Private/MuR/MemoryPrivate.h"
-#include "../Private/MuR/SerialisationPrivate.h"
+#include "MuR/MemoryPrivate.h"
+#include "MuR/SerialisationPrivate.h"
 #include "Containers/Array.h"
 #include "HAL/PlatformMath.h"
 #include "Math/Transform.h"
@@ -14,8 +13,6 @@
 #include "MuR/RefCounted.h"
 #include "MuR/Serialisation.h"
 
-#include <memory>
-
 
 namespace mu
 {
@@ -93,7 +90,7 @@ namespace mu
 		//!
 		inline void Serialise(OutputArchive& arch) const
 		{
-			uint32_t ver = 4;
+			uint32 ver = 4;
 			arch << ver;
 
 			arch << m_bones;
@@ -104,7 +101,7 @@ namespace mu
 		//!
 		inline void Unserialise(InputArchive& arch)
 		{
-			uint32_t ver;
+			uint32 ver;
 			arch >> ver;
 			check(ver >= 3);
 
diff --git a/Engine/Plugins/Experimental/Mutable/Source/MutableTools/Private/MuT/CodeOptimiser_Logic.cpp b/Engine/Plugins/Experimental/Mutable/Source/MutableTools/Private/MuT/CodeOptimiser_Logic.cpp
--- a/Engine/Plugins/Experimental/Mutable/Source/MutableTools/Private/MuT/CodeOptimiser_Logic.cpp
+++ b/Engine/Plugins/Experimental/Mutable/Source/MutableTools/Private/MuT/CodeOptimiser_Logic.cpp
@@ -482,7 +482,7 @@ namespace mu
 //                m_facts.pop_back();
 
 //                // update op if necessary
-//                if ( mutable_memcmp( &args, &program.m_code[at].args, sizeof(OP::ARGS) ) )
+//                if ( FMemory::Memcmp( &args, &program.m_code[at].args, sizeof(OP::ARGS) ) )
 //                {
 //                    OP op = program.m_code[at];
 //                    op.args = args;
@@ -547,7 +547,7 @@ namespace mu
 ////            }
 
 ////            // update op if necessary
-////            if ( mutable_memcmp( &args, &program.m_code[at].args, sizeof(OP::ARGS) ) )
+////            if ( FMemory::Memcmp( &args, &program.m_code[at].args, sizeof(OP::ARGS) ) )
 ////            {
 ////                m_modified = true;
 ////                OP op = program.m_code[at];
diff --git a/Engine/Plugins/Experimental/Mutable/Source/MutableTools/Public/MuT/NodeSurfaceNew.h b/Engine/Plugins/Experimental/Mutable/Source/MutableTools/Public/MuT/NodeSurfaceNew.h
--- a/Engine/Plugins/Experimental/Mutable/Source/MutableTools/Public/MuT/NodeSurfaceNew.h
+++ b/Engine/Plugins/Experimental/Mutable/Source/MutableTools/Public/MuT/NodeSurfaceNew.h
@@ -83,7 +83,7 @@ namespace mu
         //! Add an optional, opaque id that will be returned in the surfaces of the created
         //! instances. Can be useful to identify surfaces on the application side.
         //! See Instance::GetSurfaceCustomId
-        void SetCustomID( uint32_t id );
+        void SetCustomID( uint32 id );
 
         //! \name Tags
         //! \{
diff --git a/Engine/Plugins/Experimental/Mutable/Source/MutableTools/Public/MuT/Table.h b/Engine/Plugins/Experimental/Mutable/Source/MutableTools/Public/MuT/Table.h
--- a/Engine/Plugins/Experimental/Mutable/Source/MutableTools/Public/MuT/Table.h
+++ b/Engine/Plugins/Experimental/Mutable/Source/MutableTools/Public/MuT/Table.h
@@ -5,9 +5,6 @@
 #include "MuR/Ptr.h"
 #include "MuR/RefCounted.h"
 
-#include <stdint.h>
-
-
 
 namespace mu
 {
@@ -74,17 +71,17 @@ namespace mu
 		int FindColumn( const char* strName ) const;
 
 		//!
-        void AddRow( uint32_t id );
+        void AddRow( uint32 id );
 
 		//! Adds a "NONE" option
 		void SetNoneOption(bool bAddOption);
 
 		//!
-        void SetCell( int column, uint32_t rowId, float value );
-        void SetCell( int column, uint32_t rowId, float r, float g, float b );
-        void SetCell( int column, uint32_t rowId, Image* pImage );
-        void SetCell( int column, uint32_t rowId, Mesh* pMesh );
-        void SetCell( int column, uint32_t rowId, const char* strValue );
+        void SetCell( int column, uint32 rowId, float value );
+        void SetCell( int column, uint32 rowId, float r, float g, float b );
+        void SetCell( int column, uint32 rowId, Image* pImage );
+        void SetCell( int column, uint32 rowId, Mesh* pMesh );
+        void SetCell( int column, uint32 rowId, const char* strValue );
 
 
 		//-----------------------------------------------------------------------------------------
