#rnx

diff --git a/Engine/Plugins/Experimental/Mutable/Source/CustomizableObjectEditor/Private/MuCOE/CustomizableObjectEditorViewportClient.cpp b/Engine/Plugins/Experimental/Mutable/Source/CustomizableObjectEditor/Private/MuCOE/CustomizableObjectEditorViewportClient.cpp
--- a/Engine/Plugins/Experimental/Mutable/Source/CustomizableObjectEditor/Private/MuCOE/CustomizableObjectEditorViewportClient.cpp
+++ b/Engine/Plugins/Experimental/Mutable/Source/CustomizableObjectEditor/Private/MuCOE/CustomizableObjectEditorViewportClient.cpp
@@ -2523,9 +2523,9 @@ void FCustomizableObjectEditorViewportClient::BakeInstance()
 					for (int32 m = 0; m < Mesh->GetMaterials().Num(); ++m)
 					{
 						UMaterialInterface* Interface = Mesh->GetMaterials()[m].MaterialInterface;
-						Material = Interface->GetMaterial();
+						Material = Interface ? Interface->GetMaterial() : nullptr;
 
-						if (Interface && Material)
+						if (Material)
 						{
 							ResourceName = ObjectName + "_Material_" + Material->GetName();
 
