#rnx

diff --git a/Engine/Plugins/Experimental/Mutable/Source/CustomizableObject/Private/MuCO/CustomizableObjectInstance.cpp b/Engine/Plugins/Experimental/Mutable/Source/CustomizableObject/Private/MuCO/CustomizableObjectInstance.cpp
--- a/Engine/Plugins/Experimental/Mutable/Source/CustomizableObject/Private/MuCO/CustomizableObjectInstance.cpp
+++ b/Engine/Plugins/Experimental/Mutable/Source/CustomizableObject/Private/MuCO/CustomizableObjectInstance.cpp
@@ -1992,7 +1992,11 @@ bool UCustomizableInstancePrivateData::UpdateSkeletalMesh_PostBeginUpdate0(UCust
 						LODInfo.ScreenSize = LODData.LODInfo.ScreenSize;
 						LODInfo.LODHysteresis =  LODData.LODInfo.LODHysteresis;
 						LODInfo.bSupportUniformlyDistributedSampling =  LODData.LODInfo.bSupportUniformlyDistributedSampling;
-						LODInfo.bAllowCPUAccess =  LODData.LODInfo.bAllowCPUAccess;
+						LODInfo.bAllowCPUAccess = LODData.LODInfo.bAllowCPUAccess;
+						
+						// Disable LOD simplification when baking instances
+						LODInfo.ReductionSettings.NumOfTrianglesPercentage = 1.f;
+						LODInfo.ReductionSettings.NumOfVertPercentage = 1.f;
 						
 #if WITH_EDITORONLY_DATA
 						LODInfo.BuildSettings.bUseFullPrecisionUVs = true;
