[mutable] Instance Update optimization improvements.

Mutable:
- Moved Multilayer Projectors to the Instance Descriptor.
- Moved RealodParametersFromObject to the Instance Descriptor.
- Created new UpdateResult Instance and Component delegatres.
- Created Instance and Component Updated callbacks.
- Now when an Instance update is optimized Update delegates are also being called.
- Fixed a UCustomizableObjectInstance::Clone crash.

diff --git a/Engine/Plugins/Experimental/Mutable/Source/CustomizableObject/Classes/MuCO/CustomizableObjectInstance.h b/Engine/Plugins/Experimental/Mutable/Source/CustomizableObject/Classes/MuCO/CustomizableObjectInstance.h
--- a/Engine/Plugins/Experimental/Mutable/Source/CustomizableObject/Classes/MuCO/CustomizableObjectInstance.h
+++ b/Engine/Plugins/Experimental/Mutable/Source/CustomizableObject/Classes/MuCO/CustomizableObjectInstance.h
@@ -92,12 +92,19 @@ namespace ESkeletalMeshState
 };
 
 
-class USkeletalMesh;
+/** Instance Update Result. */
+enum class EUpdateResult : uint8
+{
+	Success,
+	Error
+};
+
 
 DECLARE_DELEGATE(FObjectInstanceUpdateBeginDelegate);
 DECLARE_DYNAMIC_MULTICAST_DELEGATE_OneParam(FObjectInstanceUpdatedDelegate, UCustomizableObjectInstance*, Instance);
 DECLARE_MULTICAST_DELEGATE_OneParam(FObjectInstanceUpdatedNativeDelegate, UCustomizableObjectInstance*);
 DECLARE_DELEGATE_OneParam( FProjectorStateChangedDelegate, FString );
+DECLARE_MULTICAST_DELEGATE_OneParam(FBeginDestroyDelegate, UCustomizableObjectInstance*);
 
 DECLARE_DYNAMIC_DELEGATE_TwoParams(FEachComponentAnimInstanceClassDelegate, int32, SlotIndex, TSubclassOf<UAnimInstance>, AnimInstClass);
 DECLARE_DELEGATE_TwoParams(FEachComponentAnimInstanceClassNativeDelegate, int32 /*SlotIndex*/, TSubclassOf<UAnimInstance> /*AnimInstClass*/);
@@ -105,9 +112,11 @@ DECLARE_DELEGATE_TwoParams(FEachComponentAnimInstanceClassNativeDelegate, int32
 UCLASS( Blueprintable, BlueprintType, HideCategories=(CustomizableObjectInstance) )
 class CUSTOMIZABLEOBJECT_API UCustomizableObjectInstance : public UObject
 {
-public:
-	 GENERATED_BODY()
+	GENERATED_BODY()
+
+	friend UCustomizableInstancePrivateData;
 
+public:
 	UCustomizableObjectInstance();
 
 	FCustomizableObjectInstanceDescriptor& GetDescriptor();
@@ -132,6 +141,9 @@ public:
 	//       callbacks are safe, avoid using it outside editor functionality
 	FObjectInstanceUpdateBeginDelegate UpdateBeginDelegate;
 
+	/** Broadcast when UObject::BeginDestroy is being called. */
+	FBeginDestroyDelegate BeginDestroyDelegate;
+
 	TMap<FString, bool> ParamNameToExpandedMap; // Used to check whether a mutable param is expanded in the editor to show its child params
 
 	// Will store status description of current skeletal mesh generation (for instance, "EmptyLOD0" or "EmptyMesh"
@@ -558,10 +570,13 @@ public:
 	bool GetUseCurrentMinLODAsBaseLOD() const;
 
 	/** Instance updated. */
-	void Updated();
+	void Updated(EUpdateResult Result);
 	
-	/** Hash representing the actual state of the instance. */
-	uint32 GetLastUpdateDescriptorHash() const;
+	/** Hash representing the actual state of the instance (meshes, textures...). This does not include the parameters. */
+	uint32 GetDescriptorHash() const;
+
+	/** Hash representing the state of the instance when the update was requested. */
+	uint32 GetUpdateDescriptorHash() const;
 
 	// --------------------------------------------------------------------
 
@@ -624,13 +639,12 @@ private:
 
 	UPROPERTY( Transient )
 	TObjectPtr<UCustomizableInstancePrivateData> PrivateData;
-	
-	/** Multilayer Projector helpers. See FMultilayerProjector.*/
-	UPROPERTY()
-	TMap<FName, FMultilayerProjector> MultilayerProjectors;
 
-	/** Hash of the UCustomizableObjectInstance::Descriptor on the last update. */
-	uint32 LastUpdateDescriptorHash = 0;
+	/** Hash of the UCustomizableObjectInstance::Descriptor when UpdateSkeletalMeshAsync has been called. */
+	uint32 UpdateDescriptorHash = 0;
+	
+	/** Hash of the UCustomizableObjectInstance::Descriptor on the last successful update. */
+	uint32 DescriptorHash = 0;
 
 	// Deprecated properties
 	
@@ -655,6 +669,9 @@ private:
 	UPROPERTY()
 	TArray<FCustomizableObjectProjectorParameterValue> ProjectorParameters_DEPRECATED;
 	
+   	UPROPERTY()
+   	TMap<FName, FMultilayerProjector> MultilayerProjectors_DEPRECATED;
+
 	UPROPERTY()
 	bool bBuildParameterDecorations_DEPRECATED;
 };
diff --git a/Engine/Plugins/Experimental/Mutable/Source/CustomizableObject/Classes/MuCO/CustomizableSkeletalComponent.h b/Engine/Plugins/Experimental/Mutable/Source/CustomizableObject/Classes/MuCO/CustomizableSkeletalComponent.h
--- a/Engine/Plugins/Experimental/Mutable/Source/CustomizableObject/Classes/MuCO/CustomizableSkeletalComponent.h
+++ b/Engine/Plugins/Experimental/Mutable/Source/CustomizableObject/Classes/MuCO/CustomizableSkeletalComponent.h
@@ -18,7 +18,8 @@ class UCustomizableSkeletalComponent;
 class UObject;
 class USkeletalMesh;
 struct FFrame;
-
+class UCustomizableObjectInstance;
+enum class EUpdateResult : uint8;
 
 DECLARE_DELEGATE_TwoParams(FCustomizableSkeletalComponentPreUpdateDelegate, UCustomizableSkeletalComponent* Component, USkeletalMesh* NextMesh);
 DECLARE_DELEGATE(FCustomizableSkeletalComponentUpdatedDelegate);
@@ -41,14 +42,16 @@ public:
 	float SkippedLastRenderTime;
 
 	UPROPERTY(BlueprintReadWrite, EditAnywhere, Category = CustomizableSkeletalMesh)
-	TObjectPtr<class UCustomizableObjectInstance> CustomizableObjectInstance;
+	TObjectPtr<UCustomizableObjectInstance> CustomizableObjectInstance;
 
 	UPROPERTY(BlueprintReadWrite, EditAnywhere, Category = CustomizableSkeletalMesh)
 	int32 ComponentIndex;
 
 	FCustomizableSkeletalComponentPreUpdateDelegate PreUpdateDelegate;
 	FCustomizableSkeletalComponentUpdatedDelegate UpdatedDelegate;
-
+	
+	void Updated(EUpdateResult Result) const;
+	
 	USkeletalMesh* GetSkeletalMesh() const;
 	USkeletalMesh* GetAttachedSkeletalMesh() const;
 
@@ -61,7 +64,10 @@ public:
 
 	void SetVisibilityOfSkeletalMeshSectionWithMaterialName(bool bVisible, const FString& MaterialName, int32 LOD);
 
+	// USceneComponent interface
 	virtual void BeginDestroy() override;
+
+	// UObject interface
 	virtual void PreSave(FObjectPreSaveContext ObjectSaveContext) override;
 
 #if WITH_EDITOR
diff --git a/Engine/Plugins/Experimental/Mutable/Source/CustomizableObject/Private/MuCO/CustomizableObjectInstance.cpp b/Engine/Plugins/Experimental/Mutable/Source/CustomizableObject/Private/MuCO/CustomizableObjectInstance.cpp
--- a/Engine/Plugins/Experimental/Mutable/Source/CustomizableObject/Private/MuCO/CustomizableObjectInstance.cpp
+++ b/Engine/Plugins/Experimental/Mutable/Source/CustomizableObject/Private/MuCO/CustomizableObjectInstance.cpp
@@ -359,6 +359,8 @@ bool UCustomizableObjectInstance::IsEditorOnly() const
 
 void UCustomizableObjectInstance::BeginDestroy()
 {
+	BeginDestroyDelegate.Broadcast(this);
+
 	if (PrivateData && PrivateData->StreamingHandle)
 	{
 		PrivateData->StreamingHandle->CancelHandle();
@@ -475,6 +477,11 @@ void UCustomizableObjectInstance::Serialize(FArchive& Ar)
 	{
 		Descriptor.bBuildParameterDecorations = bBuildParameterDecorations_DEPRECATED;
 	}
+
+	if (CustomizableObjectCustomVersion < FCustomizableObjectCustomVersion::DescriptorMultilayerProjectors)
+	{
+		Descriptor.MultilayerProjectors = MultilayerProjectors_DEPRECATED;
+	}
 }
 
 
@@ -552,7 +559,7 @@ FCustomizableObjectProjector UCustomizableObjectInstance::GetProjectorDefaultVal
 
 mu::ParametersPtr UCustomizableInstancePrivateData::ReloadParametersFromObject(UCustomizableObjectInstance* Public, bool ClearLastMeshIds)
 {
-	UCustomizableObject* CustomizableObject = Public->GetCustomizableObject();
+	const UCustomizableObject* CustomizableObject = Public->GetCustomizableObject();
 	if (!CustomizableObject)
 	{
 		// May happen when deleting assets in editor, while they are still open.
@@ -577,566 +584,25 @@ mu::ParametersPtr UCustomizableInstancePrivateData::ReloadParametersFromObject(U
 		NumMaxLODsToStream = 0;
 	}
 	
-	int32 OldState = Public->GetState();
-	TArray<FCustomizableObjectBoolParameterValue> OldBoolParameters = Public->GetBoolParameters();
-	TArray<FCustomizableObjectIntParameterValue> OldIntParameters = Public->GetIntParameters();
-	TArray<FCustomizableObjectFloatParameterValue> OldFloatParameters = Public->GetFloatParameters();
-	TArray<FCustomizableObjectTextureParameterValue> OldTextureParameters = Public->GetTextureParameters();
-	TArray<FCustomizableObjectVectorParameterValue> OldVectorParameters = Public->GetVectorParameters();
-	TArray<FCustomizableObjectProjectorParameterValue> OldProjectorParameters = Public->GetProjectorParameters();
-	TMap<TPair<FString, int32>, EProjectorState::Type> OldProjectorStates = ProjectorStates;
-
-	if (ClearLastMeshIds)
-	{
-		ClearAllLastMeshIds();
-	}
-
-	mu::ParametersPtr MutableParameters;
-	Public->SetState(FMath::Clamp(OldState, 0, CustomizableObject->GetStateCount() - 1));
-	Public->GetBoolParameters().Reset();
-	Public->GetIntParameters().Reset();
-	Public->GetFloatParameters().Reset();
-	Public->GetTextureParameters().Reset();
-	Public->GetVectorParameters().Reset();
-	Public->GetProjectorParameters().Reset();
-
 #if WITH_EDITOR
 	if(!Public->ProjectorAlphaChange)
 #endif
 	{
 		ProjectorStates.Reset();
 	}
-
-	if (CustomizableObject && CustomizableObject->GetPrivate()->GetModel())
-	{
-		TArray<bool> OldIntParametersUsed;
-		OldIntParametersUsed.SetNumUninitialized(OldIntParameters.Num());
-		FMemory::Memzero(OldIntParametersUsed.GetData(), OldIntParametersUsed.GetAllocatedSize());
-
-		MutableParameters = CustomizableObject->GetPrivate()->GetModel()->NewParameters();
-		int32 paramCount = MutableParameters->GetCount();
-
-		for (int32 paramIndex = 0; paramIndex<paramCount; ++paramIndex)
-		{
-			FString Name = MutableParameters->GetName(paramIndex);
-			FString Uid = MutableParameters->GetUid(paramIndex);
-			mu::PARAMETER_TYPE mutableType = MutableParameters->GetType(paramIndex);
-
-			switch (mutableType)
-			{
-			case mu::PARAMETER_TYPE::T_BOOL:
-			{
-				FCustomizableObjectBoolParameterValue param;
-				param.ParameterName = Name;
-				param.Uid = Uid;
-
-				bool found = false;
-				for (int32 i = 0; i < OldBoolParameters.Num(); ++i)
-				{
-					if (OldBoolParameters[i].ParameterName == Name)
-					{
-						found = true;
-						param.ParameterValue = OldBoolParameters[i].ParameterValue;
-						MutableParameters->SetBoolValue(paramIndex, param.ParameterValue);
-
-						if (OldBoolParameters[i].Uid.IsEmpty())
-						{
-							Public->Modify();
-						}
-
-						break;
-					}
-				}
-
-				if (!found)
-				{
-					for (int32 i = 0; i < OldBoolParameters.Num(); ++i)
-					{
-						if (!Uid.IsEmpty() && OldBoolParameters[i].Uid == Uid)
-						{
-							found = true;
-							param.ParameterValue = OldBoolParameters[i].ParameterValue;
-							MutableParameters->SetBoolValue(paramIndex, param.ParameterValue);
-
-							if (OldBoolParameters[i].Uid.IsEmpty())
-							{
-								Public->Modify();
-							}
-
-							break;
-						}
-					}
-				}
-
-				if (!found)
-				{
-					param.ParameterValue = MutableParameters->GetBoolValue(paramIndex);
-				}
-
-				Public->GetBoolParameters().Add(param);
-				break;
-			}
-
-			case mu::PARAMETER_TYPE::T_INT:
-			{
-				FString ParameterValueName;
-				TArray<FString> ParameterRangeValueNames;
-
-				bool found = false;
-				for (int32 i = 0; i < OldIntParameters.Num(); ++i)
-				{
-					if (OldIntParametersUsed[i] == false && OldIntParameters[i].ParameterName.Equals(Name, ESearchCase::CaseSensitive))
-					{
-						found = true;
-						OldIntParametersUsed[i] = true;
-
-						mu::RangeIndexPtr RangeIdxPtr = MutableParameters->NewRangeIndex(paramIndex);
-
-						if (!RangeIdxPtr)
-						{
-						ParameterValueName = OldIntParameters[i].ParameterValueName;
-						int32 Value = CustomizableObject->FindIntParameterValue(paramIndex, ParameterValueName);
-						MutableParameters->SetIntValue(paramIndex, Value);
-						}
-						else
-						{
-							for (int RangeIndex = 0; RangeIndex < OldIntParameters[i].ParameterRangeValueNames.Num(); ++RangeIndex)
-							{
-								FString AuxParameterValueName = OldIntParameters[i].ParameterRangeValueNames[RangeIndex];
-								ParameterRangeValueNames.Add(AuxParameterValueName);
-								int32 Value = CustomizableObject->FindIntParameterValue(paramIndex, AuxParameterValueName);
-								RangeIdxPtr->SetPosition(0, RangeIndex);
-								MutableParameters->SetIntValue(paramIndex, Value, RangeIdxPtr);
-							}
-						}
-
-						if (OldIntParameters[i].Uid.IsEmpty())
-						{
-							Public->Modify();
-						}
-
-						break;
-					}
-				}
-
-				if (!found)
-				{
-					for (int32 i = 0; i < OldIntParameters.Num(); ++i)
-					{
-						if (!Uid.IsEmpty() && OldIntParametersUsed[i] == false && OldIntParameters[i].Uid == Uid)
-						{
-							found = true;
-							OldIntParametersUsed[i] = true;
-
-							mu::RangeIndexPtr RangeIdxPtr = MutableParameters->NewRangeIndex(paramIndex);
-
-							if (!RangeIdxPtr)
-							{
-							ParameterValueName = OldIntParameters[i].ParameterValueName;
-							int32 Value = CustomizableObject->FindIntParameterValue(paramIndex, ParameterValueName);
-							MutableParameters->SetIntValue(paramIndex, Value);
-							}
-							else
-							{
-								for (int RangeIndex = 0; RangeIndex < OldIntParameters[i].ParameterRangeValueNames.Num(); ++RangeIndex)
-								{
-									FString AuxParameterValueName = OldIntParameters[i].ParameterRangeValueNames[RangeIndex];
-									ParameterRangeValueNames.Add(AuxParameterValueName);
-									int32 Value = CustomizableObject->FindIntParameterValue(paramIndex, AuxParameterValueName);
-									RangeIdxPtr->SetPosition(0, RangeIndex);
-									MutableParameters->SetIntValue(paramIndex, Value, RangeIdxPtr);
-								}
-							}
-
-							if (OldIntParameters[i].Uid.IsEmpty())
-							{
-								Public->Modify();
-							}
-
-							break;
-						}
-					}
-				}
-
-				if (!found)
-				{
-					mu::RangeIndexPtr RangeIdxPtr = MutableParameters->NewRangeIndex(paramIndex);
-
-					if (!RangeIdxPtr.get())
-					{
-					int32 Value = MutableParameters->GetIntValue(paramIndex);
-					ParameterValueName = CustomizableObject->FindIntParameterValueName(paramIndex, Value);
-					}
-					else
-					{
-						int32 ValueCount = MutableParameters->GetValueCount(paramIndex);
-
-						for (int32 ValueIterator = 0; ValueIterator < ValueCount; ++ValueIterator)
-						{
-							mu::RangeIndexPtr RangeValueIdxPtr = MutableParameters->GetValueIndex(paramIndex, ValueIterator);
-							int32 ValueIndex = RangeValueIdxPtr->GetPosition(0);
-
-							//if (!param.RangeValues.IsValidIndex(ValueIndex))
-							//{
-							//	param.RangeValues.AddDefaulted(ValueIndex + 1 - param.RangeValues.Num());
-							//}
-
-							int32 Value = MutableParameters->GetIntValue(paramIndex, RangeValueIdxPtr);
-							FString AuxParameterValueName = CustomizableObject->FindIntParameterValueName(paramIndex, Value);
-							//param.RangeValues[ValueIndex] = AuxParameterValueName;
-
-							if (!ParameterRangeValueNames.IsValidIndex(ValueIndex))
-							{
-								ParameterRangeValueNames.AddDefaulted(ValueIndex + 1 - ParameterRangeValueNames.Num());
-							}
-
-							ParameterRangeValueNames[ValueIndex] = AuxParameterValueName;
-						}
-					}
-
-					Public->Modify();
-				}
-
-				//Public->IntParameters.Add(param);
-				Public->GetIntParameters().Emplace(Name, ParameterValueName, Uid, ParameterRangeValueNames);
-
-				break;
-			}
-
-			case mu::PARAMETER_TYPE::T_FLOAT:
-			{
-				FCustomizableObjectFloatParameterValue param;
-				param.ParameterName = Name;
-				param.Uid = Uid;
-
-				bool found = false;
-				for (int32 i = 0; i < OldFloatParameters.Num(); ++i)
-				{
-					if (OldFloatParameters[i].ParameterName == Name || (!Uid.IsEmpty() && OldFloatParameters[i].Uid == Uid))
-					{
-						found = true;
-						param.ParameterValue = OldFloatParameters[i].ParameterValue;
-						param.ParameterRangeValues = OldFloatParameters[i].ParameterRangeValues;
-
-						mu::RangeIndexPtr RangeIdxPtr = MutableParameters->NewRangeIndex(paramIndex);
-
-						if (!RangeIdxPtr)
-						{
-							MutableParameters->SetFloatValue(paramIndex, param.ParameterValue);
-						}
-						else
-						{
-							for (int RangeIndex = 0; RangeIndex < OldFloatParameters[i].ParameterRangeValues.Num(); ++RangeIndex)
-							{
-								RangeIdxPtr->SetPosition(0, RangeIndex);
-								MutableParameters->SetFloatValue(paramIndex, param.ParameterRangeValues[RangeIndex], RangeIdxPtr);
-							}
-						}
-
-						if (OldFloatParameters[i].Uid.IsEmpty())
-						{
-							Public->Modify();
-						}
-					}
-				}
-
-				if (!found)
-				{
-					mu::RangeIndexPtr RangeIdxPtr = MutableParameters->NewRangeIndex(paramIndex);
-
-					if (!RangeIdxPtr.get())
-					{
-					param.ParameterValue = MutableParameters->GetFloatValue(paramIndex);
-				}
-					else
-					{
-					param.ParameterValue = MutableParameters->GetFloatValue(paramIndex);
-
-						int32 ValueCount = MutableParameters->GetValueCount(paramIndex);
-
-						for (int32 ValueIterator = 0; ValueIterator < ValueCount; ++ValueIterator)
-						{
-							mu::RangeIndexPtr RangeValueIdxPtr = MutableParameters->GetValueIndex(paramIndex, ValueIterator);
-							int32 ValueIndex = RangeValueIdxPtr->GetPosition(0);
-
-							float Value = MutableParameters->GetFloatValue(paramIndex, RangeValueIdxPtr);
-
-							if (!param.ParameterRangeValues.IsValidIndex(ValueIndex))
-							{
-								param.ParameterRangeValues.AddDefaulted(ValueIndex + 1 - param.ParameterRangeValues.Num());
-							}
-
-							param.ParameterRangeValues[ValueIndex] = Value;
-						}
-					}
-				}
-
-				Public->GetFloatParameters().Add(param);
-				break;
-			}
-
-			case mu::PARAMETER_TYPE::T_COLOUR:
-			{
-				FCustomizableObjectVectorParameterValue param;
-				param.ParameterName = Name;
-				param.Uid = Uid;
-
-				bool found = false;
-				for (int32 i = 0; i < OldVectorParameters.Num(); ++i)
-				{
-					if (OldVectorParameters[i].ParameterName == Name || (!Uid.IsEmpty() && OldVectorParameters[i].Uid == Uid))
-					{
-						found = true;
-						param.ParameterValue = OldVectorParameters[i].ParameterValue;
-						MutableParameters->SetColourValue(paramIndex, param.ParameterValue.R, param.ParameterValue.G, param.ParameterValue.B);
-
-						if (OldVectorParameters[i].Uid.IsEmpty())
-						{
-							Public->Modify();
-						}
-					}
-				}
-
-				if (!found)
-				{
-					MutableParameters->GetColourValue(paramIndex, &param.ParameterValue.R, &param.ParameterValue.G, &param.ParameterValue.B);
-					param.ParameterValue.A = 1.0f;
-				}
-
-				Public->GetVectorParameters().Add(param);
-				break;
-			}
-
-			case mu::PARAMETER_TYPE::T_PROJECTOR:
-			{
-				FCustomizableObjectProjectorParameterValue param;
-				param.ParameterName = Name;
-				param.Uid = Uid;
-
-				bool found = false;
-				for (int32 i = 0; i < OldProjectorParameters.Num(); ++i)
-				{
-					if (OldProjectorParameters[i].ParameterName == Name || (!Uid.IsEmpty() && OldProjectorParameters[i].Uid == Uid))
-					{
-						found = true;
-
-						param.Value = OldProjectorParameters[i].Value;
-						param.RangeValues = OldProjectorParameters[i].RangeValues;
-
-						mu::RangeIndexPtr RangeIdxPtr = MutableParameters->NewRangeIndex(paramIndex);
-
-						TArray<FCustomizableObjectProjector> AuxValuesArray;
-						AuxValuesArray.Add(param.Value);
-						TArray<FCustomizableObjectProjector>& Values = RangeIdxPtr.get() ? param.RangeValues : AuxValuesArray;
-						int32 ArrayIndex = 0;
-
-						for (FCustomizableObjectProjector& Value : Values)
-						{
-							if (RangeIdxPtr.get())
-							{
-								check(RangeIdxPtr->GetRangeCount() == 1);
-								RangeIdxPtr->SetPosition(0, ArrayIndex);
-							}
-
-							switch (Value.ProjectionType)
-							{
-							case ECustomizableObjectProjectorType::Planar:
-							case ECustomizableObjectProjectorType::Wrapping:
-							{
-								MutableParameters->SetProjectorValue(paramIndex,
-									Value.Position[0], Value.Position[1], Value.Position[2],
-									Value.Direction[0], Value.Direction[1], Value.Direction[2],
-									Value.Up[0], Value.Up[1], Value.Up[2],
-									Value.Scale[0], Value.Scale[1], Value.Scale[2],
-									Value.Angle,
-									RangeIdxPtr);
-								break;
-							}
-
-							case ECustomizableObjectProjectorType::Cylindrical:
-							{
-								// Apply strange swizzle for scales
-								// TODO: try to avoid this
-									float Radius = FMath::Abs(Value.Scale[0] / 2.0f);
-									float Height = Value.Scale[2];
-								// TODO: try to avoid this
-								MutableParameters->SetProjectorValue(paramIndex,
-										Value.Position[0], Value.Position[1], Value.Position[2],
-										-Value.Direction[0], -Value.Direction[1], -Value.Direction[2],
-										-Value.Up[0], -Value.Up[1], -Value.Up[2],
-										-Height, Radius, Radius,
-										Value.Angle,
-										RangeIdxPtr);
-								break;
-							}
-
-							default:
-								// Not implemented.
-								check(false);
-							}
-							
-							ArrayIndex++;
-						}
-					}
-
-					if (OldProjectorParameters[i].Uid.IsEmpty())
-					{
-						Public->Modify();
-					}
-				}
-
-				if (!found)
-				{
-					mu::RangeIndexPtr RangeIdxPtr = MutableParameters->NewRangeIndex(paramIndex);
-
-					if (!RangeIdxPtr.get())
-					{
-					mu::PROJECTOR_TYPE type;
-					MutableParameters->GetProjectorValue(paramIndex,
-						&type,
-						&param.Value.Position[0], &param.Value.Position[1], &param.Value.Position[2],
-						&param.Value.Direction[0], &param.Value.Direction[1], &param.Value.Direction[2],
-						&param.Value.Up[0], &param.Value.Up[1], &param.Value.Up[2],
-						&param.Value.Scale[0], &param.Value.Scale[1], &param.Value.Scale[2],
-						&param.Value.Angle );
-					switch (type)
-					{
-					case mu::PROJECTOR_TYPE::PLANAR:
-						param.Value.ProjectionType = ECustomizableObjectProjectorType::Planar;
-						break;
-
-					case mu::PROJECTOR_TYPE::CYLINDRICAL:
-						// Unapply strange swizzle for scales.
-						// TODO: try to avoid this
-						param.Value.ProjectionType = ECustomizableObjectProjectorType::Cylindrical;
-						param.Value.Direction = -param.Value.Direction;
-						param.Value.Up = -param.Value.Up;
-						param.Value.Scale[2] = -param.Value.Scale[0];
-						param.Value.Scale[0] = param.Value.Scale[1] = param.Value.Scale[1] * 2.0f;
-						break;
-
-					case mu::PROJECTOR_TYPE::WRAPPING:
-						param.Value.ProjectionType = ECustomizableObjectProjectorType::Wrapping;
-						break;
-					default:
-						// not implemented
-						check(false);
-					}
-				}
-					else
-					{
-						int32 ValueCount = MutableParameters->GetValueCount(paramIndex);
-
-						for (int32 ValueIterator = 0; ValueIterator < ValueCount; ++ValueIterator)
-						{
-							mu::RangeIndexPtr RangeValueIdxPtr = MutableParameters->GetValueIndex(paramIndex, ValueIterator);
-							int32 ValueIndex = RangeValueIdxPtr->GetPosition(0);
-
-							if (!param.RangeValues.IsValidIndex(ValueIndex))
-							{
-								param.RangeValues.AddDefaulted(ValueIndex + 1 - param.RangeValues.Num());
-							}
-
-							mu::PROJECTOR_TYPE type;
-							MutableParameters->GetProjectorValue(paramIndex,
-							&type,
-							&param.RangeValues[ValueIndex].Position[0], &param.RangeValues[ValueIndex].Position[1], &param.RangeValues[ValueIndex].Position[2],
-							&param.RangeValues[ValueIndex].Direction[0], &param.RangeValues[ValueIndex].Direction[1], &param.RangeValues[ValueIndex].Direction[2],
-							&param.RangeValues[ValueIndex].Up[0], &param.RangeValues[ValueIndex].Up[1], &param.RangeValues[ValueIndex].Up[2],
-							&param.RangeValues[ValueIndex].Scale[0], &param.RangeValues[ValueIndex].Scale[1], &param.RangeValues[ValueIndex].Scale[2],
-							&param.RangeValues[ValueIndex].Angle,
-							RangeValueIdxPtr);
-
-							switch (type)
-							{
-							case mu::PROJECTOR_TYPE::PLANAR:
-								param.RangeValues[ValueIndex].ProjectionType = ECustomizableObjectProjectorType::Planar;
-								break;
-
-							case mu::PROJECTOR_TYPE::CYLINDRICAL:
-								// Unapply strange swizzle for scales.
-								// TODO: try to avoid this
-								param.RangeValues[ValueIndex].ProjectionType = ECustomizableObjectProjectorType::Cylindrical;
-								param.RangeValues[ValueIndex].Direction = -param.RangeValues[ValueIndex].Direction;
-								param.RangeValues[ValueIndex].Up = -param.RangeValues[ValueIndex].Up;
-								param.RangeValues[ValueIndex].Scale[2] = -param.RangeValues[ValueIndex].Scale[0];
-								param.RangeValues[ValueIndex].Scale[0] = param.RangeValues[ValueIndex].Scale[1] = param.RangeValues[ValueIndex].Scale[1] * 2.0f;
-								break;
-
-							case mu::PROJECTOR_TYPE::WRAPPING:
-								param.RangeValues[ValueIndex].ProjectionType = ECustomizableObjectProjectorType::Wrapping;
-								break;
-							default:
-								// not implemented
-								check(false);
-							}
-						}
-					}
-				}
-
-				Public->GetProjectorParameters().Add(param);
-				break;
-			}
-
-			case mu::PARAMETER_TYPE::T_IMAGE:
-			{
-				FCustomizableObjectTextureParameterValue param;
-				param.ParameterName = Name;
-				param.Uid = Uid;
-
-				bool found = false;
-				for (int32 i = 0; i < OldTextureParameters.Num(); ++i)
-				{
-					if (OldTextureParameters[i].ParameterName == Name || (!Uid.IsEmpty() && OldTextureParameters[i].Uid == Uid))
-					{
-						found = true;
-						param.ParameterValue = OldTextureParameters[i].ParameterValue;
-						MutableParameters->SetImageValue(paramIndex, param.ParameterValue);
-					}
-
-					if (OldTextureParameters[i].Uid.IsEmpty())
-					{
-						Public->Modify();
-					}
-				}
-
-				if (!found)
-				{
-					param.ParameterValue = MutableParameters->GetImageValue(paramIndex);
-				}
-
-				Public->GetTextureParameters().Add(param);
-				break;
-			}
-
-			default:
-				// TODO
-				break;
-
-			}
-		}
-	}
-	else
+	
+	if (ClearLastMeshIds)
 	{
-		if (!CustomizableObject)
-		{
-			UE_LOG(LogMutable, Warning, TEXT("[ReloadParametersFromObject] No CustomizableObject, generated empty parameters for [%s] "), *GetName());
-		}
-		else if (!CustomizableObject->GetPrivate()->GetModel())
-		{
-			UE_LOG(LogMutable, Warning, TEXT("[ReloadParametersFromObject] No model in object [%s], generated empty parameters for [%s] "), *CustomizableObject->GetName(), *GetName());
-		}
+		ClearAllLastMeshIds();
 	}
 	
-	Public->GetDescriptor().CreateParametersLookupTable();
-
-	return MutableParameters;
+	return Public->GetDescriptor().ReloadParametersFromObject();
 }
 
 
 void UCustomizableObjectInstance::SetObject(UCustomizableObject* InObject)
 {
-	Descriptor.CustomizableObject = InObject;
+	Descriptor.SetCustomizableObject(InObject);
 	PrivateData->ReloadParametersFromObject(this, true);
 }
 
@@ -1347,6 +813,8 @@ void UCustomizableInstancePrivateData::DoUpdateSkeletalMesh(UCustomizableObjectI
 	check(IsInGameThread());
 	check(bAsync);
 
+	CustomizableInstance->UpdateDescriptorHash = GetTypeHash(CustomizableInstance->Descriptor);
+
 	UCustomizableObjectSystem* System = UCustomizableObjectSystem::GetInstance();
 	
 	UCustomizableObject* CustomizableObject = CustomizableInstance->GetCustomizableObject();
@@ -2156,114 +1624,79 @@ int32 UCustomizableObjectInstance::RemoveValueFromProjectorRange(const FString&
 
 bool UCustomizableObjectInstance::CreateMultiLayerProjector(const FName& ProjectorParamName)
 {
-	if (const FMultilayerProjector* Result = MultilayerProjectors.Find(ProjectorParamName))
-	{
-		checkCode(Result->CheckInstanceParameters());
-	}
-	else
-	{
-		if (!FMultilayerProjector::AreInstanceParametersValid(*this, ProjectorParamName))
-		{
-			return false;
-		}
-		
-		const FMultilayerProjector MultilayerProjector(*this, ProjectorParamName);
-		MultilayerProjectors.Add(ProjectorParamName, MultilayerProjector);
-	}
-
-	return true;
+	return Descriptor.CreateMultiLayerProjector(ProjectorParamName);
 }
 
 
 void UCustomizableObjectInstance::RemoveMultilayerProjector(const FName& ProjectorParamName)
 {
-	MultilayerProjectors.Remove(ProjectorParamName);
+	Descriptor.RemoveMultilayerProjector(ProjectorParamName);
 }
 
 
 int32 UCustomizableObjectInstance::MultilayerProjectorNumLayers(const FName& ProjectorParamName) const
 {
-	check(MultilayerProjectors.Contains(ProjectorParamName)); // Multilayer Projector not created.
-	
-	return MultilayerProjectors[ProjectorParamName].NumLayers();
+	return Descriptor.MultilayerProjectorNumLayers(ProjectorParamName);
 }
 
 
 void UCustomizableObjectInstance::MultilayerProjectorCreateLayer(const FName& ProjectorParamName, int32 Index)
 {
-	check(MultilayerProjectors.Contains(ProjectorParamName)); // Multilayer Projector not created.
-	
-	MultilayerProjectors[ProjectorParamName].CreateLayer(Index);
+	Descriptor.MultilayerProjectorCreateLayer(ProjectorParamName, Index);
 }
 
 
 void UCustomizableObjectInstance::MultilayerProjectorRemoveLayerAt(const FName& ProjectorParamName, int32 Index)
 {
-	check(MultilayerProjectors.Contains(ProjectorParamName)); // Multilayer Projector not created.
-	
-	MultilayerProjectors[ProjectorParamName].RemoveLayerAt(Index);
+	Descriptor.MultilayerProjectorRemoveLayerAt(ProjectorParamName, Index);
 }
 
 
 FMultilayerProjectorLayer UCustomizableObjectInstance::MultilayerProjectorGetLayer(const FName& ProjectorParamName, int32 Index) const
 {
-	check(MultilayerProjectors.Contains(ProjectorParamName)); // Multilayer Projector not created.
-	
-	return MultilayerProjectors[ProjectorParamName].GetLayer(Index);
+	return Descriptor.MultilayerProjectorGetLayer(ProjectorParamName, Index);
 }
 
 
 void UCustomizableObjectInstance::MultilayerProjectorUpdateLayer(const FName& ProjectorParamName, int32 Index, const FMultilayerProjectorLayer& Layer)
 {
-	check(MultilayerProjectors.Contains(ProjectorParamName)); // Multilayer Projector not created.
-	
-	MultilayerProjectors[ProjectorParamName].UpdateLayer(Index, Layer);
+	Descriptor.MultilayerProjectorUpdateLayer(ProjectorParamName, Index, Layer);
 }
 
 
 TArray<FName> UCustomizableObjectInstance::MultilayerProjectorGetVirtualLayers(const FName& ProjectorParamName) const
 {
-	check(MultilayerProjectors.Contains(ProjectorParamName)); // Multilayer Projector not created.
-	
-	return MultilayerProjectors[ProjectorParamName].GetVirtualLayers();
+	return Descriptor.MultilayerProjectorGetVirtualLayers(ProjectorParamName);
 }
 
 
 void UCustomizableObjectInstance::MultilayerProjectorCreateVirtualLayer(const FName& ProjectorParamName, const FName& Id)
 {
-	check(MultilayerProjectors.Contains(ProjectorParamName)); // Multilayer Projector not created.
-	
-	MultilayerProjectors[ProjectorParamName].CreateVirtualLayer(Id);
+	Descriptor.MultilayerProjectorCreateVirtualLayer(ProjectorParamName, Id);
 }
 
 
 FMultilayerProjectorVirtualLayer UCustomizableObjectInstance::MultilayerProjectorFindOrCreateVirtualLayer(const FName& ProjectorParamName, const FName& Id)
 {
-	check(MultilayerProjectors.Contains(ProjectorParamName)); // Multilayer Projector not created.
-	
-	return MultilayerProjectors[ProjectorParamName].FindOrCreateVirtualLayer(Id);
+	return Descriptor.MultilayerProjectorFindOrCreateVirtualLayer(ProjectorParamName, Id);
 }
 
 
 void UCustomizableObjectInstance::MultilayerProjectorRemoveVirtualLayer(const FName& ProjectorParamName, const FName& Id)
 {
-	check(MultilayerProjectors.Contains(ProjectorParamName)); // Multilayer Projector not created.
-	
-	MultilayerProjectors[ProjectorParamName].RemoveVirtualLayer(Id);
+	Descriptor.MultilayerProjectorRemoveVirtualLayer(ProjectorParamName, Id);
 }
 
+
 FMultilayerProjectorVirtualLayer UCustomizableObjectInstance::MultilayerProjectorGetVirtualLayer(const FName& ProjectorParamName, const FName& Id) const
 {
-	check(MultilayerProjectors.Contains(ProjectorParamName)); // Multilayer Projector not created.
-	
-	return MultilayerProjectors[ProjectorParamName].GetVirtualLayer(Id);
+	return Descriptor.MultilayerProjectorGetVirtualLayer(ProjectorParamName, Id);
 }
 
+
 void UCustomizableObjectInstance::MultilayerProjectorUpdateVirtualLayer(const FName& ProjectorParamName, const FName& Id, const FMultilayerProjectorVirtualLayer& Layer)
 {
-	check(MultilayerProjectors.Contains(ProjectorParamName)); // Multilayer Projector not created.
-	
-	MultilayerProjectors[ProjectorParamName].UpdateVirtualLayer(Id, Layer);
+	Descriptor.MultilayerProjectorUpdateVirtualLayer(ProjectorParamName, Id, Layer);
 }
 
 
@@ -2837,16 +2270,39 @@ void SetTexturePropertiesFromMutableImageProps(UTexture2D* Texture, const FMutab
 }
 
 
-void UCustomizableObjectInstance::Updated()
+void UCustomizableObjectInstance::Updated(EUpdateResult Result)
+{
+	// Call Customizable Skeletal Components updated callbacks.
+	for (TObjectIterator<UCustomizableSkeletalComponent> It; It; ++It)
+	{
+		if (const UCustomizableSkeletalComponent* CustomizableSkeletalComponent = *It;
+			CustomizableSkeletalComponent &&
+			CustomizableSkeletalComponent->CustomizableObjectInstance == this)
+		{
+			CustomizableSkeletalComponent->Updated(Result);
+		}
+	}
+
+	// Call Instance updated (this) callbacks.
+	DescriptorHash = UpdateDescriptorHash;
+
+	if (Result == EUpdateResult::Success)
+	{
+		UpdatedDelegate.Broadcast(this);
+		UpdatedNativeDelegate.Broadcast(this);
+	}	
+}
+
+
+uint32 UCustomizableObjectInstance::GetDescriptorHash() const
 {
-	UpdatedDelegate.Broadcast(this);
-	UpdatedNativeDelegate.Broadcast(this);
+	return DescriptorHash;	
 }
 
 
-uint32 UCustomizableObjectInstance::GetLastUpdateDescriptorHash() const
+uint32 UCustomizableObjectInstance::GetUpdateDescriptorHash() const
 {
-	return LastUpdateDescriptorHash;
+	return UpdateDescriptorHash;
 }
 
 
diff --git a/Engine/Plugins/Experimental/Mutable/Source/CustomizableObject/Private/MuCO/CustomizableObjectInstanceDescriptor.cpp b/Engine/Plugins/Experimental/Mutable/Source/CustomizableObject/Private/MuCO/CustomizableObjectInstanceDescriptor.cpp
--- a/Engine/Plugins/Experimental/Mutable/Source/CustomizableObject/Private/MuCO/CustomizableObjectInstanceDescriptor.cpp
+++ b/Engine/Plugins/Experimental/Mutable/Source/CustomizableObject/Private/MuCO/CustomizableObjectInstanceDescriptor.cpp
@@ -43,9 +43,8 @@ FString GetAvailableOptionsString(const UCustomizableObject& CustomizableObject,
 
 FCustomizableObjectInstanceDescriptor::FCustomizableObjectInstanceDescriptor(UCustomizableObject &Object)
 {
-	CustomizableObject = &Object;
-	
-	CreateParametersLookupTable();
+	CustomizableObject = &Object;	
+	Init();
 }
 
 
@@ -447,6 +446,516 @@ void FCustomizableObjectInstanceDescriptor::SetBuildParameterDecorations(const b
 }
 
 
+void FCustomizableObjectInstanceDescriptor::SetCustomizableObject(UCustomizableObject* InCustomizableObject)
+{
+	CustomizableObject = InCustomizableObject;
+	Init();
+}
+
+
+mu::ParametersPtr FCustomizableObjectInstanceDescriptor::ReloadParametersFromObject()
+{
+	if (!CustomizableObject)
+	{
+		return nullptr;
+	}
+
+	if (!CustomizableObject->IsCompiled())
+	{	
+		return nullptr;
+	}
+	
+	int32 OldState = GetState();
+	TArray<FCustomizableObjectBoolParameterValue> OldBoolParameters = GetBoolParameters();
+	TArray<FCustomizableObjectIntParameterValue> OldIntParameters = GetIntParameters();
+	TArray<FCustomizableObjectFloatParameterValue> OldFloatParameters = GetFloatParameters();
+	TArray<FCustomizableObjectTextureParameterValue> OldTextureParameters = GetTextureParameters();
+	TArray<FCustomizableObjectVectorParameterValue> OldVectorParameters = GetVectorParameters();
+	TArray<FCustomizableObjectProjectorParameterValue> OldProjectorParameters = GetProjectorParameters();
+
+	mu::ParametersPtr MutableParameters;
+	SetState(FMath::Clamp(OldState, 0, CustomizableObject->GetStateCount() - 1));
+	GetBoolParameters().Reset();
+	GetIntParameters().Reset();
+	GetFloatParameters().Reset();
+	GetTextureParameters().Reset();
+	GetVectorParameters().Reset();
+	GetProjectorParameters().Reset();
+	
+	if (!CustomizableObject->GetPrivate()->GetModel())
+	{
+		UE_LOG(LogMutable, Warning, TEXT("[ReloadParametersFromObject] No model in object [%s], generated empty parameters for [%s] "), *CustomizableObject->GetName());
+		return nullptr;
+	}
+
+	TArray<bool> OldIntParametersUsed;
+	OldIntParametersUsed.SetNumUninitialized(OldIntParameters.Num());
+	FMemory::Memzero(OldIntParametersUsed.GetData(), OldIntParametersUsed.GetAllocatedSize());
+
+	MutableParameters = CustomizableObject->GetPrivate()->GetModel()->NewParameters();
+	int32 paramCount = MutableParameters->GetCount();
+
+	for (int32 paramIndex = 0; paramIndex<paramCount; ++paramIndex)
+	{
+		FString Name = MutableParameters->GetName(paramIndex);
+		FString Uid = MutableParameters->GetUid(paramIndex);
+		mu::PARAMETER_TYPE mutableType = MutableParameters->GetType(paramIndex);
+
+		switch (mutableType)
+		{
+		case mu::PARAMETER_TYPE::T_BOOL:
+		{
+			FCustomizableObjectBoolParameterValue param;
+			param.ParameterName = Name;
+			param.Uid = Uid;
+
+			bool found = false;
+			for (int32 i = 0; i < OldBoolParameters.Num(); ++i)
+			{
+				if (OldBoolParameters[i].ParameterName == Name)
+				{
+					found = true;
+					param.ParameterValue = OldBoolParameters[i].ParameterValue;
+					MutableParameters->SetBoolValue(paramIndex, param.ParameterValue);
+					break;
+				}
+			}
+
+			if (!found)
+			{
+				for (int32 i = 0; i < OldBoolParameters.Num(); ++i)
+				{
+					if (!Uid.IsEmpty() && OldBoolParameters[i].Uid == Uid)
+					{
+						found = true;
+						param.ParameterValue = OldBoolParameters[i].ParameterValue;
+						MutableParameters->SetBoolValue(paramIndex, param.ParameterValue);
+						break;
+					}
+				}
+			}
+
+			if (!found)
+			{
+				param.ParameterValue = MutableParameters->GetBoolValue(paramIndex);
+			}
+
+			GetBoolParameters().Add(param);
+			break;
+		}
+
+		case mu::PARAMETER_TYPE::T_INT:
+		{
+			FString ParameterValueName;
+			TArray<FString> ParameterRangeValueNames;
+
+			bool found = false;
+			for (int32 i = 0; i < OldIntParameters.Num(); ++i)
+			{
+				if (OldIntParametersUsed[i] == false && OldIntParameters[i].ParameterName.Equals(Name, ESearchCase::CaseSensitive))
+				{
+					found = true;
+					OldIntParametersUsed[i] = true;
+
+					mu::RangeIndexPtr RangeIdxPtr = MutableParameters->NewRangeIndex(paramIndex);
+
+					if (!RangeIdxPtr)
+					{
+					ParameterValueName = OldIntParameters[i].ParameterValueName;
+					int32 Value = CustomizableObject->FindIntParameterValue(paramIndex, ParameterValueName);
+					MutableParameters->SetIntValue(paramIndex, Value);
+					}
+					else
+					{
+						for (int RangeIndex = 0; RangeIndex < OldIntParameters[i].ParameterRangeValueNames.Num(); ++RangeIndex)
+						{
+							FString AuxParameterValueName = OldIntParameters[i].ParameterRangeValueNames[RangeIndex];
+							ParameterRangeValueNames.Add(AuxParameterValueName);
+							int32 Value = CustomizableObject->FindIntParameterValue(paramIndex, AuxParameterValueName);
+							RangeIdxPtr->SetPosition(0, RangeIndex);
+							MutableParameters->SetIntValue(paramIndex, Value, RangeIdxPtr);
+						}
+					}
+
+					break;
+				}
+			}
+
+			if (!found)
+			{
+				for (int32 i = 0; i < OldIntParameters.Num(); ++i)
+				{
+					if (!Uid.IsEmpty() && OldIntParametersUsed[i] == false && OldIntParameters[i].Uid == Uid)
+					{
+						found = true;
+						OldIntParametersUsed[i] = true;
+
+						mu::RangeIndexPtr RangeIdxPtr = MutableParameters->NewRangeIndex(paramIndex);
+
+						if (!RangeIdxPtr)
+						{
+						ParameterValueName = OldIntParameters[i].ParameterValueName;
+						int32 Value = CustomizableObject->FindIntParameterValue(paramIndex, ParameterValueName);
+						MutableParameters->SetIntValue(paramIndex, Value);
+						}
+						else
+						{
+							for (int RangeIndex = 0; RangeIndex < OldIntParameters[i].ParameterRangeValueNames.Num(); ++RangeIndex)
+							{
+								FString AuxParameterValueName = OldIntParameters[i].ParameterRangeValueNames[RangeIndex];
+								ParameterRangeValueNames.Add(AuxParameterValueName);
+								int32 Value = CustomizableObject->FindIntParameterValue(paramIndex, AuxParameterValueName);
+								RangeIdxPtr->SetPosition(0, RangeIndex);
+								MutableParameters->SetIntValue(paramIndex, Value, RangeIdxPtr);
+							}
+						}
+
+						break;
+					}
+				}
+			}
+
+			if (!found)
+			{
+				mu::RangeIndexPtr RangeIdxPtr = MutableParameters->NewRangeIndex(paramIndex);
+
+				if (!RangeIdxPtr.get())
+				{
+				int32 Value = MutableParameters->GetIntValue(paramIndex);
+				ParameterValueName = CustomizableObject->FindIntParameterValueName(paramIndex, Value);
+				}
+				else
+				{
+					int32 ValueCount = MutableParameters->GetValueCount(paramIndex);
+
+					for (int32 ValueIterator = 0; ValueIterator < ValueCount; ++ValueIterator)
+					{
+						mu::RangeIndexPtr RangeValueIdxPtr = MutableParameters->GetValueIndex(paramIndex, ValueIterator);
+						int32 ValueIndex = RangeValueIdxPtr->GetPosition(0);
+
+						//if (!param.RangeValues.IsValidIndex(ValueIndex))
+						//{
+						//	param.RangeValues.AddDefaulted(ValueIndex + 1 - param.RangeValues.Num());
+						//}
+
+						int32 Value = MutableParameters->GetIntValue(paramIndex, RangeValueIdxPtr);
+						FString AuxParameterValueName = CustomizableObject->FindIntParameterValueName(paramIndex, Value);
+						//param.RangeValues[ValueIndex] = AuxParameterValueName;
+
+						if (!ParameterRangeValueNames.IsValidIndex(ValueIndex))
+						{
+							ParameterRangeValueNames.AddDefaulted(ValueIndex + 1 - ParameterRangeValueNames.Num());
+						}
+
+						ParameterRangeValueNames[ValueIndex] = AuxParameterValueName;
+					}
+				}
+			}
+
+			GetIntParameters().Emplace(Name, ParameterValueName, Uid, ParameterRangeValueNames);
+
+			break;
+		}
+
+		case mu::PARAMETER_TYPE::T_FLOAT:
+		{
+			FCustomizableObjectFloatParameterValue param;
+			param.ParameterName = Name;
+			param.Uid = Uid;
+
+			bool found = false;
+			for (int32 i = 0; i < OldFloatParameters.Num(); ++i)
+			{
+				if (OldFloatParameters[i].ParameterName == Name || (!Uid.IsEmpty() && OldFloatParameters[i].Uid == Uid))
+				{
+					found = true;
+					param.ParameterValue = OldFloatParameters[i].ParameterValue;
+					param.ParameterRangeValues = OldFloatParameters[i].ParameterRangeValues;
+
+					mu::RangeIndexPtr RangeIdxPtr = MutableParameters->NewRangeIndex(paramIndex);
+
+					if (!RangeIdxPtr)
+					{
+						MutableParameters->SetFloatValue(paramIndex, param.ParameterValue);
+					}
+					else
+					{
+						for (int RangeIndex = 0; RangeIndex < OldFloatParameters[i].ParameterRangeValues.Num(); ++RangeIndex)
+						{
+							RangeIdxPtr->SetPosition(0, RangeIndex);
+							MutableParameters->SetFloatValue(paramIndex, param.ParameterRangeValues[RangeIndex], RangeIdxPtr);
+						}
+					}
+				}
+			}
+
+			if (!found)
+			{
+				mu::RangeIndexPtr RangeIdxPtr = MutableParameters->NewRangeIndex(paramIndex);
+
+				if (!RangeIdxPtr.get())
+				{
+				param.ParameterValue = MutableParameters->GetFloatValue(paramIndex);
+			}
+				else
+				{
+				param.ParameterValue = MutableParameters->GetFloatValue(paramIndex);
+
+					int32 ValueCount = MutableParameters->GetValueCount(paramIndex);
+
+					for (int32 ValueIterator = 0; ValueIterator < ValueCount; ++ValueIterator)
+					{
+						mu::RangeIndexPtr RangeValueIdxPtr = MutableParameters->GetValueIndex(paramIndex, ValueIterator);
+						int32 ValueIndex = RangeValueIdxPtr->GetPosition(0);
+
+						float Value = MutableParameters->GetFloatValue(paramIndex, RangeValueIdxPtr);
+
+						if (!param.ParameterRangeValues.IsValidIndex(ValueIndex))
+						{
+							param.ParameterRangeValues.AddDefaulted(ValueIndex + 1 - param.ParameterRangeValues.Num());
+						}
+
+						param.ParameterRangeValues[ValueIndex] = Value;
+					}
+				}
+			}
+
+			GetFloatParameters().Add(param);
+			break;
+		}
+
+		case mu::PARAMETER_TYPE::T_COLOUR:
+		{
+			FCustomizableObjectVectorParameterValue param;
+			param.ParameterName = Name;
+			param.Uid = Uid;
+
+			bool found = false;
+			for (int32 i = 0; i < OldVectorParameters.Num(); ++i)
+			{
+				if (OldVectorParameters[i].ParameterName == Name || (!Uid.IsEmpty() && OldVectorParameters[i].Uid == Uid))
+				{
+					found = true;
+					param.ParameterValue = OldVectorParameters[i].ParameterValue;
+					MutableParameters->SetColourValue(paramIndex, param.ParameterValue.R, param.ParameterValue.G, param.ParameterValue.B);
+				}
+			}
+
+			if (!found)
+			{
+				MutableParameters->GetColourValue(paramIndex, &param.ParameterValue.R, &param.ParameterValue.G, &param.ParameterValue.B);
+				param.ParameterValue.A = 1.0f;
+			}
+
+			GetVectorParameters().Add(param);
+			break;
+		}
+
+		case mu::PARAMETER_TYPE::T_PROJECTOR:
+		{
+			FCustomizableObjectProjectorParameterValue param;
+			param.ParameterName = Name;
+			param.Uid = Uid;
+
+			bool found = false;
+			for (int32 i = 0; i < OldProjectorParameters.Num(); ++i)
+			{
+				if (OldProjectorParameters[i].ParameterName == Name || (!Uid.IsEmpty() && OldProjectorParameters[i].Uid == Uid))
+				{
+					found = true;
+
+					param.Value = OldProjectorParameters[i].Value;
+					param.RangeValues = OldProjectorParameters[i].RangeValues;
+
+					mu::RangeIndexPtr RangeIdxPtr = MutableParameters->NewRangeIndex(paramIndex);
+
+					TArray<FCustomizableObjectProjector> AuxValuesArray;
+					AuxValuesArray.Add(param.Value);
+					TArray<FCustomizableObjectProjector>& Values = RangeIdxPtr.get() ? param.RangeValues : AuxValuesArray;
+					int32 ArrayIndex = 0;
+
+					for (FCustomizableObjectProjector& Value : Values)
+					{
+						if (RangeIdxPtr.get())
+						{
+							check(RangeIdxPtr->GetRangeCount() == 1);
+							RangeIdxPtr->SetPosition(0, ArrayIndex);
+						}
+
+						switch (Value.ProjectionType)
+						{
+						case ECustomizableObjectProjectorType::Planar:
+						case ECustomizableObjectProjectorType::Wrapping:
+						{
+							MutableParameters->SetProjectorValue(paramIndex,
+								Value.Position[0], Value.Position[1], Value.Position[2],
+								Value.Direction[0], Value.Direction[1], Value.Direction[2],
+								Value.Up[0], Value.Up[1], Value.Up[2],
+								Value.Scale[0], Value.Scale[1], Value.Scale[2],
+								Value.Angle,
+								RangeIdxPtr);
+							break;
+						}
+
+						case ECustomizableObjectProjectorType::Cylindrical:
+						{
+							// Apply strange swizzle for scales
+							// TODO: try to avoid this
+								float Radius = FMath::Abs(Value.Scale[0] / 2.0f);
+								float Height = Value.Scale[2];
+							// TODO: try to avoid this
+							MutableParameters->SetProjectorValue(paramIndex,
+									Value.Position[0], Value.Position[1], Value.Position[2],
+									-Value.Direction[0], -Value.Direction[1], -Value.Direction[2],
+									-Value.Up[0], -Value.Up[1], -Value.Up[2],
+									-Height, Radius, Radius,
+									Value.Angle,
+									RangeIdxPtr);
+							break;
+						}
+
+						default:
+							// Not implemented.
+							check(false);
+						}
+						
+						ArrayIndex++;
+					}
+				}
+			}
+
+			if (!found)
+			{
+				mu::RangeIndexPtr RangeIdxPtr = MutableParameters->NewRangeIndex(paramIndex);
+
+				if (!RangeIdxPtr.get())
+				{
+				mu::PROJECTOR_TYPE type;
+				MutableParameters->GetProjectorValue(paramIndex,
+					&type,
+					&param.Value.Position[0], &param.Value.Position[1], &param.Value.Position[2],
+					&param.Value.Direction[0], &param.Value.Direction[1], &param.Value.Direction[2],
+					&param.Value.Up[0], &param.Value.Up[1], &param.Value.Up[2],
+					&param.Value.Scale[0], &param.Value.Scale[1], &param.Value.Scale[2],
+					&param.Value.Angle );
+				switch (type)
+				{
+				case mu::PROJECTOR_TYPE::PLANAR:
+					param.Value.ProjectionType = ECustomizableObjectProjectorType::Planar;
+					break;
+
+				case mu::PROJECTOR_TYPE::CYLINDRICAL:
+					// Unapply strange swizzle for scales.
+					// TODO: try to avoid this
+					param.Value.ProjectionType = ECustomizableObjectProjectorType::Cylindrical;
+					param.Value.Direction = -param.Value.Direction;
+					param.Value.Up = -param.Value.Up;
+					param.Value.Scale[2] = -param.Value.Scale[0];
+					param.Value.Scale[0] = param.Value.Scale[1] = param.Value.Scale[1] * 2.0f;
+					break;
+
+				case mu::PROJECTOR_TYPE::WRAPPING:
+					param.Value.ProjectionType = ECustomizableObjectProjectorType::Wrapping;
+					break;
+				default:
+					// not implemented
+					check(false);
+				}
+			}
+				else
+				{
+					int32 ValueCount = MutableParameters->GetValueCount(paramIndex);
+
+					for (int32 ValueIterator = 0; ValueIterator < ValueCount; ++ValueIterator)
+					{
+						mu::RangeIndexPtr RangeValueIdxPtr = MutableParameters->GetValueIndex(paramIndex, ValueIterator);
+						int32 ValueIndex = RangeValueIdxPtr->GetPosition(0);
+
+						if (!param.RangeValues.IsValidIndex(ValueIndex))
+						{
+							param.RangeValues.AddDefaulted(ValueIndex + 1 - param.RangeValues.Num());
+						}
+
+						mu::PROJECTOR_TYPE type;
+						MutableParameters->GetProjectorValue(paramIndex,
+						&type,
+						&param.RangeValues[ValueIndex].Position[0], &param.RangeValues[ValueIndex].Position[1], &param.RangeValues[ValueIndex].Position[2],
+						&param.RangeValues[ValueIndex].Direction[0], &param.RangeValues[ValueIndex].Direction[1], &param.RangeValues[ValueIndex].Direction[2],
+						&param.RangeValues[ValueIndex].Up[0], &param.RangeValues[ValueIndex].Up[1], &param.RangeValues[ValueIndex].Up[2],
+						&param.RangeValues[ValueIndex].Scale[0], &param.RangeValues[ValueIndex].Scale[1], &param.RangeValues[ValueIndex].Scale[2],
+						&param.RangeValues[ValueIndex].Angle,
+						RangeValueIdxPtr);
+
+						switch (type)
+						{
+						case mu::PROJECTOR_TYPE::PLANAR:
+							param.RangeValues[ValueIndex].ProjectionType = ECustomizableObjectProjectorType::Planar;
+							break;
+
+						case mu::PROJECTOR_TYPE::CYLINDRICAL:
+							// Unapply strange swizzle for scales.
+							// TODO: try to avoid this
+							param.RangeValues[ValueIndex].ProjectionType = ECustomizableObjectProjectorType::Cylindrical;
+							param.RangeValues[ValueIndex].Direction = -param.RangeValues[ValueIndex].Direction;
+							param.RangeValues[ValueIndex].Up = -param.RangeValues[ValueIndex].Up;
+							param.RangeValues[ValueIndex].Scale[2] = -param.RangeValues[ValueIndex].Scale[0];
+							param.RangeValues[ValueIndex].Scale[0] = param.RangeValues[ValueIndex].Scale[1] = param.RangeValues[ValueIndex].Scale[1] * 2.0f;
+							break;
+
+						case mu::PROJECTOR_TYPE::WRAPPING:
+							param.RangeValues[ValueIndex].ProjectionType = ECustomizableObjectProjectorType::Wrapping;
+							break;
+						default:
+							// not implemented
+							check(false);
+						}
+					}
+				}
+			}
+
+			GetProjectorParameters().Add(param);
+			break;
+		}
+
+		case mu::PARAMETER_TYPE::T_IMAGE:
+		{
+			FCustomizableObjectTextureParameterValue param;
+			param.ParameterName = Name;
+			param.Uid = Uid;
+
+			bool found = false;
+			for (int32 i = 0; i < OldTextureParameters.Num(); ++i)
+			{
+				if (OldTextureParameters[i].ParameterName == Name || (!Uid.IsEmpty() && OldTextureParameters[i].Uid == Uid))
+				{
+					found = true;
+					param.ParameterValue = OldTextureParameters[i].ParameterValue;
+					MutableParameters->SetImageValue(paramIndex, param.ParameterValue);
+				}
+			}
+
+			if (!found)
+			{
+				param.ParameterValue = MutableParameters->GetImageValue(paramIndex);
+			}
+
+			GetTextureParameters().Add(param);
+			break;
+		}
+
+		default:
+			// TODO
+			break;
+
+		}
+	}
+	
+	CreateParametersLookupTable();
+
+	return MutableParameters;
+}
+
+
 TArray<FCustomizableObjectBoolParameterValue>& FCustomizableObjectInstanceDescriptor::GetBoolParameters()
 {
 	return BoolParameters;
@@ -532,6 +1041,8 @@ bool FCustomizableObjectInstanceDescriptor::HasAnyParameters() const
 
 const FString& FCustomizableObjectInstanceDescriptor::GetIntParameterSelectedOption(const FString& ParamName, const int32 RangeIndex) const
 {
+	check(CustomizableObject);
+	
 	const int32 ParameterIndexInObject = CustomizableObject->FindParameter(ParamName);
 
 	const int32 IntParamIndex = FindIntParameterNameIndex(ParamName);
@@ -560,6 +1071,8 @@ const FString& FCustomizableObjectInstanceDescriptor::GetIntParameterSelectedOpt
 
 void FCustomizableObjectInstanceDescriptor::SetIntParameterSelectedOption(const int32 IntParamIndex, const FString& SelectedOption, const int32 RangeIndex)
 {
+	check(CustomizableObject);
+
 	check(IntParameters.IsValidIndex(IntParamIndex));
 
 	if (IntParameters.IsValidIndex(IntParamIndex))
@@ -613,6 +1126,8 @@ void FCustomizableObjectInstanceDescriptor::SetIntParameterSelectedOption(const
 
 float FCustomizableObjectInstanceDescriptor::GetFloatParameterSelectedOption(const FString& FloatParamName, const int32 RangeIndex) const
 {
+	check(CustomizableObject);
+
 	const int32 ParameterIndexInObject = CustomizableObject->FindParameter(FloatParamName);
 
 	const int32 FloatParamIndex = FindFloatParameterNameIndex(FloatParamName);
@@ -641,6 +1156,8 @@ float FCustomizableObjectInstanceDescriptor::GetFloatParameterSelectedOption(con
 
 void FCustomizableObjectInstanceDescriptor::SetFloatParameterSelectedOption(const FString& FloatParamName, const float FloatValue, const int32 RangeIndex)
 {
+	check(CustomizableObject);
+
 	const int32 ParameterIndexInObject = CustomizableObject->FindParameter(FloatParamName);
 
 	const int32 FloatParamIndex = FindFloatParameterNameIndex(FloatParamName);
@@ -672,6 +1189,8 @@ void FCustomizableObjectInstanceDescriptor::SetFloatParameterSelectedOption(cons
 
 FLinearColor FCustomizableObjectInstanceDescriptor::GetColorParameterSelectedOption(const FString& ColorParamName) const
 {
+	check(CustomizableObject);
+
 	const int32 ParameterIndexInObject = CustomizableObject->FindParameter(ColorParamName);
 
 	const int32 ColorParamIndex = FindVectorParameterNameIndex(ColorParamName);
@@ -693,6 +1212,8 @@ void FCustomizableObjectInstanceDescriptor::SetColorParameterSelectedOption(cons
 
 bool FCustomizableObjectInstanceDescriptor::GetBoolParameterSelectedOption(const FString& BoolParamName) const
 {
+	check(CustomizableObject);
+
 	const int32 ParameterIndexInObject = CustomizableObject->FindParameter(BoolParamName);
 
 	const int32 BoolParamIndex = FindBoolParameterNameIndex(BoolParamName);
@@ -708,6 +1229,8 @@ bool FCustomizableObjectInstanceDescriptor::GetBoolParameterSelectedOption(const
 
 void FCustomizableObjectInstanceDescriptor::SetBoolParameterSelectedOption(const FString& BoolParamName, const bool BoolValue)
 {
+	check(CustomizableObject);
+
 	const int32 ParameterIndexInObject = CustomizableObject->FindParameter(BoolParamName);
 
 	const int32 BoolParamIndex = FindBoolParameterNameIndex(BoolParamName);
@@ -721,6 +1244,8 @@ void FCustomizableObjectInstanceDescriptor::SetBoolParameterSelectedOption(const
 
 void FCustomizableObjectInstanceDescriptor::SetVectorParameterSelectedOption(const FString& VectorParamName, const FLinearColor& VectorValue)
 {
+	check(CustomizableObject);
+
 	const int32 ParameterIndexInObject = CustomizableObject->FindParameter(VectorParamName);
 
 	const int32 VectorParamIndex = FindVectorParameterNameIndex(VectorParamName);
@@ -737,6 +1262,8 @@ void FCustomizableObjectInstanceDescriptor::SetProjectorValue(const FString& Pro
 	const float Angle,
 	const int32 RangeIndex)
 {
+	check(CustomizableObject);
+
 	const int32 ParameterIndexInObject = CustomizableObject->FindParameter(ProjectorParamName);
 
 	const int32 ProjectorParamIndex = FindProjectorParameterNameIndex(ProjectorParamName);
@@ -776,6 +1303,8 @@ void FCustomizableObjectInstanceDescriptor::SetProjectorValue(const FString& Pro
 
 void FCustomizableObjectInstanceDescriptor::SetProjectorPosition(const FString& ProjectorParamName, const FVector3f& Pos, const int32 RangeIndex)
 {
+	check(CustomizableObject);
+
 	const int32 ParameterIndexInObject = CustomizableObject->FindParameter(ProjectorParamName);
 
 	const int32 ProjectorParamIndex = FindProjectorParameterNameIndex(ProjectorParamName);
@@ -813,6 +1342,8 @@ void FCustomizableObjectInstanceDescriptor::GetProjectorValue(const FString& Pro
 	float& OutAngle, ECustomizableObjectProjectorType& OutType,
 	const int32 RangeIndex) const
 {
+	check(CustomizableObject);
+
 	const int32 ParameterIndexInObject = CustomizableObject->FindParameter(ProjectorParamName);
 
 	const int32 ProjectorParamIndex = FindProjectorParameterNameIndex(ProjectorParamName);
@@ -853,6 +1384,8 @@ void FCustomizableObjectInstanceDescriptor::GetProjectorValueF(const FString& Pr
 	float& OutAngle, ECustomizableObjectProjectorType& OutType,
 	const int32 RangeIndex) const
 {
+	check(CustomizableObject);
+
 	const int32 ParameterIndexInObject = CustomizableObject->FindParameter(ProjectorParamName);
 
 	const int32 ProjectorParamIndex = FindProjectorParameterNameIndex(ProjectorParamName);
@@ -890,6 +1423,8 @@ void FCustomizableObjectInstanceDescriptor::GetProjectorValueF(const FString& Pr
 
 FVector FCustomizableObjectInstanceDescriptor::GetProjectorPosition(const FString& ParamName, const int32 RangeIndex) const
 {
+	check(CustomizableObject);
+
 	const int32 ParameterIndexInObject = CustomizableObject->FindParameter(ParamName);
 
 	const int32 ProjectorParamIndex = FindProjectorParameterNameIndex(ParamName);
@@ -912,6 +1447,8 @@ FVector FCustomizableObjectInstanceDescriptor::GetProjectorPosition(const FStrin
 
 FVector FCustomizableObjectInstanceDescriptor::GetProjectorDirection(const FString& ParamName, const int32 RangeIndex) const
 {
+	check(CustomizableObject);
+
 	const int32 ParameterIndexInObject = CustomizableObject->FindParameter(ParamName);
 
 	const int32 ProjectorParamIndex = FindProjectorParameterNameIndex(ParamName);
@@ -934,6 +1471,8 @@ FVector FCustomizableObjectInstanceDescriptor::GetProjectorDirection(const FStri
 
 FVector FCustomizableObjectInstanceDescriptor::GetProjectorUp(const FString& ParamName, const int32 RangeIndex) const
 {
+	check(CustomizableObject);
+
 	const int32 ParameterIndexInObject = CustomizableObject->FindParameter(ParamName);
 
 	const int32 ProjectorParamIndex = FindProjectorParameterNameIndex(ParamName);
@@ -956,6 +1495,8 @@ FVector FCustomizableObjectInstanceDescriptor::GetProjectorUp(const FString& Par
 
 FVector FCustomizableObjectInstanceDescriptor::GetProjectorScale(const FString& ParamName, const int32 RangeIndex) const
 {
+	check(CustomizableObject);
+
 	const int32 ParameterIndexInObject = CustomizableObject->FindParameter(ParamName);
 
 	const int32 ProjectorParamIndex = FindProjectorParameterNameIndex(ParamName);
@@ -978,6 +1519,8 @@ FVector FCustomizableObjectInstanceDescriptor::GetProjectorScale(const FString&
 
 float FCustomizableObjectInstanceDescriptor::GetProjectorAngle(const FString& ParamName, const int32 RangeIndex) const
 {
+	check(CustomizableObject);
+
 	const int32 ParameterIndexInObject = CustomizableObject->FindParameter(ParamName);
 
 	const int32 ProjectorParamIndex = FindProjectorParameterNameIndex(ParamName);
@@ -1000,6 +1543,8 @@ float FCustomizableObjectInstanceDescriptor::GetProjectorAngle(const FString& Pa
 
 ECustomizableObjectProjectorType FCustomizableObjectInstanceDescriptor::GetProjectorParameterType(const FString& ParamName, const int32 RangeIndex) const
 {
+	check(CustomizableObject);
+
 	const int32 ParameterIndexInObject = CustomizableObject->FindParameter(ParamName);
 
 	const int32 ProjectorParamIndex = FindProjectorParameterNameIndex(ParamName);
@@ -1112,6 +1657,8 @@ int32 FCustomizableObjectInstanceDescriptor::FindProjectorParameterNameIndex(con
 
 bool FCustomizableObjectInstanceDescriptor::IsParamMultidimensional(const FString& ParamName) const
 {
+	check(CustomizableObject);
+
 	const int32 ParameterIndex = CustomizableObject->FindParameter(ParamName);
 	return IsParamMultidimensional(ParameterIndex);
 }
@@ -1119,6 +1666,8 @@ bool FCustomizableObjectInstanceDescriptor::IsParamMultidimensional(const FStrin
 
 bool FCustomizableObjectInstanceDescriptor::IsParamMultidimensional(const int32 ParamIndex) const
 {
+	check(CustomizableObject);
+
 	const mu::ParametersPtr MutableParameters = CustomizableObject->GetPrivate()->GetModel()->NewParameters();
 	check(ParamIndex < MutableParameters->GetCount());
 	const mu::RangeIndexPtr RangeIdxPtr = MutableParameters->NewRangeIndex(ParamIndex);
@@ -1129,6 +1678,8 @@ bool FCustomizableObjectInstanceDescriptor::IsParamMultidimensional(const int32
 
 int32 FCustomizableObjectInstanceDescriptor::CurrentParamRange(const FString& ParamName) const
 {
+	check(CustomizableObject);
+
 	const int32 ParameterIndexInObject = CustomizableObject->FindParameter(ParamName);
 
 	const int32 ProjectorParamIndex = FindProjectorParameterNameIndex(ParamName);
@@ -1146,6 +1697,8 @@ int32 FCustomizableObjectInstanceDescriptor::CurrentParamRange(const FString& Pa
 
 int32 FCustomizableObjectInstanceDescriptor::AddValueToIntRange(const FString& ParamName)
 {
+	check(CustomizableObject);
+
 	const int32 intParameterIndex = FindIntParameterNameIndex(ParamName);
 	if (intParameterIndex != -1)
 	{
@@ -1176,6 +1729,8 @@ int32 FCustomizableObjectInstanceDescriptor::AddValueToFloatRange(const FString&
 
 int32 FCustomizableObjectInstanceDescriptor::AddValueToProjectorRange(const FString& ParamName)
 {
+	check(CustomizableObject);
+
 	const int32 projectorParameterIndex = FindProjectorParameterNameIndex(ParamName);
 	if (projectorParameterIndex != -1)
 	{
@@ -1287,6 +1842,8 @@ int32 FCustomizableObjectInstanceDescriptor::RemoveValueFromProjectorRange(const
 
 FCustomizableObjectProjector FCustomizableObjectInstanceDescriptor::GetProjectorDefaultValue(int32 const ParamIndex) const
 {
+	check(CustomizableObject);
+
 	const mu::ParametersPtr MutableParameters = CustomizableObject->GetPrivate()->GetModel()->NewParameters();
 	check(ParamIndex < MutableParameters->GetCount());
 
@@ -1338,6 +1895,8 @@ int32 FCustomizableObjectInstanceDescriptor::GetState() const
 
 FString FCustomizableObjectInstanceDescriptor::GetCurrentState() const
 {
+	check(CustomizableObject);
+
 	return CustomizableObject->GetStateName(GetState());
 }
 
@@ -1350,12 +1909,16 @@ void FCustomizableObjectInstanceDescriptor::SetState(const int32 InState)
 
 void FCustomizableObjectInstanceDescriptor::SetCurrentState(const FString& StateName)
 {
+	check(CustomizableObject);
+
 	SetState(CustomizableObject->FindState(StateName));
 }
 
 
 void FCustomizableObjectInstanceDescriptor::SetRandomValues()
 {
+	check(CustomizableObject);
+
 	for (int32 i = 0; i < FloatParameters.Num(); ++i)
 	{
 		FloatParameters[i].ParameterValue = FMath::SRand();
@@ -1385,41 +1948,159 @@ void FCustomizableObjectInstanceDescriptor::SetRandomValues()
 }
 
 
+bool FCustomizableObjectInstanceDescriptor::CreateMultiLayerProjector(const FName& ProjectorParamName)
+{
+	if (const FMultilayerProjector* Result = MultilayerProjectors.Find(ProjectorParamName))
+	{
+		checkCode(Result->CheckDescriptorParameters(*this));
+	}
+	else
+	{
+		if (!FMultilayerProjector::AreDescriptorParametersValid(*this, ProjectorParamName))
+		{
+			return false;
+		}
+		
+		const FMultilayerProjector MultilayerProjector(ProjectorParamName);
+		MultilayerProjectors.Add(ProjectorParamName, MultilayerProjector);
+	}
+
+	return true;
+}
+
+
+void FCustomizableObjectInstanceDescriptor::RemoveMultilayerProjector(const FName& ProjectorParamName)
+{
+	MultilayerProjectors.Remove(ProjectorParamName);
+}
+
+
+int32 FCustomizableObjectInstanceDescriptor::MultilayerProjectorNumLayers(const FName& ProjectorParamName) const
+{
+	check(MultilayerProjectors.Contains(ProjectorParamName)); // Multilayer Projector not created.
+	
+	return MultilayerProjectors[ProjectorParamName].NumLayers(*this);
+}
+
+
+void FCustomizableObjectInstanceDescriptor::MultilayerProjectorCreateLayer(const FName& ProjectorParamName, int32 Index)
+{
+	check(MultilayerProjectors.Contains(ProjectorParamName)); // Multilayer Projector not created.
+	
+	MultilayerProjectors[ProjectorParamName].CreateLayer(*this, Index);
+}
+
+
+void FCustomizableObjectInstanceDescriptor::MultilayerProjectorRemoveLayerAt(const FName& ProjectorParamName, int32 Index)
+{
+	check(MultilayerProjectors.Contains(ProjectorParamName)); // Multilayer Projector not created.
+	
+	MultilayerProjectors[ProjectorParamName].RemoveLayerAt(*this, Index);
+}
+
+
+FMultilayerProjectorLayer FCustomizableObjectInstanceDescriptor::MultilayerProjectorGetLayer(const FName& ProjectorParamName, int32 Index) const
+{
+	check(MultilayerProjectors.Contains(ProjectorParamName)); // Multilayer Projector not created.
+	
+	return MultilayerProjectors[ProjectorParamName].GetLayer(*this, Index);
+}
+
+
+void FCustomizableObjectInstanceDescriptor::MultilayerProjectorUpdateLayer(const FName& ProjectorParamName, int32 Index, const FMultilayerProjectorLayer& Layer)
+{
+	check(MultilayerProjectors.Contains(ProjectorParamName)); // Multilayer Projector not created.
+	
+	MultilayerProjectors[ProjectorParamName].UpdateLayer(*this, Index, Layer);
+}
+
+
+TArray<FName> FCustomizableObjectInstanceDescriptor::MultilayerProjectorGetVirtualLayers(const FName& ProjectorParamName) const
+{
+	check(MultilayerProjectors.Contains(ProjectorParamName)); // Multilayer Projector not created.
+	
+	return MultilayerProjectors[ProjectorParamName].GetVirtualLayers();
+}
+
+
+void FCustomizableObjectInstanceDescriptor::MultilayerProjectorCreateVirtualLayer(const FName& ProjectorParamName, const FName& Id)
+{
+	check(MultilayerProjectors.Contains(ProjectorParamName)); // Multilayer Projector not created.
+	
+	MultilayerProjectors[ProjectorParamName].CreateVirtualLayer(*this, Id);
+}
+
+
+FMultilayerProjectorVirtualLayer FCustomizableObjectInstanceDescriptor::MultilayerProjectorFindOrCreateVirtualLayer(const FName& ProjectorParamName, const FName& Id)
+{
+	check(MultilayerProjectors.Contains(ProjectorParamName)); // Multilayer Projector not created.
+	
+	return MultilayerProjectors[ProjectorParamName].FindOrCreateVirtualLayer(*this, Id);
+}
+
+
+void FCustomizableObjectInstanceDescriptor::MultilayerProjectorRemoveVirtualLayer(const FName& ProjectorParamName, const FName& Id)
+{
+	check(MultilayerProjectors.Contains(ProjectorParamName)); // Multilayer Projector not created.
+	
+	MultilayerProjectors[ProjectorParamName].RemoveVirtualLayer(*this, Id);
+}
+
+
+FMultilayerProjectorVirtualLayer FCustomizableObjectInstanceDescriptor::MultilayerProjectorGetVirtualLayer(const FName& ProjectorParamName, const FName& Id) const
+{
+	check(MultilayerProjectors.Contains(ProjectorParamName)); // Multilayer Projector not created.
+	
+	return MultilayerProjectors[ProjectorParamName].GetVirtualLayer(*this, Id);
+}
+
+
+void FCustomizableObjectInstanceDescriptor::MultilayerProjectorUpdateVirtualLayer(const FName& ProjectorParamName, const FName& Id, const FMultilayerProjectorVirtualLayer& Layer)
+{
+	check(MultilayerProjectors.Contains(ProjectorParamName)); // Multilayer Projector not created.
+	
+	MultilayerProjectors[ProjectorParamName].UpdateVirtualLayer(*this, Id, Layer);
+}
+
+
+void FCustomizableObjectInstanceDescriptor::Init()
+{
+	ReloadParametersFromObject();
+}
+
+
 void FCustomizableObjectInstanceDescriptor::CreateParametersLookupTable()
 {
+	check(CustomizableObject);
+
+	const int32 NumParameters = CustomizableObject->GetParameterCount();
+
 	IntParametersLookupTable.Reset();
-	IntParametersLookupTable.Reserve(IntParameters.Num());
 
-	for (int32 Index = 0, Num = IntParameters.Num(); Index < Num; ++Index)
+	int32 IntParameter = 0;
+	for (int32 ParamIndex = 0; ParamIndex < NumParameters; ++ParamIndex)
 	{
-		const FCustomizableObjectIntParameterValue & Value = IntParameters[Index];
-
-#if WITH_EDITOR
-		if (IntParametersLookupTable.Contains(Value.ParameterName))
+		if (CustomizableObject->GetParameterType(ParamIndex) == EMutableParameterType::Int)
 		{
-			int ExistIndex = IntParametersLookupTable[Value.ParameterName];
-
-			UE_LOG(LogMutable, Warning,
-				TEXT("Name '%s' is already in IntParametersLookupTable (%s/%s/%s/#%d)"),
-				*Value.ParameterName, *Value.ParameterName, *Value.ParameterValueName, *Value.Uid, ExistIndex);
+			IntParametersLookupTable.Add(CustomizableObject->GetParameterName(ParamIndex), IntParameter++);
 		}
-#endif
-
-		IntParametersLookupTable.Add(Value.ParameterName, Index);
 	}
 }
 
 
 uint32 GetTypeHash(const FCustomizableObjectInstanceDescriptor& Key)
 {
-	FCustomizableObjectIntParameterValue a;
-	
-	uint32 Hash = GetTypeHash(Key.CustomizableObject);
+	uint32 Hash = 0;
+
+	if (Key.CustomizableObject) 
+	{
+		Hash = HashCombine(Hash, GetTypeHash(Key.CustomizableObject->GetPathName()));
+	}
 
 	for (const FCustomizableObjectBoolParameterValue& Value : Key.BoolParameters)
 	{
 		Hash = HashCombine(Hash, GetTypeHash(Value));
-	}
+	}	
 
 	for (const FCustomizableObjectIntParameterValue& Value : Key.IntParameters)
 	{
@@ -1450,5 +2131,11 @@ uint32 GetTypeHash(const FCustomizableObjectInstanceDescriptor& Key)
 
 	Hash = HashCombine(Hash, GetTypeHash(Key.bBuildParameterDecorations));
 
+	for (const TTuple<FName, FMultilayerProjector>& Pair : Key.MultilayerProjectors)
+	{
+		// Hash = HashCombine(Hash, GetTypeHash(Pair.Key)); // Already hashed by the FMultilayerProjector.
+		Hash = HashCombine(Hash, GetTypeHash(Pair.Value));
+	}
+	
 	return Hash;
 }
diff --git a/Engine/Plugins/Experimental/Mutable/Source/CustomizableObject/Private/MuCO/CustomizableObjectSystem.cpp b/Engine/Plugins/Experimental/Mutable/Source/CustomizableObject/Private/MuCO/CustomizableObjectSystem.cpp
--- a/Engine/Plugins/Experimental/Mutable/Source/CustomizableObject/Private/MuCO/CustomizableObjectSystem.cpp
+++ b/Engine/Plugins/Experimental/Mutable/Source/CustomizableObject/Private/MuCO/CustomizableObjectSystem.cpp
@@ -109,6 +109,10 @@ class AActor;
 class UAnimInstance;
 class UMaterialInterface;
 
+namespace impl
+{
+	void Task_Game_Callbacks_Work(UCustomizableObjectInstance* CustomizableObjectInstance);
+}
 
 DEFINE_STAT(STAT_MutableNumSkeletalMeshes);
 DEFINE_STAT(STAT_MutableNumCachedSkeletalMeshes);
@@ -694,6 +698,84 @@ static FAutoConsoleVariableRef CVarEnableMutableProgressiveMipStreaming(
 	ECVF_Default);
 
 
+/** Update the given Instance Skeletal Meshes and call its callbacks. */
+void UpdateSkeletalMesh(UCustomizableObjectInstance* CustomizableObjectInstance)
+{
+	// This used to be CustomizableObjectInstance::UpdateSkeletalMesh_PostBeginUpdate3
+	{
+		MUTABLE_CPUPROFILER_SCOPE(UpdateSkeletalMesh_PostBeginUpdate3);
+
+#if !WITH_EDITOR
+		for (int32 ComponentIndex = 0; ComponentIndex < CustomizableObjectInstance->SkeletalMeshes.Num(); ++ComponentIndex)
+		{
+			if (CustomizableObjectInstance->SkeletalMeshes[ComponentIndex])
+			{
+				CustomizableObjectInstance->SkeletalMeshes[ComponentIndex]->RebuildSocketMap();
+			}
+		}
+#endif
+
+		for (TObjectIterator<UCustomizableSkeletalComponent> It; It; ++It)
+		{
+			UCustomizableSkeletalComponent* CustomizableSkeletalComponent = *It;
+
+			if (CustomizableSkeletalComponent && CustomizableSkeletalComponent->CustomizableObjectInstance == CustomizableObjectInstance
+				&& CustomizableObjectInstance->SkeletalMeshes.IsValidIndex(CustomizableSkeletalComponent->ComponentIndex))
+			{
+				MUTABLE_CPUPROFILER_SCOPE(UpdateSkeletalMesh_SetSkeletalMesh);
+
+				const bool bIsCreatingSkeletalMesh = CustomizableObjectInstance->GetPrivate()->HasCOInstanceFlags(CreatingSkeletalMesh);
+				CustomizableSkeletalComponent->SetSkeletalMesh(CustomizableObjectInstance->SkeletalMeshes[CustomizableSkeletalComponent->ComponentIndex], false, bIsCreatingSkeletalMesh);
+
+				if (CustomizableObjectInstance->GetPrivate()->HasCOInstanceFlags(ReplacePhysicsAssets))
+				{
+					CustomizableSkeletalComponent->SetPhysicsAsset(
+						CustomizableObjectInstance->SkeletalMeshes[CustomizableSkeletalComponent->ComponentIndex] ? 
+						CustomizableObjectInstance->SkeletalMeshes[CustomizableSkeletalComponent->ComponentIndex]->GetPhysicsAsset() : nullptr);
+				}
+			}
+		}
+
+		CustomizableObjectInstance->GetPrivate()->SetCOInstanceFlags(Generated);
+		CustomizableObjectInstance->GetPrivate()->ClearCOInstanceFlags(CreatingSkeletalMesh);
+
+		CustomizableObjectInstance->bEditorPropertyChanged = false;
+
+		{
+			MUTABLE_CPUPROFILER_SCOPE(UpdateSkeletalMesh_UpdatedDelegate);
+
+			CustomizableObjectInstance->Updated(EUpdateResult::Success);
+
+#if WITH_EDITOR
+			CustomizableObjectInstance->InstanceUpdated = true;
+
+			// \TODO: Review if we can avoid all this editor code here.
+			for (int32 MeshIndex = 0; MeshIndex < CustomizableObjectInstance->SkeletalMeshes.Num(); ++MeshIndex)
+			{
+				if (USkeletalMesh* SkeletalMesh = CustomizableObjectInstance->GetSkeletalMesh(MeshIndex))
+				{
+					FUnrealBakeHelpers::BakeHelper_RegenerateImportedModel(SkeletalMesh);
+				}
+			}
+#endif //WITH_EDITOR
+		}
+
+		FCustomizableObjectSystemPrivate* CustomizableObjectSystem = UCustomizableObjectSystem::GetInstance()->GetPrivate();
+		if (CustomizableObjectSystem->bReleaseTexturesImmediately)
+		{
+			FMutableResourceCache& Cache = CustomizableObjectSystem->GetObjectCache(CustomizableObjectInstance->GetCustomizableObject());
+
+			for (TPair<uint32, FGeneratedTexture>& Item : CustomizableObjectInstance->GetPrivate()->TexturesToRelease)
+			{
+				UCustomizableInstancePrivateData::ReleaseMutableTexture(Item.Value.Id, Item.Value.Texture, Cache);
+			}
+
+			CustomizableObjectInstance->GetPrivate()->TexturesToRelease.Empty();
+		}
+	}
+}
+
+
 void FCustomizableObjectSystemPrivate::InitUpdateSkeletalMesh(UCustomizableObjectInstance* InCustomizableObjectInstance, FMutableQueueElem::EQueuePriorityType Priority)
 {
 	MUTABLE_CPUPROFILER_SCOPE(FCustomizableObjectSystemPrivate::InitUpdateSkeletalMesh);
@@ -703,6 +785,7 @@ void FCustomizableObjectSystemPrivate::InitUpdateSkeletalMesh(UCustomizableObjec
 	if (!InCustomizableObjectInstance ||
 		!InCustomizableObjectInstance->IsValidLowLevel())
 	{
+		InCustomizableObjectInstance->Updated(EUpdateResult::Error);
 		return;
 	}
 	
@@ -720,20 +803,20 @@ void FCustomizableObjectSystemPrivate::InitUpdateSkeletalMesh(UCustomizableObjec
 		{
 			MipsToSkip = 255; // This means skip all possible mips until only UTexture::GetStaticMinTextureResidentMipCount() are left
 		}
-
+		
 		const TSharedPtr<FMutableOperation> Operation = MakeShared<FMutableOperation>(FMutableOperation::CreateInstanceUpdate(InCustomizableObjectInstance, bNeverStream, MipsToSkip));
 
 		// Skip the update if the Instance has the same descriptor (pending or applied).
 		// If there is a hash collision, at worse, we would add an unnecessary update.
 		{
-			const uint32 DescriptorHash = Operation->InstanceDescriptorHash;
+			const uint32 UpdateDescriptorHash = InCustomizableObjectInstance->GetUpdateDescriptorHash();
 
 			if (const FMutableQueueElem* QueueElem = MutableOperationQueue.Get(InCustomizableObjectInstance))
 			{
 				if (const TSharedPtr<FMutableOperation>& QueuedOperation = QueueElem->Operation;
 					QueuedOperation &&
 					QueuedOperation->Type == FMutableOperation::EOperationType::Update &&
-					QueuedOperation->InstanceDescriptorHash == DescriptorHash)
+					QueuedOperation->InstanceDescriptorHash == UpdateDescriptorHash)
 				{
 					return;
 				}
@@ -741,19 +824,24 @@ void FCustomizableObjectSystemPrivate::InitUpdateSkeletalMesh(UCustomizableObjec
 		
 			if (CurrentMutableOperation &&
 				CurrentMutableOperation->Type == FMutableOperation::EOperationType::Update &&
-				CurrentMutableOperation->InstanceDescriptorHash == DescriptorHash)
+				CurrentMutableOperation->InstanceDescriptorHash == UpdateDescriptorHash)
 			{
 				return;
 			}
 			
-			if (InCustomizableObjectInstance->GetLastUpdateDescriptorHash() == DescriptorHash)
+			if (InCustomizableObjectInstance->GetDescriptorHash() == UpdateDescriptorHash)
 			{
+				UpdateSkeletalMesh(InCustomizableObjectInstance);
 				return;
 			}
 		}
 
 		MutableOperationQueue.Enqueue(FMutableQueueElem::Create(Operation, Priority, InCustomizableObjectInstance->GetPrivate()->MinSquareDistFromComponentToPlayer));
 	}
+	else
+	{
+		InCustomizableObjectInstance->Updated(EUpdateResult::Error);
+	}
 }
 
 
@@ -1629,7 +1717,7 @@ namespace impl
 		ImageToPlatformDataMap.Reset();
 	}
 
-
+	
 	void Task_Game_Callbacks(TSharedPtr<FMutableOperationData> OperationData, const TWeakObjectPtr<UCustomizableObjectInstance>& CustomizableObjectInstancePtr)
 	{
 		MUTABLE_CPUPROFILER_SCOPE(Task_Game_Callbacks)
@@ -1652,95 +1740,7 @@ namespace impl
 		}
 
 		// Actual work
-
-		// This used to be CustomizableObjectInstance::UpdateSkeletalMesh_PostBeginUpdate3
-		{
-			MUTABLE_CPUPROFILER_SCOPE(UpdateSkeletalMesh_PostBeginUpdate3);
-
-#if !WITH_EDITOR
-			for (int32 ComponentIndex = 0; ComponentIndex < CustomizableObjectInstance->SkeletalMeshes.Num(); ++ComponentIndex)
-			{
-				if (CustomizableObjectInstance->SkeletalMeshes[ComponentIndex])
-				{
-					CustomizableObjectInstance->SkeletalMeshes[ComponentIndex]->RebuildSocketMap();
-				}
-			}
-#endif
-
-			for (TObjectIterator<UCustomizableSkeletalComponent> It; It; ++It)
-			{
-				UCustomizableSkeletalComponent* CustomizableSkeletalComponent = *It;
-
-				if (CustomizableSkeletalComponent && CustomizableSkeletalComponent->CustomizableObjectInstance == CustomizableObjectInstance
-					&& CustomizableObjectInstance->SkeletalMeshes.IsValidIndex(CustomizableSkeletalComponent->ComponentIndex))
-				{
-					MUTABLE_CPUPROFILER_SCOPE(UpdateSkeletalMesh_SetSkeletalMesh);
-
-					const bool bIsCreatingSkeletalMesh = CustomizableObjectInstance->GetPrivate()->HasCOInstanceFlags(CreatingSkeletalMesh);
-					CustomizableSkeletalComponent->SetSkeletalMesh(CustomizableObjectInstance->SkeletalMeshes[CustomizableSkeletalComponent->ComponentIndex], false, bIsCreatingSkeletalMesh);
-
-					if (CustomizableObjectInstance->GetPrivate()->HasCOInstanceFlags(ReplacePhysicsAssets))
-					{
-						CustomizableSkeletalComponent->SetPhysicsAsset(
-							CustomizableObjectInstance->SkeletalMeshes[CustomizableSkeletalComponent->ComponentIndex] ? 
-							CustomizableObjectInstance->SkeletalMeshes[CustomizableSkeletalComponent->ComponentIndex]->GetPhysicsAsset() : nullptr);
-					}
-				}
-			}
-
-			CustomizableObjectInstance->GetPrivate()->SetCOInstanceFlags(Generated);
-			CustomizableObjectInstance->GetPrivate()->ClearCOInstanceFlags(CreatingSkeletalMesh);
-
-			if (CustomizableObjectInstance->bEditorPropertyChanged)
-			{
-				CustomizableObjectInstance->bEditorPropertyChanged = false;
-			}
-
-			for (TObjectIterator<UCustomizableSkeletalComponent> It; It; ++It)
-			{
-				UCustomizableSkeletalComponent* CustomizableSkeletalComponent = *It;
-
-				if (CustomizableSkeletalComponent && CustomizableSkeletalComponent->CustomizableObjectInstance == CustomizableObjectInstance)
-				{
-					MUTABLE_CPUPROFILER_SCOPE(UpdateSkeletalMesh_ComponentUpdatedDelegate);
-
-					CustomizableSkeletalComponent->UpdatedDelegate.ExecuteIfBound();
-				}
-			}
-
-			{
-				MUTABLE_CPUPROFILER_SCOPE(UpdateSkeletalMesh_UpdatedDelegate);
-
-				CustomizableObjectInstance->Updated();
-
-#if WITH_EDITOR
-				CustomizableObjectInstance->InstanceUpdated = true;
-
-				// \TODO: Review if we can avoid all this editor code here.
-				for (int32 MeshIndex = 0; MeshIndex < CustomizableObjectInstance->SkeletalMeshes.Num(); ++MeshIndex)
-				{
-					if (USkeletalMesh* SkeletalMesh = CustomizableObjectInstance->GetSkeletalMesh(MeshIndex))
-					{
-						FUnrealBakeHelpers::BakeHelper_RegenerateImportedModel(SkeletalMesh);
-					}
-				}
-#endif //WITH_EDITOR
-			}
-
-			FCustomizableObjectSystemPrivate* CustomizableObjectSystem = UCustomizableObjectSystem::GetInstance()->GetPrivate();
-			if (CustomizableObjectSystem->bReleaseTexturesImmediately)
-			{
-				FMutableResourceCache& Cache = CustomizableObjectSystem->GetObjectCache(CustomizableObjectInstance->GetCustomizableObject());
-
-				for (TPair<uint32, FGeneratedTexture>& Item : CustomizableObjectInstance->GetPrivate()->TexturesToRelease)
-				{
-					UCustomizableInstancePrivateData::ReleaseMutableTexture(Item.Value.Id, Item.Value.Texture, Cache);
-				}
-
-				CustomizableObjectInstance->GetPrivate()->TexturesToRelease.Empty();
-			}
-		}
-
+		UpdateSkeletalMesh(CustomizableObjectInstance);
 
 		// TODO: T2927
 		if (LogBenchmarkUtil::isLoggingActive())
@@ -2421,7 +2421,7 @@ FMutableOperation FMutableOperation::CreateInstanceUpdate(UCustomizableObjectIns
 	Op.bNeverStream = bInNeverStream;
 	Op.MipsToSkip = InMipsToSkip;
 	Op.CustomizableObjectInstance = InCustomizableObjectInstance;
-	Op.InstanceDescriptorHash = GetTypeHash(InCustomizableObjectInstance->GetDescriptor());
+	Op.InstanceDescriptorHash = InCustomizableObjectInstance->GetUpdateDescriptorHash();
 	Op.bStarted = false;
 	Op.bBuildParameterDecorations = InCustomizableObjectInstance->GetBuildParameterDecorations();
 	Op.Parameters = InCustomizableObjectInstance->GetPrivate()->ReloadParametersFromObject(InCustomizableObjectInstance, false);
@@ -2458,7 +2458,7 @@ void FMutableOperation::MutableIsDisabledCase()
 	{
 		UCustomizableSkeletalComponent* CustomizableSkeletalComponent = *It;
 
-		if (CustomizableSkeletalComponent && CustomizableSkeletalComponent->CustomizableObjectInstance == CustomizableObjectInstance)
+		if (CustomizableSkeletalComponent && CustomizableSkeletalComponent->CustomizableObjectInstance == CustomizableObjectInstance.Get())
 		{
 			CustomizableSkeletalComponent->SetSkeletalMesh(CustomizableObjectInstance->GetCustomizableObject()->GetRefSkeletalMesh(CustomizableSkeletalComponent->ComponentIndex), false);
 		}
@@ -2472,17 +2472,7 @@ void FMutableOperation::MutableIsDisabledCase()
 		CustomizableObjectInstance->bEditorPropertyChanged = false;
 	}
 
-	for (TObjectIterator<UCustomizableSkeletalComponent> It; It; ++It)
-	{
-		UCustomizableSkeletalComponent* CustomizableSkeletalComponent = *It;
-
-		if (CustomizableSkeletalComponent && CustomizableSkeletalComponent->CustomizableObjectInstance == CustomizableObjectInstance)
-		{
-			CustomizableSkeletalComponent->UpdatedDelegate.ExecuteIfBound();
-		}
-	}
-
-	CustomizableObjectInstance->Updated();
+	CustomizableObjectInstance->Updated(EUpdateResult::Success);
 
 #if WITH_EDITOR
 	CustomizableObjectInstance->InstanceUpdated = true;
diff --git a/Engine/Plugins/Experimental/Mutable/Source/CustomizableObject/Private/MuCO/CustomizableSkeletalComponent.cpp b/Engine/Plugins/Experimental/Mutable/Source/CustomizableObject/Private/MuCO/CustomizableSkeletalComponent.cpp
--- a/Engine/Plugins/Experimental/Mutable/Source/CustomizableObject/Private/MuCO/CustomizableSkeletalComponent.cpp
+++ b/Engine/Plugins/Experimental/Mutable/Source/CustomizableObject/Private/MuCO/CustomizableSkeletalComponent.cpp
@@ -40,6 +40,15 @@ UCustomizableSkeletalComponent::UCustomizableSkeletalComponent()
 }
 
 
+void UCustomizableSkeletalComponent::Updated(EUpdateResult Result) const
+{
+	if (Result == EUpdateResult::Success)
+	{
+		UpdatedDelegate.ExecuteIfBound();
+	}	
+}
+
+
 USkeletalMesh* UCustomizableSkeletalComponent::GetSkeletalMesh() const
 {
 	return CustomizableObjectInstance ? CustomizableObjectInstance->GetSkeletalMesh(ComponentIndex) : nullptr;
@@ -150,6 +159,7 @@ void UCustomizableSkeletalComponent::PreSave(FObjectPreSaveContext ObjectSaveCon
 	Super::PreSave(ObjectSaveContext);
 }
 
+
 #if WITH_EDITOR
 
 void UCustomizableSkeletalComponent::UpdateDistFromComponentToLevelEditorCamera(const FVector& CameraPosition)
@@ -321,11 +331,11 @@ void UCustomizableSkeletalComponent::TickComponent(float DeltaTime, ELevelTick T
 
 			bPendingSetSkeletalMesh = true;
 
-			if (CustomizableObjectInstance && CustomizableObjectInstance->GetSkeletalMesh(ComponentIndex))
+			if (CustomizableObjectInstance->GetSkeletalMesh(ComponentIndex))
 			{
 				Parent->SetSkeletalMesh(CustomizableObjectInstance->GetSkeletalMesh(ComponentIndex));
 			}
-			else if (!bSkipSetReferenceSkeletalMesh && CustomizableObjectInstance && CustomizableObject)
+			else if (!bSkipSetReferenceSkeletalMesh && CustomizableObject)
 			{
 				Parent->SetSkeletalMesh(CustomizableObject->GetRefSkeletalMesh(ComponentIndex));
 			}
diff --git a/Engine/Plugins/Experimental/Mutable/Source/CustomizableObject/Private/MuCO/MultilayerProjector.cpp b/Engine/Plugins/Experimental/Mutable/Source/CustomizableObject/Private/MuCO/MultilayerProjector.cpp
--- a/Engine/Plugins/Experimental/Mutable/Source/CustomizableObject/Private/MuCO/MultilayerProjector.cpp
+++ b/Engine/Plugins/Experimental/Mutable/Source/CustomizableObject/Private/MuCO/MultilayerProjector.cpp
@@ -6,21 +6,21 @@
 #include "Misc/AssertionMacros.h"
 #include "MuCO/CustomizableObject.h"
 #include "MuCO/CustomizableObjectInstance.h"
+#include "MuCO/CustomizableObjectInstanceDescriptor.h"
 #include "MuCO/CustomizableObjectParameterTypeDefinitions.h"
 #include "Templates/Tuple.h"
 
 
-void FMultilayerProjectorLayer::Read(const FMultilayerProjector& MultilayerProjector, const int32 Index)
+void FMultilayerProjectorLayer::Read(const FCustomizableObjectInstanceDescriptor& Descriptor, const FMultilayerProjector& MultilayerProjector, const int32 Index)
 {
-	checkCode(MultilayerProjector.CheckInstanceParameters());
-	check(Index >= 0 && Index < MultilayerProjector.NumLayers()); // Layer out of range.
+	checkCode(MultilayerProjector.CheckDescriptorParameters(Descriptor));
+	check(Index >= 0 && Index < MultilayerProjector.NumLayers(Descriptor)); // Layer out of range.
 	
-	const UCustomizableObjectInstance& Instance = *MultilayerProjector.Instance;
 	const FString ParamName = MultilayerProjector.ParamName.ToString();
 	
 	{
-		const int32 ProjectorParamIndex = Instance.FindProjectorParameterNameIndex(ParamName);
-		const FCustomizableObjectProjector& Projector = Instance.GetProjectorParameters()[ProjectorParamIndex].RangeValues[Index];
+		const int32 ProjectorParamIndex = Descriptor.FindProjectorParameterNameIndex(ParamName);
+		const FCustomizableObjectProjector& Projector = Descriptor.GetProjectorParameters()[ProjectorParamIndex].RangeValues[Index];
 		Position = static_cast<FVector3d>(Projector.Position);
 		Direction = static_cast<FVector3d>(Projector.Direction);
 		Up = static_cast<FVector3d>(Projector.Up);
@@ -29,28 +29,27 @@ void FMultilayerProjectorLayer::Read(const FMultilayerProjector& MultilayerProje
 	}
 	
 	{
-		const int32 ImageParamIndex = Instance.FindIntParameterNameIndex(ParamName + FMultilayerProjector::IMAGE_PARAMETER_POSTFIX);
-		Image = Instance.GetIntParameters()[ImageParamIndex].ParameterRangeValueNames[Index];
+		const int32 ImageParamIndex = Descriptor.FindIntParameterNameIndex(ParamName + FMultilayerProjector::IMAGE_PARAMETER_POSTFIX);
+		Image = Descriptor.GetIntParameters()[ImageParamIndex].ParameterRangeValueNames[Index];
 	}
 
 	{
-		const int32 OpacityParamIndex = Instance.FindFloatParameterNameIndex(ParamName + FMultilayerProjector::OPACITY_PARAMETER_POSTFIX);
-		Opacity = Instance.GetFloatParameters()[OpacityParamIndex].ParameterRangeValues[Index];
+		const int32 OpacityParamIndex = Descriptor.FindFloatParameterNameIndex(ParamName + FMultilayerProjector::OPACITY_PARAMETER_POSTFIX);
+		Opacity = Descriptor.GetFloatParameters()[OpacityParamIndex].ParameterRangeValues[Index];
 	}
 }
 
 
-void FMultilayerProjectorLayer::Write(const FMultilayerProjector& MultilayerProjector, const int32 Index) const
+void FMultilayerProjectorLayer::Write(FCustomizableObjectInstanceDescriptor& Descriptor, const FMultilayerProjector& MultilayerProjector, const int32 Index) const
 {
-	checkCode(MultilayerProjector.CheckInstanceParameters());
-	check(Index >= 0 && Index < MultilayerProjector.NumLayers()); // Layer out of range.
+	checkCode(MultilayerProjector.CheckDescriptorParameters(Descriptor));
+	check(Index >= 0 && Index < MultilayerProjector.NumLayers(Descriptor)); // Layer out of range.
 
-	UCustomizableObjectInstance& Instance = *MultilayerProjector.Instance;
 	const FString ParamName = MultilayerProjector.ParamName.ToString();
 
 	{
-		const int32 ProjectorParamIndex = Instance.FindProjectorParameterNameIndex(ParamName);
-		FCustomizableObjectProjector& Projector = Instance.GetProjectorParameters()[ProjectorParamIndex].RangeValues[Index];
+		const int32 ProjectorParamIndex = Descriptor.FindProjectorParameterNameIndex(ParamName);
+		FCustomizableObjectProjector& Projector = Descriptor.GetProjectorParameters()[ProjectorParamIndex].RangeValues[Index];
 		Projector.Position = static_cast<FVector3f>(Position);
 		Projector.Direction = static_cast<FVector3f>(Direction);
 		Projector.Up = static_cast<FVector3f>(Up);
@@ -59,17 +58,32 @@ void FMultilayerProjectorLayer::Write(const FMultilayerProjector& MultilayerProj
 	}
 	
 	{
-		const int32 ImageParamIndex = Instance.FindIntParameterNameIndex(ParamName + FMultilayerProjector::IMAGE_PARAMETER_POSTFIX);
-		Instance.GetIntParameters()[ImageParamIndex].ParameterRangeValueNames[Index] = Image;
+		const int32 ImageParamIndex = Descriptor.FindIntParameterNameIndex(ParamName + FMultilayerProjector::IMAGE_PARAMETER_POSTFIX);
+		Descriptor.GetIntParameters()[ImageParamIndex].ParameterRangeValueNames[Index] = Image;
 	}
 
 	{
-		const int32 OpacityParamIndex = Instance.FindFloatParameterNameIndex(ParamName + FMultilayerProjector::OPACITY_PARAMETER_POSTFIX);
-		Instance.GetFloatParameters()[OpacityParamIndex].ParameterRangeValues[Index] = Opacity;
+		const int32 OpacityParamIndex = Descriptor.FindFloatParameterNameIndex(ParamName + FMultilayerProjector::OPACITY_PARAMETER_POSTFIX);
+		Descriptor.GetFloatParameters()[OpacityParamIndex].ParameterRangeValues[Index] = Opacity;
 	}
 }
 
 
+uint32 GetTypeHash(const FMultilayerProjectorLayer& Key)
+{
+	uint32 Hash = GetTypeHash(Key.Position);
+
+	Hash = HashCombine(Hash, GetTypeHash(Key.Direction));
+	Hash = HashCombine(Hash, GetTypeHash(Key.Up));
+	Hash = HashCombine(Hash, GetTypeHash(Key.Scale));
+	Hash = HashCombine(Hash, GetTypeHash(Key.Angle));
+	Hash = HashCombine(Hash, GetTypeHash(Key.Image));
+	Hash = HashCombine(Hash, GetTypeHash(Key.Opacity));
+
+	return Hash;
+}
+
+
 FMultilayerProjectorVirtualLayer::FMultilayerProjectorVirtualLayer(const FMultilayerProjectorLayer& Layer, const bool bEnabled, const int32 Order):
 	FMultilayerProjectorLayer(Layer),
 	bEnabled(bEnabled),
@@ -87,48 +101,46 @@ const FString FMultilayerProjector::IMAGE_PARAMETER_POSTFIX = FString("_Selected
 const FString FMultilayerProjector::POSE_PARAMETER_POSTFIX = FString("_SelectedPoses");
 
 
-int32 FMultilayerProjector::NumLayers() const
+int32 FMultilayerProjector::NumLayers(const FCustomizableObjectInstanceDescriptor& Descriptor) const
 {
-	check(Instance); // Multilayer Projector created without using the Instance factory.
-
 	const FString NumLayersParamName = ParamName.ToString() + NUM_LAYERS_PARAMETER_POSTFIX;
 	
-	const int32 FloatParameterIndex = Instance->FindFloatParameterNameIndex(NumLayersParamName);
+	const int32 FloatParameterIndex = Descriptor.FindFloatParameterNameIndex(NumLayersParamName);
 	check(FloatParameterIndex != -1); // Parameter not found.
 
-	return Instance->GetFloatParameters()[FloatParameterIndex].ParameterValue;
+	return Descriptor.GetFloatParameters()[FloatParameterIndex].ParameterValue;
 }
 
 
-void FMultilayerProjector::CreateLayer(const int32 Index) const
+void FMultilayerProjector::CreateLayer(FCustomizableObjectInstanceDescriptor& Descriptor, const int32 Index) const
 {
-	check(Instance); // Multilayer Projector created without using the Instance factory.
-	checkCode(CheckInstanceParameters());
-	check(Index >= 0 && Index <= NumLayers()); // Layer is non-contiguous or out of range.
+	checkCode(CheckDescriptorParameters(Descriptor));
+	check(Index >= 0 && Index <= NumLayers(Descriptor)); // Layer is non-contiguous or out of range.
 
-	const UCustomizableObject* Object = Instance->GetCustomizableObject(); 
+	const UCustomizableObject* Object = Descriptor.GetCustomizableObject(); 
+	check(Object);
 	
 	// Num Layers.
 	{
         const FString NumLayersParamName = ParamName.ToString() + NUM_LAYERS_PARAMETER_POSTFIX;
-		const int32 FloatParameterIndex = Instance->FindFloatParameterNameIndex(NumLayersParamName);
+		const int32 FloatParameterIndex = Descriptor.FindFloatParameterNameIndex(NumLayersParamName);
 
-        Instance->GetFloatParameters()[FloatParameterIndex].ParameterValue += 1;
+        Descriptor.GetFloatParameters()[FloatParameterIndex].ParameterValue += 1;
     }
 	// Projector Range.
 	{
-		const int32 ProjectorParameterIndex = Instance->FindProjectorParameterNameIndex(ParamName.ToString());
+		const int32 ProjectorParameterIndex = Descriptor.FindProjectorParameterNameIndex(ParamName.ToString());
 		
-		FCustomizableObjectProjectorParameterValue& ProjectorParameter = Instance->GetProjectorParameters()[ProjectorParameterIndex];
-		const FCustomizableObjectProjector Projector = Instance->GetProjectorDefaultValue(Object->FindParameter(ParamName.ToString()));
+		FCustomizableObjectProjectorParameterValue& ProjectorParameter = Descriptor.GetProjectorParameters()[ProjectorParameterIndex];
+		const FCustomizableObjectProjector Projector = Descriptor.GetProjectorDefaultValue(Object->FindParameter(ParamName.ToString()));
 		ProjectorParameter.RangeValues.Insert(Projector, Index);
 	}
 	
 	// Selected Image Range.
 	{
-		const int32 IntParameterIndex = Instance->FindIntParameterNameIndex(ParamName.ToString() + FMultilayerProjector::IMAGE_PARAMETER_POSTFIX);
+		const int32 IntParameterIndex = Descriptor.FindIntParameterNameIndex(ParamName.ToString() + FMultilayerProjector::IMAGE_PARAMETER_POSTFIX);
 
-		FCustomizableObjectIntParameterValue& IntParameter = Instance->GetIntParameters()[IntParameterIndex];
+		FCustomizableObjectIntParameterValue& IntParameter = Descriptor.GetIntParameters()[IntParameterIndex];
 		const int32 ParamIndexInObject = Object->FindParameter(IntParameter.ParameterName);
 
 		const FString DefaultValue = Object->GetIntParameterAvailableOption(ParamIndexInObject, 0); // TODO: Define the default option in the editor instead of taking the first available, like it's currently defined for GetProjectorDefaultValue()
@@ -137,65 +149,64 @@ void FMultilayerProjector::CreateLayer(const int32 Index) const
 	
 	// Opacity Range.
 	{
-		const int32 FloatParameterIndex = Instance->FindFloatParameterNameIndex(ParamName.ToString() + FMultilayerProjector::OPACITY_PARAMETER_POSTFIX);
+		const int32 FloatParameterIndex = Descriptor.FindFloatParameterNameIndex(ParamName.ToString() + FMultilayerProjector::OPACITY_PARAMETER_POSTFIX);
 
-		FCustomizableObjectFloatParameterValue& FloatParameter = Instance->GetFloatParameters()[FloatParameterIndex];
+		FCustomizableObjectFloatParameterValue& FloatParameter = Descriptor.GetFloatParameters()[FloatParameterIndex];
 		FloatParameter.ParameterRangeValues.Insert(0.5, Index); // TODO: Define the default float in the editor instead of [0.5f], like it's currently defined for GetProjectorDefaultValue()
 	}
 }
 
 
-void FMultilayerProjector::RemoveLayerAt(const int32 Index) const
+void FMultilayerProjector::RemoveLayerAt(FCustomizableObjectInstanceDescriptor& Descriptor, const int32 Index) const
 {
-	check(Instance); // Multilayer Projector created without using the Instance factory.
-	checkCode(CheckInstanceParameters());
-	check(Index >= 0 && Index < NumLayers()); // Layer out of range.
+	checkCode(CheckDescriptorParameters(Descriptor));
+	check(Index >= 0 && Index < NumLayers(Descriptor)); // Layer out of range.
 	
 	// Num Layers.
 	{
 		const FString NumLayersParamName = ParamName.ToString() + NUM_LAYERS_PARAMETER_POSTFIX;
-        const int32 FloatParameterIndex = Instance->FindFloatParameterNameIndex(NumLayersParamName);
+        const int32 FloatParameterIndex = Descriptor.FindFloatParameterNameIndex(NumLayersParamName);
 		
-    	Instance->GetFloatParameters()[FloatParameterIndex].ParameterValue -= 1;
+    	Descriptor.GetFloatParameters()[FloatParameterIndex].ParameterValue -= 1;
     }
     
 	// Projector Range.
 	{
-		const int32 ProjectorParameterIndex = Instance->FindProjectorParameterNameIndex(ParamName.ToString());
+		const int32 ProjectorParameterIndex = Descriptor.FindProjectorParameterNameIndex(ParamName.ToString());
 		
-		FCustomizableObjectProjectorParameterValue& ProjectorParameter = Instance->GetProjectorParameters()[ProjectorParameterIndex];
+		FCustomizableObjectProjectorParameterValue& ProjectorParameter = Descriptor.GetProjectorParameters()[ProjectorParameterIndex];
 		ProjectorParameter.RangeValues.RemoveAt(Index);
 	}
 	
 	// Selected Image Range.
 	{
-		const int32 IntParameterIndex = Instance->FindIntParameterNameIndex(ParamName.ToString() + IMAGE_PARAMETER_POSTFIX);
+		const int32 IntParameterIndex = Descriptor.FindIntParameterNameIndex(ParamName.ToString() + IMAGE_PARAMETER_POSTFIX);
 		
-		FCustomizableObjectIntParameterValue& IntParameter = Instance->GetIntParameters()[IntParameterIndex];
+		FCustomizableObjectIntParameterValue& IntParameter = Descriptor.GetIntParameters()[IntParameterIndex];
 		IntParameter.ParameterRangeValueNames.RemoveAt(Index);
 	}
 	
 	// Opacity Range.
 	{
-		const int32 FloatParameterIndex = Instance->FindFloatParameterNameIndex(ParamName.ToString() + OPACITY_PARAMETER_POSTFIX);
+		const int32 FloatParameterIndex = Descriptor.FindFloatParameterNameIndex(ParamName.ToString() + OPACITY_PARAMETER_POSTFIX);
 		
-		FCustomizableObjectFloatParameterValue& FloatParameter = Instance->GetFloatParameters()[FloatParameterIndex];
+		FCustomizableObjectFloatParameterValue& FloatParameter = Descriptor.GetFloatParameters()[FloatParameterIndex];
 		FloatParameter.ParameterRangeValues.RemoveAt(Index);
 	}
 }
 
 
-FMultilayerProjectorLayer FMultilayerProjector::GetLayer(const int32 Index) const
+FMultilayerProjectorLayer FMultilayerProjector::GetLayer(const FCustomizableObjectInstanceDescriptor& Descriptor, const int32 Index) const
 {
 	FMultilayerProjectorLayer MultilayerProjectorLayer;
-	MultilayerProjectorLayer.Read(*this, Index);
+	MultilayerProjectorLayer.Read(Descriptor, *this, Index);
 	return MultilayerProjectorLayer;
 }
 
 
-void FMultilayerProjector::UpdateLayer(const int32 Index, const FMultilayerProjectorLayer& Layer) const
+void FMultilayerProjector::UpdateLayer(FCustomizableObjectInstanceDescriptor& Descriptor, const int32 Index, const FMultilayerProjectorLayer& Layer) const
 {
-	Layer.Write(*this, Index);
+	Layer.Write(Descriptor, *this, Index);
 }
 
 
@@ -207,20 +218,20 @@ TArray<FName> FMultilayerProjector::GetVirtualLayers() const
 }
 
 
-void FMultilayerProjector::CreateVirtualLayer(const FName& Id)
+void FMultilayerProjector::CreateVirtualLayer(FCustomizableObjectInstanceDescriptor& Descriptor, const FName& Id)
 {
 	if (!VirtualLayersMapping.Contains(Id))
 	{
-		const int32 Index = NumLayers();
+		const int32 Index = NumLayers(Descriptor);
 		
-		CreateLayer(Index);
+		CreateLayer(Descriptor, Index);
 		VirtualLayersMapping.Add(Id, Index);
 		VirtualLayersOrder.Add(Id, NEW_VIRTUAL_LAYER_ORDER);
 	}
 }
 
 
-FMultilayerProjectorVirtualLayer FMultilayerProjector::FindOrCreateVirtualLayer(const FName& Id)
+FMultilayerProjectorVirtualLayer FMultilayerProjector::FindOrCreateVirtualLayer(FCustomizableObjectInstanceDescriptor& Descriptor, const FName& Id)
 {
 	FMultilayerProjectorLayer Layer;
 	bool bEnabled;
@@ -235,7 +246,7 @@ FMultilayerProjectorVirtualLayer FMultilayerProjector::FindOrCreateVirtualLayer(
 		}
 		else
 		{
-			Layer = GetLayer(*Index);
+			Layer = GetLayer(Descriptor, *Index);
 			bEnabled = true;
 		}
 
@@ -243,14 +254,14 @@ FMultilayerProjectorVirtualLayer FMultilayerProjector::FindOrCreateVirtualLayer(
 	}
 	else
 	{
-		const int32 NewIndex = NumLayers();
+		const int32 NewIndex = NumLayers(Descriptor);
 		constexpr int32 NewOrder = NEW_VIRTUAL_LAYER_ORDER;
 		
-		CreateLayer(NewIndex);
+		CreateLayer(Descriptor, NewIndex);
 		VirtualLayersMapping.Add(Id, NewIndex);
 		VirtualLayersOrder.Add(Id, NewOrder);
 
-		Layer = GetLayer(NewIndex);
+		Layer = GetLayer(Descriptor, NewIndex);
 		bEnabled = true;
 		Order = NewOrder;
 	}
@@ -259,7 +270,7 @@ FMultilayerProjectorVirtualLayer FMultilayerProjector::FindOrCreateVirtualLayer(
 }
 
 
-void FMultilayerProjector::RemoveVirtualLayer(const FName& Id)
+void FMultilayerProjector::RemoveVirtualLayer(FCustomizableObjectInstanceDescriptor& Descriptor, const FName& Id)
 {
 	const int32* Index = VirtualLayersMapping.Find(Id);
 	check(Index); // Virtual Layer not created.
@@ -270,7 +281,7 @@ void FMultilayerProjector::RemoveVirtualLayer(const FName& Id)
 	}
 	else
 	{
-		RemoveLayerAt(*Index);
+		RemoveLayerAt(Descriptor, *Index);
 		
 		for (TMap<FName, int32>::TIterator It =  VirtualLayersMapping.CreateIterator(); It; ++It)
 		{
@@ -289,12 +300,12 @@ void FMultilayerProjector::RemoveVirtualLayer(const FName& Id)
 }
 
 
-FMultilayerProjectorVirtualLayer FMultilayerProjector::GetVirtualLayer(const FName& Id) const
+FMultilayerProjectorVirtualLayer FMultilayerProjector::GetVirtualLayer(const FCustomizableObjectInstanceDescriptor& Descriptor, const FName& Id) const
 {
 	const int32* Index = VirtualLayersMapping.Find(Id);
 	check(Index); // Virtual Layer not created.
 
-	const FMultilayerProjectorLayer Layer = GetLayer(*Index);
+	const FMultilayerProjectorLayer Layer = GetLayer(Descriptor, *Index);
 	const bool bEnabled = *Index != VIRTUAL_LAYER_DISABLED;
 	const int32 Order = VirtualLayersOrder[Id];
 	
@@ -302,7 +313,7 @@ FMultilayerProjectorVirtualLayer FMultilayerProjector::GetVirtualLayer(const FNa
 }
 
 
-void FMultilayerProjector::UpdateVirtualLayer(const FName& Id, const FMultilayerProjectorVirtualLayer& Layer)
+void FMultilayerProjector::UpdateVirtualLayer(FCustomizableObjectInstanceDescriptor& Descriptor, const FName& Id, const FMultilayerProjectorVirtualLayer& Layer)
 {
 	const int32* Index = VirtualLayersMapping.Find(Id);
 	check(Index); // Virtual Layer not created.
@@ -323,7 +334,7 @@ void FMultilayerProjector::UpdateVirtualLayer(const FName& Id, const FMultilayer
 			int32 NewIndex = CalculateVirtualLayerIndex(Id, Layer.Order);
 			if (OldIndex != NewIndex) // Move required. Could be optimized by moving only the in-between values.
 			{
-				RemoveLayerAt(OldIndex);
+				RemoveLayerAt(Descriptor, OldIndex);
 				UpdateMappingVirtualLayerDisabled(Id, OldIndex);
 
 				if (OldIndex < NewIndex)
@@ -331,14 +342,14 @@ void FMultilayerProjector::UpdateVirtualLayer(const FName& Id, const FMultilayer
 					NewIndex -= 1;
 				}
 				
-				CreateLayer(NewIndex);
+				CreateLayer(Descriptor, NewIndex);
 				UpdateMappingVirtualLayerEnabled(Id, NewIndex);
 			}
 			
 			*Order = Layer.Order;
 		}
 
-		UpdateLayer(*Index, static_cast<FMultilayerProjectorLayer>(Layer)); // Update enabled layer.
+		UpdateLayer(Descriptor, *Index, static_cast<FMultilayerProjectorLayer>(Layer)); // Update enabled layer.
 	}
 	
 	// Enable or disable virtual layer.
@@ -346,16 +357,16 @@ void FMultilayerProjector::UpdateVirtualLayer(const FName& Id, const FMultilayer
 	{
 		const int32 NewIndex = CalculateVirtualLayerIndex(Id, VirtualLayersOrder[Id]);
 
-		CreateLayer(NewIndex);
+		CreateLayer(Descriptor, NewIndex);
 		UpdateMappingVirtualLayerEnabled(Id, NewIndex);
 
-		UpdateLayer(NewIndex, static_cast<FMultilayerProjectorLayer>(Layer));
+		UpdateLayer(Descriptor, NewIndex, static_cast<FMultilayerProjectorLayer>(Layer));
 		
 		DisableVirtualLayers.Remove(Id);
 	}
 	else if (!Layer.bEnabled && bEnabled)
 	{
-		RemoveLayerAt(*Index);
+		RemoveLayerAt(Descriptor, *Index);
 		UpdateMappingVirtualLayerDisabled(Id, *Index);
 		
 		DisableVirtualLayers.Add(Id, Layer);
@@ -363,9 +374,8 @@ void FMultilayerProjector::UpdateVirtualLayer(const FName& Id, const FMultilayer
 }
 
 
-FMultilayerProjector::FMultilayerProjector(UCustomizableObjectInstance& Instance, const FName& ParamName) :
-	Instance(&Instance),
-	ParamName(ParamName)
+FMultilayerProjector::FMultilayerProjector(const FName& InParamName) :
+	ParamName(InParamName)
 {
 }
 
@@ -428,47 +438,45 @@ void FMultilayerProjector::UpdateMappingVirtualLayerDisabled(const FName& Id, co
 }
 
 
-void FMultilayerProjector::CheckInstanceParameters() const
+void FMultilayerProjector::CheckDescriptorParameters(const FCustomizableObjectInstanceDescriptor& Descriptor) const
 {
-	check(Instance); // Multilayer Projector created without using the Instance factory.
-
 	const FString ParamNameString = ParamName.ToString();
 
 	// Num layers.
 	{
 		const FString NumLayersParamName = ParamNameString + NUM_LAYERS_PARAMETER_POSTFIX;
-		const int32 FloatParameterIndex = Instance->FindFloatParameterNameIndex(NumLayersParamName);
-		check(FloatParameterIndex >= 0); // Instance Parameter does not exist.
+		const int32 FloatParameterIndex = Descriptor.FindFloatParameterNameIndex(NumLayersParamName);
+		check(FloatParameterIndex >= 0); // Descriptor Parameter does not exist.
 	}
     
 	// Projector.
 	{
-		const int32 ProjectorParameterIndex = Instance->FindProjectorParameterNameIndex(ParamNameString);
-		check(ProjectorParameterIndex >= 0) // Instance Parameter does not exist.
+		const int32 ProjectorParameterIndex = Descriptor.FindProjectorParameterNameIndex(ParamNameString);
+		check(ProjectorParameterIndex >= 0) // Descriptor Parameter does not exist.
 	}
 	
 	// Selected Image.
 	{
-		const int32 IntParameterIndex = Instance->FindIntParameterNameIndex(ParamNameString + IMAGE_PARAMETER_POSTFIX);
-		check(IntParameterIndex >= 0) // Instance Parameter does not exist.
+		const int32 IntParameterIndex = Descriptor.FindIntParameterNameIndex(ParamNameString + IMAGE_PARAMETER_POSTFIX);
+		check(IntParameterIndex >= 0) // Descriptor Parameter does not exist.
 	}
 	
 	// Opacity.
 	{
-		const int32 FloatParameterIndex = Instance->FindFloatParameterNameIndex(ParamNameString + OPACITY_PARAMETER_POSTFIX);
-		check(FloatParameterIndex >= 0) // Instance Parameter does not exist.
+		const int32 FloatParameterIndex = Descriptor.FindFloatParameterNameIndex(ParamNameString + OPACITY_PARAMETER_POSTFIX);
+		check(FloatParameterIndex >= 0) // Descriptor Parameter does not exist.
 	}
 }
 
 
-bool FMultilayerProjector::AreInstanceParametersValid(const UCustomizableObjectInstance& Instance, const FName& ParamName)
+bool FMultilayerProjector::AreDescriptorParametersValid(const FCustomizableObjectInstanceDescriptor& Descriptor, const FName& ParamName)
 {
 	const FString ParamNameString = ParamName.ToString();
 
 	// Num layers.
 	{
 		const FString NumLayersParamName = ParamNameString + NUM_LAYERS_PARAMETER_POSTFIX;
-		const int32 FloatParameterIndex = Instance.FindFloatParameterNameIndex(NumLayersParamName);
+		const int32 FloatParameterIndex = Descriptor.FindFloatParameterNameIndex(NumLayersParamName);
 		if (FloatParameterIndex < 0)
 		{
 			return false;
@@ -477,7 +485,7 @@ bool FMultilayerProjector::AreInstanceParametersValid(const UCustomizableObjectI
     
 	// Projector.
 	{
-		const int32 ProjectorParameterIndex = Instance.FindProjectorParameterNameIndex(ParamNameString);
+		const int32 ProjectorParameterIndex = Descriptor.FindProjectorParameterNameIndex(ParamNameString);
 		if (ProjectorParameterIndex < 0)
 		{
 			return false;
@@ -486,7 +494,7 @@ bool FMultilayerProjector::AreInstanceParametersValid(const UCustomizableObjectI
 	
 	// Selected Image.
 	{
-		const int32 IntParameterIndex = Instance.FindIntParameterNameIndex(ParamNameString + IMAGE_PARAMETER_POSTFIX);
+		const int32 IntParameterIndex = Descriptor.FindIntParameterNameIndex(ParamNameString + IMAGE_PARAMETER_POSTFIX);
 		if (IntParameterIndex < 0)
 		{
 			return false;
@@ -495,7 +503,7 @@ bool FMultilayerProjector::AreInstanceParametersValid(const UCustomizableObjectI
 	
 	// Opacity.
 	{
-		const int32 FloatParameterIndex = Instance.FindFloatParameterNameIndex(ParamNameString + OPACITY_PARAMETER_POSTFIX);
+		const int32 FloatParameterIndex = Descriptor.FindFloatParameterNameIndex(ParamNameString + OPACITY_PARAMETER_POSTFIX);
 		if (FloatParameterIndex < 0)
 		{
 			return false;
@@ -505,3 +513,28 @@ bool FMultilayerProjector::AreInstanceParametersValid(const UCustomizableObjectI
 	return true;
 }
 
+
+uint32 GetTypeHash(const FMultilayerProjector& Key)
+{
+	uint32 Hash = GetTypeHash(Key.ParamName);
+
+	for (const TTuple<FName, int32>& Pair : Key.VirtualLayersMapping)
+	{
+		Hash = HashCombine(Hash, GetTypeHash(Pair.Key));
+		Hash = HashCombine(Hash, GetTypeHash(Pair.Value));
+	}
+	
+	for (const TTuple<FName, int32>& Pair : Key.VirtualLayersOrder)
+	{
+		Hash = HashCombine(Hash, GetTypeHash(Pair.Key));
+		Hash = HashCombine(Hash, GetTypeHash(Pair.Value));
+	}
+
+	for (const TTuple<FName, FMultilayerProjectorLayer>& Pair : Key.DisableVirtualLayers)
+	{
+		Hash = HashCombine(Hash, GetTypeHash(Pair.Key));
+		Hash = HashCombine(Hash, GetTypeHash(Pair.Value));
+	}
+	
+	return Hash;
+}
\ No newline at end of file
diff --git a/Engine/Plugins/Experimental/Mutable/Source/CustomizableObject/Public/MuCO/CustomizableObjectCustomVersion.h b/Engine/Plugins/Experimental/Mutable/Source/CustomizableObject/Public/MuCO/CustomizableObjectCustomVersion.h
--- a/Engine/Plugins/Experimental/Mutable/Source/CustomizableObject/Public/MuCO/CustomizableObjectCustomVersion.h
+++ b/Engine/Plugins/Experimental/Mutable/Source/CustomizableObject/Public/MuCO/CustomizableObjectCustomVersion.h
@@ -57,6 +57,8 @@ struct CUSTOMIZABLEOBJECT_API FCustomizableObjectCustomVersion
 		CustomizableObjectInstanceDescriptor,
 		
 		DescriptorBuildParameterDecorations,
+
+		DescriptorMultilayerProjectors,
 		
 		// -----<new versions can be added above this line>-------------------------------------------------
 		VersionPlusOne,
diff --git a/Engine/Plugins/Experimental/Mutable/Source/CustomizableObject/Public/MuCO/CustomizableObjectInstanceDescriptor.h b/Engine/Plugins/Experimental/Mutable/Source/CustomizableObject/Public/MuCO/CustomizableObjectInstanceDescriptor.h
--- a/Engine/Plugins/Experimental/Mutable/Source/CustomizableObject/Public/MuCO/CustomizableObjectInstanceDescriptor.h
+++ b/Engine/Plugins/Experimental/Mutable/Source/CustomizableObject/Public/MuCO/CustomizableObjectInstanceDescriptor.h
@@ -8,6 +8,8 @@
 #include "Math/Color.h"
 #include "Math/UnrealMathSSE.h"
 #include "MuCO/CustomizableObjectParameterTypeDefinitions.h"
+#include "MuCO/CustomizableObject.h"
+#include "MuCo/MultilayerProjector.h"
 #include "UObject/ObjectPtr.h"
 
 #include "CustomizableObjectInstanceDescriptor.generated.h"
@@ -17,19 +19,27 @@ class UCustomizableInstancePrivateData;
 class UCustomizableObject;
 class UCustomizableObjectInstance;
 
+namespace mu
+{
+	class Parameters;
+	
+	template<typename Type>
+	class Ptr;
+}
+
 
 /** Set of parameters + state that defines a CustomizableObjectInstance.
  *
  * This object has the same parameters + state interfice as UCustomizableObjectInstance.
- * Be aware that current implementation does not support Multilayer Projector helpers! */
+ * UCustomizableObjectInstance must share the same interface. Any public methods added here should also end up in the Instance. */
 USTRUCT()
-struct FCustomizableObjectInstanceDescriptor
+struct CUSTOMIZABLEOBJECT_API FCustomizableObjectInstanceDescriptor
 {
 	GENERATED_BODY()
 
 	FCustomizableObjectInstanceDescriptor() = default;
 	
-	FCustomizableObjectInstanceDescriptor(UCustomizableObject& Object);
+	explicit FCustomizableObjectInstanceDescriptor(UCustomizableObject& Object);
 
 	FCustomizableObjectInstanceDescriptor(const FCustomizableObjectInstanceDescriptor& Other) = default;
 
@@ -45,6 +55,10 @@ struct FCustomizableObjectInstanceDescriptor
 	
 	void SetBuildParameterDecorations(bool Value);
 	
+	void SetCustomizableObject(UCustomizableObject* InCustomizableObject);
+
+	mu::Ptr<mu::Parameters> ReloadParametersFromObject(); // TODO Split function in 2: 1. Load parameters from model, 2. Get mu::ParametersPtr
+    
 	// ------------------------------------------------------------
 	// Parameters
 	// ------------------------------------------------------------
@@ -227,6 +241,56 @@ struct FCustomizableObjectInstanceDescriptor
 	
 	void SetRandomValues();
 
+	// ------------------------------------------------------------
+	// Multilayer Projectors
+	// ------------------------------------------------------------
+	
+	/** Given Multilayer Projector name, create a new Multilayer Projector Helper (if non-existent). See FMultilayerProjector.
+	 *
+	 * @return ture if successfully created (or was already created).
+	 */
+	bool CreateMultiLayerProjector(const FName& ProjectorParamName);
+	
+	/** Given Multilayer Projector name, remove a Multilayer Projector Helper. See FMultilayerProjector. */
+	void RemoveMultilayerProjector(const FName& ProjectorParamName);
+	
+	// Layers
+
+	/** See FMultilayerProjector::NumLayers. */
+	int32 MultilayerProjectorNumLayers(const FName& ProjectorParamName) const;
+
+	/** See FMultilayerProjector::CreateLayer. */
+	void MultilayerProjectorCreateLayer(const FName& ProjectorParamName, int32 Index);
+
+	/** See FMultilayerProjector::RemoveLayerAt. */
+	void MultilayerProjectorRemoveLayerAt(const FName& ProjectorParamName, int32 Index);
+
+	/** See FMultilayerProjector::GetLayer. */
+	FMultilayerProjectorLayer MultilayerProjectorGetLayer(const FName& ProjectorParamName, int32 Index) const;
+
+	/** See FMultilayerProjector::UpdateLayer. */
+	void MultilayerProjectorUpdateLayer(const FName& ProjectorParamName, int32 Index, const FMultilayerProjectorLayer& Layer);
+
+	// Virtual layers
+
+	/** See FMultilayerProjector::GetVirtualLayers. */
+	TArray<FName> MultilayerProjectorGetVirtualLayers(const FName& ProjectorParamName) const;
+	
+	/** See FMultilayerProjector::VirtualLayer. */
+	void MultilayerProjectorCreateVirtualLayer(const FName& ProjectorParamName, const FName& Id);
+
+	/** See FMultilayerProjector::FindOrCreateVirtualLayer. */
+	FMultilayerProjectorVirtualLayer MultilayerProjectorFindOrCreateVirtualLayer(const FName& ProjectorParamName, const FName& Id);
+
+	/** See FMultilayerProjector::RemoveVirtualLayer. */
+	void MultilayerProjectorRemoveVirtualLayer(const FName& ProjectorParamName, const FName& Id);
+
+	/** See FMultilayerProjector::GetVirtualLayer. */
+	FMultilayerProjectorVirtualLayer MultilayerProjectorGetVirtualLayer(const FName& ProjectorParamName, const FName& Id) const;
+
+	/** See FMultilayerProjector::UpdateVirtualLayer. */
+	void MultilayerProjectorUpdateVirtualLayer(const FName& ProjectorParamName, const FName& Id, const FMultilayerProjectorVirtualLayer& Layer);
+
 private:
 	UPROPERTY()
 	TObjectPtr<UCustomizableObject> CustomizableObject = nullptr;
@@ -257,13 +321,20 @@ private:
 
 	TMap<FString, int32> IntParametersLookupTable;
 	
-	void CreateParametersLookupTable();
+	/** Multilayer Projector helpers. See FMultilayerProjector.*/
+	UPROPERTY()
+	TMap<FName, FMultilayerProjector> MultilayerProjectors;
+	
+	void Init();
 
+	void CreateParametersLookupTable();
+	
 	// Friends
-	friend uint32 GetTypeHash(const FCustomizableObjectInstanceDescriptor& Key);
+	friend CUSTOMIZABLEOBJECT_API uint32 GetTypeHash(const FCustomizableObjectInstanceDescriptor& Key);
 	friend UCustomizableObjectInstance;
 	friend UCustomizableInstancePrivateData;
+	friend FMultilayerProjector;
 };
 
 
-uint32 GetTypeHash(const FCustomizableObjectInstanceDescriptor& Key);
\ No newline at end of file
+CUSTOMIZABLEOBJECT_API uint32 GetTypeHash(const FCustomizableObjectInstanceDescriptor& Key);
\ No newline at end of file
diff --git a/Engine/Plugins/Experimental/Mutable/Source/CustomizableObject/Public/MuCO/MultilayerProjector.h b/Engine/Plugins/Experimental/Mutable/Source/CustomizableObject/Public/MuCO/MultilayerProjector.h
--- a/Engine/Plugins/Experimental/Mutable/Source/CustomizableObject/Public/MuCO/MultilayerProjector.h
+++ b/Engine/Plugins/Experimental/Mutable/Source/CustomizableObject/Public/MuCO/MultilayerProjector.h
@@ -12,6 +12,7 @@
 
 #include "MultilayerProjector.generated.h"
 
+struct FCustomizableObjectInstanceDescriptor;
 class UCustomizableObjectInstance;
 struct FMultilayerProjector;
 
@@ -26,17 +27,19 @@ struct FMultilayerProjectorLayer
 	
 	/** Read the Layer from the Instance Parameters.
 	 *
+	 * @param Descriptor Instance Descriptor.
 	 * @param MultilayerProjector Multilayer Projector helper.
 	 * @param Index Index where to read the Layer in the Instance Parameters.
 	 */
-	void Read(const FMultilayerProjector& MultilayerProjector, int32 Index);
+	void Read(const FCustomizableObjectInstanceDescriptor& Descriptor, const FMultilayerProjector& MultilayerProjector, int32 Index);
 
 	/** Write to Layer to the Instance Parameters.
 	 *
+ 	 * @param Descriptor Instance Descriptor.
 	 * @param MultilayerProjector Multilayer Projector helper.
 	 * @param Index Index where to write the Layer in the Instance Parameters.
 	 */
-	void Write(const FMultilayerProjector& MultilayerProjector, int32 Index) const;
+	void Write(FCustomizableObjectInstanceDescriptor& Descriptor, const FMultilayerProjector& MultilayerProjector, int32 Index) const;
 	
 	/** Layer position. */
 	FVector Position;
@@ -61,6 +64,9 @@ struct FMultilayerProjectorLayer
 };
 
 
+uint32 GetTypeHash(const FMultilayerProjectorLayer& Key);
+
+
 /** Data structure representing a Multilayer Projector Virtual Layer. */
 USTRUCT(BlueprintType)
 struct FMultilayerProjectorVirtualLayer : public FMultilayerProjectorLayer
@@ -93,8 +99,10 @@ struct CUSTOMIZABLEOBJECT_API FMultilayerProjector
 {
 	GENERATED_BODY()
 
-	friend UCustomizableObjectInstance;
+	friend FCustomizableObjectInstanceDescriptor;
 	friend FMultilayerProjectorLayer;
+	friend uint32 GetTypeHash(const FMultilayerProjector& Key);
+
 
 	// Parameters encoding
 	static const FString NUM_LAYERS_PARAMETER_POSTFIX;
@@ -105,24 +113,24 @@ struct CUSTOMIZABLEOBJECT_API FMultilayerProjector
 	// Constructors
 	FMultilayerProjector() = default;
 
-	FMultilayerProjector(UCustomizableObjectInstance& Instance, const FName& ParamName);
+	explicit FMultilayerProjector(const FName& InParamName);
 
 	// Layers
 	
 	/** Returns the number of layers. */
-	int32 NumLayers() const;
+	int32 NumLayers(const FCustomizableObjectInstanceDescriptor& Descriptor) const;
 
 	/** Insert the Layer at the given index moving all the following layers one position. */
-	void CreateLayer(int32 Index) const;
+	void CreateLayer(FCustomizableObjectInstanceDescriptor& Descriptor, int32 Index) const;
 
 	/** Remove the Layer at the given index moving all the following layers one position. */
-	void RemoveLayerAt(int32 Index) const;
+	void RemoveLayerAt(FCustomizableObjectInstanceDescriptor& Descriptor, int32 Index) const;
 
 	/** Get the properties of the Layer at the given index. */
-	FMultilayerProjectorLayer GetLayer(int32 Index) const;
+	FMultilayerProjectorLayer GetLayer(const FCustomizableObjectInstanceDescriptor& Descriptor, int32 Index) const;
 
 	/** Update the Layer properties at the given index. */
-	void UpdateLayer(int32 Index, const FMultilayerProjectorLayer& Layer) const;
+	void UpdateLayer(FCustomizableObjectInstanceDescriptor& Descriptor, int32 Index, const FMultilayerProjectorLayer& Layer) const;
 
 	// Virtual layers
 
@@ -130,28 +138,24 @@ struct CUSTOMIZABLEOBJECT_API FMultilayerProjector
 	TArray<FName> GetVirtualLayers() const;
 	
 	/** Given a new identifier, create a new Virtual Layer (if non-existent).*/
-	void CreateVirtualLayer(const FName& Id);
+	void CreateVirtualLayer(FCustomizableObjectInstanceDescriptor& Descriptor, const FName& Id);
 
 	/** Given a new identifier, find or create a new Virtual Layer.*/
-	FMultilayerProjectorVirtualLayer FindOrCreateVirtualLayer(const FName& Id);
+	FMultilayerProjectorVirtualLayer FindOrCreateVirtualLayer(FCustomizableObjectInstanceDescriptor& Descriptor, const FName& Id);
 
 	/** Given its identifier, remove the Virtual Layer. */
-	void RemoveVirtualLayer(const FName& Id);
+	void RemoveVirtualLayer(FCustomizableObjectInstanceDescriptor& Descriptor, const FName& Id);
 
 	/** Given its identifier, get the Virtual Layer properties. */
-	FMultilayerProjectorVirtualLayer GetVirtualLayer(const FName& Id) const;
+	FMultilayerProjectorVirtualLayer GetVirtualLayer(const FCustomizableObjectInstanceDescriptor& Descriptor, const FName& Id) const;
 
 	/** Given its identifier, update a Virtual Layer properties. */
-	void UpdateVirtualLayer(const FName& Id, const FMultilayerProjectorVirtualLayer& Layer);
+	void UpdateVirtualLayer(FCustomizableObjectInstanceDescriptor& Descriptor, const FName& Id, const FMultilayerProjectorVirtualLayer& Layer);
 
 private:
-	/** Always valid. */
-	UPROPERTY()
-	TObjectPtr<UCustomizableObjectInstance> Instance = nullptr;
-
 	/** Multilayer Projector Parameter name. */
 	UPROPERTY()
-	FName ParamName;	
+	FName ParamName;
 	
 	/** Maps a Virtual Layer to a Layer.
 	 *
@@ -196,10 +200,13 @@ private:
 	 */
 	void UpdateMappingVirtualLayerDisabled(const FName& Id, int32 Index);
 
-	/** Rise an assert if the instance does not contain the necessary Instance Parameters. */
-	void CheckInstanceParameters() const;
+	/** Rise an assert if the Instance Descriptor does not contain the necessary Parameters. */
+	void CheckDescriptorParameters(const FCustomizableObjectInstanceDescriptor& Descriptor) const;
 	
 	/** Return false if the Instance does not contain the necessary Instance Parameters. */
-	static bool AreInstanceParametersValid(const UCustomizableObjectInstance& Instance, const FName& ParamName);
+	static bool AreDescriptorParametersValid(const FCustomizableObjectInstanceDescriptor& Descriptor, const FName& ParamName);
 };
 
+
+uint32 GetTypeHash(const FMultilayerProjector& Key);
+
