[mutable] Fix include path errors.

diff --git a/Engine/Plugins/Experimental/Mutable/Source/CustomizableObject/Public/MuCO/CustomizableObjectInstanceDescriptor.h b/Engine/Plugins/Experimental/Mutable/Source/CustomizableObject/Public/MuCO/CustomizableObjectInstanceDescriptor.h
--- a/Engine/Plugins/Experimental/Mutable/Source/CustomizableObject/Public/MuCO/CustomizableObjectInstanceDescriptor.h
+++ b/Engine/Plugins/Experimental/Mutable/Source/CustomizableObject/Public/MuCO/CustomizableObjectInstanceDescriptor.h
@@ -9,7 +9,7 @@
 #include "Math/UnrealMathSSE.h"
 #include "MuCO/CustomizableObjectParameterTypeDefinitions.h"
 #include "MuCO/CustomizableObject.h"
-#include "MuCo/MultilayerProjector.h"
+#include "MuCO/MultilayerProjector.h"
 #include "UObject/ObjectPtr.h"
 
 #include "CustomizableObjectInstanceDescriptor.generated.h"
