#rnx

diff --git a/Engine/Plugins/Experimental/Mutable/Source/CustomizableObjectEditor/Private/MuCOE/CustomizableObjectBakeHelpers.cpp b/Engine/Plugins/Experimental/Mutable/Source/CustomizableObjectEditor/Private/MuCOE/CustomizableObjectBakeHelpers.cpp
--- a/Engine/Plugins/Experimental/Mutable/Source/CustomizableObjectEditor/Private/MuCOE/CustomizableObjectBakeHelpers.cpp
+++ b/Engine/Plugins/Experimental/Mutable/Source/CustomizableObjectEditor/Private/MuCOE/CustomizableObjectBakeHelpers.cpp
@@ -476,12 +476,13 @@ UObject* BakeHelper_DuplicateAsset(UObject* Object, const FString& ObjName, cons
 
 UTexture2D* BakeHelper_CreateAssetTexture(UTexture2D* SrcTex, const FString& TexObjName, const FString& TexPkgName, const UTexture* OrgTex, bool ResetDuplicatedFlags, TMap<UObject*, UObject*>& ReplacementMap, bool OverwritePackage)
 {
-	const bool IsMutableTexutre = !SrcTex->Source.IsValid();
-	if (IsMutableTexutre)
+	const bool bIsMutableTexutre = !SrcTex->Source.IsValid();
+	if (bIsMutableTexutre)
 	{
-		int sx = SrcTex->GetPlatformData()->SizeX;
-		int sy = SrcTex->GetPlatformData()->SizeY;
-		ETextureSourceFormat PixelFormat = SrcTex->GetPlatformData()->PixelFormat == PF_BC4 ? TSF_G8 : TSF_BGRA8;
+		int32 sx = SrcTex->GetPlatformData()->SizeX;
+		int32 sy = SrcTex->GetPlatformData()->SizeY;
+		EPixelFormat SrcPixelFormat = SrcTex->GetPlatformData()->PixelFormat;
+		ETextureSourceFormat PixelFormat = SrcPixelFormat == PF_BC4 || SrcPixelFormat == PF_G8 ? TSF_G8 : TSF_BGRA8;
 
 		// Begin duplicate texture
 		FString FinalObjName = TexObjName;
@@ -500,7 +501,7 @@ UTexture2D* BakeHelper_CreateAssetTexture(UTexture2D* SrcTex, const FString& Tex
 
 		ReplacementMap.Add(SrcTex, DupTex);
 		
-		// The garbage collector is called in the middle of the bake process, and this can destroy this temporary objects. 
+		// The garbage collector is called in the middle of the bake process, and this can destroy this temporary object. 
 		// We add them to the garbage root to prevent this. This will avoid them being unloaded while the editor is running, but this
 		// action is not used often.
 		DupTex->AddToRoot();
@@ -542,7 +543,11 @@ UTexture2D* BakeHelper_CreateAssetTexture(UTexture2D* SrcTex, const FString& Tex
 		switch (SrcTex->GetPlatformData()->PixelFormat)
 		{
 		case PF_R8G8B8A8:
-			FMemory::Memcpy(dst, src, sx*sy * 4);
+			FMemory::Memcpy(dst, src, sx * sy * 4);
+			break;
+
+		case PF_G8:
+			FMemory::Memcpy(dst, src, sx * sy * 1);
 			break;
 
 		case PF_DXT1:
@@ -608,6 +613,17 @@ UTexture2D* BakeHelper_CreateAssetTexture(UTexture2D* SrcTex, const FString& Tex
 			}
 		}
 
+		if (PixelFormat == TSF_G8 || PixelFormat == TSF_G16)
+		{
+			// If compression settings are not set to TC_Grayscale the texture will get a DXT format
+			// instead of G8 or G16.
+			FTextureFormatSettings Settings;
+			Settings.CompressionSettings = TC_Grayscale;
+			DupTex->SetLayerFormatSettings(0, Settings);
+
+			DupTex->CompressionSettings = TC_Grayscale;
+		}
+
 		SrcTex->GetPlatformData()->Mips[0].BulkData.Unlock();
 
 		DupTex->Source.UnlockMip(0);
