[mutable] Fix for search casting exception when looking for words on the COEditor.
- We were trying to cast a UGraphNode to a CustomizableObjectNode with won't work for UE standard nodes (like the comment one). By explicitly only logging the nodes that are children of UCustomizableObjectNode we avoid that check fail.
- The editor crash was being produced not by the casting error itself but due to an unknown error during the CastLogError() operation over Casts.cpp triggered by that casting error.
- In terms of performance it should be better since now we only perform one casting operation per node being searched.

diff --git a/Engine/Plugins/Experimental/Mutable/Source/CustomizableObjectEditor/Private/MuCOE/CustomizableObjectEditor.cpp b/Engine/Plugins/Experimental/Mutable/Source/CustomizableObjectEditor/Private/MuCOE/CustomizableObjectEditor.cpp
--- a/Engine/Plugins/Experimental/Mutable/Source/CustomizableObjectEditor/Private/MuCOE/CustomizableObjectEditor.cpp
+++ b/Engine/Plugins/Experimental/Mutable/Source/CustomizableObjectEditor/Private/MuCOE/CustomizableObjectEditor.cpp
@@ -2353,6 +2353,13 @@ void FCustomizableObjectEditor::OnEnterText(const FText & NewText, ETextCommit::
 		{
 			UObject* Node = Graph->Nodes[i];
 
+			// Ensure we are working with a Customizable object Node
+			UCustomizableObjectNode* CustomizableObjectNode = Cast<UCustomizableObjectNode>(Node);
+			if (!CustomizableObjectNode)
+			{
+				continue;
+			}
+
 			//Find Node Coincidence 
 			FString NodeName = Graph->Nodes[i]->GetNodeTitle(ENodeTitleType::ListView).ToString();
 			int32 cPos = NodeName.Find("\n");
@@ -2365,7 +2372,7 @@ void FCustomizableObjectEditor::OnEnterText(const FText & NewText, ETextCommit::
 
 			if (NodeName.Contains(NewText.ToString(), ESearchCase::IgnoreCase))
 			{
-				LogSearchResult(Graph->Nodes[i], "Node", found, NodeName);
+				LogSearchResult(CustomizableObjectNode, "Node", found, NodeName);
 				found = true;
 			}
 
@@ -2381,7 +2388,7 @@ void FCustomizableObjectEditor::OnEnterText(const FText & NewText, ETextCommit::
 				//Find Variable name
 				if (ResultName.Contains(VariableName))
 				{
-					LogSearchResult(Graph->Nodes[i], "Variable", found, ResultName);
+					LogSearchResult(CustomizableObjectNode, "Variable", found, ResultName);
 					found = true;
 				}
 
@@ -2392,7 +2399,7 @@ void FCustomizableObjectEditor::OnEnterText(const FText & NewText, ETextCommit::
 
 					if (StringResult->Contains(NewText.ToString()))
 					{
-						LogSearchResult(Graph->Nodes[i], "Value", found, *StringResult);
+						LogSearchResult(CustomizableObjectNode, "Value", found, *StringResult);
 						found = true;
 					}
 				}
@@ -2440,7 +2447,7 @@ void FCustomizableObjectEditor::OnEnterText(const FText & NewText, ETextCommit::
 
 					if (StringResult.Contains(NewText.ToString()))
 					{
-						LogSearchResult(Graph->Nodes[i], "Value", found, StringResult);
+						LogSearchResult(CustomizableObjectNode, "Value", found, StringResult);
 						found = true;
 					}
 				}
@@ -2464,7 +2471,7 @@ void FCustomizableObjectEditor::OnEnterText(const FText & NewText, ETextCommit::
 						{
 							if (OptionName.Contains(VariableName))
 							{
-								LogSearchResult(Graph->Nodes[i], "Variable", found, OptionName);
+								LogSearchResult(CustomizableObjectNode, "Variable", found, OptionName);
 								found = true;
 							}
 
@@ -2472,7 +2479,7 @@ void FCustomizableObjectEditor::OnEnterText(const FText & NewText, ETextCommit::
 
 							if (OptionNameContent.Contains(NewText.ToString()))
 							{
-								LogSearchResult(Graph->Nodes[i], "Value", found, Array[a].OptionName);
+								LogSearchResult(CustomizableObjectNode, "Value", found, Array[a].OptionName);
 								found = true;
 							}
 						}
@@ -2489,7 +2496,7 @@ void FCustomizableObjectEditor::OnEnterText(const FText & NewText, ETextCommit::
 						{
 							if (PoseName.Contains(VariableName))
 							{
-								LogSearchResult(Graph->Nodes[i], "Variable", found, PoseName);
+								LogSearchResult(CustomizableObjectNode, "Variable", found, PoseName);
 								found = true;
 							}
 
@@ -2497,7 +2504,7 @@ void FCustomizableObjectEditor::OnEnterText(const FText & NewText, ETextCommit::
 
 							if (OptionNameContent.Contains(NewText.ToString()))
 							{
-								LogSearchResult(Graph->Nodes[i], "Value", found, Array[a].PoseName);
+								LogSearchResult(CustomizableObjectNode, "Value", found, Array[a].PoseName);
 								found = true;
 							}
 						}
@@ -2513,7 +2520,7 @@ void FCustomizableObjectEditor::OnEnterText(const FText & NewText, ETextCommit::
 						{
 							if (Name.Contains(VariableName))
 							{
-								LogSearchResult(Graph->Nodes[i], "Variable", found, Name);
+								LogSearchResult(CustomizableObjectNode, "Variable", found, Name);
 								found = true;
 							}
 
@@ -2521,7 +2528,7 @@ void FCustomizableObjectEditor::OnEnterText(const FText & NewText, ETextCommit::
 
 							if (OptionNameContent.Contains(NewText.ToString()))
 							{
-								LogSearchResult(Graph->Nodes[i], "Value", found, Array[a].Name);
+								LogSearchResult(CustomizableObjectNode, "Value", found, Array[a].Name);
 								found = true;
 							}
 						}
@@ -2537,7 +2544,7 @@ void FCustomizableObjectEditor::OnEnterText(const FText & NewText, ETextCommit::
 						{
 							if (Name.Contains(VariableName))
 							{
-								LogSearchResult(Graph->Nodes[i], "Variable", found, Name);
+								LogSearchResult(CustomizableObjectNode, "Variable", found, Name);
 								found = true;
 							}
 
@@ -2545,7 +2552,7 @@ void FCustomizableObjectEditor::OnEnterText(const FText & NewText, ETextCommit::
 
 							if (OptionNameContent.Contains(NewText.ToString()))
 							{
-								LogSearchResult(Graph->Nodes[i], "Value", found, Array[a].Name);
+								LogSearchResult(CustomizableObjectNode, "Value", found, Array[a].Name);
 								found = true;
 							}
 						}
@@ -2561,7 +2568,7 @@ void FCustomizableObjectEditor::OnEnterText(const FText & NewText, ETextCommit::
 						{
 							if (Name.Contains(VariableName))
 							{
-								LogSearchResult(Graph->Nodes[i], "Variable", found, Name);
+								LogSearchResult(CustomizableObjectNode, "Variable", found, Name);
 								found = true;
 							}
 
@@ -2569,7 +2576,7 @@ void FCustomizableObjectEditor::OnEnterText(const FText & NewText, ETextCommit::
 
 							if (OptionNameContent.Contains(NewText.ToString()))
 							{
-								LogSearchResult(Graph->Nodes[i], "Value", found, Array[a].Name);
+								LogSearchResult(CustomizableObjectNode, "Value", found, Array[a].Name);
 								found = true;
 							}
 						}
@@ -2585,7 +2592,7 @@ void FCustomizableObjectEditor::OnEnterText(const FText & NewText, ETextCommit::
 						{
 							if (Name.Contains(VariableName))
 							{
-								LogSearchResult(Graph->Nodes[i], "Variable", found, Name);
+								LogSearchResult(CustomizableObjectNode, "Variable", found, Name);
 								found = true;
 							}
 
@@ -2593,7 +2600,7 @@ void FCustomizableObjectEditor::OnEnterText(const FText & NewText, ETextCommit::
 
 							if (OptionNameContent.Contains(NewText.ToString()))
 							{
-								LogSearchResult(Graph->Nodes[i], "Value", found, Array[a].Name);
+								LogSearchResult(CustomizableObjectNode, "Value", found, Array[a].Name);
 								found = true;
 							}
 						}
@@ -2609,7 +2616,7 @@ void FCustomizableObjectEditor::OnEnterText(const FText & NewText, ETextCommit::
 						{
 							if (Name.Contains(VariableName))
 							{
-								LogSearchResult(Graph->Nodes[i], "Variable", found, Name);
+								LogSearchResult(CustomizableObjectNode, "Variable", found, Name);
 								found = true;
 							}
 
@@ -2617,7 +2624,7 @@ void FCustomizableObjectEditor::OnEnterText(const FText & NewText, ETextCommit::
 
 							if (NameContent.Contains(NewText.ToString()))
 							{
-								LogSearchResult(Graph->Nodes[i], "Value", found, Array[a].Name);
+								LogSearchResult(CustomizableObjectNode, "Value", found, Array[a].Name);
 								found = true;
 							}
 						}
@@ -2633,7 +2640,7 @@ void FCustomizableObjectEditor::OnEnterText(const FText & NewText, ETextCommit::
 						{
 							if (Name.Contains(VariableName))
 							{
-								LogSearchResult(Graph->Nodes[i], "Variable", found, Name);
+								LogSearchResult(CustomizableObjectNode, "Variable", found, Name);
 								found = true;
 							}
 
@@ -2641,7 +2648,7 @@ void FCustomizableObjectEditor::OnEnterText(const FText & NewText, ETextCommit::
 
 							if (NameContent.Contains(NewText.ToString()))
 							{
-								LogSearchResult(Graph->Nodes[i], "Value", found, Array[a].Name);
+								LogSearchResult(CustomizableObjectNode, "Value", found, Array[a].Name);
 								found = true;
 							}
 						}
@@ -2657,7 +2664,7 @@ void FCustomizableObjectEditor::OnEnterText(const FText & NewText, ETextCommit::
 						{
 							if (Tag.Contains(VariableName))
 							{
-								LogSearchResult(Graph->Nodes[i], "Variable", found, Tag);
+								LogSearchResult(CustomizableObjectNode, "Variable", found, Tag);
 								found = true;
 							}
 
@@ -2665,7 +2672,7 @@ void FCustomizableObjectEditor::OnEnterText(const FText & NewText, ETextCommit::
 
 							if (OptionNameContent.Contains(NewText.ToString()))
 							{
-								LogSearchResult(Graph->Nodes[i], "Value", found, Array[a].Tag);
+								LogSearchResult(CustomizableObjectNode, "Value", found, Array[a].Tag);
 								found = true;
 							}
 						}
@@ -2679,7 +2686,7 @@ void FCustomizableObjectEditor::OnEnterText(const FText & NewText, ETextCommit::
 						{
 							if (Array[a].Contains(NewText.ToString()))
 							{
-								LogSearchResult(Graph->Nodes[i], "Value", found, Array[a]);
+								LogSearchResult(CustomizableObjectNode, "Value", found, Array[a]);
 								found = true;
 							}
 						}
@@ -2697,19 +2704,19 @@ void FCustomizableObjectEditor::OnEnterText(const FText & NewText, ETextCommit::
 						{
 							if (Name.Contains(VariableName))
 							{
-								LogSearchResult(Graph->Nodes[i], "Variable", found, Name);
+								LogSearchResult(CustomizableObjectNode, "Variable", found, Name);
 								found = true;
 							}
 
 							if (Array[a].Name.Contains(NewText.ToString()))
 							{
-								LogSearchResult(Graph->Nodes[i], "Value", found, Array[a].Name);
+								LogSearchResult(CustomizableObjectNode, "Value", found, Array[a].Name);
 								found = true;
 							}
 
 							if (RuntimeParameters.Contains(VariableName))
 							{
-								LogSearchResult(Graph->Nodes[i], "Variable", found, Name);
+								LogSearchResult(CustomizableObjectNode, "Variable", found, Name);
 								found = true;
 							}
 
@@ -2717,7 +2724,7 @@ void FCustomizableObjectEditor::OnEnterText(const FText & NewText, ETextCommit::
 							{
 								if (Array[a].RuntimeParameters[s].Contains(NewText.ToString()))
 								{
-									LogSearchResult(Graph->Nodes[i], "Value", found, Array[a].RuntimeParameters[s]);
+									LogSearchResult(CustomizableObjectNode, "Value", found, Array[a].RuntimeParameters[s]);
 									found = true;
 								}
 							}
@@ -2734,7 +2741,7 @@ void FCustomizableObjectEditor::OnEnterText(const FText & NewText, ETextCommit::
 						{
 							if (LODs.Contains(VariableName))
 							{
-								LogSearchResult(Graph->Nodes[i], "Variable", found, LODs);
+								LogSearchResult(CustomizableObjectNode, "Variable", found, LODs);
 								found = true;
 							}
 
@@ -2744,13 +2751,13 @@ void FCustomizableObjectEditor::OnEnterText(const FText & NewText, ETextCommit::
 
 								if (Name.Contains(VariableName))
 								{
-									LogSearchResult(Graph->Nodes[i], "Variable", found, Name);
+									LogSearchResult(CustomizableObjectNode, "Variable", found, Name);
 									found = true;
 								}
 
 								if (Array[a].Materials[s].Name.Contains(NewText.ToString()))
 								{
-									LogSearchResult(Graph->Nodes[i], "Value", found, Array[a].Materials[s].Name);
+									LogSearchResult(CustomizableObjectNode, "Value", found, Array[a].Materials[s].Name);
 									found = true;
 								}
 							}
@@ -2767,7 +2774,7 @@ void FCustomizableObjectEditor::OnEnterText(const FText & NewText, ETextCommit::
 						{
 							if (LODs.Contains(VariableName))
 							{
-								LogSearchResult(Graph->Nodes[i], "Variable", found, LODs);
+								LogSearchResult(CustomizableObjectNode, "Variable", found, LODs);
 								found = true;
 							}
 
@@ -2777,13 +2784,13 @@ void FCustomizableObjectEditor::OnEnterText(const FText & NewText, ETextCommit::
 
 								if (Name.Contains(VariableName))
 								{
-									LogSearchResult(Graph->Nodes[i], "Variable", found, Name);
+									LogSearchResult(CustomizableObjectNode, "Variable", found, Name);
 									found = true;
 								}
 
 								if (Array[a].Materials[s].Name.Contains(NewText.ToString()))
 								{
-									LogSearchResult(Graph->Nodes[i], "Value", found, Array[a].Materials[s].Name);
+									LogSearchResult(CustomizableObjectNode, "Value", found, Array[a].Materials[s].Name);
 									found = true;
 								}
 							}
@@ -2805,9 +2812,9 @@ void FCustomizableObjectEditor::OnEnterText(const FText & NewText, ETextCommit::
 }
 
 
-void FCustomizableObjectEditor::LogSearchResult(UEdGraphNode* Node, FString Type, bool IsFirst, FString Result) const
+void FCustomizableObjectEditor::LogSearchResult(UCustomizableObjectNode* Node, FString Type, bool bIsFirst, FString Result) const
 {
-	if (!IsFirst)
+	if (!bIsFirst)
 	{
 		FCustomizableObjectEditorLogger::CreateLog(LOCTEXT("SearchResults", "Search Results:"))
 		.Notification(false)
@@ -2815,7 +2822,7 @@ void FCustomizableObjectEditor::LogSearchResult(UEdGraphNode* Node, FString Type
 	}
 	
 	FCustomizableObjectEditorLogger::CreateLog(FText::FromString(Type + ": " + Result))
-	.Node(*CastChecked<UCustomizableObjectNode>(Node))
+	.Node(*Node)
 	.BaseObject()
 	.Notification(false)
 	.Log();
diff --git a/Engine/Plugins/Experimental/Mutable/Source/CustomizableObjectEditor/Private/MuCOE/CustomizableObjectEditor.h b/Engine/Plugins/Experimental/Mutable/Source/CustomizableObjectEditor/Private/MuCOE/CustomizableObjectEditor.h
--- a/Engine/Plugins/Experimental/Mutable/Source/CustomizableObjectEditor/Private/MuCOE/CustomizableObjectEditor.h
+++ b/Engine/Plugins/Experimental/Mutable/Source/CustomizableObjectEditor/Private/MuCOE/CustomizableObjectEditor.h
@@ -367,8 +367,13 @@ private:
 	/** Searches a node that contains the inserted word */
 	void OnEnterText(const FText& NewText, ETextCommit::Type TextType);
 
-	/** Logs the sersults of the search */
-	void LogSearchResult(UEdGraphNode* Node, FString Type, bool IsFirst, FString result) const;
+	/** Logs the search results of the search
+	 * @param Node The Customizable Object Node we have found to be related with the searched string.
+	 * @param Type The type of relation with the searched word. It is a node, a value or maybe a variable?
+	 * @param bIsFirst Is this the first time we encountered something during our search?
+	 * @param Result The string containing the search word we are looking for in Node
+	 */
+	void LogSearchResult(UCustomizableObjectNode* Node, FString Type, bool bIsFirst, FString Result) const;
 
 	/** Open the Texture Analyzer tab */
 	void OpenTextureAnalyzerTab();
